[Diff tools] list differing properties
- fixed property paths when calling for nested calls to 'IdenticalHelper'
- added new callbacks to FDetailsDiffControl and SDetailsDiff to allow custom entries and ordering
- added new callback to SDetailsSplitter to control row highlighting

diff --git a/Engine/Plugins/Experimental/InstanceDataObjectFixupTool/Source/Private/InstanceDataObjectFixupTool.cpp b/Engine/Plugins/Experimental/InstanceDataObjectFixupTool/Source/Private/InstanceDataObjectFixupTool.cpp
--- a/Engine/Plugins/Experimental/InstanceDataObjectFixupTool/Source/Private/InstanceDataObjectFixupTool.cpp
+++ b/Engine/Plugins/Experimental/InstanceDataObjectFixupTool/Source/Private/InstanceDataObjectFixupTool.cpp
@@ -24,7 +24,7 @@ public:
 	{}
 
 	// when diffing, use the redirects to match properties so that renames are respected
-	virtual bool AreMatching(const TWeakPtr<FDetailTreeNode>& TreeNodeA, const TWeakPtr<FDetailTreeNode>& TreeNodeB) const override
+	virtual bool AreMatching(const TWeakPtr<FDetailTreeNode>& TreeNodeA, const TWeakPtr<FDetailTreeNode>& TreeNodeB, TArray<FPropertySoftPath>* OutDifferingProperties = nullptr) const override
 	{
 		const TSharedPtr<FDetailTreeNode> PinnedTreeNodeA = TreeNodeA.Pin();
 		const TSharedPtr<FDetailTreeNode> PinnedTreeNodeB = TreeNodeB.Pin();
diff --git a/Engine/Source/Editor/Kismet/Private/DiffControl.cpp b/Engine/Source/Editor/Kismet/Private/DiffControl.cpp
--- a/Engine/Source/Editor/Kismet/Private/DiffControl.cpp
+++ b/Engine/Source/Editor/Kismet/Private/DiffControl.cpp
@@ -1,7 +1,7 @@
 // Copyright Epic Games, Inc. All Rights Reserved.
 
 #include "DiffControl.h"
-
+#include "Algo/Transform.h"
 #include "DiffResults.h"
 #include "GraphDiffControl.h"
 #include "SBlueprintDiff.h"
@@ -351,24 +351,61 @@ void FDetailsDiffControl::GenerateTreeEntriesWithoutComments(
 	for (auto&[Object, DetailsDiff] : DetailsDiffs)
 	{
 		Algo::Transform(DifferingProperties, PropertyAllowList,
-        [&InObject = Object](const FSingleObjectDiffEntry& DiffEntry)
-        {
-        	return DiffEntry.Identifier.ResolvePath(InObject);
-        });
-        
-        DetailsDiff.DetailsWidget()->UpdatePropertyAllowList(PropertyAllowList);
+		[&InObject = Object](const FSingleObjectDiffEntry& DiffEntry)
+		{
+			return DiffEntry.Identifier.ResolvePath(InObject);
+		});
+
+		DetailsDiff.DetailsWidget()->UpdatePropertyAllowList(PropertyAllowList);
 	}
-	
-	for (const FSingleObjectDiffEntry& Difference : DifferingProperties)
+
+	if (GenerateCustomEntriesCallback.IsBound())
 	{
-		TSharedPtr<FBlueprintDifferenceTreeEntry> Entry = MakeShared<FBlueprintDifferenceTreeEntry>(
+		TArray<FSingleObjectDiffEntry> CustomEntries;
+		GenerateCustomEntriesCallback.Execute(CustomEntries);
+		DifferingProperties.Append(CustomEntries);
+	}
+
+	auto GenerateDiffEntry = [&](const FSingleObjectDiffEntry& Difference)
+	{
+		return MakeShared<FBlueprintDifferenceTreeEntry>(
 			FOnDiffEntryFocused::CreateSP(TSharedFromThis<FDetailsDiffControl>::AsShared(), &FDetailsDiffControl::OnSelectDiffEntry, Difference.Identifier),
 			FGenerateDiffEntryWidget::CreateStatic(&GenerateObjectDiffWidget, Difference, RightRevision));
-		Children.Push(Entry);
-		OutRealDifferences.Push(Entry);
-		if (bPopulateOutTreeEntries)
+	};
+
+	if (OrganizeEntriesCallback.IsBound())
+	{
+		auto GenerateCategoryEntry = [&](FText& CategoryName)
+		{
+			return MakeShared<FBlueprintDifferenceTreeEntry>(
+				/*InOnFocus*/nullptr,
+				FGenerateDiffEntryWidget::CreateStatic(&GenerateSimpleDiffWidget, CategoryName));
+		};
+
+		TArray<TSharedPtr<FBlueprintDifferenceTreeEntry>> Entries;
+		OrganizeEntriesCallback.Execute(Entries, DifferingProperties, GenerateDiffEntry, GenerateCategoryEntry);
+
+		for (const TSharedPtr<FBlueprintDifferenceTreeEntry>& Entry : Entries)
 		{
-			OutTreeEntries.Push(Entry);
+			Children.Push(Entry);
+			OutRealDifferences.Push(Entry);
+			if (bPopulateOutTreeEntries)
+			{
+				OutTreeEntries.Push(Entry);
+			}
+		}
+	}
+	else
+	{
+		for (const FSingleObjectDiffEntry& Difference : DifferingProperties)
+		{
+			TSharedPtr<FBlueprintDifferenceTreeEntry> Entry = GenerateDiffEntry(Difference);
+			Children.Push(Entry);
+			OutRealDifferences.Push(Entry);
+			if (bPopulateOutTreeEntries)
+			{
+				OutTreeEntries.Push(Entry);
+			}
 		}
 	}
 }
@@ -413,8 +450,6 @@ TSharedRef<IDetailsView> FDetailsDiffControl::InsertObject(const UObject* Object
 		const TSharedRef<IDetailsView> OtherDetailsView = OtherDetailsDiff.DetailsWidget();
 		const auto ScrollRate = GetLinkedScrollRateAttribute(DetailsView, OtherDetailsView);
 		FDetailsDiff::LinkScrolling(DetailsDiff, OtherDetailsDiff, ScrollRate);
-		
-		const TSharedPtr<FAsyncDetailViewDiff> DifferencesWithRight = MakeShared<FAsyncDetailViewDiff>(OtherDetailsView, DetailsView);
 
 		PropertyTreeDifferences[OtherObject].Left = MakeShared<FAsyncDetailViewDiff>(DetailsView, OtherDetailsView);
 		PropertyTreeDifferences[Object].Right = PropertyTreeDifferences[OtherObject].Left;
diff --git a/Engine/Source/Editor/Kismet/Private/DiffControl.h b/Engine/Source/Editor/Kismet/Private/DiffControl.h
--- a/Engine/Source/Editor/Kismet/Private/DiffControl.h
+++ b/Engine/Source/Editor/Kismet/Private/DiffControl.h
@@ -2,13 +2,15 @@
 #pragma once
 
 #include "SCSDiff.h"
+#if UE_ENABLE_INCLUDE_ORDER_DEPRECATED_IN_5_5
 #include "CoreMinimal.h"
+#include "Algo/Transform.h"
+#endif // UE_ENABLE_INCLUDE_ORDER_DEPRECATED_IN_5_5
 #include "DetailsDiff.h"
 #include "DiffResults.h"
 #include "IAssetTypeActions.h"
 #include "IDetailsView.h"
 #include "ReviewComments.h"
-#include "Algo/Transform.h"
 #include "Widgets/SWidget.h"
 #include "DiffUtils.h"
 
@@ -111,6 +113,10 @@ public:
 
 	// to support comment posting, set this to the tree view that contains the comments.
 	virtual void EnableComments(TWeakPtr<STreeView<TSharedPtr<FBlueprintDifferenceTreeEntry>>> TreeView) override;
+
+	DiffUtils::FOnGenerateCustomDiffEntries GenerateCustomEntriesCallback;
+	DiffUtils::FOnGenerateCustomDiffEntryWidget GenerateCustomEntryWidgetCallback;
+	DiffUtils::FOnOrganizeDiffEntries OrganizeEntriesCallback;
 protected:
 	virtual void OnSelectDiffEntry(FPropertySoftPath PropertyName);
 	TAttribute<TArray<FVector2f>> GetLinkedScrollRateAttribute(const TSharedRef<IDetailsView>& OldDetailsView, const TSharedRef<IDetailsView>& NewDetailsView);
@@ -384,4 +390,4 @@ protected:
 
 	// we only want one reply draft widget in existence so we can use a static to track it
 	static TWeakPtr<FCommentDraftTreeEntry> ReplyDraftEntry;
-};
+};
\ No newline at end of file
diff --git a/Engine/Source/Editor/Kismet/Private/SDetailsDiff.cpp b/Engine/Source/Editor/Kismet/Private/SDetailsDiff.cpp
--- a/Engine/Source/Editor/Kismet/Private/SDetailsDiff.cpp
+++ b/Engine/Source/Editor/Kismet/Private/SDetailsDiff.cpp
@@ -33,9 +33,9 @@ static const FName DetailsMode = FName(TEXT("DetailsMode"));
 BEGIN_SLATE_FUNCTION_BUILD_OPTIMIZATION
 void SDetailsDiff::Construct( const FArguments& InArgs)
 {
-	check(InArgs._AssetOld || InArgs._AssetNew);
-	PanelOld.Object = InArgs._AssetOld;
-	PanelNew.Object = InArgs._AssetNew;
+	check(InArgs._OldAsset || InArgs._NewAsset);
+	PanelOld.Object = InArgs._OldAsset;
+	PanelNew.Object = InArgs._NewAsset;
 	PanelOld.RevisionInfo = InArgs._OldRevision;
 	PanelNew.RevisionInfo = InArgs._NewRevision;
 
@@ -44,6 +44,13 @@ void SDetailsDiff::Construct( const FArguments& InArgs)
 	PanelOld.bShowAssetName = InArgs._ShowAssetNames;
 	PanelNew.bShowAssetName = InArgs._ShowAssetNames;
 
+	OnCustomizeDetailsWidget = InArgs._OnCustomizeDetailsWidget;
+	OnGenerateCustomDiffEntries = InArgs._OnGenerateCustomDiffEntries;
+	OnGenerateCustomDiffEntryWidget = InArgs._OnGenerateCustomDiffEntryWidget;
+	OnOrganizeDiffEntries = InArgs._OnOrganizeDiffEntries;
+	ShouldHighlightRow = InArgs._ShouldHighlightRow;
+	RowHighlightColor = InArgs._RowHighlightColor;
+
 	if (InArgs._ParentWindow.IsValid())
 	{
 		WeakParentWindow = InArgs._ParentWindow;
@@ -262,8 +269,8 @@ TSharedRef<SDetailsDiff> SDetailsDiff::CreateDiffWindow(FText WindowTitle, const
 		.ClientSize(FVector2D(1000.f, 800.f));
 
 	TSharedRef<SDetailsDiff> DetailsDiff = SNew(SDetailsDiff)
-		.AssetOld(OldObject)
-		.AssetNew(NewObject)
+		.OldAsset(OldObject)
+		.NewAsset(NewObject)
 		.OldRevision(OldRevision)
 		.NewRevision(NewRevision)
 		.ShowAssetNames(!bIsSingleAsset)
@@ -394,12 +401,23 @@ SDetailsDiff::FDiffControl SDetailsDiff::GenerateDetailsPanel(const TFunction<co
 {
 	const UObject* OldObject = Redirector ? Redirector(PanelOld.Object) : PanelOld.Object;
 	const UObject* NewObject = Redirector ? Redirector(PanelNew.Object) : PanelNew.Object;
-	
+
 	const TSharedPtr<FDetailsDiffControl> NewDiffControl = MakeShared<FDetailsDiffControl>(OldObject, NewObject, FOnDiffEntryFocused::CreateRaw(this, &SDetailsDiff::SetCurrentMode, DetailsMode), true);
 	NewDiffControl->EnableComments(DifferencesTreeView.ToWeakPtr());
 	NewDiffControl->GenerateTreeEntries(PrimaryDifferencesList, RealDifferences);
-	
-	const TSharedRef<SDetailsSplitter> Splitter = SNew(SDetailsSplitter);
+	NewDiffControl->GenerateCustomEntriesCallback = OnGenerateCustomDiffEntries;
+	NewDiffControl->GenerateCustomEntryWidgetCallback = OnGenerateCustomDiffEntryWidget;
+	NewDiffControl->OrganizeEntriesCallback = OnOrganizeDiffEntries;
+	if (OnCustomizeDetailsWidget.IsBound())
+	{
+		OnCustomizeDetailsWidget.Execute(NewDiffControl->GetDetailsWidget(OldObject));
+		OnCustomizeDetailsWidget.Execute(NewDiffControl->GetDetailsWidget(NewObject));
+	}
+
+	const TSharedRef<SDetailsSplitter> Splitter = SNew(SDetailsSplitter)
+		.ShouldHighlightRow(ShouldHighlightRow)
+		.RowHighlightColor(RowHighlightColor);
+
 	if (OldObject)
 	{
 		Splitter->AddSlot(
diff --git a/Engine/Source/Editor/Kismet/Public/SDetailsDiff.h b/Engine/Source/Editor/Kismet/Public/SDetailsDiff.h
--- a/Engine/Source/Editor/Kismet/Public/SDetailsDiff.h
+++ b/Engine/Source/Editor/Kismet/Public/SDetailsDiff.h
@@ -1,7 +1,10 @@
 // Copyright Epic Games, Inc. All Rights Reserved.
 #pragma once
 
+#include "AsyncDetailViewDiff.h"
+#if UE_ENABLE_INCLUDE_ORDER_DEPRECATED_IN_5_5
 #include "CoreMinimal.h"
+#endif // UE_ENABLE_INCLUDE_ORDER_DEPRECATED_IN_5_5
 #include "DiffUtils.h"
 #include "IAssetTypeActions.h"
 #include "Widgets/DeclarativeSyntaxSupport.h"
@@ -14,6 +17,7 @@ class FTabManager;
 class IDiffControl;
 class FUICommandList;
 class SDetailsSplitter;
+class FDetailsDiffControl;
 enum class EAssetEditorCloseReason : uint8;
 
 /** Panel used to display the details */
@@ -39,13 +43,28 @@ private:
 class  KISMET_API SDetailsDiff: public SCompoundWidget
 {
 public:
-	SLATE_BEGIN_ARGS( SDetailsDiff ){}
-			SLATE_ARGUMENT( const class UObject*, AssetOld )
-			SLATE_ARGUMENT( const class UObject*, AssetNew )
-			SLATE_ARGUMENT( struct FRevisionInfo, OldRevision )
-			SLATE_ARGUMENT( struct FRevisionInfo, NewRevision )
-			SLATE_ARGUMENT( bool, ShowAssetNames )
-			SLATE_ARGUMENT(TSharedPtr<SWindow>, ParentWindow)
+
+	DECLARE_DELEGATE_OneParam(FOnCustomizeDetailsWidget, const TSharedRef<IDetailsView>&)
+	DECLARE_DELEGATE_RetVal_OneParam(FLinearColor, FRowHighlightColor, const TUniquePtr<FAsyncDetailViewDiff::DiffNodeType>&)
+	DECLARE_DELEGATE_RetVal_OneParam(bool, FShouldHighlightRow, const TUniquePtr<FAsyncDetailViewDiff::DiffNodeType>&)
+
+	SLATE_BEGIN_ARGS(SDetailsDiff) {}
+		SLATE_ARGUMENT_DEPRECATED(const class UObject*, AssetOld, 5.5, "Use OldAsset instead")
+		SLATE_ARGUMENT_DEPRECATED(const class UObject*, AssetNew, 5.5, "Use NewAsset instead")
+		SLATE_ARGUMENT(const class UObject*, OldAsset)
+		SLATE_ARGUMENT(const class UObject*, NewAsset)
+		SLATE_ARGUMENT(struct FRevisionInfo, OldRevision)
+		SLATE_ARGUMENT(struct FRevisionInfo, NewRevision)
+		SLATE_ARGUMENT(bool, ShowAssetNames)
+		SLATE_ARGUMENT(TSharedPtr<SWindow>, ParentWindow)
+		SLATE_ARGUMENT(TSharedPtr<IDetailPropertyExtensionHandler>, ExtensionHandler)
+
+		SLATE_EVENT(FOnCustomizeDetailsWidget, OnCustomizeDetailsWidget)
+		SLATE_EVENT(DiffUtils::FOnGenerateCustomDiffEntries, OnGenerateCustomDiffEntries)
+		SLATE_EVENT(DiffUtils::FOnGenerateCustomDiffEntryWidget, OnGenerateCustomDiffEntryWidget)
+		SLATE_EVENT(DiffUtils::FOnOrganizeDiffEntries, OnOrganizeDiffEntries)
+		SLATE_EVENT(FShouldHighlightRow, ShouldHighlightRow)
+		SLATE_EVENT(FRowHighlightColor, RowHighlightColor)
 	SLATE_END_ARGS()
 
 	void Construct(const FArguments& InArgs);
@@ -123,7 +142,7 @@ protected:
 
 	FName CurrentMode;
 
-	/*The two panels used to show the old & new revision*/ 
+	/** The two panels used to show the old & new revision */
 	FDetailsDiffPanel PanelOld;
 	FDetailsDiffPanel PanelNew;
 	
@@ -133,6 +152,13 @@ protected:
 	DECLARE_MULTICAST_DELEGATE(FOnSetOutputObjectEvent)
 	FOnSetOutputObjectEvent OnOutputObjectSetEvent;
 
+	FOnCustomizeDetailsWidget OnCustomizeDetailsWidget;
+	DiffUtils::FOnGenerateCustomDiffEntries OnGenerateCustomDiffEntries;
+	DiffUtils::FOnGenerateCustomDiffEntryWidget OnGenerateCustomDiffEntryWidget;
+	DiffUtils::FOnOrganizeDiffEntries OnOrganizeDiffEntries;
+	FShouldHighlightRow ShouldHighlightRow;
+	FRowHighlightColor RowHighlightColor;
+
 	/** Contents widget that we swap when mode changes (defaults, components, etc) */
 	TSharedPtr<SBox> ModeContents;
 
diff --git a/Engine/Source/Editor/PropertyEditor/Internal/AsyncDetailViewDiff.h b/Engine/Source/Editor/PropertyEditor/Internal/AsyncDetailViewDiff.h
deleted file mode 100644
diff --git a/Engine/Source/Editor/PropertyEditor/Private/AsyncDetailViewDiff.cpp b/Engine/Source/Editor/PropertyEditor/Private/AsyncDetailViewDiff.cpp
--- a/Engine/Source/Editor/PropertyEditor/Private/AsyncDetailViewDiff.cpp
+++ b/Engine/Source/Editor/PropertyEditor/Private/AsyncDetailViewDiff.cpp
@@ -22,7 +22,7 @@ namespace AsyncDetailViewDiffHelpers
 
 
 
-bool TTreeDiffSpecification<TWeakPtr<FDetailTreeNode>>::AreValuesEqual(const TWeakPtr<FDetailTreeNode>& TreeNodeA, const TWeakPtr<FDetailTreeNode>& TreeNodeB) const
+bool TTreeDiffSpecification<TWeakPtr<FDetailTreeNode>>::AreValuesEqual(const TWeakPtr<FDetailTreeNode>& TreeNodeA, const TWeakPtr<FDetailTreeNode>& TreeNodeB, TArray<FPropertySoftPath>* OutDifferingProperties) const
 {
 	const TSharedPtr<FDetailTreeNode> PinnedTreeNodeA = TreeNodeA.Pin();
 	const TSharedPtr<FDetailTreeNode> PinnedTreeNodeB = TreeNodeB.Pin();
@@ -62,10 +62,16 @@ bool TTreeDiffSpecification<TWeakPtr<FDetailTreeNode>>::AreValuesEqual(const TWe
 	{
 		return true;
 	}
+
+	if (OutDifferingProperties)
+	{
+		return DiffUtils::Identical(*OutDifferingProperties, PropertyHandleA, PropertyHandleB, OwningObjectsA, OwningObjectsB);
+	}
+
 	return DiffUtils::Identical(PropertyHandleA, PropertyHandleB, OwningObjectsA, OwningObjectsB);
 }
 
-bool TTreeDiffSpecification<TWeakPtr<FDetailTreeNode>>::AreMatching(const TWeakPtr<FDetailTreeNode>& TreeNodeA, const TWeakPtr<FDetailTreeNode>& TreeNodeB) const
+bool TTreeDiffSpecification<TWeakPtr<FDetailTreeNode>>::AreMatching(const TWeakPtr<FDetailTreeNode>& TreeNodeA, const TWeakPtr<FDetailTreeNode>& TreeNodeB, TArray<FPropertySoftPath>* OutDifferingProperties) const
 {
 	const TSharedPtr<FDetailTreeNode> PinnedTreeNodeA = TreeNodeA.Pin();
 	const TSharedPtr<FDetailTreeNode> PinnedTreeNodeB = TreeNodeB.Pin();
@@ -92,6 +98,12 @@ bool TTreeDiffSpecification<TWeakPtr<FDetailTreeNode>>::AreMatching(const TWeakP
 		{
 			const TArray<TWeakObjectPtr<UObject>> OwningObjectsA = AsyncDetailViewDiffHelpers::GetObjects(PinnedTreeNodeA);
 			const TArray<TWeakObjectPtr<UObject>> OwningObjectsB = AsyncDetailViewDiffHelpers::GetObjects(PinnedTreeNodeB);
+
+			if (OutDifferingProperties)
+			{
+				return DiffUtils::Identical(*OutDifferingProperties, PropertyHandleA, PropertyHandleB, OwningObjectsA, OwningObjectsB);
+			}
+
 			return DiffUtils::Identical(KeyHandleA, KeyHandleB, OwningObjectsA, OwningObjectsB);
 		}
 
@@ -102,6 +114,12 @@ bool TTreeDiffSpecification<TWeakPtr<FDetailTreeNode>>::AreMatching(const TWeakP
 			// match set elements by value
 			const TArray<TWeakObjectPtr<UObject>> OwningObjectsA = AsyncDetailViewDiffHelpers::GetObjects(PinnedTreeNodeA);
 			const TArray<TWeakObjectPtr<UObject>> OwningObjectsB = AsyncDetailViewDiffHelpers::GetObjects(PinnedTreeNodeB);
+
+			if (OutDifferingProperties)
+			{
+				return DiffUtils::Identical(*OutDifferingProperties, PropertyHandleA, PropertyHandleB, OwningObjectsA, OwningObjectsB);
+			}
+
 			return DiffUtils::Identical(PropertyHandleA, PropertyHandleB, OwningObjectsA, OwningObjectsB);
 		}
 		
@@ -178,24 +196,36 @@ void FAsyncDetailViewDiff::GetPropertyDifferences(TArray<FSingleObjectDiffEntry>
 		}
 		
 		EPropertyDiffType::Type PropertyDiffType;
-        switch(Node->DiffResult)
-        {
-        case ETreeDiffResult::MissingFromTree1: 
-            PropertyDiffType = EPropertyDiffType::PropertyAddedToB;
-            break;
-        case ETreeDiffResult::MissingFromTree2: 
-            PropertyDiffType = EPropertyDiffType::PropertyAddedToA;
-            break;
-        case ETreeDiffResult::DifferentValues: 
-            PropertyDiffType = EPropertyDiffType::PropertyValueChanged;
-            break;
-        default:
-        	// only include changes
-            return ETreeTraverseControl::Continue;
-        }
-
-		OutDiffEntries.Add(FSingleObjectDiffEntry(PropertyPath, PropertyDiffType));
-        return ETreeTraverseControl::SkipChildren; // only include top-most properties
+		switch (Node->DiffResult)
+		{
+		case ETreeDiffResult::MissingFromTree1:
+			PropertyDiffType = EPropertyDiffType::PropertyAddedToB;
+			break;
+		case ETreeDiffResult::MissingFromTree2:
+			PropertyDiffType = EPropertyDiffType::PropertyAddedToA;
+			break;
+		case ETreeDiffResult::DifferentValues:
+			PropertyDiffType = EPropertyDiffType::PropertyValueChanged;
+			break;
+		default:
+			// only include changes
+			return ETreeTraverseControl::Continue;
+		}
+
+		const FPropertySoftPath RootPath(PropertyPath);
+		if (Node->DifferingProperties.Num())
+		{
+			for (FPropertySoftPath& DifferingProperty : Node->DifferingProperties)
+			{
+				OutDiffEntries.Add(FSingleObjectDiffEntry(FPropertySoftPath(RootPath, DifferingProperty), PropertyDiffType));
+			}
+		}
+		else
+		{
+			OutDiffEntries.Add(FSingleObjectDiffEntry(RootPath, PropertyDiffType));
+		}
+
+		return ETreeTraverseControl::SkipChildren; // only include top-most properties
 	});
 }
 
diff --git a/Engine/Source/Editor/PropertyEditor/Private/SDetailsSplitter.cpp b/Engine/Source/Editor/PropertyEditor/Private/SDetailsSplitter.cpp
--- a/Engine/Source/Editor/PropertyEditor/Private/SDetailsSplitter.cpp
+++ b/Engine/Source/Editor/PropertyEditor/Private/SDetailsSplitter.cpp
@@ -407,7 +407,16 @@ void SDetailsSplitter::Construct(const FArguments& InArgs)
 	Splitter = SNew(SSplitter).PhysicalSplitterHandleSize(5.f);
 	
 	GetRowHighlightColor = InArgs._RowHighlightColor.IsBound() ? InArgs._RowHighlightColor : FRowHighlightColor::CreateStatic(
-		[](const TUniquePtr<FAsyncDetailViewDiff::DiffNodeType>&) {return FLinearColor(0.f, 1.f, 1.f, .7f);});
+		[](const TUniquePtr<FAsyncDetailViewDiff::DiffNodeType>&)
+		{
+			return FLinearColor(0.f, 1.f, 1.f, .7f);
+		});
+
+	GetShouldHighlightRow = InArgs._ShouldHighlightRow.IsBound() ? InArgs._ShouldHighlightRow : FShouldHighlightRow::CreateStatic(
+		[](const TUniquePtr<FAsyncDetailViewDiff::DiffNodeType>& DiffNode)
+		{
+			return DiffNode->DiffResult != ETreeDiffResult::Identical;
+		});
 	
 	for(const FSlot::FSlotArguments& SlotArgs : InArgs._Slots)
 	{
@@ -549,6 +558,7 @@ int32 SDetailsSplitter::OnPaint(const FPaintArgs& Args, const FGeometry& Allotte
 			
 			Diff->ForEachRow([&](const TUniquePtr<FAsyncDetailViewDiff::DiffNodeType>& DiffNode, int32, int32)->ETreeTraverseControl
 			{
+				const bool bShouldHighlightRow = GetShouldHighlightRow.Execute(DiffNode);
 				const FLinearColor Color = GetRowHighlightColor.Execute(DiffNode);
 				
 				FSlateRect LeftPropertyRect;
@@ -564,7 +574,7 @@ int32 SDetailsSplitter::OnPaint(const FPaintArgs& Args, const FGeometry& Allotte
 					LeftPropertyRect = PrevLeftPropertyRect;
 					LeftPropertyRect.Top = LeftPropertyRect.Bottom;
 				}
-				if (LeftPropertyRect.IsValid() && DiffNode->DiffResult != ETreeDiffResult::Identical)
+				if (LeftPropertyRect.IsValid() && bShouldHighlightRow)
 				{
 					if (!LeftPanel.ShouldIgnoreRow.Execute(DiffNode->ValueA))
 					{
@@ -589,15 +599,15 @@ int32 SDetailsSplitter::OnPaint(const FPaintArgs& Args, const FGeometry& Allotte
 						RightPropertyRect = PrevRightPropertyRect;
 						RightPropertyRect.Top = RightPropertyRect.Bottom;
 					}
-					if (RightPropertyRect.IsValid() && DiffNode->DiffResult != ETreeDiffResult::Identical)
+					if (RightPropertyRect.IsValid() && bShouldHighlightRow)
 					{
 						if (!RightPanel.ShouldIgnoreRow.Execute(DiffNode->ValueB))
 						{
 							RowHighlights.Add(RightPropertyRect, Color);
 						}
 					}
-					
-					if (LeftPropertyRect.IsValid() && RightPropertyRect.IsValid() && DiffNode->DiffResult != ETreeDiffResult::Identical)
+
+					if (LeftPropertyRect.IsValid() && RightPropertyRect.IsValid() && bShouldHighlightRow)
 					{
 						if (!LeftPanel.ShouldIgnoreRow.Execute(DiffNode->ValueA) && !RightPanel.ShouldIgnoreRow.Execute(DiffNode->ValueB))
 						{
diff --git a/Engine/Source/Editor/PropertyEditor/Public/AsyncDetailViewDiff.h b/Engine/Source/Editor/PropertyEditor/Public/AsyncDetailViewDiff.h
new file mode 100644
--- /dev/null
+++ b/Engine/Source/Editor/PropertyEditor/Public/AsyncDetailViewDiff.h
@@ -0,0 +1,74 @@
+// Copyright Epic Games, Inc. All Rights Reserved.
+#pragma once
+#include "AsyncTreeDifferences.h"
+#include "DiffUtils.h"
+#include "IDetailsView.h"
+
+template<>
+class PROPERTYEDITOR_API TTreeDiffSpecification<TWeakPtr<FDetailTreeNode>>
+{
+public:
+	virtual ~TTreeDiffSpecification() = default;
+	/**
+	 * determine whether the values stored in two nodes are equal.
+	 * @param TreeNodeA node from the first user provided tree (guaranteed not to be null)
+	 * @param TreeNodeB node from the second user provided tree (guaranteed not to be null)
+	 * @param OutDifferingProperties optional output list of property paths for the properties that differs between the two nodes
+	 */
+	virtual bool AreValuesEqual(const TWeakPtr<FDetailTreeNode>& TreeNodeA, const TWeakPtr<FDetailTreeNode>& TreeNodeB, TArray<FPropertySoftPath>* OutDifferingProperties = nullptr) const;
+
+	/**
+	 * determine whether two nodes occupy the same space in their trees
+	 * for example if you have a tree key/value pairs, AreMatching should compare the keys while AreValuesEqual
+	 * should compare the values
+	 * @param TreeNodeA node from the first user provided tree (guaranteed not to be null)
+	 * @param TreeNodeB node from the second user provided tree (guaranteed not to be null)
+	 * @param OutDifferingProperties optional output list of property paths for the properties that differs between the two nodes
+	 */
+	virtual bool AreMatching(const TWeakPtr<FDetailTreeNode>& TreeNodeA, const TWeakPtr<FDetailTreeNode>& TreeNodeB, TArray<FPropertySoftPath>* OutDifferingProperties = nullptr) const;
+
+	/**
+	 * retrieves an array of children nodes from the parent node
+	 * @param InParent node from one of the two user provided trees (guaranteed not to be null)
+	 * @param[out] OutChildren to be filled with the children of parent
+	 */
+	virtual void GetChildren(const TWeakPtr<FDetailTreeNode>& InParent, TArray<TWeakPtr<FDetailTreeNode>>& OutChildren) const;
+
+	/**
+	 * return true for nodes that match using AreValuesEqual first, and pair up by position second
+	 * this is useful for arrays since we often want to keep elements with the same data paired while diffing other elements in order
+	 * @param TreeNode node from one of the two user provided trees (guaranteed not to be null)
+	*/
+	virtual bool ShouldMatchByValue(const TWeakPtr<FDetailTreeNode>& TreeNode) const;
+
+	/**
+	 * return true if TreeNode is considered equal when all it's children are equal.
+	 * This avoids an unnecessary call to AreValuesEqual
+	*/
+	virtual bool ShouldInheritEqualFromChildren(const TWeakPtr<FDetailTreeNode>& TreeNodeA, const TWeakPtr<FDetailTreeNode>& TreeNodeB) const;
+};
+
+// asynchronously update a difference tree with changes to details views
+// note: users are expected to call TAsyncTreeDifferences::Tick or TAsyncTreeDifferences::FlushQueue to get accurate data
+class PROPERTYEDITOR_API FAsyncDetailViewDiff : public TAsyncTreeDifferences<TWeakPtr<FDetailTreeNode>>
+{
+public:
+	FAsyncDetailViewDiff(TSharedRef<IDetailsView> LeftView, TSharedRef<IDetailsView> RightView);
+
+	// Note: by default this list is not necessarily accurate or exhaustive.
+	// if you need a perfectly accurate list, call FlushQueue() first
+	void GetPropertyDifferences(TArray<FSingleObjectDiffEntry>& OutDiffEntries) const;
+
+	// execute a method for each node associated with a row in at least one details view
+	// method provides current left and right row number
+	// returns the number of rows in each details view
+	TPair<int32, int32> ForEachRow(const TFunction<ETreeTraverseControl(const TUniquePtr<DiffNodeType>&, int32, int32)>& Method) const;
+
+	// generates scroll data used by LinkableScrollbar
+	TArray<FVector2f> GenerateScrollSyncRate() const;
+
+private:
+	static TAttribute<TArray<TWeakPtr<FDetailTreeNode>>> RootNodesAttribute(TWeakPtr<IDetailsView> DetailsView);
+	TWeakPtr<IDetailsView> LeftView;
+	TWeakPtr<IDetailsView> RightView;
+};
diff --git a/Engine/Source/Editor/PropertyEditor/Public/SDetailsSplitter.h b/Engine/Source/Editor/PropertyEditor/Public/SDetailsSplitter.h
--- a/Engine/Source/Editor/PropertyEditor/Public/SDetailsSplitter.h
+++ b/Engine/Source/Editor/PropertyEditor/Public/SDetailsSplitter.h
@@ -12,6 +12,7 @@ public:
 	// Callback that's called before highlighting a row. If true, the row is skipped.
 	DECLARE_DELEGATE_RetVal_OneParam(bool, FShouldIgnoreRow, const TWeakPtr<FDetailTreeNode>&)
 	DECLARE_DELEGATE_RetVal_OneParam(FLinearColor, FRowHighlightColor, const TUniquePtr<FAsyncDetailViewDiff::DiffNodeType>&)
+	DECLARE_DELEGATE_RetVal_OneParam(bool, FShouldHighlightRow, const TUniquePtr<FAsyncDetailViewDiff::DiffNodeType>&)
 
 	class FSlot : public TSlotBase<FSlot>
 	{
@@ -39,6 +40,7 @@ public:
 	SLATE_BEGIN_ARGS(SDetailsSplitter)
 	{}
 		SLATE_EVENT(FRowHighlightColor, RowHighlightColor) // default cyan: FLinearColor(0.f, 1.f, 1.f, .7f)
+		SLATE_EVENT(FShouldHighlightRow, ShouldHighlightRow) // default: highlight when node difference result is 'not identical'
 		SLATE_SLOT_ARGUMENT(FSlot, Slots)
 	SLATE_END_ARGS()
 
@@ -86,4 +88,5 @@ private:
 	TArray<FPanel> Panels;
 	FCopyPropertyButton HoveredCopyButton;
 	FRowHighlightColor GetRowHighlightColor;
+	FShouldHighlightRow GetShouldHighlightRow;
 };
diff --git a/Engine/Source/Editor/UnrealEd/Internal/AsyncTreeDifferences.h b/Engine/Source/Editor/UnrealEd/Internal/AsyncTreeDifferences.h
deleted file mode 100644
diff --git a/Engine/Source/Editor/UnrealEd/Private/DiffUtils.cpp b/Engine/Source/Editor/UnrealEd/Private/DiffUtils.cpp
--- a/Engine/Source/Editor/UnrealEd/Private/DiffUtils.cpp
+++ b/Engine/Source/Editor/UnrealEd/Private/DiffUtils.cpp
@@ -596,17 +596,20 @@ FPropertyPath FPropertySoftPath::ResolvePath(const UObject* Object) const
 	return Ret;
 }
 
-FString FPropertySoftPath::ToDisplayName() const
+FString FPropertySoftPath::ToDisplayName(const int32 NumberOfElements /*= INDEX_NONE*/) const
 {
 	FString Ret;
-	for( FChainElement Element : PropertyChain )
+	const int32 Count = NumberOfElements != INDEX_NONE
+		? FMath::Clamp(NumberOfElements, 1, PropertyChain.Num())
+		: PropertyChain.Num();
+	for (int32 i = PropertyChain.Num() - Count; i < PropertyChain.Num(); i++)
 	{
-		FString PropertyAsString = Element.DisplayString;
-		if(Ret.IsEmpty())
+		FString PropertyAsString = PropertyChain[i].DisplayString;
+		if (Ret.IsEmpty())
 		{
 			Ret.Append(PropertyAsString);
 		}
-		else if( PropertyAsString.IsNumeric())
+		else if (PropertyAsString.IsNumeric())
 		{
 			Ret.AppendChar('[');
 			Ret.Append(PropertyAsString);
@@ -621,6 +624,22 @@ FString FPropertySoftPath::ToDisplayName() const
 	return Ret;
 }
 
+FPropertySoftPath FPropertySoftPath::GetRootProperty(int32 Depth) const
+{
+	FPropertySoftPath RootPath;
+	Depth = FMath::Clamp(Depth, 1, PropertyChain.Num());
+	for (int32 ElementIndex = 0; ElementIndex < Depth; ElementIndex++)
+	{
+		RootPath.PropertyChain.Push(PropertyChain[ElementIndex]);
+	}
+	return RootPath;
+}
+
+int32 FPropertySoftPath::TryReadIndex(int32 Index) const
+{
+	return TryReadIndex(PropertyChain, Index);
+};
+
 const UObject* DiffUtils::GetCDO(const UBlueprint* ForBlueprint)
 {
 	if (!ForBlueprint
@@ -820,7 +839,7 @@ static void IdenticalHelper(const FProperty* AProperty, const FProperty* BProper
 		{
 			int32 Offset = AProperty->GetElementSize() * I;
 
-			DiffParameters.RootPath = FPropertySoftPath(DiffParameters.RootPath, I);
+			DiffParameters.RootPath = FPropertySoftPath(RootPath, I);
 			DiffParameters.bShouldDiffArrayElements = false;
 
 			const void* CurAValue = reinterpret_cast<const void*>(reinterpret_cast<const uint8*>(AValue) + Offset);
@@ -857,7 +876,7 @@ static void IdenticalHelper(const FProperty* AProperty, const FProperty* BProper
 				{
 					const FProperty* StructProp = *PropertyIt;
 
-					DiffParameters.RootPath = FPropertySoftPath(DiffParameters.RootPath, StructProp);
+					DiffParameters.RootPath = FPropertySoftPath(RootPath, StructProp);
 
 					const void* SubValueA = StructProp->ContainerPtrToValuePtr<void>(AValue, 0);
 					const void* SubValueB = StructProp->ContainerPtrToValuePtr<void>(BValue, 0);
@@ -882,7 +901,7 @@ static void IdenticalHelper(const FProperty* AProperty, const FProperty* BProper
 			// note any differences in contained types:
 			for (int32 ArrayIndex = 0; ArrayIndex < ArrayHelperA.Num() && ArrayIndex < ArrayHelperB.Num(); ArrayIndex++)
 			{
-				DiffParameters.RootPath = FPropertySoftPath(DiffParameters.RootPath, ArrayIndex);
+				DiffParameters.RootPath = FPropertySoftPath(RootPath, ArrayIndex);
 
 				const void* SubValueA = ArrayHelperA.GetRawPtr(ArrayIndex);
 				const void* SubValueB = ArrayHelperB.GetRawPtr(ArrayIndex);
@@ -921,7 +940,7 @@ static void IdenticalHelper(const FProperty* AProperty, const FProperty* BProper
 			FScriptSetHelper::FIterator IteratorB(SetHelperB);
 			for (; IteratorA && IteratorB; ++IteratorA, ++IteratorB)
 			{
-				DiffParameters.RootPath = FPropertySoftPath(DiffParameters.RootPath, IteratorA.GetLogicalIndex());
+				DiffParameters.RootPath = FPropertySoftPath(RootPath, IteratorA.GetLogicalIndex());
 
 				const void* SubValueA = SetHelperA.GetElementPtr(IteratorA);
 				const void* SubValueB = SetHelperB.GetElementPtr(IteratorB);
@@ -952,7 +971,7 @@ static void IdenticalHelper(const FProperty* AProperty, const FProperty* BProper
 			FScriptMapHelper::FIterator IteratorB(MapHelperB);
 			for (; IteratorA && IteratorB; ++IteratorA, ++IteratorB)
 			{
-				DiffParameters.RootPath = FPropertySoftPath(DiffParameters.RootPath, IteratorA.GetLogicalIndex());
+				DiffParameters.RootPath = FPropertySoftPath(RootPath, IteratorA.GetLogicalIndex());
 
 				IdenticalHelper(APropAsMap->KeyProp, BPropAsMap->KeyProp, MapHelperA.GetKeyPtr(IteratorA), MapHelperB.GetKeyPtr(IteratorB),
 					OwningOuterA, OwningOuterB, DiffParameters, DifferingSubProperties);
@@ -991,7 +1010,7 @@ static void IdenticalHelper(const FProperty* AProperty, const FProperty* BProper
 				for (TFieldIterator<FProperty> PropertyIt(AClass); PropertyIt; ++PropertyIt)
 				{
 					const FProperty* ClassProp = *PropertyIt;
-					DiffParameters.RootPath = FPropertySoftPath(DiffParameters.RootPath, ClassProp);
+					DiffParameters.RootPath = FPropertySoftPath(RootPath, ClassProp);
 
 					const void* SubValueA = ClassProp->ContainerPtrToValuePtr<void>(A, 0);
 					const void* SubValueB = ClassProp->ContainerPtrToValuePtr<void>(B, 0);
@@ -1086,6 +1105,18 @@ bool DiffUtils::Identical(const FResolvedProperty& AProp, const FResolvedPropert
 
 bool DiffUtils::Identical(const TSharedPtr<IPropertyHandle>& PropertyHandleA, const TSharedPtr<IPropertyHandle>& PropertyHandleB,
                           const TArray<TWeakObjectPtr<UObject>>& OwningOutersA, const TArray<TWeakObjectPtr<UObject>>& OwningOutersB)
+{
+	TArray<FPropertySoftPath> DifferingProperties;
+	return Identical(DifferingProperties, PropertyHandleA, PropertyHandleB, OwningOutersA, OwningOutersB);
+}
+
+bool DiffUtils::Identical(
+	TArray<FPropertySoftPath>& OutDifferingProperties,
+	const TSharedPtr<IPropertyHandle>& PropertyHandleA,
+	const TSharedPtr<IPropertyHandle>& PropertyHandleB,
+	const TArray<TWeakObjectPtr<UObject>>& OwningOutersA,
+	const TArray<TWeakObjectPtr<UObject>>& OwningOutersB
+	)
 {
 	TArray<void*> ValuesA;
 	TArray<void*> ValuesB;
@@ -1131,14 +1162,13 @@ bool DiffUtils::Identical(const TSharedPtr<IPropertyHandle>& PropertyHandleA, co
 
 		// note that we're not directly calling FProperty::Identical because sub-object properties should be weakly compared based on
 		// their paths instead of their pointers or data
-		TArray<FPropertySoftPath> DifferingProperties;
 		FDiffParameters DiffParameters;
 		DiffParameters.bShouldDiffArrayElements = false;
 
 		IdenticalHelper(PropertyHandleA->GetProperty(), PropertyHandleB->GetProperty(), ValueA, ValueB,
-			OwningOuterA, OwningOuterB, DiffParameters, DifferingProperties);
+			OwningOuterA, OwningOuterB, DiffParameters, OutDifferingProperties);
 
-		return DifferingProperties.IsEmpty();
+		return OutDifferingProperties.IsEmpty();
 	};
 	
 	if (ValuesA.Num() == ValuesB.Num())
diff --git a/Engine/Source/Editor/UnrealEd/Private/MergeUtils.cpp b/Engine/Source/Editor/UnrealEd/Private/MergeUtils.cpp
--- a/Engine/Source/Editor/UnrealEd/Private/MergeUtils.cpp
+++ b/Engine/Source/Editor/UnrealEd/Private/MergeUtils.cpp
@@ -395,14 +395,13 @@ template<>
 class TTreeDiffSpecification<FPropertyInstance>
 {
 public:
-	virtual ~TTreeDiffSpecification() = default;
-	
-	virtual bool AreValuesEqual(const FPropertyInstance& TreeNodeA, const FPropertyInstance& TreeNodeB) const
+
+	bool AreValuesEqual(const FPropertyInstance& TreeNodeA, const FPropertyInstance& TreeNodeB, TArray<FPropertySoftPath>* OutDifferingProperties = nullptr) const
 	{
 		return TreeNodeA == TreeNodeB;
 	}
 	
-	virtual bool AreMatching(const FPropertyInstance& TreeNodeA, const FPropertyInstance& TreeNodeB) const
+	bool AreMatching(const FPropertyInstance& TreeNodeA, const FPropertyInstance& TreeNodeB, TArray<FPropertySoftPath>* OutDifferingProperties = nullptr) const
 	{
 		if (TreeNodeA.KeyProperty->GetName() == TreeNodeB.KeyProperty->GetName())
 		{
@@ -420,18 +419,18 @@ public:
 		return false;
 	}
 	
-	virtual void GetChildren(const FPropertyInstance& InParent, TArray<FPropertyInstance>& OutChildren) const
+	void GetChildren(const FPropertyInstance& InParent, TArray<FPropertyInstance>& OutChildren) const
 	{
 		return InParent.GetChildren(OutChildren);
 	}
 
-	virtual bool ShouldMatchByValue(const FPropertyInstance& TreeNodeA) const
+	bool ShouldMatchByValue(const FPropertyInstance& TreeNodeA) const
 	{
 		// array elements should match by value
 		return CastField<FArrayProperty>(TreeNodeA.KeyProperty->Owner.ToField()) != nullptr;
 	}
 	
-	virtual bool ShouldInheritEqualFromChildren(const FPropertyInstance& TreeNodeA, const FPropertyInstance& TreeNodeB) const
+	bool ShouldInheritEqualFromChildren(const FPropertyInstance& TreeNodeA, const FPropertyInstance& TreeNodeB) const
 	{
 		return true;
 	}
diff --git a/Engine/Source/Editor/UnrealEd/Public/AsyncTreeDifferences.h b/Engine/Source/Editor/UnrealEd/Public/AsyncTreeDifferences.h
new file mode 100644
--- /dev/null
+++ b/Engine/Source/Editor/UnrealEd/Public/AsyncTreeDifferences.h
@@ -0,0 +1,663 @@
+// Copyright Epic Games, Inc. All Rights Reserved.
+#pragma once
+
+#include "DiffUtils.h"
+#include "Engine/EngineTypes.h"
+#include "Misc/Attribute.h"
+#include "Templates/UniquePtr.h"
+
+enum class ETreeDiffResult
+{
+	Invalid,
+	MissingFromTree1,
+	MissingFromTree2,
+	DifferentValues,
+	Identical
+};
+
+enum class ETreeTraverseOrder
+{
+	PreOrder, // parent then children
+	PostOrder, // children then parent
+	
+	ReversePreOrder, // parent then children in reverse order
+	ReversePostOrder // children in reverse order then parent
+};
+
+enum class ETreeTraverseControl
+{
+	Continue, // continue traversing
+	Break, // stop
+
+	// PreOrder traversal only:
+	SkipChildren, // don't iterate the children of this node
+};
+
+/// To use TAsyncTreeDifferences, define the following template specializations for your type
+////////////////////////////////////////////////////////////////////////////////////////////////////////////////////
+//namespace TreeDiffSpecification
+//{
+//	/**
+//	 * determine whether the values stored in two nodes are equal.
+//	 * @param TreeNodeA node from the first user provided tree (guaranteed not to be null)
+//	 * @param TreeNodeB node from the second user provided tree (guaranteed not to be null)
+//	 * @param OutDifferingProperties optional output list of property paths for the properties that differs between the two nodes
+//	 */
+//	template<typename InNodeType>
+//	bool AreValuesEqual(const InNodeType& TreeNodeA, const InNodeType& TreeNodeB, TArray<FPropertySoftPath>* OutDifferingProperties = nullptr);
+//
+//	/**
+//	 * determine whether two nodes occupy the same space in their trees
+//	 * for example if you have a tree key/value pairs, AreMatching should compare the keys while AreValuesEqual
+//	 * should compare the values
+//	 * @param TreeNodeA node from the first user provided tree (guaranteed not to be null)
+//	 * @param TreeNodeB node from the second user provided tree (guaranteed not to be null)
+//	 * @param OutDifferingProperties optional output list of property paths for the properties that differs between the two nodes
+//	 */
+//	template<typename InNodeType>
+//	bool AreMatching(const InNodeType& TreeNodeA, const InNodeType& TreeNodeB, TArray<FPropertySoftPath>* OutDifferingProperties = nullptr);
+//
+//	/**
+//	 * retrieves an array of children nodes from the parent node
+//	 * @param InParent node from one of the two user provided trees (guaranteed not to be null)
+//	 * @param[out] OutChildren to be filled with the children of parent
+//	 */
+//	template<typename InNodeType>
+//	void GetChildren(const InNodeType& InParent, TArray<InNodeType>& OutChildren);
+//
+//	/**
+//	 * return true for nodes that match using AreValuesEqual first, and pair up by position second
+//	 * this is useful for arrays since we often want to keep elements with the same data paired while diffing other elements in order
+//	 * @param TreeNode node from one of the two user provided trees (guaranteed not to be null)
+//	*/
+//	template<typename InNodeType>
+//	bool ShouldMatchByValue(const InNodeType& TreeNode);
+//
+//}
+
+template<typename InNodeType>
+class TTreeDiffSpecification
+{
+public:
+	virtual ~TTreeDiffSpecification() = default;
+	
+	/**
+	 * determine whether the values stored in two nodes are equal.
+	 * @param TreeNodeA node from the first user provided tree (guaranteed not to be null)
+	 * @param TreeNodeB node from the second user provided tree (guaranteed not to be null)
+	 * @param OutDifferingProperties optional output list of property paths for the properties that differs between the two nodes
+	 */
+	bool AreValuesEqual(const InNodeType& TreeNodeA, const InNodeType& TreeNodeB, TArray<FPropertySoftPath>* OutDifferingProperties = nullptr) const
+	{
+		return TreeNodeA == TreeNodeB;
+	}
+
+	/**
+	 * determine whether two nodes occupy the same space in their trees
+	 * for example if you have a tree key/value pairs, AreMatching should compare the keys while AreValuesEqual
+	 * should compare the values
+	 * @param TreeNodeA node from the first user provided tree (guaranteed not to be null)
+	 * @param TreeNodeB node from the second user provided tree (guaranteed not to be null)
+	 * @param OutDifferingProperties optional output list of property paths for the properties that differs between the two nodes
+	 */
+	bool AreMatching(const InNodeType& TreeNodeA, const InNodeType& TreeNodeB, TArray<FPropertySoftPath>* OutDifferingProperties = nullptr) const
+	{
+		return TreeNodeA == TreeNodeB;
+	}
+
+	/**
+	 * retrieves an array of children nodes from the parent node
+	 * @param InParent node from one of the two user provided trees (guaranteed not to be null)
+	 * @param[out] OutChildren to be filled with the children of parent
+	 */
+	void GetChildren(const InNodeType& InParent, TArray<InNodeType>& OutChildren) const
+	{}
+
+	/**
+	 * return true for nodes that match using AreValuesEqual first, and pair up by position second
+	 * this is useful for arrays since we often want to keep elements with the same data paired while diffing other elements in order
+	 * @param TreeNode node from one of the two user provided trees (guaranteed not to be null)
+	*/
+	bool ShouldMatchByValue(const InNodeType& TreeNode) const
+	{
+		return true;
+	}
+
+	/**
+	 * return true if TreeNode is considered equal when all it's children are equal.
+	 * This avoids an unnecessary call to AreValuesEqual
+	*/
+	bool ShouldInheritEqualFromChildren(const InNodeType& TreeNodeA, const InNodeType& TreeNodeB) const
+	{
+		return false;
+	}
+};
+
+
+// the TAsyncTreeDifferences structure's goal is to build and maintain an updated tree made of these nodes.
+// where TDiffNode::ValueA and TDiffNode::ValueB are matching nodes from two diffed trees. Note that these values can
+// be both "matching" and not equal which would set TDiffNode::DiffResult to DifferentValues
+template<typename InNodeType>
+struct TDiffNode
+{
+	// because we've got a bit of node-type inception going on, we'll refer to the nodes in the user provided trees as
+	// values rather than nodes. (but it's still sometimes helpful to remember that they are nodes in trees)
+	using ValueType = InNodeType;
+	
+	ValueType ValueA = NullValue;
+	ValueType ValueB = NullValue;
+	
+	TArray<TUniquePtr<TDiffNode>> Children = {};
+	TArray<FPropertySoftPath> DifferingProperties;
+	const TDiffNode* Parent = nullptr;
+	ETreeDiffResult DiffResult = ETreeDiffResult::Invalid;
+
+	TDiffNode() = default;
+	TDiffNode(TDiffNode &&Other) = default;
+	
+	TDiffNode(const ValueType& InValueA, const ValueType& InValueB, const TDiffNode* InParent)
+		: ValueA(InValueA), ValueB(InValueB), Children(), Parent(InParent)
+	{}
+
+	void SetDiffType(TUniquePtr<TTreeDiffSpecification<ValueType>>& Specification);
+
+	bool operator==(const TDiffNode& Other) const;
+	
+	static inline ValueType NullValue{}; // assume default constructor to signify a "null" node
+};
+
+template<typename InNodeType>
+class TAsyncTreeDifferences
+{
+public:
+	using ValueType = InNodeType;
+	using DiffNodeType = TDiffNode<ValueType>;
+	
+	TAsyncTreeDifferences(const TAttribute<TArray<ValueType>>& InRootValuesA, const TAttribute<TArray<ValueType>>& InRootValuesB);
+	
+	// update the differences list over time
+	void Tick(float MaxAllottedTimeMs = 1.f);
+
+	// process all tasks synchronously until all data is up to date. (slow)
+	void FlushQueue();
+
+	// calls Method(DiffNode) on each diff node in the tree.
+	void ForEach(ETreeTraverseOrder TraversalOrder, const TFunction<ETreeTraverseControl(const TUniquePtr<DiffNodeType>&)>& Method) const;
+
+	template<typename SpecificationType, typename... TArgs>
+	void SetDiffSpecification(TArgs... Args);
+
+	// head of the main diff tree (note: not for user modification. Tick() will overwrite user changes)
+	TUniquePtr<DiffNodeType> Head;
+
+	// every time the update queue is completed this is called (max of once per Tick)
+	DECLARE_MULTICAST_DELEGATE_OneParam(FOnQueueFlushed, const TAsyncTreeDifferences&);
+	FOnQueueFlushed OnQueueFlushed;
+private:
+
+	// process one element from the top of the queue
+	void ProcessTopOfQueue();
+	
+	void QueueParallelNodeLists(const TArray<ValueType>& ValuesA, const TArray<ValueType>& ValuesB, DiffNodeType* ParentNode);
+	
+	TArray<TArray<int32>> CalculateLCSTableForMatchingValues(const TArray<ValueType>& ValuesA, const TArray<ValueType>& ValuesB);
+	static TArray<TArray<int32>> CalculateLCSTableForDiffNodes(const TArray<DiffNodeType>& FoundNodes, const TArray<TUniquePtr<DiffNodeType>>& ExpectedNodes);
+
+	bool AreMatching(const ValueType& ValueA, const ValueType& ValueB);
+	
+	static bool PreOrderRecursive(const TUniquePtr<DiffNodeType>& Node, const TFunction<ETreeTraverseControl(const TUniquePtr<DiffNodeType>&)>& Method);
+	static bool PostOrderRecursive(const TUniquePtr<DiffNodeType>& Node, const TFunction<ETreeTraverseControl(const TUniquePtr<DiffNodeType>&)>& Method);
+	static bool ReversePreOrderRecursive(const TUniquePtr<DiffNodeType>& Node, const TFunction<ETreeTraverseControl(const TUniquePtr<DiffNodeType>&)>& Method);
+	static bool ReversePostOrderRecursive(const TUniquePtr<DiffNodeType>& Node, const TFunction<ETreeTraverseControl(const TUniquePtr<DiffNodeType>&)>& Method);
+
+	TArray<DiffNodeType*> UpdateQueue;
+	TAttribute<TArray<ValueType>> RootValuesA;
+	TAttribute<TArray<ValueType>> RootValuesB;
+	TUniquePtr<TTreeDiffSpecification<ValueType>> Specification;
+};
+
+template <typename InNodeType>
+void TDiffNode<InNodeType>::SetDiffType(TUniquePtr<TTreeDiffSpecification<ValueType>>& Specification)
+{
+	DifferingProperties.Reset();
+
+	if (ValueA != NullValue && ValueB != NullValue)
+	{
+		if (Specification->ShouldInheritEqualFromChildren(ValueA, ValueB) && !Children.IsEmpty())
+		{
+			// iterate children. Parent is identical iff all of it's children are identical
+			for (const TUniquePtr<TDiffNode>& Child : Children)
+			{
+				if (Child->DiffResult != ETreeDiffResult::Identical)
+				{
+					DiffResult = ETreeDiffResult::DifferentValues;
+					return;
+				}
+			}
+			DiffResult = ETreeDiffResult::Identical;
+			return;
+		}
+		else
+		{
+			if (Specification->AreValuesEqual(ValueA, ValueB, &DifferingProperties))
+			{
+				DiffResult = ETreeDiffResult::Identical;
+				return;
+			}
+			DiffResult = ETreeDiffResult::DifferentValues;
+			return;
+		}
+	}
+	if (ValueA != NullValue && ValueB == NullValue)
+	{
+		DiffResult = ETreeDiffResult::MissingFromTree2;
+		return;
+	}
+	if (ValueA == NullValue && ValueB != NullValue)
+	{
+		DiffResult = ETreeDiffResult::MissingFromTree1;
+		return;
+	}
+	
+	check(false);
+	DiffResult = ETreeDiffResult::Invalid;
+}
+
+template <typename InNodeType>
+bool TDiffNode<InNodeType>::operator==(const TDiffNode& Other) const
+{
+	return Other.ValueA == ValueA && Other.ValueB == ValueB;
+}
+
+template <typename InNodeType>
+TAsyncTreeDifferences<InNodeType>::TAsyncTreeDifferences(const TAttribute<TArray<ValueType>>& InRootValuesA, const TAttribute<TArray<ValueType>>& InRootValuesB)
+	: Head(MakeUnique<DiffNodeType>())
+	, RootValuesA(InRootValuesA)
+	, RootValuesB(InRootValuesB)
+	, Specification(MakeUnique<TTreeDiffSpecification<ValueType>>())
+{
+	QueueParallelNodeLists(RootValuesA.Get(), RootValuesB.Get(), Head.Get());
+}
+
+template <typename InNodeType>
+void TAsyncTreeDifferences<InNodeType>::Tick(float MaxAllottedTimeMs)
+{
+	if (UpdateQueue.IsEmpty())
+	{
+		QueueParallelNodeLists(RootValuesA.Get(), RootValuesB.Get(), Head.Get());
+	}
+	
+	const double StartTime = FPlatformTime::Seconds();
+	while(!UpdateQueue.IsEmpty() && FPlatformTime::Seconds() - StartTime < MaxAllottedTimeMs)
+	{
+		ProcessTopOfQueue();
+	}
+	
+	if (UpdateQueue.IsEmpty())
+	{
+		OnQueueFlushed.Broadcast(*this);
+	}
+}
+
+template <typename InNodeType>
+void TAsyncTreeDifferences<InNodeType>::FlushQueue()
+{
+	while(!UpdateQueue.IsEmpty())
+	{
+		ProcessTopOfQueue();
+	}
+	OnQueueFlushed.Broadcast(*this);
+}
+
+template <typename InNodeType>
+void TAsyncTreeDifferences<InNodeType>::ForEach(ETreeTraverseOrder TraversalOrder, const TFunction<ETreeTraverseControl(const TUniquePtr<DiffNodeType>&)>& Method) const
+{
+	switch(TraversalOrder)
+	{
+	case ETreeTraverseOrder::PreOrder:
+		PreOrderRecursive(Head, Method);
+		break;
+	case ETreeTraverseOrder::PostOrder:
+		PostOrderRecursive(Head, Method);
+		break;
+	case ETreeTraverseOrder::ReversePreOrder:
+		ReversePreOrderRecursive(Head, Method);
+		break;
+	case ETreeTraverseOrder::ReversePostOrder:
+		ReversePostOrderRecursive(Head, Method);
+		break;
+	default: check(false);
+	}
+}
+
+template <typename InNodeType>
+template <typename SpecificationType, typename... TArgs>
+void TAsyncTreeDifferences<InNodeType>::SetDiffSpecification(TArgs... Args)
+{
+	Specification = MakeUnique<SpecificationType>(Args...);
+}
+
+template <typename InNodeType>
+bool TAsyncTreeDifferences<InNodeType>::PreOrderRecursive(const TUniquePtr<DiffNodeType>& Node, const TFunction<ETreeTraverseControl(const TUniquePtr<DiffNodeType>&)>& Method)
+{
+	if (Node->DiffResult != ETreeDiffResult::Invalid)
+	{
+		switch (Method(Node))
+		{
+			case ETreeTraverseControl::Break: return false;
+			case ETreeTraverseControl::SkipChildren: return true;
+		}
+	}
+	
+	for (const TUniquePtr<DiffNodeType>& Child : Node->Children)
+	{
+		if (!PreOrderRecursive(Child, Method))
+		{
+			return false;
+		}
+	}
+
+	return true;
+}
+
+template <typename InNodeType>
+bool TAsyncTreeDifferences<InNodeType>::PostOrderRecursive(const TUniquePtr<DiffNodeType>& Node, const TFunction<ETreeTraverseControl(const TUniquePtr<DiffNodeType>&)>& Method)
+{
+	for (const TUniquePtr<DiffNodeType>& Child : Node->Children)
+	{
+		if (!PostOrderRecursive(Child, Method))
+		{
+			return false;
+		}
+	}
+	
+	if (Node->DiffResult != ETreeDiffResult::Invalid)
+	{
+		switch (Method(Node))
+		{
+			case ETreeTraverseControl::Break: return false;
+		}
+	}
+
+	return true;
+}
+
+template <typename InNodeType>
+bool TAsyncTreeDifferences<InNodeType>::ReversePreOrderRecursive(const TUniquePtr<DiffNodeType>& Node, const TFunction<ETreeTraverseControl(const TUniquePtr<DiffNodeType>&)>& Method)
+{
+	if (Node->DiffResult != ETreeDiffResult::Invalid)
+	{
+		switch (Method(Node))
+		{
+			case ETreeTraverseControl::Break: return false;
+			case ETreeTraverseControl::SkipChildren: return true;
+		}
+	}
+	
+	for (int32 I = Node->Children.Num() - 1; I >= 0; --I)
+	{
+		const TUniquePtr<DiffNodeType>& Child = Node->Children[I];
+		if (!ReversePreOrderRecursive(Child, Method))
+		{
+			return false;
+		}
+	}
+
+	return true;
+}
+
+template <typename InNodeType>
+bool TAsyncTreeDifferences<InNodeType>::ReversePostOrderRecursive(const TUniquePtr<DiffNodeType>& Node, const TFunction<ETreeTraverseControl(const TUniquePtr<DiffNodeType>&)>& Method)
+{
+	for (int32 I = Node->Children.Num() - 1; I >= 0; --I)
+	{
+		const TUniquePtr<DiffNodeType>& Child = Node->Children[I];
+		if (!ReversePostOrderRecursive(Child, Method))
+		{
+			return false;
+		}
+	}
+	
+	if (Node->DiffResult != ETreeDiffResult::Invalid)
+	{
+		switch (Method(Node))
+		{
+			case ETreeTraverseControl::Break: return false;
+		}
+	}
+
+	return true;
+}
+
+template <typename InNodeType>
+void TAsyncTreeDifferences<InNodeType>::ProcessTopOfQueue()
+{
+	DiffNodeType* DiffNode = UpdateQueue.Pop();
+	
+	TArray<ValueType> ChildrenA;
+	TArray<ValueType> ChildrenB;
+	if (DiffNode->ValueA != DiffNodeType::NullValue)
+	{
+		Specification->GetChildren(DiffNode->ValueA, ChildrenA);
+	}
+	if (DiffNode->ValueB != DiffNodeType::NullValue)
+	{
+		Specification->GetChildren(DiffNode->ValueB, ChildrenB);
+	}
+	
+	if (!ChildrenA.IsEmpty() || !ChildrenB.IsEmpty())
+	{
+		QueueParallelNodeLists(ChildrenA, ChildrenB, DiffNode);
+	}
+}
+
+
+// this is where the magic happens :)
+template <typename InNodeType>
+void TAsyncTreeDifferences<InNodeType>::QueueParallelNodeLists(const TArray<ValueType>& ValuesA, const TArray<ValueType>& ValuesB, DiffNodeType* ParentNode)
+{
+	// calculate the diff nodes that should be children of ParentNode
+	TArray<DiffNodeType> FoundChildren;
+	{
+		// reversing data before calculating LCS puts "add" diff results towards the end which is a friendlier format
+		// (don't worry we'll return the data to it's original order)
+		Algo::Reverse(const_cast<TArray<ValueType>&>(ValuesA));
+		Algo::Reverse(const_cast<TArray<ValueType>&>(ValuesB));
+		
+		// Find the Longest Common Subsequence between node lists
+        const TArray<TArray<int32>> LCS = CalculateLCSTableForMatchingValues(ValuesA, ValuesB);
+        
+        // Using an LCS Table, we can determine which tree nodes match one another in each tree and queue them up to update together
+        int32 IndexA = ValuesA.Num();
+        int32 IndexB = ValuesB.Num();
+        while (IndexA > 0 || IndexB > 0)
+        {
+        	if (IndexA == 0)
+        	{
+        		// ValueB doesn't match. push it on it's own
+        		FoundChildren.Add(DiffNodeType(DiffNodeType::NullValue, ValuesB[--IndexB], ParentNode));
+        	}
+        	else if (IndexB == 0)
+        	{
+        		// ValueA doesn't match. push it on it's own
+        		FoundChildren.Add(DiffNodeType(ValuesA[--IndexA], DiffNodeType::NullValue, ParentNode));
+        	}
+        	else if (AreMatching(ValuesA[IndexA - 1], ValuesB[IndexB - 1]))
+        	{
+        		// found a match between both nodes
+        		FoundChildren.Add(DiffNodeType(ValuesA[--IndexA], ValuesB[--IndexB], ParentNode));
+        	}
+        	else if (LCS[IndexA - 1][IndexB] <= LCS[IndexA][IndexB - 1])
+        	{
+        		// ValueB doesn't match. push it on it's own
+        		FoundChildren.Add(DiffNodeType(DiffNodeType::NullValue, ValuesB[--IndexB], ParentNode));
+        	}
+        	else
+        	{
+        		// ValueA doesn't match. push it on it's own
+        		FoundChildren.Add(DiffNodeType(ValuesA[--IndexA], DiffNodeType::NullValue, ParentNode));
+        	}
+        }
+
+		// undo the temp reverse from above
+		Algo::Reverse(const_cast<TArray<ValueType>&>(ValuesA));
+		Algo::Reverse(const_cast<TArray<ValueType>&>(ValuesB));
+	}
+	
+	// compress consecutive value matched entries
+	TArray<DiffNodeType> AllFoundChildren = MoveTemp(FoundChildren);
+	FoundChildren = TArray<DiffNodeType>();
+
+	int32 CompressIndex = INDEX_NONE;
+	int32 CompressCount = 0;
+	for (DiffNodeType& FoundChild : AllFoundChildren)
+	{
+		if (FoundChild.ValueA == DiffNodeType::NullValue &&
+			Specification->ShouldMatchByValue(FoundChild.ValueB))
+		{
+			if (CompressCount == 0)
+			{
+				CompressIndex = FoundChildren.Num();
+			}
+			++CompressCount;
+			FoundChildren.Add(MoveTemp(FoundChild));
+			continue;
+		}
+		if (CompressCount > 0 &&
+            FoundChild.ValueB == DiffNodeType::NullValue &&
+			Specification->ShouldMatchByValue(FoundChild.ValueA))
+		{
+			FoundChildren[CompressIndex].ValueA = FoundChild.ValueA;
+			++CompressIndex;
+			--CompressCount;
+			continue;
+		}
+		
+		FoundChildren.Add(MoveTemp(FoundChild));
+		CompressIndex = INDEX_NONE;
+		CompressCount = 0;
+	}
+	
+	// compare/replace results with what's *actually* in the parent node to avoid modifying nodes that didn't change
+	{
+		TArray<TUniquePtr<DiffNodeType>> ExpectedChildren = MoveTemp(ParentNode->Children);
+		ParentNode->Children = TArray<TUniquePtr<DiffNodeType>>();
+		const TArray<TArray<int32>> LCS = CalculateLCSTableForDiffNodes(FoundChildren, ExpectedChildren);
+		int32 FoundChildrenIndex = FoundChildren.Num();
+		int32 ExpectedChildrenIndex = ExpectedChildren.Num();
+		while (FoundChildrenIndex > 0)
+		{
+			if (ExpectedChildrenIndex == 0)
+			{
+				// a new child was found. move it into the children
+				ParentNode->Children.Add(MakeUnique<DiffNodeType>(MoveTemp(FoundChildren[--FoundChildrenIndex])));
+			}
+			else if (FoundChildren[FoundChildrenIndex - 1] == *ExpectedChildren[ExpectedChildrenIndex - 1])
+			{
+				// found a match between both nodes. Preserve the old data by copying it into the new child array
+				ParentNode->Children.Add(MoveTemp(ExpectedChildren[--ExpectedChildrenIndex]));
+				--FoundChildrenIndex;
+			}
+			else if (LCS[FoundChildrenIndex - 1][ExpectedChildrenIndex] <= LCS[FoundChildrenIndex][ExpectedChildrenIndex - 1])
+			{
+				// a child was removed. ignore the old data
+				--ExpectedChildrenIndex;
+			}
+			else
+			{
+				// a new child was found. move it into the children
+				ParentNode->Children.Add(MakeUnique<DiffNodeType>(MoveTemp(FoundChildren[--FoundChildrenIndex])));
+			}
+		}
+	}
+	Algo::Reverse(ParentNode->Children);
+
+	// diff the values of all the children
+	for (const TUniquePtr<DiffNodeType>& Child : ParentNode->Children)
+	{
+		Child->SetDiffType(Specification);
+	}
+
+	// queue all the children to update
+	for (const TUniquePtr<DiffNodeType>& Child : ParentNode->Children)
+	{
+		UpdateQueue.Add(Child.Get());
+	}
+}
+
+template <typename InNodeType>
+TArray<TArray<int32>> TAsyncTreeDifferences<InNodeType>::CalculateLCSTableForMatchingValues(
+	const TArray<ValueType>& ValuesA, const TArray<ValueType>& ValuesB)
+{
+	TArray<TArray<int32>> LCS;
+	LCS.SetNum(ValuesA.Num() + 1);
+	for (int32 I = 0; I <= ValuesA.Num(); I++)
+	{
+		LCS[I].SetNum(ValuesB.Num() + 1);
+		if (I == 0)
+		{
+			continue;
+		}
+	
+		for (int32 J = 1; J <= ValuesB.Num(); J++)
+		{
+			if (AreMatching(ValuesA[I - 1], ValuesB[J - 1]))
+			{
+				LCS[I][J] = LCS[I - 1][J - 1] + 1;
+			}
+			else
+			{
+				LCS[I][J] = FMath::Max(LCS[I - 1][J], LCS[I][J - 1]);
+			}
+		}
+	}
+	return LCS;
+}
+
+template <typename InNodeType>
+TArray<TArray<int32>> TAsyncTreeDifferences<InNodeType>::CalculateLCSTableForDiffNodes(const TArray<DiffNodeType>& FoundNodes,
+	const TArray<TUniquePtr<DiffNodeType>>& ExpectedNodes)
+{
+	TArray<TArray<int32>> LCS;
+	LCS.SetNum(FoundNodes.Num() + 1);
+	for (int32 I = 0; I <= FoundNodes.Num(); I++)
+	{
+		LCS[I].SetNum(ExpectedNodes.Num() + 1);
+		if (I == 0)
+		{
+			continue;
+		}
+	
+		for (int32 J = 1; J <= ExpectedNodes.Num(); J++)
+		{
+			if (FoundNodes[I - 1] == *ExpectedNodes[J - 1])
+			{
+				LCS[I][J] = LCS[I - 1][J - 1] + 1;
+			}
+			else
+			{
+				LCS[I][J] = FMath::Max(LCS[I - 1][J], LCS[I][J - 1]);
+			}
+		}
+	}
+	return LCS;
+}
+
+template <typename InNodeType>
+bool TAsyncTreeDifferences<InNodeType>::AreMatching(const ValueType& ValueA, const ValueType& ValueB)
+{
+	const bool bMatchValueAByValue = Specification->ShouldMatchByValue(ValueA);
+	const bool bMatchValueBByValue = Specification->ShouldMatchByValue(ValueB);
+	if (bMatchValueAByValue && bMatchValueBByValue)
+	{
+		return Specification->AreValuesEqual(ValueA, ValueB);
+	}
+	if (!bMatchValueAByValue && !bMatchValueBByValue)
+	{
+		return Specification->AreMatching(ValueA, ValueB);
+	}
+
+	// a node that should match by value will never match a node that shouldn't match by value
+	return false;
+}
+
diff --git a/Engine/Source/Editor/UnrealEd/Public/DiffUtils.h b/Engine/Source/Editor/UnrealEd/Public/DiffUtils.h
--- a/Engine/Source/Editor/UnrealEd/Public/DiffUtils.h
+++ b/Engine/Source/Editor/UnrealEd/Public/DiffUtils.h
@@ -3,9 +3,11 @@
 
 #include "Containers/Array.h"
 #include "Containers/Set.h"
+#if UE_ENABLE_INCLUDE_ORDER_DEPRECATED_IN_5_5
 #include "Containers/SparseArray.h"
-#include "Containers/UnrealString.h"
 #include "CoreMinimal.h"
+#endif // UE_ENABLE_INCLUDE_ORDER_DEPRECATED_IN_5_5
+#include "Containers/UnrealString.h"
 #include "Delegates/Delegate.h"
 #include "HAL/Platform.h"
 #include "HAL/PlatformCrt.h"
@@ -23,6 +25,7 @@
 #include "Widgets/Views/STreeView.h"
 #include "PropertyHandle.h"
 
+class FBlueprintDifferenceTreeEntry;
 class ISourceControlRevision;
 class SWidget;
 class UBlueprint;
@@ -74,9 +77,24 @@ struct FPropertySoftPath
 	UNREALED_API FResolvedProperty Resolve(const UObject* Object) const;
 	UNREALED_API FResolvedProperty Resolve(const UStruct* Struct, const void* StructData) const;
 	UNREALED_API FPropertyPath ResolvePath(const UObject* Object) const;
-	UNREALED_API FString ToDisplayName() const;
 
-	inline bool IsSubPropertyMatch(const FPropertySoftPath& PotentialBasePropertyPath) const
+	UNREALED_API FPropertySoftPath GetRootProperty(int32 Depth) const;
+	UNREALED_API int32 TryReadIndex(int32 Index) const;
+
+	/**
+	 * Provides a string of all elements in the property chain using the following format
+	 * e.g., "X[3] Y Z"
+	 * @param NumberOfElements Optional argument to return a string containing only the last 'NumberOfElements' elements of the path
+	 * @return String of all elements in the property chain
+	 */
+	UNREALED_API FString ToDisplayName(const int32 NumberOfElements = INDEX_NONE) const;
+
+	/**
+	 * Indicates whether a given path is a base path of the current path.
+	 * @param PotentialBasePropertyPath The path to look for as a base common path
+	 * @return Whether the path has the provided path as base path or not
+	 */
+	bool IsSubPropertyMatch(const FPropertySoftPath& PotentialBasePropertyPath) const
 	{
 		if (PropertyChain.Num() <= PotentialBasePropertyPath.PropertyChain.Num())
 		{
@@ -94,6 +112,39 @@ struct FPropertySoftPath
 		return true;
 	}
 
+	/**
+	 * Indicates whether a given list of property names is a base path of the current path.
+	 * @param PotentialBasePropertyNameArray The list of property names to look for as a base common path
+	 * @return Whether the path has the provided list of names as base path or not
+	 */
+	bool IsSubPropertyMatch(TConstArrayView<FName> PotentialBasePropertyNameArray) const
+	{
+		if (PropertyChain.Num() <= PotentialBasePropertyNameArray.Num())
+		{
+			return false;
+		}
+
+		for (int32 CurChainElement = 0; CurChainElement < PotentialBasePropertyNameArray.Num(); CurChainElement++)
+		{
+			if (PotentialBasePropertyNameArray[CurChainElement] != PropertyChain[CurChainElement].PropertyName)
+			{
+				return false;
+			}
+		}
+
+		return true;
+	}
+
+	/**
+	 * Indicates whether a given property name is a base path of the current path.
+	 * @param PotentialBasePropertyName Name of the single property to match
+	 * @return Whether the path contains a single element with the name provided or not
+	 */
+	bool IsSubPropertyMatch(const FName PotentialBasePropertyName) const
+	{
+		return IsSubPropertyMatch(MakeArrayView({PotentialBasePropertyName}));
+	}
+
 	inline bool operator==(FPropertySoftPath const& RHS) const
 	{
 		return PropertyChain == RHS.PropertyChain;
@@ -257,6 +308,14 @@ struct FSCSDiffRoot
 
 namespace DiffUtils
 {
+	DECLARE_DELEGATE_OneParam(FOnGenerateCustomDiffEntries, TArray<FSingleObjectDiffEntry>&)
+	DECLARE_DELEGATE_RetVal_TwoParams(TSharedRef<SWidget>, FOnGenerateCustomDiffEntryWidget, const FSingleObjectDiffEntry&, FText&)
+	DECLARE_DELEGATE_FourParams(FOnOrganizeDiffEntries,
+		TArray<TSharedPtr<FBlueprintDifferenceTreeEntry>>&,
+		const TArray<FSingleObjectDiffEntry>&,
+		TFunctionRef<TSharedPtr<FBlueprintDifferenceTreeEntry>(const FSingleObjectDiffEntry&)>,
+		TFunctionRef<TSharedPtr<FBlueprintDifferenceTreeEntry>(FText&)>)
+
 	struct FDiffParameters
 	{
 		FDiffParameters() = default;
@@ -296,6 +355,7 @@ namespace DiffUtils
 	UNREALED_API bool Identical(const FResolvedProperty& AProp, const FResolvedProperty& BProp, const UObject* OwningOuterA, const UObject* OwningOuterB);
 	UNREALED_API bool Identical(const FResolvedProperty& AProp, const FResolvedProperty& BProp, const UObject* OwningOuterA, const UObject* OwningOuterB, FDiffParameters DiffParameters, TArray<FPropertySoftPath>& DifferingProperties);
 	UNREALED_API bool Identical(const TSharedPtr<IPropertyHandle>& PropertyHandleA, const TSharedPtr<IPropertyHandle>& PropertyHandleB, const TArray<TWeakObjectPtr<UObject>>& OwningOutersA = {}, const TArray<TWeakObjectPtr<UObject>>& OwningOutersB = {});
+	UNREALED_API bool Identical(TArray<FPropertySoftPath>& OutDifferingProperties, const TSharedPtr<IPropertyHandle>& PropertyHandleA, const TSharedPtr<IPropertyHandle>& PropertyHandleB, const TArray<TWeakObjectPtr<UObject>>& OwningOutersA = {}, const TArray<TWeakObjectPtr<UObject>>& OwningOutersB = {});
 	UNREALED_API TArray<FPropertySoftPath> GetVisiblePropertiesInOrderDeclared(const UStruct* ForStruct, const FPropertySoftPath& Scope = FPropertySoftPath());
 
 	UNREALED_API TArray<FPropertyPath> ResolveAll(const UObject* Object, const TArray<FPropertySoftPath>& InSoftProperties);
@@ -308,7 +368,6 @@ namespace DiffUtils
 	 */
 	UNREALED_API UPackage* LoadPackageForDiff(const FPackagePath& InTempPackagePath, const FPackagePath& InOriginalPackagePath);
 	UNREALED_API UPackage* LoadPackageForDiff(TSharedPtr<ISourceControlRevision> Revision);
-
 }
 
 DECLARE_DELEGATE(FOnDiffEntryFocused);
