Add timing logs for Localize job on Main (temporary)

diff --git a/Engine/Source/Editor/UnrealEd/Private/Commandlets/GatherTextFromAssetsCommandlet.cpp b/Engine/Source/Editor/UnrealEd/Private/Commandlets/GatherTextFromAssetsCommandlet.cpp
--- a/Engine/Source/Editor/UnrealEd/Private/Commandlets/GatherTextFromAssetsCommandlet.cpp
+++ b/Engine/Source/Editor/UnrealEd/Private/Commandlets/GatherTextFromAssetsCommandlet.cpp
@@ -1372,6 +1372,8 @@ UGatherTextFromAssetsCommandlet::EPackageLocCacheState UGatherTextFromAssetsComm
 
 int32 UGatherTextFromAssetsCommandlet::Main(const FString& Params)
 {
+	double TimePrev = FPlatformTime::Seconds();
+
 	UE_SCOPED_TIMER(TEXT("UGatherTextFromAssetsCommandlet::Main"), LogGatherTextFromAssetsCommandlet, Display);
 	// Parse command line.
 	if (!ParseCommandLineHelper(Params))
@@ -1379,6 +1381,9 @@ int32 UGatherTextFromAssetsCommandlet::Main(const FString& Params)
 		return -1;
 	}
 
+	UE_LOG(LogGatherTextFromAssetsCommandlet, Display, TEXT("0088:01 %.2f seconds."), FPlatformTime::Seconds() - TimePrev);
+	TimePrev = FPlatformTime::Seconds();
+
 	// If the editor has loaded a persistent world then create an empty world prior to starting the asset gather
 	// This avoids any issues when loading and initializing worlds during the gather, as WP needs to re-initialize the world
 	if (GEditor)
@@ -1404,6 +1409,9 @@ int32 UGatherTextFromAssetsCommandlet::Main(const FString& Params)
 		return -1;
 	}
 
+	UE_LOG(LogGatherTextFromAssetsCommandlet, Display, TEXT("0088:02 %.2f seconds."), FPlatformTime::Seconds() - TimePrev);
+	TimePrev = FPlatformTime::Seconds();
+
 	if (!bShouldExcludeDerivedClasses)
 	{
 		if (!PerformExcludeExactClassesFilter(AssetDataArray))
@@ -1424,11 +1432,17 @@ int32 UGatherTextFromAssetsCommandlet::Main(const FString& Params)
 	// Discover the external actors for any worlds that are pending gather
 	DiscoverExternalActors(AssetDataArray);
 
+	UE_LOG(LogGatherTextFromAssetsCommandlet, Display, TEXT("0088:03 %.2f seconds."), FPlatformTime::Seconds() - TimePrev);
+	TimePrev = FPlatformTime::Seconds();
+
 	// Collect the basic information about the packages that we're going to gather from
 	TSet<FName> PackageNamesToGather = GetPackageNamesToGather(AssetDataArray);
 	AssetDataArray.Empty();
 	PopulatePackagesPendingGather(MoveTemp(PackageNamesToGather));
 
+	UE_LOG(LogGatherTextFromAssetsCommandlet, Display, TEXT("0088:04 %.2f seconds."), FPlatformTime::Seconds() - TimePrev);
+	TimePrev = FPlatformTime::Seconds();
+
 	UE_LOG(LogGatherTextFromAssetsCommandlet, Display, TEXT("Discovering assets to gather took %.2f seconds."), FPlatformTime::Seconds() - DiscoveringAssetsStartTime);
 
 	// These are external actor packages that are stale or are missing a gather cache from their package
@@ -1438,12 +1452,18 @@ int32 UGatherTextFromAssetsCommandlet::Main(const FString& Params)
 	UE_LOG(LogGatherTextFromAssetsCommandlet, Display, TEXT("Processing assets to gather..."));
 	ProcessAndRemoveCachedPackages(ExternalActorsWithStaleOrMissingCaches);
 
+	UE_LOG(LogGatherTextFromAssetsCommandlet, Display, TEXT("0088:05 %.2f seconds."), FPlatformTime::Seconds() - TimePrev);
+	TimePrev = FPlatformTime::Seconds();
+
 	UE::Private::GatherTextFromAssetsCommandlet::AssetGatherCacheMetrics.LogMetrics();
 
 	// Merge any pending WP map requests back into PackagesPendingGather
 	MergeInExternalActorsWithStaleOrMissingCaches(ExternalActorsWithStaleOrMissingCaches);
 	ExternalActorsWithStaleOrMissingCaches.Reset();
 
+	UE_LOG(LogGatherTextFromAssetsCommandlet, Display, TEXT("0088:06 %.2f seconds."), FPlatformTime::Seconds() - TimePrev);
+	TimePrev = FPlatformTime::Seconds();
+
 	// All packages left in PackagesPendingGather should now have to be loaded 
 	if (PackagesPendingGather.Num() == 0)
 	{
@@ -1456,10 +1476,16 @@ int32 UGatherTextFromAssetsCommandlet::Main(const FString& Params)
 
 	CalculateDependenciesForPackagesPendingGather();
 
+	UE_LOG(LogGatherTextFromAssetsCommandlet, Display, TEXT("0088:07 %.2f seconds."), FPlatformTime::Seconds() - TimePrev);
+	TimePrev = FPlatformTime::Seconds();
+
 	// Collect garbage before beginning to load packages
 	// This also sorts the list of packages into the best processing order
 	PurgeGarbage(/*bPurgeReferencedPackages*/false);
 
+	UE_LOG(LogGatherTextFromAssetsCommandlet, Display, TEXT("0088:08 %.2f seconds."), FPlatformTime::Seconds() - TimePrev);
+	TimePrev = FPlatformTime::Seconds();
+
 	// We don't need to have compiled shaders to gather text
 	bool bWasShaderCompilationEnabled = false;
 	if (GShaderCompilingManager)
@@ -1468,15 +1494,24 @@ int32 UGatherTextFromAssetsCommandlet::Main(const FString& Params)
 		GShaderCompilingManager->SkipShaderCompilation(true);
 	}
 
+	UE_LOG(LogGatherTextFromAssetsCommandlet, Display, TEXT("0088:09 %.2f seconds."), FPlatformTime::Seconds() - TimePrev);
+	TimePrev = FPlatformTime::Seconds();
+
 	TArray<FName> PackagesWithStaleGatherCache;
 	LoadAndProcessUncachedPackages(PackagesWithStaleGatherCache);
 
+	UE_LOG(LogGatherTextFromAssetsCommandlet, Display, TEXT("0088:10 %.2f seconds."), FPlatformTime::Seconds() - TimePrev);
+	TimePrev = FPlatformTime::Seconds();
+
 	UE_LOG(LogGatherTextFromAssetsCommandlet, Display, TEXT("Loading packages took %.2f seconds."), FPlatformTime::Seconds() - PackageLoadingStartTime);
 
 	// Collect garbage after loading all packages
 	// This reclaims as much memory as possible for the rest of the gather pipeline
 	PurgeGarbage(/*bPurgeReferencedPackages*/true);
 	
+	UE_LOG(LogGatherTextFromAssetsCommandlet, Display, TEXT("0088:11 %.2f seconds."), FPlatformTime::Seconds() - TimePrev);
+	TimePrev = FPlatformTime::Seconds();
+
 	if (GShaderCompilingManager)
 	{
 		GShaderCompilingManager->SkipShaderCompilation(!bWasShaderCompilationEnabled);
@@ -1487,6 +1522,9 @@ int32 UGatherTextFromAssetsCommandlet::Main(const FString& Params)
 		ReportStaleGatherCache(PackagesWithStaleGatherCache);
 	}
 
+	UE_LOG(LogGatherTextFromAssetsCommandlet, Display, TEXT("0088:12 %.2f seconds."), FPlatformTime::Seconds() - TimePrev);
+	TimePrev = FPlatformTime::Seconds();
+
 	return 0;
 }
 
