Cooker: CookSkipRequests: Turn off StartupPackages, TargetPlatform->GetExtraPackagesToCook and bCookAllByDefault when !NoGameAlwaysCookPackages && !StartupPackages.
Cooker: -package flag: Support localpaths for -package.
#rnx

diff --git a/Engine/Source/Editor/UnrealEd/Classes/CookOnTheSide/CookOnTheFlyServer.h b/Engine/Source/Editor/UnrealEd/Classes/CookOnTheSide/CookOnTheFlyServer.h
--- a/Engine/Source/Editor/UnrealEd/Classes/CookOnTheSide/CookOnTheFlyServer.h
+++ b/Engine/Source/Editor/UnrealEd/Classes/CookOnTheSide/CookOnTheFlyServer.h
@@ -984,6 +984,9 @@ private:
 	/** AddFileToCook add file to cook list */
 	void AddFileToCook( TArray<FName>& InOutFilesToCook, TMap<FName, UE::Cook::FInstigator>& InOutInstigators,
 		const FString &InFilename, const UE::Cook::FInstigator& Instigator) const;
+	/** AddFileToCook add file to cook list */
+	void AddFlexPathToCook(TArray<FName>& InOutFilesToCook, TMap<FName, UE::Cook::FInstigator>& InOutInstigators,
+		const FString& InFlexPath, const UE::Cook::FInstigator& Instigator) const;
 
 	/** Return the name to use for the project's global shader library */
 	FString GetProjectShaderLibraryName() const;
diff --git a/Engine/Source/Editor/UnrealEd/Private/CookOnTheFlyServer.cpp b/Engine/Source/Editor/UnrealEd/Private/CookOnTheFlyServer.cpp
--- a/Engine/Source/Editor/UnrealEd/Private/CookOnTheFlyServer.cpp
+++ b/Engine/Source/Editor/UnrealEd/Private/CookOnTheFlyServer.cpp
@@ -9140,6 +9140,65 @@ void UCookOnTheFlyServer::GenerateLongPackageNames(TArray<FName>& FilesInPath, T
 	Swap(Instigators, NewInstigators);
 }
 
+void UCookOnTheFlyServer::AddFlexPathToCook(TArray<FName>& InOutFilesToCook,
+	TMap<FName, UE::Cook::FInstigator>& InOutInstigators,
+	const FString& InFlexPath, const UE::Cook::FInstigator& Instigator) const
+{
+	using namespace UE::Cook;
+
+	FString FlexPath(InFlexPath);
+	// Convert \ to / so that IsShortPackageName works.
+	// We can still interpret the path as a filepath even with \ converted to /
+	FlexPath.ReplaceCharInline('\\', '/');
+	if (FPackageName::IsShortPackageName(FlexPath))
+	{
+		TArray<FName> LongPackageNames;
+		AssetRegistry->GetPackagesByName(FlexPath, LongPackageNames);
+		if (LongPackageNames.IsEmpty())
+		{
+			LogCookerMessage(FString::Printf(TEXT("Unable to find package for path `%s`."), *InFlexPath),
+				EMessageSeverity::Warning);
+		}
+		else if (LongPackageNames.Num() > 1)
+		{
+			constexpr int32 MaxMessageLen = 256;
+			TStringBuilder<256> Message;
+			Message.Appendf(
+				TEXT("Multiple packages found for path `%s`; it will not be added. Specify the full LongPackageName. Packages found:"),
+				*InFlexPath);
+			for (FName LongPackageName : LongPackageNames)
+			{
+				Message << TEXT("\n\t");
+				if (Message.Len() >= MaxMessageLen)
+				{
+					Message << TEXT("...");
+					break;
+				}
+				else
+				{
+					Message << LongPackageName;
+				}
+			}
+			LogCookerMessage(FString(*Message), EMessageSeverity::Warning);
+		}
+		else
+		{
+			AddFileToCook(InOutFilesToCook, InOutInstigators, LongPackageNames[0].ToString(), Instigator);
+		}
+	}
+	else
+	{
+		FString PackageName;
+		if (!FPackageName::TryConvertFilenameToLongPackageName(FlexPath, PackageName))
+		{
+			LogCookerMessage(FString::Printf(TEXT("Unable to find package for path `%s`."), *InFlexPath),
+				EMessageSeverity::Warning);
+			return;
+		}
+		AddFileToCook(InOutFilesToCook, InOutInstigators, PackageName, Instigator);
+	}
+}
+
 void UCookOnTheFlyServer::AddFileToCook( TArray<FName>& InOutFilesToCook,
 	TMap<FName, UE::Cook::FInstigator>& InOutInstigators,
 	const FString &InFilename, const UE::Cook::FInstigator& Instigator) const
@@ -9253,7 +9312,7 @@ void UCookOnTheFlyServer::CollectFilesToCook(TArray<FName>& FilesInPath, TMap<FN
 
 			for (int32 MapIdx = 0; MapIdx < MapList.Num(); MapIdx++)
 			{
-				UE_LOG(LogCook, Verbose, TEXT("Maplist contains has %s "), *MapList[MapIdx]);
+				UE_LOG(LogCook, Verbose, TEXT("Maplist contains %s "), *MapList[MapIdx]);
 				AddFileToCook(FilesInPath, Instigators, MapList[MapIdx], EInstigator::AlwaysCookMap);
 			}
 		}
@@ -9388,46 +9447,9 @@ void UCookOnTheFlyServer::CollectFilesToCook(TArray<FName>& FilesInPath, TMap<FN
 		}
 	}
 
-	for ( const FString& CurrEntry : CookMaps )
+	for (const FString& CurrEntry : CookMaps)
 	{
-		UE_SCOPED_HIERARCHICAL_COOKTIMER(SearchForPackageOnDisk);
-		if (FPackageName::IsShortPackageName(CurrEntry))
-		{
-			TArray<FName> LongPackageNames;
-			AssetRegistry->GetPackagesByName(CurrEntry, LongPackageNames);
-			if (LongPackageNames.IsEmpty())
-			{
-				LogCookerMessage(FString::Printf(TEXT("Unable to find package for map %s."), *CurrEntry), EMessageSeverity::Warning);
-			}
-			else if (LongPackageNames.Num() > 1)
-			{
-				constexpr int32 MaxMessageLen = 256;
-				TStringBuilder<256> Message;
-				Message.Appendf(TEXT("Multiple packages found for map %s; it will not be added. Specify the full LongPackageName. Packages found:"), *CurrEntry);
-				for (FName LongPackageName : LongPackageNames)
-				{
-					Message << TEXT("\n\t");
-					if (Message.Len() >= MaxMessageLen)
-					{
-						Message << TEXT("...");
-						break;
-					}
-					else
-					{
-						Message << LongPackageName;
-					}
-				}
-				LogCookerMessage(FString(*Message), EMessageSeverity::Warning);
-			}
-			else
-			{
-				AddFileToCook(FilesInPath, Instigators, LongPackageNames[0].ToString(), EInstigator::CommandLinePackage);
-			}
-		}
-		else
-		{
-			AddFileToCook(FilesInPath, Instigators, CurrEntry, EInstigator::CommandLinePackage);
-		}
+		AddFlexPathToCook(FilesInPath, Instigators, CurrEntry, EInstigator::CommandLinePackage);
 	}
 	if (IsCookingDLC())
 	{
@@ -9445,14 +9467,18 @@ void UCookOnTheFlyServer::CollectFilesToCook(TArray<FName>& FilesInPath, TMap<FN
 		}
 	}
 
-	if (!(FilesToCookFlags & ECookByTheBookOptions::SkipSoftReferences))
+	if (!(FilesToCookFlags & ECookByTheBookOptions::SkipSoftReferences)
+		&& !(FilesToCookFlags & ECookByTheBookOptions::NoGameAlwaysCookPackages))
 	{
 		for (const ITargetPlatform* TargetPlatform : TargetPlatforms)
 		{
 			TargetPlatform->GetExtraPackagesToCook(FilesInPath);
 		}
 		UpdateInstigators(EInstigator::TargetPlatformExtraPackagesToCook);
+	}
 
+	if (!(FilesToCookFlags & ECookByTheBookOptions::SkipSoftReferences))
+	{
 		const FString ExternalMountPointName(TEXT("/Game/"));
 		for (const FNameWithInstigator& CurrEntry : CookDirectories)
 		{
@@ -9473,8 +9499,12 @@ void UCookOnTheFlyServer::CollectFilesToCook(TArray<FName>& FilesInPath, TMap<FN
 				}
 			}
 		}
+	}
 
-		// If no packages were explicitly added by command line or game callback, add all maps
+	// If no packages were explicitly added by command line or game callback, add all maps
+	if (!(FilesToCookFlags & ECookByTheBookOptions::SkipSoftReferences)
+		&& !(FilesToCookFlags & ECookByTheBookOptions::NoGameAlwaysCookPackages))
+	{
 		if (bCookAll || (UE::Cook::bCookAllByDefault && FilesInPath.Num() == InitialNum))
 		{
 			TArray<FString> Tokens;
@@ -12264,25 +12294,31 @@ void UCookOnTheFlyServer::GenerateInitialRequests(FBeginCookContext& BeginContex
 	CollectFilesToCook(FilesInPath, FilesInPathInstigators, CookMaps, CookDirectories, IniMapSections, CookOptions, TargetPlatforms, GameDefaultObjects);
 
 	// Add soft/hard startup references after collecting requested files and handling empty requests
+	FlushAsyncLoading();
 	if (!CookByTheBookOptions->bSkipHardReferences && !EnumHasAnyFlags(CookOptions, ECookByTheBookOptions::NoStartupPackages))
 	{
 		ProcessUnsolicitedPackages(&FilesInPath, &FilesInPathInstigators);
 	}
-	for (FName SoftObjectPackage : StartupSoftObjectPackages)
+	else
 	{
-		TMap<FSoftObjectPath, FSoftObjectPath> RedirectedPaths;
-
-		// If this is a redirector, extract destination from asset registry
-		if (ContainsRedirector(SoftObjectPackage, RedirectedPaths))
+		// Clear the list of startup packages currently held by the packagetracker so that we don't see them when we
+		// ProcessUnsolicitedPackages to find hard references used by the first requested package we load.
+		(void)PackageTracker->GetNewPackages();
+	}
+	if (!CookByTheBookOptions->bSkipSoftReferences && !EnumHasAnyFlags(CookOptions, ECookByTheBookOptions::NoStartupPackages))
+	{
+		for (FName SoftObjectPackage : StartupSoftObjectPackages)
 		{
-			for (TPair<FSoftObjectPath, FSoftObjectPath>& RedirectedPath : RedirectedPaths)
+			TMap<FSoftObjectPath, FSoftObjectPath> RedirectedPaths;
+
+			// If this is a redirector, extract destination from asset registry
+			if (ContainsRedirector(SoftObjectPackage, RedirectedPaths))
 			{
-				GRedirectCollector.AddAssetPathRedirection(RedirectedPath.Key, RedirectedPath.Value);
+				for (TPair<FSoftObjectPath, FSoftObjectPath>& RedirectedPath : RedirectedPaths)
+				{
+					GRedirectCollector.AddAssetPathRedirection(RedirectedPath.Key, RedirectedPath.Value);
+				}
 			}
-		}
-
-		if (!CookByTheBookOptions->bSkipSoftReferences)
-		{
 			AddFileToCook(FilesInPath, FilesInPathInstigators, SoftObjectPackage.ToString(),
 				UE::Cook::EInstigator::StartupSoftObjectPath);
 		}
