Remove duplicated cook stats in cook stats file

The last change I made contained a duplicate serialization of the misc cook stats, which is of course unnecessary.

#rnx

diff --git a/Engine/Source/Editor/UnrealEd/Private/Cooker/CookProfiling.cpp b/Engine/Source/Editor/UnrealEd/Private/Cooker/CookProfiling.cpp
--- a/Engine/Source/Editor/UnrealEd/Private/Cooker/CookProfiling.cpp
+++ b/Engine/Source/Editor/UnrealEd/Private/Cooker/CookProfiling.cpp
@@ -1016,22 +1016,6 @@ void SendLogCookStats(ECookMode::Type CookMode)
 		}
 	}
 
-	if (CookStatsWriter)
-	{
-		CookStatsWriter->WriteObjectStart(TEXT("MiscCookStats"));
-		for (FString& StatCategory : StatCategories)
-		{
-			CookStatsWriter->WriteObjectStart(*StatCategory);
-			TArray<FCookStatsManager::StringKeyValue>& StatsInCategory = StatsInCategories.FindOrAdd(StatCategory);
-			for (const FCookStatsManager::StringKeyValue& StatKeyValue : StatsInCategory)
-			{
-				CookStatsWriter->WriteValue(*StatKeyValue.Key, *StatKeyValue.Value);
-			}
-			CookStatsWriter->WriteObjectEnd();
-		}
-		CookStatsWriter->WriteObjectEnd();
-	}
-
 	// DDC Usage stats are custom formatted, and the above code just accumulated them into a TSet. Now log it with our special formatting for readability.
 	if (CookProfileData.Num() > 0)
 	{
