Fix Deleting Assets progress bar exceeding 100%

diff --git a/Engine/Source/Editor/UnrealEd/Private/ObjectTools.cpp b/Engine/Source/Editor/UnrealEd/Private/ObjectTools.cpp
--- a/Engine/Source/Editor/UnrealEd/Private/ObjectTools.cpp
+++ b/Engine/Source/Editor/UnrealEd/Private/ObjectTools.cpp
@@ -3545,7 +3545,6 @@ namespace ObjectTools
 
 		TArray<UPackage*> PackagesFailedToDelete;
 		{
-			int32 ReplaceableObjectsNum = 0;
 			{
 				for(TWeakObjectPtr<UObject>& Object : ObjectsToDelete)
 				{
@@ -3619,14 +3618,12 @@ namespace ObjectTools
 					{
 						FForceReplaceInfo ReplaceInfo;
 						ForceReplaceReferences(GetFallbackStruct(), UDStructToReplace, ReplaceInfo, false);
-						ReplaceableObjectsNum += ReplaceInfo.ReplaceableObjects.Num();
 					}
 				}
 
 				{
 					FForceReplaceInfo ReplaceInfo;
 					ForceReplaceReferences(nullptr, ObjectsToReplace, ReplaceInfo, false);
-					ReplaceableObjectsNum += ReplaceInfo.ReplaceableObjects.Num();
 				}
 			}
 
@@ -3644,6 +3641,7 @@ namespace ObjectTools
 			FAssetToolsModule& AssetToolsModule = FModuleManager::LoadModuleChecked<FAssetToolsModule>(TEXT("AssetTools"));
 
 			int32 Count = 0;
+			const int32 OriginalNumObjectsToDelete = ObjectsToDelete.Num();
 			for(auto It = ObjectsToDelete.CreateIterator(); It; ++It)
 			{
 				UObject* CurObject = It->Get();
@@ -3669,7 +3667,7 @@ namespace ObjectTools
 					It.RemoveCurrent();
 				}
 
-				GWarn->StatusUpdate(Count, ReplaceableObjectsNum, NSLOCTEXT("UnrealEd", "ConsolidateAssetsUpdate_DeletingObjects", "Deleting Assets..."));
+				GWarn->StatusUpdate(Count, OriginalNumObjectsToDelete, NSLOCTEXT("UnrealEd", "ConsolidateAssetsUpdate_DeletingObjects", "Deleting Assets..."));
 				++Count;
 
 			}
