When building WP navmesh and deleting unnecessary package, skip packages already deleted to avoid "Error deleting packages" error.

diff --git a/Engine/Source/Editor/UnrealEd/Private/WorldPartition/WorldPartitionNavigationDataBuilder.cpp b/Engine/Source/Editor/UnrealEd/Private/WorldPartition/WorldPartitionNavigationDataBuilder.cpp
--- a/Engine/Source/Editor/UnrealEd/Private/WorldPartition/WorldPartitionNavigationDataBuilder.cpp
+++ b/Engine/Source/Editor/UnrealEd/Private/WorldPartition/WorldPartitionNavigationDataBuilder.cpp
@@ -92,7 +92,19 @@ bool UWorldPartitionNavigationDataBuilder::RunInternal(UWorld* World, const FCel
 		{
 			UPackage* Package = ItActor->GetPackage();
 			check(!PackagesToClean.Find(Package));
-			PackagesToClean.Add(Package);
+
+			if (Package)
+			{
+				const FString Filename = SourceControlHelpers::PackageFilename(Package->GetName());
+				if(IPlatformFile::GetPlatformPhysical().FileExists(*Filename))
+				{
+					PackagesToClean.Add(Package);
+				}
+				else
+				{
+					UE_LOG(LogWorldPartitionNavigationDataBuilder, Verbose, TEXT("   Skippping package file already deleted: %s"), *Package->GetName());
+				}
+			}
 		}
 
 		UE_LOG(LogWorldPartitionNavigationDataBuilder, Verbose, TEXT("   Number of packages to clear: %i"), PackagesToClean.Num());
@@ -185,7 +197,15 @@ bool UWorldPartitionNavigationDataBuilder::RunInternal(UWorld* World, const FCel
 		{
 			if (UPackage::IsEmptyPackage(ActorPackage))
 			{
-				PackagesToDelete.Add(ActorPackage);
+				const FString Filename = SourceControlHelpers::PackageFilename(ActorPackage->GetName());
+				if(IPlatformFile::GetPlatformPhysical().FileExists(*Filename))
+				{
+					PackagesToDelete.Add(ActorPackage);
+				}
+				else
+				{
+					UE_LOG(LogWorldPartitionNavigationDataBuilder, Verbose, TEXT("   Skippping package file already deleted: %s"), *ActorPackage->GetName());
+				}
 			}
 			else
 			{
