Add a new cook profile timer for flushing the shader queue

Flushing the shader queue happens under high memory pressure and is hence a prime value to track to make informed choices about memory limits vs. cook times. It currently is not instrumented in a way that the content worker can access (the content worker can only look at the cook stats json output). This CL adds new instrumentation.

#rnx

diff --git a/Engine/Source/Editor/UnrealEd/Private/Commandlets/CookCommandlet.cpp b/Engine/Source/Editor/UnrealEd/Private/Commandlets/CookCommandlet.cpp
--- a/Engine/Source/Editor/UnrealEd/Private/Commandlets/CookCommandlet.cpp
+++ b/Engine/Source/Editor/UnrealEd/Private/Commandlets/CookCommandlet.cpp
@@ -671,7 +671,7 @@ void UCookCommandlet::ConditionalCollectGarbage(uint32 TickResults, UCookOnTheFl
 
 		int32 JobsToLogAt = GShaderCompilingManager->GetNumRemainingJobs();
 		double NextFlushMsgSeconds = FPlatformTime::Seconds();
-		UE_SCOPED_COOKTIMER(CookByTheBook_ShaderJobFlush);
+		UE_SCOPED_COOKTIMER_AND_DURATION(CookByTheBook_ShaderJobFlush, DetailedCookStats::ShaderFlushTimeSec);
 		UE_LOG(LogCookCommandlet, Display, TEXT("Detected max mem exceeded - forcing shader compilation flush"));
 		while (true)
 		{
@@ -714,14 +714,7 @@ void UCookCommandlet::ConditionalCollectGarbage(uint32 TickResults, UCookOnTheFl
 		UE_SCOPED_COOKTIMER(CookByTheBook_TickAssetRegistry);
 		FAssetRegistryModule::TickAssetRegistry(-1.0f);
 	}
-#if OUTPUT_COOKTIMING
-	TOptional<FScopedDurationTimer> CBTBScopedDurationTimer;
-	if (!COTFS.IsCookOnTheFlyMode())
-	{
-		CBTBScopedDurationTimer.Emplace(DetailedCookStats::TickLoopGCTimeSec);
-	}
-#endif
-	UE_SCOPED_COOKTIMER(CookCommandlet_GC);
+	UE_SCOPED_COOKTIMER_AND_DURATION(CookCommandlet_GC, DetailedCookStats::TickLoopGCTimeSec);
 
 	const FPlatformMemoryStats MemStatsBeforeGC = FPlatformMemory::GetStats();
 	int32 NumObjectsBeforeGC = GUObjectArray.GetObjectArrayNumMinusAvailable();
diff --git a/Engine/Source/Editor/UnrealEd/Private/Cooker/CookProfiling.cpp b/Engine/Source/Editor/UnrealEd/Private/Cooker/CookProfiling.cpp
--- a/Engine/Source/Editor/UnrealEd/Private/Cooker/CookProfiling.cpp
+++ b/Engine/Source/Editor/UnrealEd/Private/Cooker/CookProfiling.cpp
@@ -235,6 +235,7 @@ namespace DetailedCookStats
 	double TickLoopProcessDeferredCommandsTimeSec = 0.0;
 	double TickLoopTickCommandletStatsTimeSec = 0.0;
 	double TickLoopFlushRenderingCommandsTimeSec = 0.0;
+	double ShaderFlushTimeSec = 0.0;
 	bool IsCookAll = false;
 	bool IsCookOnTheFly = false;
 	bool IsIterativeCook = false;
@@ -887,6 +888,7 @@ FCookStatsManager::FAutoRegisterCallback RegisterCookStats([](FCookStatsManager:
 	ADD_COOK_STAT_FLT(" 0. 6", TickLoopProcessDeferredCommandsTimeSec);
 	ADD_COOK_STAT_FLT(" 0. 7", TickLoopTickCommandletStatsTimeSec);
 	ADD_COOK_STAT_FLT(" 0. 8", TickLoopFlushRenderingCommandsTimeSec);
+	ADD_COOK_STAT_FLT(" 0. 9", ShaderFlushTimeSec);
 	FString CookParameters; // Empty value to write a header with name "CookParameters"
 	ADD_COOK_STAT_FLT(" 1", CookParameters);
 	ADD_COOK_STAT_FLT(" 1. 0", TargetPlatforms);
diff --git a/Engine/Source/Editor/UnrealEd/Private/Cooker/CookProfiling.h b/Engine/Source/Editor/UnrealEd/Private/Cooker/CookProfiling.h
--- a/Engine/Source/Editor/UnrealEd/Private/Cooker/CookProfiling.h
+++ b/Engine/Source/Editor/UnrealEd/Private/Cooker/CookProfiling.h
@@ -148,6 +148,7 @@ extern double TickLoopShaderProcessAsyncResultsTimeSec;
 extern double TickLoopProcessDeferredCommandsTimeSec;
 extern double TickLoopTickCommandletStatsTimeSec;
 extern double TickLoopFlushRenderingCommandsTimeSec;
+extern double ShaderFlushTimeSec;
 extern bool IsCookAll;
 extern bool IsCookOnTheFly;
 extern bool IsIterativeCook;
