Expose validation times accumulated during a cook

#rnx

diff --git a/Engine/Source/Editor/UnrealEd/Private/CookOnTheFlyServer.cpp b/Engine/Source/Editor/UnrealEd/Private/CookOnTheFlyServer.cpp
--- a/Engine/Source/Editor/UnrealEd/Private/CookOnTheFlyServer.cpp
+++ b/Engine/Source/Editor/UnrealEd/Private/CookOnTheFlyServer.cpp
@@ -3059,7 +3059,7 @@ void UCookOnTheFlyServer::RejectPackageToLoad(UE::Cook::FPackageData& PackageDat
 
 EDataValidationResult UCookOnTheFlyServer::ValidateSourcePackage(UE::Cook::FPackageData& PackageData, UPackage* Package)
 {
-	UE_SCOPED_HIERARCHICAL_COOKTIMER(ValidateSourcePackage);
+	UE_SCOPED_HIERARCHICAL_COOKTIMER_AND_DURATION(ValidateSourcePackage, DetailedCookStats::ValidationTimeSec);
 
 	// Don't validate packages if validation is disabled
 	if (!EnumHasAnyFlags(CookByTheBookOptions->StartupOptions, ECookByTheBookOptions::RunAssetValidation | ECookByTheBookOptions::RunMapValidation))
diff --git a/Engine/Source/Editor/UnrealEd/Private/Cooker/CookProfiling.cpp b/Engine/Source/Editor/UnrealEd/Private/Cooker/CookProfiling.cpp
--- a/Engine/Source/Editor/UnrealEd/Private/Cooker/CookProfiling.cpp
+++ b/Engine/Source/Editor/UnrealEd/Private/Cooker/CookProfiling.cpp
@@ -236,6 +236,7 @@ namespace DetailedCookStats
 	double TickLoopTickCommandletStatsTimeSec = 0.0;
 	double TickLoopFlushRenderingCommandsTimeSec = 0.0;
 	double ShaderFlushTimeSec = 0.0;
+	double ValidationTimeSec = 0.0;
 	bool IsCookAll = false;
 	bool IsCookOnTheFly = false;
 	bool IsIterativeCook = false;
@@ -889,6 +890,7 @@ FCookStatsManager::FAutoRegisterCallback RegisterCookStats([](FCookStatsManager:
 	ADD_COOK_STAT_FLT(" 0. 7", TickLoopTickCommandletStatsTimeSec);
 	ADD_COOK_STAT_FLT(" 0. 8", TickLoopFlushRenderingCommandsTimeSec);
 	ADD_COOK_STAT_FLT(" 0. 9", ShaderFlushTimeSec);
+	ADD_COOK_STAT_FLT(" 0. 10", ValidationTimeSec);
 	FString CookParameters; // Empty value to write a header with name "CookParameters"
 	ADD_COOK_STAT_FLT(" 1", CookParameters);
 	ADD_COOK_STAT_FLT(" 1. 0", TargetPlatforms);
diff --git a/Engine/Source/Editor/UnrealEd/Private/Cooker/CookProfiling.h b/Engine/Source/Editor/UnrealEd/Private/Cooker/CookProfiling.h
--- a/Engine/Source/Editor/UnrealEd/Private/Cooker/CookProfiling.h
+++ b/Engine/Source/Editor/UnrealEd/Private/Cooker/CookProfiling.h
@@ -149,6 +149,7 @@ extern double TickLoopProcessDeferredCommandsTimeSec;
 extern double TickLoopTickCommandletStatsTimeSec;
 extern double TickLoopFlushRenderingCommandsTimeSec;
 extern double ShaderFlushTimeSec;
+extern double ValidationTimeSec;
 extern bool IsCookAll;
 extern bool IsCookOnTheFly;
 extern bool IsIterativeCook;
