MPCook: CookPackageSplitter: Fix for error "GenerationHelper for package %s is still allocated at end of cooksession...."

Fix a missing record that a generated package has completed its cook in some cases where its cook was suppressed. Suppressed packages were reported if suppressed on remote CookWorkers, or if suppressed when in the save state on the CookDirector, but if suppressed when in load state or earlier on the CookDirector they were not reported. Add code to report them as completed whenever their ParentGenerationHelper pointer is removed, which happens when they get demoted to idle.

Also add storage of the suppress cook reason so that diagnostics that later discover they were not cooked can report why they were not cooked.

#rnx

diff --git a/Engine/Source/Editor/UnrealEd/Private/CookOnTheFlyServer.cpp b/Engine/Source/Editor/UnrealEd/Private/CookOnTheFlyServer.cpp
--- a/Engine/Source/Editor/UnrealEd/Private/CookOnTheFlyServer.cpp
+++ b/Engine/Source/Editor/UnrealEd/Private/CookOnTheFlyServer.cpp
@@ -2779,12 +2779,16 @@ void UCookOnTheFlyServer::DemoteToIdle(UE::Cook::FPackageData& PackageData, UE::
 	{
 		WorkerRequests->ReportDemoteToIdle(PackageData, Reason);
 
+		bool bHasCookResult = PackageData.HasAllCookedPlatforms(PlatformManager->GetSessionPlatforms(), true /* bIncludeFailed */);
+
 		// If per-package display is on, write a log statement explaining that the package was reachable but skipped.
 		bool bPrintDiagnostic = !bCookListMode &
 			!!((GCookProgressDisplay & ((int32)ECookProgressDisplayMode::Instigators | (int32)ECookProgressDisplayMode::PackageNames)));
 
 		// Suppress the message in cases that cause large spam like NotInCurrentPlugin for DLC cooks.
 		bPrintDiagnostic &= (Reason != ESuppressCookReason::NotInCurrentPlugin);
+		// Iterative cooks: suppress the diagnostic for packages that were iteratively skipped
+		bPrintDiagnostic &= !bHasCookResult;
 		if (bPrintDiagnostic && IsCookingDLC() && Reason == ESuppressCookReason::AlreadyCooked && LogCook.GetVerbosity() < ELogVerbosity::Verbose)
 		{
 			bPrintDiagnostic = false;
@@ -2804,9 +2808,6 @@ void UCookOnTheFlyServer::DemoteToIdle(UE::Cook::FPackageData& PackageData, UE::
 			// Reachability: Suppress the diagnostic that were found via cookload reference traversal but are not reachable on the target platforms
 			bPrintDiagnostic &= (PackageData.HasInstigator() || Reason != ESuppressCookReason::OnlyEditorOnly);
 
-			// Iterative cooks: suppress the diagnostic for packages that were iteratively skipped
-			bPrintDiagnostic &= !PackageData.HasAllCookedPlatforms(PlatformManager->GetSessionPlatforms(), true /* bIncludeFailed */);
-
 			if (bPrintDiagnostic)
 			{
 				UE_CLOG((GCookProgressDisplay & (int32)ECookProgressDisplayMode::Instigators), LogCook, Display,
@@ -2816,6 +2817,12 @@ void UCookOnTheFlyServer::DemoteToIdle(UE::Cook::FPackageData& PackageData, UE::
 					TEXT("Cooking %s -> Rejected %s"), *PackageNameStr, LexToString(Reason));
 			}
 		}
+
+		// If the package is demoted without a cookresult, store the suppressed reason
+		if (!bHasCookResult)
+		{
+			PackageData.SetSuppressCookReason(Reason);
+		}
 	}
 	PackageData.SendToState(UE::Cook::EPackageState::Idle, SendFlags, ConvertToStateChangeReason(Reason));
 }
@@ -10505,7 +10512,7 @@ void UCookOnTheFlyServer::CookByTheBookFinishedInternal()
 	PackageDatas->LockAndEnumeratePackageDatas([&DanglingGenerationHelpers](FPackageData* PackageData)
 		{
 			TRefCountPtr<FGenerationHelper> GenerationHelper = PackageData->GetGenerationHelper();
-			if (GenerationHelper && GenerationHelper->IsInitialized())
+			if (GenerationHelper)
 			{
 				DanglingGenerationHelpers.Add(PackageData);
 			}
@@ -10516,7 +10523,10 @@ void UCookOnTheFlyServer::CookByTheBookFinishedInternal()
 		if (GenerationHelper)
 		{
 			GenerationHelper->DiagnoseWhyNotShutdown();
-			GenerationHelper->ForceUninitialize();
+			if (GenerationHelper->IsInitialized())
+			{
+				GenerationHelper->ForceUninitialize();
+			}
 		}
 	};
 
diff --git a/Engine/Source/Editor/UnrealEd/Private/Cooker/CookGenerationHelper.cpp b/Engine/Source/Editor/UnrealEd/Private/Cooker/CookGenerationHelper.cpp
--- a/Engine/Source/Editor/UnrealEd/Private/Cooker/CookGenerationHelper.cpp
+++ b/Engine/Source/Editor/UnrealEd/Private/Cooker/CookGenerationHelper.cpp
@@ -10,6 +10,7 @@
 #include "Cooker/CookWorkerServer.h"
 #include "Cooker/IWorkerRequests.h"
 #include "Cooker/PackageTracker.h"
+#include "Interfaces/ITargetPlatform.h"
 #include "Misc/CommandLine.h"
 #include "Misc/ConfigCacheIni.h"
 #include "Misc/PackageAccessTrackingOps.h"
@@ -259,20 +260,34 @@ void FGenerationHelper::DiagnoseWhyNotShutdown()
 			}
 			else
 			{
-				bool bMissingPlatform = false;
+				TArray< const ITargetPlatform*> MissingPlatforms;
 				for (const ITargetPlatform* TargetPlatform : COTFS.PlatformManager->GetSessionPlatforms())
 				{
 					const FPackagePlatformData* PlatformData = Info.PackageData->GetPlatformDatas().Find(TargetPlatform);
 					if (!PlatformData || PlatformData->GetCookResults() == ECookResult::NotAttempted)
 					{
-						bMissingPlatform = true;
+						MissingPlatforms.Add(TargetPlatform);
 					}
 				}
-				if (bMissingPlatform)
+				if (!MissingPlatforms.IsEmpty())
 				{
-					Lines.Appendf(TEXT("\t%s%s was not cooked.\n"),
+					TStringBuilder<256> MissingPlatformStr;
+					if (MissingPlatforms.Num() != COTFS.PlatformManager->GetSessionPlatforms().Num())
+					{
+						MissingPlatformStr << TEXT(" for platforms { ");
+						for (const ITargetPlatform* TargetPlatform : MissingPlatforms)
+						{
+							MissingPlatformStr << TargetPlatform->PlatformName() << TEXT(", ");
+						}
+						MissingPlatformStr.RemoveSuffix(2);
+						MissingPlatformStr << TEXT(" }");
+					}
+
+					Lines.Appendf(TEXT("\t%s%s was not cooked%s. SuppressCookReason == %s.\n"),
 						Info.IsGenerator() ? TEXT("OwnerInfo") : TEXT("GeneratedPackage "),
-						Info.IsGenerator() ? TEXT("") : *Info.GetPackageName());
+						Info.IsGenerator() ? TEXT("") : *Info.GetPackageName(),
+						*MissingPlatformStr,
+						LexToString(Info.PackageData->GetSuppressCookReason()));
 				}
 			}
 			if (!Info.HasSaved())
@@ -363,9 +378,18 @@ void FGenerationHelper::DiagnoseWhyNotShutdown()
 		Lines.RemoveSuffix(1);
 	}
 
-	UE_LOG(LogCook, Error,
-		TEXT("GenerationHelper for package %s is still allocated at end of cooksession. This is unexpected and could indicate some generated packages are missing."),
-		*GetOwner().GetPackageName().ToString());
+	FString Message = FString::Printf(
+		TEXT("GenerationHelper for package %s is still allocated%s at end of cooksession. This is unexpected and could indicate some generated packages are missing."),
+		*GetOwner().GetPackageName().ToString(), IsInitialized() ? TEXT(" and initialized") : TEXT(""));
+
+	if (IsInitialized())
+	{
+		UE_LOG(LogCook, Error, TEXT("%s"), *Message);
+	}
+	else
+	{
+		UE_LOG(LogCook, Warning, TEXT("%s"), *Message);
+	}
 	UE_LOG(LogCook, Display, TEXT("Diagnostics:\n%s"), *Lines);
 }
 
@@ -1268,12 +1292,13 @@ void FGenerationHelper::ResetSaveState(FCookGenerationInfo& Info, UPackage* Pack
 
 	if (IsTerminalStateChange(ReleaseSaveReason))
 	{
-		// The save is completed and we will not come back to it; set state back to initial
-		// state and drop our reference keeping this GenerationHelper in memory for the save.
-		Info.SetHasSaved(*this, true, FWorkerId::Local());
-
+		// The package's progress is completed and we will not come back to it; set state back to initial
+		// state, mark the package as saved in our GenerationHelper data, and drop the ParentGenerationHelper
+		// reference.
 		if (Info.IsGenerator())
 		{
+			Info.SetHasSaved(*this, true, FWorkerId::Local());
+
 			// Now that we've finished saving, we know that we will not call QueueGeneratedPackages again, so we can
 			// teardown iterative results as well
 			ClearKeepForIterative();
@@ -1284,7 +1309,8 @@ void FGenerationHelper::ResetSaveState(FCookGenerationInfo& Info, UPackage* Pack
 		}
 		else
 		{
-			Info.PackageData->SetParentGenerationHelper(nullptr);
+			// For generated packages, SetHasSaved is called inside of SetParentGenerationHelper
+			Info.PackageData->SetParentGenerationHelper(nullptr, ReleaseSaveReason, &Info);
 		}
 	}
 
diff --git a/Engine/Source/Editor/UnrealEd/Private/Cooker/CookPackageData.cpp b/Engine/Source/Editor/UnrealEd/Private/Cooker/CookPackageData.cpp
--- a/Engine/Source/Editor/UnrealEd/Private/Cooker/CookPackageData.cpp
+++ b/Engine/Source/Editor/UnrealEd/Private/Cooker/CookPackageData.cpp
@@ -103,6 +103,7 @@ FPackageData::FPackageData(FPackageDatas& PackageDatas, const FName& InPackageNa
 {
 	SetState(EPackageState::Idle);
 	SetSaveSubState(ESaveSubState::StartSave);
+	SetSuppressCookReason(ESuppressCookReason::NotSuppressed);
 
 	SendToState(EPackageState::Idle, ESendFlags::QueueAdd, EStateChangeReason::Discovered);
 }
@@ -127,7 +128,7 @@ void FPackageData::ClearReferences()
 	{
 		GenerationHelper->ClearSelfReferences();
 	}
-	ParentGenerationHelper.SafeRelease();
+	SetParentGenerationHelper(nullptr, EStateChangeReason::CookerShutdown);
 }
 
 const FName& FPackageData::GetPackageName() const
@@ -327,7 +328,7 @@ void FPackageData::SetInstigatorInternal(FInstigator&& InInstigator)
 	}
 }
 
-void FPackageData::ClearInProgressData()
+void FPackageData::ClearInProgressData(EStateChangeReason StateChangeReason)
 {
 	SetIsUrgent(false);
 	CompletionCallback = FCompletionCallback();
@@ -339,7 +340,7 @@ void FPackageData::ClearInProgressData()
 		// in progress.
 		GenerationHelper->ClearKeepForGeneratorSave();
 	}
-	ParentGenerationHelper.SafeRelease();
+	SetParentGenerationHelper(nullptr, StateChangeReason);
 }
 
 void FPackageData::SetPlatformsCooked(
@@ -444,6 +445,7 @@ void FPackageData::ClearCookResults()
 		bWasCookedThisSession = false;
 		PackageDatas.GetMonitor().OnLastCookedPlatformRemoved(*this);
 	}
+	SetSuppressCookReason(ESuppressCookReason::NotSuppressed);
 }
 
 void FPackageData::ClearCookResults(const ITargetPlatform* TargetPlatform)
@@ -695,7 +697,7 @@ void FPackageData::SendToState(EPackageState NextState, ESendFlags SendFlags, ES
 			switch (Iterator)
 			{
 			case EPackageStateProperty::InProgress:
-				OnExitInProgress();
+				OnExitInProgress(ReleaseSaveReason);
 				break;
 			case EPackageStateProperty::Loading:
 				OnExitLoading();
@@ -940,11 +942,13 @@ void FPackageData::OnEnterAssignedToWorker()
 {
 	if (IsGenerated())
 	{
-		// Clear the referencecount that we added in OnEnterInProgress; we don't want to keep
-		// the GenerationHelper referenced for the entire duration of assigned packages
-		// running on other CookWorkers. If this package gets retracted and moved into
-		// LoadState locally, we will recreate the GenerationHelper if necessary.
-		ParentGenerationHelper.SafeRelease();
+		// Clear the referencecount that we added in OnEnterInProgress; we don't want to keep the GenerationHelper
+		// referenced for the entire duration of assigned packages running on other CookWorkers. If this package gets
+		// retracted and moved into LoadState locally, we will recreate the GenerationHelper if necessary.
+		// Since we have set the ParentGenerationHelper to null, we can no automatically longer report to the
+		// GenerationHelper that the package has saved when it transitions to Idle. Reporting to the GenerationHelper
+		// that this FPackageData has saved is now the responsibility of the CookWorkerServer's RecordResults function.
+		SetParentGenerationHelper(nullptr, EStateChangeReason::Retraction);
 	}
 }
 
@@ -1028,7 +1032,7 @@ void FPackageData::OnEnterInProgress()
 	}
 }
 
-void FPackageData::OnExitInProgress()
+void FPackageData::OnExitInProgress(EStateChangeReason StateChangeReason)
 {
 	PackageDatas.GetMonitor().OnInProgressChanged(*this, false);
 	UE::Cook::FCompletionCallback LocalCompletionCallback(MoveTemp(GetCompletionCallback()));
@@ -1036,7 +1040,7 @@ void FPackageData::OnExitInProgress()
 	{
 		LocalCompletionCallback(this);
 	}
-	ClearInProgressData();
+	ClearInProgressData(StateChangeReason);
 }
 
 void FPackageData::OnEnterLoading()
@@ -1647,10 +1651,24 @@ TRefCountPtr<FGenerationHelper> FPackageData::GetParentGenerationHelper() const
 	return ParentGenerationHelper;
 }
 
-void FPackageData::SetParentGenerationHelper(FGenerationHelper* InGenerationHelper)
+void FPackageData::SetParentGenerationHelper(FGenerationHelper* InGenerationHelper,
+	EStateChangeReason StateChangeReason, FCookGenerationInfo* InfoOfPackageInGenerator)
 {
 	check(InGenerationHelper == nullptr || IsGenerated());
 	check(!(ParentGenerationHelper && InGenerationHelper) || ParentGenerationHelper == InGenerationHelper);
+
+	if (ParentGenerationHelper && !InGenerationHelper && IsTerminalStateChange(StateChangeReason))
+	{
+		// The package's progress is completed and we will not come back to it; report the package was saved.
+		if (!InfoOfPackageInGenerator)
+		{
+			InfoOfPackageInGenerator = ParentGenerationHelper->FindInfo(*this);
+		}
+		if (InfoOfPackageInGenerator)
+		{
+			InfoOfPackageInGenerator->SetHasSaved(*ParentGenerationHelper, true, FWorkerId::Local());
+		}
+	}
 	ParentGenerationHelper = InGenerationHelper;
 }
 
@@ -1671,7 +1689,7 @@ TRefCountPtr<FGenerationHelper> FPackageData::GetOrFindParentGenerationHelper()
 		return nullptr;
 	}
 
-	ParentGenerationHelper = OwnerPackageData->GetGenerationHelper();
+	SetParentGenerationHelper(OwnerPackageData->GetGenerationHelper(), EStateChangeReason::Requested);
 	return ParentGenerationHelper;
 }
 
@@ -1681,7 +1699,7 @@ TRefCountPtr<FGenerationHelper> FPackageData::TryCreateValidParentGenerationHelp
 	{
 		if (!ParentGenerationHelper->IsValid())
 		{
-			ParentGenerationHelper.SafeRelease();
+			SetParentGenerationHelper(nullptr, EStateChangeReason::Requested);
 		}
 		return ParentGenerationHelper;
 	}
diff --git a/Engine/Source/Editor/UnrealEd/Private/Cooker/CookPackageData.h b/Engine/Source/Editor/UnrealEd/Private/Cooker/CookPackageData.h
--- a/Engine/Source/Editor/UnrealEd/Private/Cooker/CookPackageData.h
+++ b/Engine/Source/Editor/UnrealEd/Private/Cooker/CookPackageData.h
@@ -379,7 +379,7 @@ public:
 	 * Clear all the inprogress variables from the current PackageData. It is invalid to call this except when
 	 * the PackageData is transitioning out of InProgress.
 	 */
-	void ClearInProgressData();
+	void ClearInProgressData(EStateChangeReason StateChangeReason);
 
 	/**
 	 * FindOrAdd each TargetPlatform and set its flags: CookAttempted=true, Succeeded=<given>.
@@ -437,6 +437,10 @@ public:
 	 * returns ECookResult::NotAttempted, otherwise returns whatever result has been set.
 	 */
 	ECookResult GetCookResults(const ITargetPlatform* Platform) const;
+	/** Get/Set the SuppressCookReason for the package. Defaults to ESuppressCookReason::NotSuppressed. */
+	ESuppressCookReason GetSuppressCookReason() const;
+	/** Get/Set the SuppressCookReason for the package. Defaults to ESuppressCookReason::NotSuppressed. */
+	void SetSuppressCookReason(ESuppressCookReason Reason);
 
 	/**
 	 * Return the package pointer. By contract it will be non-null if and only if the PackageData's state is
@@ -639,7 +643,8 @@ public:
 	 * Set the owning generator PackageData. Only valid to call with non-null if SetGenerated() has been called. Keeps
 	 * the GenerationHelper referenced until SetParentGenerationHelper(nullptr) is called.
 	 */
-	void SetParentGenerationHelper(FGenerationHelper* InGenerationHelper);
+	void SetParentGenerationHelper(FGenerationHelper* InGenerationHelper, EStateChangeReason StateChangeReason,
+		FCookGenerationInfo* InfoOfPackageInGenerator = nullptr);
 	/** Return the ParentGenerator's GenerationHelper if the pointer to it has already been set on this. */
 	TRefCountPtr<FGenerationHelper> GetParentGenerationHelper() const;
 	/**
@@ -776,7 +781,7 @@ private:
 	void OnEnterSaveStalledAssignedToWorker();
 	void OnExitSaveStalledAssignedToWorker();
 	/* Entry/Exit gates for Properties shared between multiple states */
-	void OnExitInProgress();
+	void OnExitInProgress(EStateChangeReason StateChangeReason);
 	void OnEnterInProgress();
 	void OnExitLoading();
 	void OnEnterLoading();
@@ -832,6 +837,7 @@ private:
 	FWorkerId WorkerAssignmentConstraint = FWorkerId::Invalid();
 	uint32 State : int32(EPackageState::BitCount);
 	uint32 SaveSubState : int32(ESaveSubState::BitCount);
+	uint32 SuppressCookReason : int32(ESuppressCookReason::BitCount);
 	uint32 bIsUrgent : 1;
 	uint32 bIsCookLast : 1;
 	uint32 bIsVisited : 1;
@@ -1516,6 +1522,16 @@ inline ESaveSubState FPackageData::GetSaveSubState() const
 	return static_cast<ESaveSubState>(SaveSubState);
 }
 
+inline ESuppressCookReason FPackageData::GetSuppressCookReason() const
+{
+	return static_cast<ESuppressCookReason>(SuppressCookReason);
+}
+
+inline void FPackageData::SetSuppressCookReason(ESuppressCookReason Reason)
+{
+	SuppressCookReason = static_cast<uint32>(Reason);
+}
+
 inline bool FPackageData::HasCompletedGeneration() const
 {
 	return static_cast<bool>(bCompletedGeneration);
diff --git a/Engine/Source/Editor/UnrealEd/Private/Cooker/CookTypes.h b/Engine/Source/Editor/UnrealEd/Private/Cooker/CookTypes.h
--- a/Engine/Source/Editor/UnrealEd/Private/Cooker/CookTypes.h
+++ b/Engine/Source/Editor/UnrealEd/Private/Cooker/CookTypes.h
@@ -168,6 +168,8 @@ namespace UE::Cook
 		MultiprocessAssignmentError,
 		RetractedByCookDirector,
 		CookFilter,
+		Count,
+		BitCount = FPlatformMath::ConstExprCeilLogTwo(Count),
 	};
 	const TCHAR* LexToString(UE::Cook::ESuppressCookReason Reason);
 	EStateChangeReason ConvertToStateChangeReason(ESuppressCookReason Reason);
diff --git a/Engine/Source/Editor/UnrealEd/Private/Cooker/CookWorkerServer.cpp b/Engine/Source/Editor/UnrealEd/Private/Cooker/CookWorkerServer.cpp
--- a/Engine/Source/Editor/UnrealEd/Private/Cooker/CookWorkerServer.cpp
+++ b/Engine/Source/Editor/UnrealEd/Private/Cooker/CookWorkerServer.cpp
@@ -940,7 +940,11 @@ void FCookWorkerServer::RecordResults(FPackageResultsMessage& Message)
 			else
 			{
 				ParentGenerationHelper->MarkPackageSavedRemotely(COTFS, *PackageData, GetWorkerId());
-				PackageData->SetParentGenerationHelper(nullptr);
+				EStateChangeReason StateChangeReason =
+					Result.GetSuppressCookReason() == ESuppressCookReason::NotSuppressed
+					? EStateChangeReason::Saved
+					: ConvertToStateChangeReason(Result.GetSuppressCookReason());
+				PackageData->SetParentGenerationHelper(nullptr, StateChangeReason);
 			}
 		}
 		TRefCountPtr<FGenerationHelper> GenerationHelper = PackageData->GetGenerationHelper();
