Localization:
Small adjustments to DetectOrphanedLocalizedAssetsCommandlet.
- Upgraded the log status of dumping all found localized assets from display to warning to help with CICD
- When a file is specified to be written to, all of the discovered assets are still dumped to the log.
#rnx

diff --git a/Engine/Source/Editor/UnrealEd/Private/Commandlets/DetectOrphanedLocalizedAssetsCommandlet.cpp b/Engine/Source/Editor/UnrealEd/Private/Commandlets/DetectOrphanedLocalizedAssetsCommandlet.cpp
--- a/Engine/Source/Editor/UnrealEd/Private/Commandlets/DetectOrphanedLocalizedAssetsCommandlet.cpp
+++ b/Engine/Source/Editor/UnrealEd/Private/Commandlets/DetectOrphanedLocalizedAssetsCommandlet.cpp
@@ -40,6 +40,7 @@ int32 UDetectOrphanedLocalizedAssetsCommandlet::Main(const FString& Params)
 	FCollectionManagerModule& CollectionManagerModule = FCollectionManagerModule::GetModule();
 	ICollectionManager& CollectionManager = CollectionManagerModule.Get();
 	FARFilter CollectionFilter;
+	// @TODOLocalization: Don't hardcode and instead turn this into a parameter for the commandlet 
 	bool bSuccess = CollectionManager.GetObjectsInCollection(FName("Audit_InCook"), ECollectionShareType::CST_All, CollectionFilter.SoftObjectPaths, ECollectionRecursionFlags::SelfAndChildren);
 	
 	IAssetRegistry::Get()->SearchAllAssets(true);
@@ -69,7 +70,7 @@ int32 UDetectOrphanedLocalizedAssetsCommandlet::Main(const FString& Params)
 	}
 	UE_LOG(LogDetectOrphanedLocalizedAssetsCommandlet, Display, TEXT("Iterating through all assets took %.2f seconds."), FPlatformTime::Seconds() - AllAssetsIterationStartTime);
 	float LocalizedAssetsPercentage = static_cast<float>(LocalizedAssets.Num()) / static_cast<float>(AllAssets.Num()) * 100.0f;
-	UE_LOG(LogDetectOrphanedLocalizedAssetsCommandlet, Display, TEXT("Found %d localized assets out of %d assets. %.2f percent of assets are localized."), LocalizedAssets.Num(), AllAssets.Num(), LocalizedAssetsPercentage);
+	UE_LOG(LogDetectOrphanedLocalizedAssetsCommandlet, Display, TEXT("Found %d localized assets out of %d assets. %.2f%% of assets are localized."), LocalizedAssets.Num(), AllAssets.Num(), LocalizedAssetsPercentage);
 
 	IAssetRegistry& AssetRegistry = IAssetRegistry::GetChecked();
 	FString SourceObjectPath;
@@ -116,7 +117,7 @@ int32 UDetectOrphanedLocalizedAssetsCommandlet::Main(const FString& Params)
 	}
 	UE_LOG(LogDetectOrphanedLocalizedAssetsCommandlet, Display, TEXT("Detecting orphaned localized assets took %.2f seconds."), FPlatformTime::Seconds() - DetectOrphansStartTime);
 	float OrphanedPercentage = static_cast<float>(OrphanedLocalizedAssets.Num()) / static_cast<float>(LocalizedAssets.Num()) * 100.0f;
-	UE_LOG(LogDetectOrphanedLocalizedAssetsCommandlet, Display, TEXT("%d out of %d localized assets are orphaned. %.2f of all localized assets are orphaned."), OrphanedLocalizedAssets.Num(), LocalizedAssets.Num(), OrphanedPercentage);
+	UE_LOG(LogDetectOrphanedLocalizedAssetsCommandlet, Display, TEXT("%d out of %d localized assets are orphaned. %.2f%% of all localized assets are orphaned."), OrphanedLocalizedAssets.Num(), LocalizedAssets.Num(), OrphanedPercentage);
 
 	TArray<FString> OrphanedLocalizedAssetsStrings;
 	OrphanedLocalizedAssetsStrings.Reserve(OrphanedLocalizedAssets.Num());
@@ -125,18 +126,18 @@ int32 UDetectOrphanedLocalizedAssetsCommandlet::Main(const FString& Params)
 		OrphanedLocalizedAssetsStrings.Add(OrphanedAsset.ToString());
 	}
 
+	// We sort the strings so it's easier to cluster all assets that are in the same module for fixing.
+	OrphanedLocalizedAssetsStrings.Sort();
+
 	if (ParamVals.Contains("OutputOrphans"))
 	{
 		UE_LOG(LogDetectOrphanedLocalizedAssetsCommandlet, Display, TEXT("An output file was provided. Dumping all found orphaned assets to the file: %s"), *ParamVals["OutputOrphans"]);
 		FFileHelper::SaveStringArrayToFile(OrphanedLocalizedAssetsStrings, *ParamVals["OutputOrphans"]);
 	}
-	else
+	UE_LOG(LogDetectOrphanedLocalizedAssetsCommandlet, Display, TEXT("Dumping %d found orphaned assets to the console."), OrphanedLocalizedAssetsStrings.Num());
+	for (const FString& OrphanedLocalizedAssetString : OrphanedLocalizedAssetsStrings)
 	{
-		UE_LOG(LogDetectOrphanedLocalizedAssetsCommandlet, Display, TEXT("No output file (-OutputOrphans=<PathToOutput>) was provided. Dumping all found orphaned assets to the console."));
-		for (const FString& OrphanedLocalizedAssetString : OrphanedLocalizedAssetsStrings)
-		{
-			UE_LOG(LogDetectOrphanedLocalizedAssetsCommandlet, Display, TEXT("%s"), *OrphanedLocalizedAssetString);
-		}
+		UE_LOG(LogDetectOrphanedLocalizedAssetsCommandlet, Warning, TEXT("'%s' is an orphaned localized asset. Either delete the asset or re-parent the localized asset to its source asset."), *OrphanedLocalizedAssetString);
 	}
 
 	return 0;
