Expose OnActorLabelChanged to editor Blueprints in UEditorActorSubsystem.

diff --git a/Engine/Source/Editor/UnrealEd/Private/Subsystems/EditorActorSubsystem.cpp b/Engine/Source/Editor/UnrealEd/Private/Subsystems/EditorActorSubsystem.cpp
--- a/Engine/Source/Editor/UnrealEd/Private/Subsystems/EditorActorSubsystem.cpp
+++ b/Engine/Source/Editor/UnrealEd/Private/Subsystems/EditorActorSubsystem.cpp
@@ -32,6 +32,7 @@
 #include "Layers/LayersSubsystem.h"
 #include "LevelEditorViewport.h"
 #include "Logging/MessageLog.h"
+#include "Misc/CoreDelegates.h"
 #include "Misc/FeedbackContext.h"
 #include "Misc/FileHelper.h"
 #include "Misc/MessageDialog.h"
@@ -173,6 +174,8 @@ void UEditorActorSubsystem::Initialize(FSubsystemCollectionBase& Collection)
 
 	FEditorDelegates::OnDeleteActorsBegin.AddUObject(this, &UEditorActorSubsystem::BroadcastDeleteActorsBegin);
 	FEditorDelegates::OnDeleteActorsEnd.AddUObject(this, &UEditorActorSubsystem::BroadcastDeleteActorsEnd);
+
+	FCoreDelegates::OnActorLabelChanged.AddUObject(this, &UEditorActorSubsystem::BroadcastActorLabelChanged);
 }
 
 void UEditorActorSubsystem::Deinitialize()
@@ -189,7 +192,7 @@ void UEditorActorSubsystem::Deinitialize()
 	FEditorDelegates::OnDuplicateActorsEnd.RemoveAll(this);
 	FEditorDelegates::OnDeleteActorsBegin.RemoveAll(this);
 	FEditorDelegates::OnDeleteActorsEnd.RemoveAll(this);
-
+	FCoreDelegates::OnActorLabelChanged.RemoveAll(this);
 }
 
 /** To fire before an Actor is Dropped */
@@ -264,6 +267,12 @@ void UEditorActorSubsystem::BroadcastDeleteActorsEnd()
 	OnDeleteActorsEnd.Broadcast();
 }
 
+/** To fire after an Actor has changed label */
+void UEditorActorSubsystem::BroadcastActorLabelChanged(AActor* Actor)
+{
+	OnActorLabelChanged.Broadcast(Actor);
+}
+
 void UEditorActorSubsystem::DuplicateSelectedActors(UWorld* InWorld)
 {
 	if (!GEditor || !InWorld)
diff --git a/Engine/Source/Editor/UnrealEd/Public/Subsystems/EditorActorSubsystem.h b/Engine/Source/Editor/UnrealEd/Public/Subsystems/EditorActorSubsystem.h
--- a/Engine/Source/Editor/UnrealEd/Public/Subsystems/EditorActorSubsystem.h
+++ b/Engine/Source/Editor/UnrealEd/Public/Subsystems/EditorActorSubsystem.h
@@ -38,6 +38,8 @@ DECLARE_DYNAMIC_MULTICAST_DELEGATE(FOnDuplicateActorsEnd);
 DECLARE_DYNAMIC_MULTICAST_DELEGATE(FOnDeleteActorsBegin);
 /** delegate type for after delete actors is handled */
 DECLARE_DYNAMIC_MULTICAST_DELEGATE(FOnDeleteActorsEnd);
+/** delegate type for when an actor changes its label. */
+DECLARE_DYNAMIC_MULTICAST_DELEGATE_OneParam(FOnActorLabelChanged, AActor*, Actor);
 
 /**
 * UEditorActorUtilitiesSubsystem
@@ -50,7 +52,7 @@ class UEditorActorSubsystem : public UEditorSubsystem
 
 public:
 
-	UNREALED_API virtual void Initialize(FSubsystemCollectionBase& Collection) override;
+	virtual void Initialize(FSubsystemCollectionBase& Collection) override;
 	UNREALED_API virtual void Deinitialize() override;
 
 	UPROPERTY(BlueprintAssignable, Category = "Editor Scripting | Level Utility")
@@ -89,6 +91,9 @@ public:
 	UPROPERTY(BlueprintAssignable, Category = "Editor Scripting | Level Utility")
 	FOnDeleteActorsEnd OnDeleteActorsEnd;
 
+	UPROPERTY(BlueprintAssignable, Category = "Editor Scripting | Level Utility")
+	FOnActorLabelChanged OnActorLabelChanged;
+
 	struct FActorDuplicateParameters
 	{
 		ULevel* LevelOverride = nullptr;
@@ -404,5 +409,8 @@ private:
 
 	/** To fire after an Actor is Deleted */
 	void BroadcastDeleteActorsEnd();
+	
+	/** To fire after an Actor has changed label */
+	void BroadcastActorLabelChanged(AActor* Actor);
 };
 
