Cooker: Refactor loading data and functions on FPackageData out into a new class FPackagePreloader.
#rnx

diff --git a/Engine/Source/Editor/UnrealEd/Private/Cooker/CookPackageData.cpp b/Engine/Source/Editor/UnrealEd/Private/Cooker/CookPackageData.cpp
--- a/Engine/Source/Editor/UnrealEd/Private/Cooker/CookPackageData.cpp
+++ b/Engine/Source/Editor/UnrealEd/Private/Cooker/CookPackageData.cpp
@@ -13,6 +13,7 @@
 #include "CompactBinaryTCP.h"
 #include "Cooker/CookDirector.h"
 #include "Cooker/CookGenerationHelper.h"
+#include "Cooker/CookPackagePreloader.h"
 #include "Cooker/CookPlatformManager.h"
 #include "Cooker/CookRequestCluster.h"
 #include "Cooker/CookWorkerClient.h"
@@ -20,7 +21,6 @@
 #include "Cooker/PackageTracker.h"
 #include "CookOnTheSide/CookOnTheFlyServer.h"
 #include "Containers/StringView.h"
-#include "EditorDomain/EditorDomain.h"
 #include "Engine/Console.h"
 #include "HAL/FileManager.h"
 #include "HAL/PlatformTime.h"
@@ -32,7 +32,6 @@
 #include "Misc/PackageName.h"
 #include "Misc/Parse.h"
 #include "Misc/Paths.h"
-#include "Misc/PreloadableFile.h"
 #include "Misc/ScopeExit.h"
 #include "Misc/ScopeLock.h"
 #include "Serialization/CompactBinaryWriter.h"
@@ -94,8 +93,8 @@ bool FPackagePlatformData::NeedsCooking(const ITargetPlatform* PlatformItBelongs
 FPackageData::FPackageData(FPackageDatas& PackageDatas, const FName& InPackageName, const FName& InFileName)
 	: ParentGenerationHelper(nullptr), PackageName(InPackageName), FileName(InFileName), PackageDatas(PackageDatas)
 	, Instigator(EInstigator::NotYetRequested), bIsUrgent(0), bIsCookLast(0)
-	, bIsVisited(0), bIsPreloadAttempted(0)
-	, bIsPreloaded(0), bHasSaveCache(0), bPrepareSaveFailed(0), bPrepareSaveRequiresGC(0)
+	, bIsVisited(0)
+	, bHasSaveCache(0), bPrepareSaveFailed(0), bPrepareSaveRequiresGC(0)
 	, MonitorCookResult((uint8)ECookResult::NotAttempted)
 	, bCompletedGeneration(0), bGenerated(0), bKeepReferencedDuringGC(0)
 	, bWasCookedThisSession(0)
@@ -112,14 +111,19 @@ FPackageData::~FPackageData()
 {
 	// ClearReferences should have been called earlier, but call it here in case it was missed
 	ClearReferences();
-	// FPackageDatas guarantees that all references to GenerationHelper are removed before any PackageDatas are deleted.
-	// We rely on that so that we can be sure that when this PackageData is being deleted, its GenerationHelper - which
-	// assumes the FPackageData lifetime exceeds its own - has already been deleted.
-	check(GenerationHelper == nullptr);
 	// We need to send OnLastCookedPlatformRemoved message to the monitor, so call SetPlatformsNotCooked
 	ClearCookResults();
 	// Update the monitor's counters and call exit functions
 	SendToState(EPackageState::Idle, ESendFlags::QueueNone, EStateChangeReason::CookerShutdown);
+
+	// FPackageDatas guarantees that all references to GenerationHelper are removed before any PackageDatas are deleted.
+	// We rely on that so that we can be sure that when this PackageData is being deleted, its GenerationHelper - which
+	// assumes the FPackageData lifetime exceeds its own - has already been deleted.
+	check(GenerationHelper == nullptr);
+	// FPackageDatas guarantees that all references to PackagePreloaders are removed before any PackageDatas are deleted.
+	// We rely on that so that we can be sure that when this PackageData is being deleted, its PackagePreloader - which
+	// assumes the FPackageData lifetime exceeds its own - has already been deleted.
+	check(PackagePreloader == nullptr);
 }
 
 void FPackageData::ClearReferences()
@@ -129,6 +133,10 @@ void FPackageData::ClearReferences()
 		GenerationHelper->ClearSelfReferences();
 	}
 	SetParentGenerationHelper(nullptr, EStateChangeReason::CookerShutdown);
+	if (PackagePreloader)
+	{
+		PackagePreloader->Shutdown(); // Clears references to any other preloaders
+	}
 }
 
 const FName& FPackageData::GetPackageName() const
@@ -1045,12 +1053,19 @@ void FPackageData::OnExitInProgress(EStateChangeReason StateChangeReason)
 
 void FPackageData::OnEnterLoading()
 {
-	CheckPreloadEmpty();
+	TRefCountPtr<FPackagePreloader> Local = CreatePackagePreloader();
+	Local->SetSelfReference();
+
+	check(PackagePreloader);
+	PackagePreloader->CheckPreloadEmpty();
 }
 
 void FPackageData::OnExitLoading()
 {
-	ClearPreload();
+	check(PackagePreloader); // Guaranteed by OnEnterLoading
+	PackagePreloader->ClearPreload();
+	PackagePreloader->ClearSelfReference();
+	// PackagePreloader might now be nullptr
 }
 
 void FPackageData::OnEnterSaving()
@@ -1132,191 +1147,33 @@ void FPackageData::AddCompletionCallback(TConstArrayView<const ITargetPlatform*>
 	}
 }
 
-bool FPackageData::TryPreload()
-{
-	check(IsInStateProperty(EPackageStateProperty::Loading));
-	if (GetIsPreloadAttempted())
-	{
-		return true;
-	}
-	if (FindObjectFast<UPackage>(nullptr, GetPackageName()))
-	{
-		if (AsyncRequest && !AsyncRequest->bHasFinished)
-		{
-			// In case of async loading, the object can be found while still being asynchronously serialized, we need
-			// to wait until the callback is called and the async request is completely done.
-			return false;
-		}
-
-		// If the package has already loaded, then there is no point in further preloading
-		ClearPreload();
-		SetIsPreloadAttempted(true);
-		return true;
-	}
-	if (IsGenerated())
-	{
-		// Deferred populate generated packages are loaded from their generator, not from disk
-		ClearPreload();
-		SetIsPreloadAttempted(true);
-		return true;
-	}
-	if (IsAsyncLoadingMultithreaded())
-	{
-		if (!AsyncRequest.IsValid())
-		{
-			PackageDatas.GetMonitor().OnPreloadAllocatedChanged(*this, true);
-			AsyncRequest = MakeShared<FAsyncRequest>();
-			AsyncRequest->RequestID = LoadPackageAsync(
-				GetFileName().ToString(), 
-				FLoadPackageAsyncDelegate::CreateLambda(
-					[AsyncRequest = AsyncRequest](const FName&, UPackage*, EAsyncLoadingResult::Type) 
-					{
-						AsyncRequest->bHasFinished = true;
-					}
-				),
-				32 /* Use arbitrary higher priority for preload as we're going to need them very soon */
-			);
-		}
-
-		// always return false so we continue to check the status of the load until FindObjectFast above finds the
-		// loaded object
-		return false;
-	}
-	if (!PreloadableFile.Get())
-	{
-		if (FEditorDomain* EditorDomain(FEditorDomain::Get());
-			EditorDomain && EditorDomain->IsReadingPackages())
-		{
-			EditorDomain->PrecachePackageDigest(GetPackageName());
-		}
-		TStringBuilder<NAME_SIZE> FileNameString;
-		GetFileName().ToString(FileNameString);
-		PreloadableFile.Set(MakeShared<FPreloadableArchive>(FileNameString.ToString()), *this);
-		PreloadableFile.Get()->InitializeAsync([this]()
-			{
-				TStringBuilder<NAME_SIZE> FileNameString;
-				// Note this async callback has an read of this->GetFilename and a write of PreloadableFileOpenResult
-				// outside of a critical section. This read and write is allowed because GetFilename does
-				// not change until this is destructed, and the destructor does not run and other threads do not read
-				// or write PreloadableFileOpenResult until after PreloadableFile.Get() has finished initialization
-				// and this callback is therefore complete.
-				// The code that accomplishes that waiting is in TryPreload (IsInitialized) and ClearPreload
-				// (ReleaseCache)
-				this->GetFileName().ToString(FileNameString);
-				FPackagePath PackagePath = FPackagePath::FromLocalPath(FileNameString);
-				FOpenPackageResult Result = IPackageResourceManager::Get().OpenReadPackage(PackagePath);
-				if (Result.Archive)
-				{
-					this->PreloadableFileOpenResult.CopyMetaData(Result);
-				}
-				return Result.Archive.Release();
-			},
-			FPreloadableFile::Flags::PreloadHandle | FPreloadableFile::Flags::Prime);
-	}
-	const TSharedPtr<FPreloadableArchive>& FilePtr = PreloadableFile.Get();
-	if (!FilePtr->IsInitialized())
-	{
-		if (GetIsUrgent())
-		{
-			// For urgent requests, wait on them to finish preloading rather than letting them run asynchronously
-			// and coming back to them later
-			FilePtr->WaitForInitialization();
-			check(FilePtr->IsInitialized());
-		}
-		else
-		{
-			return false;
-		}
-	}
-	if (FilePtr->TotalSize() < 0)
-	{
-		UE_LOG(LogCook, Warning, TEXT("Failed to find file when preloading %s."), *GetFileName().ToString());
-		SetIsPreloadAttempted(true);
-		PreloadableFile.Reset(*this);
-		PreloadableFileOpenResult = FOpenPackageResult();
-		return true;
-	}
-
-	TStringBuilder<NAME_SIZE> FileNameString;
-	GetFileName().ToString(FileNameString);
-	if (!IPackageResourceManager::TryRegisterPreloadableArchive(FPackagePath::FromLocalPath(FileNameString),
-		FilePtr, PreloadableFileOpenResult))
-	{
-		UE_LOG(LogCook, Warning, TEXT("Failed to register %s for preload."), *GetFileName().ToString());
-		SetIsPreloadAttempted(true);
-		PreloadableFile.Reset(*this);
-		PreloadableFileOpenResult = FOpenPackageResult();
-		return true;
-	}
-
-	SetIsPreloaded(true);
-	SetIsPreloadAttempted(true);
-	return true;
-}
-
-void FPackageData::FTrackedPreloadableFilePtr::Set(TSharedPtr<FPreloadableArchive>&& InPtr, FPackageData& Owner)
+TRefCountPtr<FPackagePreloader> FPackageData::GetPackagePreloader() const
 {
-	Reset(Owner);
-	if (InPtr)
-	{
-		Ptr = MoveTemp(InPtr);
-		Owner.PackageDatas.GetMonitor().OnPreloadAllocatedChanged(Owner, true);
-	}
+	return TRefCountPtr<FPackagePreloader>(PackagePreloader);
 }
 
-void FPackageData::FTrackedPreloadableFilePtr::Reset(FPackageData& Owner)
+TRefCountPtr<FPackagePreloader> FPackageData::CreatePackagePreloader()
 {
-	if (Ptr)
+	if (PackagePreloader)
 	{
-		Owner.PackageDatas.GetMonitor().OnPreloadAllocatedChanged(Owner, false);
-		Ptr.Reset();
+		return TRefCountPtr<FPackagePreloader>(PackagePreloader);
 	}
+	TRefCountPtr<FPackagePreloader> Result(new FPackagePreloader(*this));
+	PackagePreloader = Result.GetReference();
+	return Result;
 }
 
-void FPackageData::ClearPreload()
+void FPackageData::OnPackagePreloaderDestroyed(FPackagePreloader& InPackagePreloader)
 {
-	if (AsyncRequest)
-	{
-		if (!AsyncRequest->bHasFinished)
-		{
-			FlushAsyncLoading(AsyncRequest->RequestID);
-			check(AsyncRequest->bHasFinished);
-		}
-		PackageDatas.GetMonitor().OnPreloadAllocatedChanged(*this, false);
-		AsyncRequest.Reset();
-	}
-
-	const TSharedPtr<FPreloadableArchive>& FilePtr = PreloadableFile.Get();
-	if (GetIsPreloaded())
-	{
-		check(FilePtr);
-		TStringBuilder<NAME_SIZE> FileNameString;
-		GetFileName().ToString(FileNameString);
-		if (IPackageResourceManager::UnRegisterPreloadableArchive(FPackagePath::FromLocalPath(FileNameString)))
-		{
-			UE_LOG(LogCook, Display,
-				TEXT("PreloadableFile was created for %s but never used. This is wasteful and bad for cook performance."),
-				*PackageName.ToString());
-		}
-		FilePtr->ReleaseCache(); // ReleaseCache to conserve memory if the Linker still has a pointer to it
-	}
-	else
-	{
-		check(!FilePtr || !FilePtr->IsCacheAllocated());
-	}
-
-	PreloadableFile.Reset(*this);
-	PreloadableFileOpenResult = FOpenPackageResult();
-	SetIsPreloaded(false);
-	SetIsPreloadAttempted(false);
+	check(PackagePreloader == &InPackagePreloader);
+	PackagePreloader = nullptr;
 }
 
-void FPackageData::CheckPreloadEmpty()
+bool FPackageData::TryPreload()
 {
-	check(!AsyncRequest);
-	check(!GetIsPreloadAttempted());
-	check(!PreloadableFile.Get());
-	check(!GetIsPreloaded());
+	check(IsInStateProperty(EPackageStateProperty::Loading));
+	check(PackagePreloader != nullptr); // Guaranteed by OnEnterLoading
+	return PackagePreloader->TryPreload();
 }
 
 TArray<FCachedObjectInOuter>& FPackageData::GetCachedObjectsInOuter()
diff --git a/Engine/Source/Editor/UnrealEd/Private/Cooker/CookPackageData.h b/Engine/Source/Editor/UnrealEd/Private/Cooker/CookPackageData.h
--- a/Engine/Source/Editor/UnrealEd/Private/Cooker/CookPackageData.h
+++ b/Engine/Source/Editor/UnrealEd/Private/Cooker/CookPackageData.h
@@ -23,13 +23,9 @@
 #include "UObject/GCObject.h"
 #include "UObject/ICookInfo.h"
 #include "UObject/NameTypes.h"
-#include "UObject/PackageResourceManager.h"
 #include "UObject/WeakObjectPtr.h"
 #include "UObject/WeakObjectPtrTemplates.h"
 
-#include <atomic>
-
-class FPreloadableFile;
 class FReferenceCollector;
 class ITargetPlatform;
 class FCbFieldView;
@@ -38,6 +34,7 @@ class UCookOnTheFlyServer;
 class UObject;
 class UPackage;
 namespace UE::Cook { class FCookWorkerClient; }
+namespace UE::Cook { class FPackagePreloader; }
 namespace UE::Cook { class FRequestCluster; }
 namespace UE::Cook { struct FConstructPackageData; }
 namespace UE::Cook { struct FDiscoveredPlatformSet; }
@@ -518,16 +515,6 @@ public:
 
 	/** Try to preload the file. Return true if preloading is complete (succeeded or failed or was skipped). */
 	bool TryPreload();
-	/** Get/Set whether preloading is complete (succeeded or failed or was skipped). */
-	bool GetIsPreloadAttempted() const { return bIsPreloadAttempted != 0; }
-	void SetIsPreloadAttempted(bool bValue) { bIsPreloadAttempted = static_cast<uint32>(bValue); }
-	/** Get/Set whether preloading succeeded and completed. */
-	bool GetIsPreloaded() const { return bIsPreloaded != 0; }
-	void SetIsPreloaded(bool bValue) { bIsPreloaded = static_cast<uint32>(bValue); }
-	/** Clear any allocated preload data. */
-	void ClearPreload();
-	/** Issue check statements confirming that no preload data is allocated or flags are set. */
-	void CheckPreloadEmpty();
 
 	/**
 	 * The list of objects inside the package.  Only non-empty during saving; it is populated on demand by
@@ -606,6 +593,16 @@ public:
 	/** Swap all ITargetPlatform* stored on this instance according to the mapping in @param Remap. */
 	void RemapTargetPlatforms(const TMap<ITargetPlatform*, ITargetPlatform*>& Remap);
 
+	// The PackagePreloader holds data necessary to preload and load the UPackage for this PackageData.
+	// It is used during this PackageData's load state, but also during the load of packages that have
+	// a transitive import of it, so it is refcounted.
+	/** Return the PackagePreloader if it already exists, otherwise return nullptr. */
+	TRefCountPtr<FPackagePreloader> GetPackagePreloader() const;
+	/** Create the PackagePreloader if it does not already exist and return a non-null TRefCountPtr to it. */
+	TRefCountPtr<FPackagePreloader> CreatePackagePreloader();
+	/** Helper function for ~FPackagePreloader: clear the pointer this->PackagePreloader. */
+	void OnPackagePreloaderDestroyed(FPackagePreloader& InPackagePreloader);
+
 	// GenerationHelper is set on packages that are generator packages: they generate other packages
 	// during cook. The PackageData for the generator package has a pointer to the GenerationHelper to look
 	// it up, but does not keep it in memory. It is kept in memory by its internal state and by
@@ -804,33 +801,13 @@ private:
 	FName FileName;
 	FName ParentGenerator;
 
-	struct FAsyncRequest
-	{
-		int32 RequestID { 0 };
-		std::atomic<bool> bHasFinished { false };
-	};
-	TSharedPtr<FAsyncRequest> AsyncRequest;
-
 	TWeakObjectPtr<UPackage> Package;
 	/** The one-per-CookOnTheFlyServer owner of this PackageData. */
 	FPackageDatas& PackageDatas;
-	/**
-	 * The number of active PreloadableFiles is tracked globally; wrap the PreloadableFile in a struct that
-	 * guarantees we always update the counter when changing it
-	 */
-	struct FTrackedPreloadableFilePtr
-	{
-		const TSharedPtr<FPreloadableArchive>& Get() { return Ptr; }
-		void Set(TSharedPtr<FPreloadableArchive>&& InPtr, FPackageData& PackageData);
-		void Reset(FPackageData& PackageData);
-	private:
-		TSharedPtr<FPreloadableArchive> Ptr;
-	};
-	FTrackedPreloadableFilePtr PreloadableFile;
+	FPackagePreloader* PackagePreloader = nullptr;
 	int32 NumPendingCookedPlatformData = 0;
 	int32 CookedPlatformDataNextIndex = -1;
 	int32 NumRetriesBeginCacheOnObject = 0;
-	FOpenPackageResult PreloadableFileOpenResult;
 	FInstigator Instigator;
 
 	FWorkerId WorkerAssignment = FWorkerId::Invalid();
@@ -841,8 +818,6 @@ private:
 	uint32 bIsUrgent : 1;
 	uint32 bIsCookLast : 1;
 	uint32 bIsVisited : 1;
-	uint32 bIsPreloadAttempted : 1;
-	uint32 bIsPreloaded : 1;
 	uint32 bHasSaveCache : 1;
 	uint32 bPrepareSaveFailed : 1;
 	uint32 bPrepareSaveRequiresGC : 1;
diff --git a/Engine/Source/Editor/UnrealEd/Private/Cooker/CookPackagePreloader.cpp b/Engine/Source/Editor/UnrealEd/Private/Cooker/CookPackagePreloader.cpp
new file mode 100644
--- /dev/null
+++ b/Engine/Source/Editor/UnrealEd/Private/Cooker/CookPackagePreloader.cpp
@@ -0,0 +1,218 @@
+// Copyright Epic Games, Inc. All Rights Reserved.
+
+#include "Cooker/CookPackagePreloader.h"
+
+#include "CookOnTheSide/CookLog.h"
+#include "CoreGlobals.h"
+#include "EditorDomain/EditorDomain.h"
+#include "Logging/LogMacros.h"
+#include "Misc/AssertionMacros.h"
+#include "Misc/PackagePath.h"
+#include "Misc/PreloadableFile.h"
+#include "Misc/StringBuilder.h"
+#include "UObject/NameTypes.h"
+#include "UObject/UObjectGlobals.h"
+
+namespace UE::Cook
+{
+
+FPackagePreloader::~FPackagePreloader()
+{
+	PackageData.OnPackagePreloaderDestroyed(*this);
+}
+
+void FPackagePreloader::Shutdown()
+{
+	TRefCountPtr<FPackagePreloader> LocalRef(this);
+	ClearPreload();
+}
+
+bool FPackagePreloader::TryPreload()
+{
+	if (bIsPreloadAttempted)
+	{
+		return true;
+	}
+	if (FindObjectFast<UPackage>(nullptr, PackageData.GetPackageName()))
+	{
+		if (AsyncRequest && !AsyncRequest->bHasFinished)
+		{
+			// In case of async loading, the object can be found while still being asynchronously serialized, we need
+			// to wait until the callback is called and the async request is completely done.
+			return false;
+		}
+
+		// If the package has already loaded, then there is no point in further preloading
+		ClearPreload();
+		bIsPreloadAttempted = true;
+		return true;
+	}
+	if (PackageData.IsGenerated())
+	{
+		// Deferred populate generated packages are loaded from their generator, not from disk
+		ClearPreload();
+		bIsPreloadAttempted = true;
+		return true;
+	}
+	if (IsAsyncLoadingMultithreaded())
+	{
+		if (!AsyncRequest.IsValid())
+		{
+			PackageData.GetPackageDatas().GetMonitor().OnPreloadAllocatedChanged(PackageData, true);
+			AsyncRequest = MakeShared<FAsyncRequest>();
+			AsyncRequest->RequestID = LoadPackageAsync(
+				PackageData.GetFileName().ToString(),
+				FLoadPackageAsyncDelegate::CreateLambda(
+					[AsyncRequest = AsyncRequest](const FName&, UPackage*, EAsyncLoadingResult::Type)
+					{
+						AsyncRequest->bHasFinished = true;
+					}
+				),
+				32 /* Use arbitrary higher priority for preload as we're going to need them very soon */
+			);
+		}
+
+		// always return false so we continue to check the status of the load until FindObjectFast above finds the
+		// loaded object
+		return false;
+	}
+	if (!PreloadableFile.Get())
+	{
+		if (FEditorDomain* EditorDomain(FEditorDomain::Get());
+			EditorDomain && EditorDomain->IsReadingPackages())
+		{
+			EditorDomain->PrecachePackageDigest(PackageData.GetPackageName());
+		}
+		TStringBuilder<NAME_SIZE> FileNameString;
+		PackageData.GetFileName().ToString(FileNameString);
+		PreloadableFile.Set(MakeShared<FPreloadableArchive>(FileNameString.ToString()), PackageData);
+		PreloadableFile.Get()->InitializeAsync([this]()
+			{
+				TStringBuilder<NAME_SIZE> FileNameString;
+				// Note this async callback has an read of PackageData->GetFilename and a write of this->PreloadableFileOpenResult
+				// outside of a critical section. This read and write is allowed because GetFilename does
+				// not change until the PackageData is destructed, and the destructor does not run and other threads do not read
+				// or write PreloadableFileOpenResult until after PreloadableFile.Get() has finished initialization
+				// and this callback is therefore complete.
+				// The code that accomplishes that waiting is in TryPreload (IsInitialized) and ClearPreload
+				// (ReleaseCache)
+				PackageData.GetFileName().ToString(FileNameString);
+				FPackagePath PackagePath = FPackagePath::FromLocalPath(FileNameString);
+				FOpenPackageResult Result = IPackageResourceManager::Get().OpenReadPackage(PackagePath);
+				if (Result.Archive)
+				{
+					this->PreloadableFileOpenResult.CopyMetaData(Result);
+				}
+				return Result.Archive.Release();
+			},
+			FPreloadableFile::Flags::PreloadHandle | FPreloadableFile::Flags::Prime);
+	}
+	const TSharedPtr<FPreloadableArchive>& FilePtr = PreloadableFile.Get();
+	if (!FilePtr->IsInitialized())
+	{
+		if (PackageData.GetIsUrgent())
+		{
+			// For urgent requests, wait on them to finish preloading rather than letting them run asynchronously
+			// and coming back to them later
+			FilePtr->WaitForInitialization();
+			check(FilePtr->IsInitialized());
+		}
+		else
+		{
+			return false;
+		}
+	}
+	if (FilePtr->TotalSize() < 0)
+	{
+		UE_LOG(LogCook, Warning, TEXT("Failed to find file when preloading %s."),
+			*PackageData.GetFileName().ToString());
+		bIsPreloadAttempted = true;
+		PreloadableFile.Reset(PackageData);
+		PreloadableFileOpenResult = FOpenPackageResult();
+		return true;
+	}
+
+	TStringBuilder<NAME_SIZE> FileNameString;
+	PackageData.GetFileName().ToString(FileNameString);
+	if (!IPackageResourceManager::TryRegisterPreloadableArchive(FPackagePath::FromLocalPath(FileNameString),
+		FilePtr, PreloadableFileOpenResult))
+	{
+		UE_LOG(LogCook, Warning, TEXT("Failed to register %s for preload."),
+			*PackageData.GetFileName().ToString());
+		bIsPreloadAttempted = true;
+		PreloadableFile.Reset(PackageData);
+		PreloadableFileOpenResult = FOpenPackageResult();
+		return true;
+	}
+
+	bIsPreloaded = true;
+	bIsPreloadAttempted = true;
+	return true;
+}
+
+void FPackagePreloader::ClearPreload()
+{
+	if (AsyncRequest)
+	{
+		if (!AsyncRequest->bHasFinished)
+		{
+			FlushAsyncLoading(AsyncRequest->RequestID);
+			check(AsyncRequest->bHasFinished);
+		}
+		PackageData.GetPackageDatas().GetMonitor().OnPreloadAllocatedChanged(PackageData, false);
+		AsyncRequest.Reset();
+	}
+
+	const TSharedPtr<FPreloadableArchive>& FilePtr = PreloadableFile.Get();
+	if (bIsPreloaded)
+	{
+		check(FilePtr);
+		TStringBuilder<NAME_SIZE> FileNameString;
+		PackageData.GetFileName().ToString(FileNameString);
+		if (IPackageResourceManager::UnRegisterPreloadableArchive(FPackagePath::FromLocalPath(FileNameString)))
+		{
+			UE_LOG(LogCook, Display,
+				TEXT("PreloadableFile was created for %s but never used. This is wasteful and bad for cook performance."),
+				*PackageData.GetPackageName().ToString());
+		}
+		FilePtr->ReleaseCache(); // ReleaseCache to conserve memory if the Linker still has a pointer to it
+	}
+	else
+	{
+		check(!FilePtr || !FilePtr->IsCacheAllocated());
+	}
+
+	PreloadableFile.Reset(PackageData);
+	PreloadableFileOpenResult = FOpenPackageResult();
+	bIsPreloaded = false;
+	bIsPreloadAttempted = false;
+}
+
+void FPackagePreloader::CheckPreloadEmpty()
+{
+	check(!AsyncRequest);
+	check(!bIsPreloadAttempted);
+	check(!PreloadableFile.Get());
+	check(!bIsPreloaded);
+}
+
+void FPackagePreloader::FTrackedPreloadableFilePtr::Set(TSharedPtr<FPreloadableArchive>&& InPtr, FPackageData& Owner)
+{
+	Reset(Owner);
+	if (InPtr)
+	{
+		Ptr = MoveTemp(InPtr);
+		Owner.GetPackageDatas().GetMonitor().OnPreloadAllocatedChanged(Owner, true);
+	}
+}
+
+void FPackagePreloader::FTrackedPreloadableFilePtr::Reset(FPackageData& Owner)
+{
+	if (Ptr)
+	{
+		Owner.GetPackageDatas().GetMonitor().OnPreloadAllocatedChanged(Owner, false);
+		Ptr.Reset();
+	}
+}
+
+}
\ No newline at end of file
diff --git a/Engine/Source/Editor/UnrealEd/Private/Cooker/CookPackagePreloader.h b/Engine/Source/Editor/UnrealEd/Private/Cooker/CookPackagePreloader.h
new file mode 100644
--- /dev/null
+++ b/Engine/Source/Editor/UnrealEd/Private/Cooker/CookPackagePreloader.h
@@ -0,0 +1,100 @@
+// Copyright Epic Games, Inc. All Rights Reserved.
+
+#pragma once
+
+#include "Cooker/CookPackageData.h"
+#include "Templates/RefCounting.h"
+#include "Templates/SharedPointer.h"
+#include "UObject/PackageResourceManager.h"
+
+#include <atomic>
+
+namespace UE::Cook
+{
+
+/**
+ * Helper class for the Load state on FPackageData. Also supports loads of other packages that depend on the owner
+ * package. Handles preloading or asyncloading the package before we put the package into LoadReadyQueue for
+ * a call to LoadPackage.
+ */
+class FPackagePreloader : public FThreadSafeRefCountedObject
+{
+public:
+	FPackagePreloader(FPackageData& InPackageData);
+	~FPackagePreloader();
+
+	/** Free any IO requests or buffers, and remove all references to any other FPackagePreloaders. */
+	void Shutdown();
+
+	/** Set the SelfReference that keeps this object in memory while the owner package is using it. */
+	void SetSelfReference();
+	/** Clear the SelfReference; the owner package is done using this object. */
+	void ClearSelfReference();
+
+	/** Try to preload the file. Return true if preloading is complete (succeeded or failed or was skipped). */
+	bool TryPreload();
+	/** Clear any allocated preload data. */
+	void ClearPreload();
+	/** Issue check statements confirming that no preload data is allocated or flags are set. */
+	void CheckPreloadEmpty();
+
+private:
+	/**
+	 * The number of active PreloadableFiles is tracked globally; wrap the PreloadableFile in a struct that
+	 * guarantees we always update the counter when changing it
+	 */
+	struct FTrackedPreloadableFilePtr
+	{
+		const TSharedPtr<FPreloadableArchive>& Get() { return Ptr; }
+		void Set(TSharedPtr<FPreloadableArchive>&& InPtr, FPackageData& PackageData);
+		void Reset(FPackageData& PackageData);
+	private:
+		TSharedPtr<FPreloadableArchive> Ptr;
+	};
+
+	/**
+	 * Structure used when we are doing async packageloads rather than just preloading. Holds the
+	 * request data that gets written from the AsyncLoading thread.
+	 */
+	struct FAsyncRequest
+	{
+		int32 RequestID{ 0 };
+		std::atomic<bool> bHasFinished{ false };
+	};
+
+private:
+	FTrackedPreloadableFilePtr PreloadableFile;
+	FOpenPackageResult PreloadableFileOpenResult;
+	FPackageData& PackageData;
+	TRefCountPtr<FPackagePreloader> SelfReference;
+	TSharedPtr<FAsyncRequest> AsyncRequest;
+	bool bIsPreloadAttempted = false;
+	bool bIsPreloaded = false;
+};
+
+}
+
+///////////////////////////////////////////////////////
+// Inline implementations
+///////////////////////////////////////////////////////
+
+namespace UE::Cook
+{
+
+inline FPackagePreloader::FPackagePreloader(FPackageData& InPackageData)
+: PackageData(InPackageData)
+{
+
+}
+
+inline void FPackagePreloader::SetSelfReference()
+{
+	SelfReference = this;
+}
+
+inline void FPackagePreloader::ClearSelfReference()
+{
+	SelfReference.SafeRelease();
+}
+
+}
\ No newline at end of file
