Editor: Remove log message when clicking background

This appears to be leftover debugging code but no longer serves any useful purpose

#rnx

diff --git a/Engine/Source/Editor/UnrealEd/Private/LevelViewportClickHandlers.cpp b/Engine/Source/Editor/UnrealEd/Private/LevelViewportClickHandlers.cpp
--- a/Engine/Source/Editor/UnrealEd/Private/LevelViewportClickHandlers.cpp
+++ b/Engine/Source/Editor/UnrealEd/Private/LevelViewportClickHandlers.cpp
@@ -898,7 +898,6 @@ namespace LevelViewportClickHandlers
 			if( !Click.IsControlDown() )
 			{
 				const FScopedTransaction Transaction( NSLOCTEXT("UnrealEd", "ClickingBackground", "Clicking Background") );
-				UE_LOG(LogEditorViewport, Log,  TEXT("Clicking Background") );
 				GEditor->SelectNone( true, true );
 			}
 		}
