Fix infinite recursion UMaterialGraphNode::PropagatePropertyChange.

- Check to see if we've already visited this node by checking `bNeedToUpdatePreview` which is set by PropagatePropertyChange.
- If so, we just early out of the function.

diff --git a/Engine/Source/Editor/UnrealEd/Private/MaterialGraphNode.cpp b/Engine/Source/Editor/UnrealEd/Private/MaterialGraphNode.cpp
--- a/Engine/Source/Editor/UnrealEd/Private/MaterialGraphNode.cpp
+++ b/Engine/Source/Editor/UnrealEd/Private/MaterialGraphNode.cpp
@@ -817,7 +817,13 @@ void UMaterialGraphNode::CreateOutputPins()
 
 void UMaterialGraphNode::PropagatePropertyChange()
 {
-	MaterialExpression->bNeedToUpdatePreview = 1;
+	// If we've already propagated an update to this node return
+	if (MaterialExpression->bNeedToUpdatePreview)
+	{
+		return;
+	}
+
+	MaterialExpression->bNeedToUpdatePreview = true;
 
 	// Loop through all the output pins.
 	for (UEdGraphPin* Pin : Pins)
