[ToolMenus] Moved ToolbarLabelOverride to nested Toolbar type

#rnx

diff --git a/Engine/Plugins/Animation/ControlRig/Source/ControlRigEditor/Private/Editor/ControlRigViewportToolbarExtensions.cpp b/Engine/Plugins/Animation/ControlRig/Source/ControlRigEditor/Private/Editor/ControlRigViewportToolbarExtensions.cpp
--- a/Engine/Plugins/Animation/ControlRig/Source/ControlRigEditor/Private/Editor/ControlRigViewportToolbarExtensions.cpp
+++ b/Engine/Plugins/Animation/ControlRig/Source/ControlRigEditor/Private/Editor/ControlRigViewportToolbarExtensions.cpp
@@ -339,7 +339,7 @@ void UE::ControlRig::PopulateControlRigViewportToolbarShowSubmenu(
 				);
 				MotionPathsSubmenu.SetShowInToolbarTopLevel(true);
 				// Override the label when this item is raised to the top-level toolbar.
-				MotionPathsSubmenu.ToolbarLabelOverride = TAttribute<FText>::CreateLambda(
+				MotionPathsSubmenu.ToolBarData.LabelOverride = TAttribute<FText>::CreateLambda(
 					[Settings]() -> FText
 					{
 						// TODO: This only uses the FramesBefore value at the moment. It should be updated.
diff --git a/Engine/Source/Developer/ToolMenus/Private/ToolMenus.cpp b/Engine/Source/Developer/ToolMenus/Private/ToolMenus.cpp
--- a/Engine/Source/Developer/ToolMenus/Private/ToolMenus.cpp
+++ b/Engine/Source/Developer/ToolMenus/Private/ToolMenus.cpp
@@ -1678,9 +1678,9 @@ void UToolMenus::PopulateToolBarBuilderWithEntry(
 	const FUIAction UIAction = UToolMenus::ConvertUIAction(Block, MenuData->Context);
 
 	TAttribute<FText> ToolbarLabelOverride;
-	if (Block.ToolbarLabelOverride.IsSet())
+	if (Block.ToolBarData.LabelOverride.IsSet())
 	{
-		ToolbarLabelOverride = Block.ToolbarLabelOverride;
+		ToolbarLabelOverride = Block.ToolBarData.LabelOverride;
 	}
 	else if (const bool bHasIcon = Block.Icon.IsSet() || (Block.Command.IsValid() && Block.Command->GetIcon().IsSet());
 			 bHasIcon && bIsRaisingToTopLevel)
diff --git a/Engine/Source/Developer/ToolMenus/Public/ToolMenuEntry.h b/Engine/Source/Developer/ToolMenus/Public/ToolMenuEntry.h
--- a/Engine/Source/Developer/ToolMenus/Public/ToolMenuEntry.h
+++ b/Engine/Source/Developer/ToolMenus/Public/ToolMenuEntry.h
@@ -51,6 +51,9 @@ public:
 	{
 	}
 
+	/** Optional override label to use when the entry appears in a toolbar. */
+	TAttribute<FText> LabelOverride;
+
 	/** Delegate that generates a widget for this combo button's menu content. Called when the menu is summoned. */
 	FNewToolMenuChoice ComboButtonContextMenuGenerator;
 
@@ -233,9 +236,6 @@ public:
 	TAttribute<FSlateIcon> Icon;
 	TAttribute<FText> InputBindingLabel;
 
-	/** Optional override label to use when the entry appears in a toolbar. */
-	TAttribute<FText> ToolbarLabelOverride;
-
 private:
 
 	friend class UToolMenus;
diff --git a/Engine/Source/Editor/LevelEditor/Private/ViewportToolbar/LevelEditorViewportToolbarSections.cpp b/Engine/Source/Editor/LevelEditor/Private/ViewportToolbar/LevelEditorViewportToolbarSections.cpp
--- a/Engine/Source/Editor/LevelEditor/Private/ViewportToolbar/LevelEditorViewportToolbarSections.cpp
+++ b/Engine/Source/Editor/LevelEditor/Private/ViewportToolbar/LevelEditorViewportToolbarSections.cpp
@@ -1514,7 +1514,7 @@ FToolMenuEntry CreateViewportToolbarShowSubmenu()
 			}
 		)
 	);
-	Entry.ToolbarLabelOverride = FText();
+	Entry.ToolBarData.LabelOverride = FText();
 	Entry.Icon = FSlateIcon(FAppStyle::GetAppStyleSetName(), "Level.VisibleHighlightIcon16x");
 	return Entry;
 }
@@ -1710,7 +1710,7 @@ FToolMenuEntry CreateViewportToolbarPerformanceAndScalabilitySubmenu()
 		)
 	);
 	Entry.Icon = FSlateIcon(FAppStyle::GetAppStyleSetName(), "Icons.Scalability");
-	Entry.ToolbarLabelOverride = FText();
+	Entry.ToolBarData.LabelOverride = FText();
 	return Entry;
 }
 
@@ -2053,7 +2053,7 @@ FToolMenuEntry CreateLevelEditorViewportToolbarSettingsSubmenu()
 		)
 	);
 	Entry.Icon = FSlateIcon(FAppStyle::GetAppStyleSetName(), "LevelEditor.GameSettings");
-	Entry.ToolbarLabelOverride = FText();
+	Entry.ToolBarData.LabelOverride = FText();
 	return Entry;
 }
 
@@ -2551,7 +2551,7 @@ void ExtendCameraSubmenu(FName InCameraOptionsSubmenuName)
 							FSlateIcon(FSlateIcon(FAppStyle::GetAppStyleSetName(), "EditorViewport.Speed"))
 						);
 
-						Entry.ToolbarLabelOverride = TAttribute<FText>::CreateLambda(
+						Entry.ToolBarData.LabelOverride = TAttribute<FText>::CreateLambda(
 							[LevelViewportWeak]()
 							{
 								return UE::UnrealEd::GetCameraSpeedLabel(LevelViewportWeak);
diff --git a/Engine/Source/Editor/UnrealEd/Private/ViewportToolbar/UnrealEdViewportToolbar.cpp b/Engine/Source/Editor/UnrealEd/Private/ViewportToolbar/UnrealEdViewportToolbar.cpp
--- a/Engine/Source/Editor/UnrealEd/Private/ViewportToolbar/UnrealEdViewportToolbar.cpp
+++ b/Engine/Source/Editor/UnrealEd/Private/ViewportToolbar/UnrealEdViewportToolbar.cpp
@@ -158,7 +158,7 @@ FToolMenuEntry CreateSurfaceSnapCheckboxMenu()
 		MakeMenuDelegate
 	);
 
-	Entry.ToolbarLabelOverride = TAttribute<FText>::CreateLambda(
+	Entry.ToolBarData.LabelOverride = TAttribute<FText>::CreateLambda(
 		[]()
 		{
 			const ULevelEditorViewportSettings* const Settings = GetMutableDefault<ULevelEditorViewportSettings>();
@@ -355,7 +355,7 @@ FToolMenuEntry CreateLocationSnapCheckboxMenu()
 		);
 	}
 
-	Entry.ToolbarLabelOverride = TAttribute<FText>::Create(&UE::UnrealEd::GetLocationGridLabel);
+	Entry.ToolBarData.LabelOverride = TAttribute<FText>::Create(&UE::UnrealEd::GetLocationGridLabel);
 	Entry.Icon = FSlateIcon(FAppStyle::GetAppStyleSetName(), "EditorViewport.LocationGridSnap");
 
 	return Entry;
@@ -426,7 +426,7 @@ FToolMenuEntry CreateRotationSnapCheckboxMenu()
 		MakeMenuDelegate
 	);
 
-	Entry.ToolbarLabelOverride = TAttribute<FText>::Create(&UE::UnrealEd::GetRotationGridLabel);
+	Entry.ToolBarData.LabelOverride = TAttribute<FText>::Create(&UE::UnrealEd::GetRotationGridLabel);
 	Entry.Icon = FSlateIcon(FAppStyle::GetAppStyleSetName(), "EditorViewport.RotationGridSnap");
 
 	return Entry;
@@ -530,7 +530,7 @@ FToolMenuEntry CreateScaleSnapCheckboxMenu()
 		MakeMenuDelegate
 	);
 
-	Entry.ToolbarLabelOverride = TAttribute<FText>::Create(&UE::UnrealEd::GetScaleGridLabel);
+	Entry.ToolBarData.LabelOverride = TAttribute<FText>::Create(&UE::UnrealEd::GetScaleGridLabel);
 	Entry.Icon = FSlateIcon(FAppStyle::GetAppStyleSetName(), "EditorViewport.ScaleGridSnap");
 
 	return Entry;
@@ -675,7 +675,7 @@ FToolMenuEntry CreateViewportToolbarTransformsSection()
 								}
 							);
 						}
-						CoordianteSystemSubmenu.ToolbarLabelOverride = FText();
+						CoordianteSystemSubmenu.ToolBarData.LabelOverride = FText();
 						CoordianteSystemSubmenu.SetShowInToolbarTopLevel(true);
 					}
 				}
@@ -732,7 +732,7 @@ FToolMenuEntry CreateViewportToolbarTransformsSection()
 		)
 	);
 
-	Entry.ToolbarLabelOverride = FText();
+	Entry.ToolBarData.LabelOverride = FText();
 	Entry.Icon = FSlateIcon(FAppStyle::GetAppStyleSetName(), "LevelEditor.SelectMode");
 	return Entry;
 }
@@ -1081,7 +1081,7 @@ FToolMenuEntry CreateViewportToolbarSnappingSubmenu()
 		)
 	);
 
-	Entry.ToolbarLabelOverride = FText();
+	Entry.ToolBarData.LabelOverride = FText();
 	Entry.Icon = FSlateIcon(FAppStyle::GetAppStyleSetName(), "Icons.Snap");
 	return Entry;
 }
