[Viewport Toolbar] Camera Speed slider implementation

For the moment being, this design replicates the current design, with a slider + scalar speed setting.

SAnimViewportToolbar was registering a callback to camera speed updates. This has been replicated.

diff --git a/Engine/Source/Editor/Kismet/Private/SSCSEditorViewport.cpp b/Engine/Source/Editor/Kismet/Private/SSCSEditorViewport.cpp
--- a/Engine/Source/Editor/Kismet/Private/SSCSEditorViewport.cpp
+++ b/Engine/Source/Editor/Kismet/Private/SSCSEditorViewport.cpp
@@ -69,48 +69,9 @@ void CreateCameraSpeedMenu(UToolMenu* InMenu, TWeakPtr<FSCSEditorViewportClient>
 		LOCTEXT("CameraSpeedSubMenu", "Camera Speed"),
 		LOCTEXT("CameraSpeedSubMenu_ToolTip", "Camera Speed related actions"),
 		FNewToolMenuDelegate::CreateLambda(
-			[InWeakViewportClient](UToolMenu* InMenu)
+			[](UToolMenu* InMenu)
 			{
-				// It seems like the speed settings can only be set via the mouse wheel for BP editor.
-				// The other slider values (see legacy toolbar) are not actually changing speed.
-
-				// Taken from legacy toolbar values
-				constexpr int32 MinSpeed = 1;
-				constexpr int32 MaxSpeed = 8;
-
-				FToolMenuEntry CameraSpeedSlider = UnrealEd::CreateNumericEntry(
-					"CameraSpeed",
-					LOCTEXT("CameraSpeedLabel", "Camera Speed"),
-					LOCTEXT("CameraSpeedTooltip", "Camera Speed"),
-					FCanExecuteAction(),
-					UnrealEd::FNumericEntryExecuteActionDelegate::CreateLambda(
-						[InWeakViewportClient](float InValue)
-						{
-							if (TSharedPtr<FSCSEditorViewportClient> LevelViewport = InWeakViewportClient.Pin())
-							{
-								// TODO: properly implement
-								//LevelViewport->SetCameraSpeedSetting(InValue);
-							}
-						}
-					),
-					TAttribute<float>::CreateLambda(
-						[InWeakViewportClient]() -> float
-						{
-							// TODO: properly implement
-							/*if (TSharedPtr<FSCSEditorViewportClient> LevelViewport = InWeakViewportClient.Pin())
-							{
-								return LevelViewport->GetCameraSpeedSetting();
-							}*/
-
-							return 1;
-						}
-					),
-					MinSpeed,
-					MaxSpeed,
-					0
-				);
-
-				InMenu->AddMenuEntry("CameraSpeed", CameraSpeedSlider);
+				UE::UnrealEd::CreateCameraSpeedMenu(InMenu);
 			}
 		),
 		false,
diff --git a/Engine/Source/Editor/LevelEditor/Private/ViewportToolbar/LevelEditorViewportToolbarSections.cpp b/Engine/Source/Editor/LevelEditor/Private/ViewportToolbar/LevelEditorViewportToolbarSections.cpp
--- a/Engine/Source/Editor/LevelEditor/Private/ViewportToolbar/LevelEditorViewportToolbarSections.cpp
+++ b/Engine/Source/Editor/LevelEditor/Private/ViewportToolbar/LevelEditorViewportToolbarSections.cpp
@@ -176,41 +176,6 @@ float GetFarViewPlaneValue(const TSharedRef<::SLevelViewport>& InLevelViewport)
 	return ViewportClient.GetFarClipPlaneOverride();
 }
 
-// TODO: properly implement
-void SetCameraSpeed(const TSharedRef<::SLevelViewport>& InLevelViewport, float NewValue)
-{
-}
-
-void SetCameraSpeedScalarValue(const TSharedRef<::SLevelViewport>& InLevelViewport, float NewValue)
-{
-	if (InLevelViewport->GetViewportClient().IsValid())
-	{
-		InLevelViewport->GetViewportClient()->SetCameraSpeedScalar(NewValue);
-
-		// TODO: make sure something like this gets called if needed (e.g. future menus sharing code)
-		// Also, verify where/how to deal with this callback
-		// OnCamSpeedScalarChanged.ExecuteIfBound(NewValue);
-	}
-}
-
-// TODO: properly implement
-float GetCamSpeedSliderPosition(const TSharedRef<::SLevelViewport>& InLevelViewport)
-{
-	return 1.0f;
-}
-
-float GetCamSpeedScalarSliderPosition(const TSharedRef<::SLevelViewport>& InLevelViewport)
-{
-	float CamSpeedScalar = 1.0f;
-
-	if (InLevelViewport->GetViewportClient().IsValid())
-	{
-		CamSpeedScalar = (InLevelViewport->GetViewportClient()->GetCameraSpeedScalar());
-	}
-
-	return CamSpeedScalar;
-}
-
 bool AddJumpToBookmarkMenu(UToolMenu* InMenu, const TWeakPtr<::SLevelViewport>& InViewport)
 {
 	FToolMenuSection& Section =
@@ -2254,104 +2219,6 @@ FToolMenuEntry CreateFarViewPlaneMenu(TWeakPtr<::SLevelViewport> InLevelViewport
 	);
 }
 
-FToolMenuEntry CreateCameraSpeedSlider(TWeakPtr<::SLevelViewport> InLevelViewportWeak)
-{
-	constexpr float MinSpeed = 0.033f;
-	constexpr float MaxSpeed = 32.0f;
-
-	return UnrealEd::CreateNumericEntry(
-		"CameraSpeed",
-		LOCTEXT("CameraSpeedLabel", "Camera Speed"),
-		LOCTEXT("CameraSpeedTooltip", "Camera Speed"),
-		FCanExecuteAction(),
-		UnrealEd::FNumericEntryExecuteActionDelegate::CreateLambda(
-			[InLevelViewportWeak](float InValue)
-			{
-				if (TSharedPtr<::SLevelViewport> LevelViewport = InLevelViewportWeak.Pin())
-				{
-					Private::SetCameraSpeed(LevelViewport.ToSharedRef(), InValue);
-				}
-			}
-		),
-		TAttribute<float>::CreateLambda(
-			[InLevelViewportWeak]()
-			{
-				if (TSharedPtr<::SLevelViewport> Viewport = InLevelViewportWeak.Pin())
-				{
-					return Private::GetCamSpeedSliderPosition(Viewport.ToSharedRef());
-				}
-
-				return 1.0f;
-			}
-		),
-		MinSpeed,
-		MaxSpeed,
-		3
-	);
-}
-
-FToolMenuEntry CreateCameraSpeedScalarSlider(TWeakPtr<::SLevelViewport> InLevelViewportWeak)
-{
-	constexpr float MinSpeed = 1.0f;
-	constexpr float MaxSpeed = 128.0f;
-
-	return UnrealEd::CreateNumericEntry(
-		"CameraSpeedScalar",
-		LOCTEXT("CameraSpeedScalarLabel", "Speed Scalar"),
-		LOCTEXT("CameraSpeedScalarTooltip", "Scalar to increase camera movement range"),
-		FCanExecuteAction(),
-		UnrealEd::FNumericEntryExecuteActionDelegate::CreateLambda(
-			[InLevelViewportWeak](float InValue)
-			{
-				if (TSharedPtr<::SLevelViewport> LevelViewport = InLevelViewportWeak.Pin())
-				{
-					Private::SetCameraSpeedScalarValue(LevelViewport.ToSharedRef(), InValue);
-				}
-			}
-		),
-		TAttribute<float>::CreateLambda(
-			[InLevelViewportWeak]()
-			{
-				if (TSharedPtr<::SLevelViewport> Viewport = InLevelViewportWeak.Pin())
-				{
-					return Private::GetCamSpeedScalarSliderPosition(Viewport.ToSharedRef());
-				}
-
-				return MinSpeed;
-			}
-		),
-		MinSpeed,
-		MaxSpeed,
-		1
-	);
-}
-
-void CreateCameraSpeedMenu(UToolMenu* InMenu)
-{
-	if (!InMenu)
-	{
-		return;
-	}
-
-	FToolMenuSection& Section = InMenu->AddSection("Section");
-
-	ULevelViewportContext* const LevelViewportContext = InMenu->FindContext<ULevelViewportContext>();
-	if (!LevelViewportContext)
-	{
-		return;
-	}
-
-	const TSharedPtr<::SLevelViewport> LevelViewport = LevelViewportContext->LevelViewport.Pin();
-
-	if (!LevelViewport)
-	{
-		return;
-	}
-
-	Section.AddEntry(CreateCameraSpeedSlider(LevelViewport));
-	Section.AddEntry(CreateCameraSpeedScalarSlider(LevelViewport));
-}
-
 void AddCameraActorSelectSection(UToolMenu* InMenu)
 {
 	ULevelViewportContext* const LevelViewportContext = InMenu->FindContext<ULevelViewportContext>();
@@ -2544,7 +2411,7 @@ void ExtendCameraSubmenu(FName InCameraOptionsSubmenuName)
 							FNewToolMenuDelegate::CreateLambda(
 								[LevelViewportWeak](UToolMenu* InMenu)
 								{
-									CreateCameraSpeedMenu(InMenu);
+									UE::UnrealEd::CreateCameraSpeedMenu(InMenu);
 								}
 							),
 							false,
@@ -2558,6 +2425,7 @@ void ExtendCameraSubmenu(FName InCameraOptionsSubmenuName)
 							}
 						);
 						Entry.SetShowInToolbarTopLevel(true);
+						Entry.ToolBarData.PlacementOverride = MenuPlacement_BelowRightAnchor;
 					}
 
 					PositioningSection.AddSeparator("PositioningSeparator_1");
diff --git a/Engine/Source/Editor/LevelEditor/Private/ViewportToolbar/LevelEditorViewportToolbarSections.h b/Engine/Source/Editor/LevelEditor/Private/ViewportToolbar/LevelEditorViewportToolbarSections.h
--- a/Engine/Source/Editor/LevelEditor/Private/ViewportToolbar/LevelEditorViewportToolbarSections.h
+++ b/Engine/Source/Editor/LevelEditor/Private/ViewportToolbar/LevelEditorViewportToolbarSections.h
@@ -16,7 +16,6 @@ namespace UE::LevelEditor
 {
 void CreateCameraSpawnMenu(UToolMenu* InMenu);
 void CreateBookmarksMenu(UToolMenu* InMenu);
-void CreateCameraSpeedMenu(UToolMenu* InMenu);
 void AddCameraActorSelectSection(UToolMenu* InMenu);
 
 FToolMenuEntry CreateFOVMenu(TWeakPtr<SLevelViewport> InLevelViewportWeak);
diff --git a/Engine/Source/Editor/Persona/Private/SAnimViewportToolBar.cpp b/Engine/Source/Editor/Persona/Private/SAnimViewportToolBar.cpp
--- a/Engine/Source/Editor/Persona/Private/SAnimViewportToolBar.cpp
+++ b/Engine/Source/Editor/Persona/Private/SAnimViewportToolBar.cpp
@@ -387,8 +387,46 @@ void SAnimViewportToolBar::Construct(const FArguments& InArgs, TSharedPtr<class
 			.OnCamSpeedChanged(this, &SAnimViewportToolBar::OnCamSpeedChanged)
 			.OnCamSpeedScalarChanged(this, &SAnimViewportToolBar::OnCamSpeedScalarChanged)
 		];
-			
-	
+
+	TWeakPtr<SAnimViewportToolBar> AnimViewportToolbarWeak = SharedThis(this);
+	UE::UnrealEd::OnViewportClientCamSpeedChanged().BindLambda(
+		[AnimViewportToolbarWeak, EditorViewportWeak = InRealViewport.ToWeakPtr()](const TSharedRef<SEditorViewport>& InEditorViewport, int32 InNewValue)
+		{
+			TSharedPtr<SAnimViewportToolBar> AnimViewportToolbar = AnimViewportToolbarWeak.Pin();
+			if (!AnimViewportToolbar)
+			{
+				return;
+			}
+
+			if (TSharedPtr<SEditorViewport> EditorViewport = EditorViewportWeak.Pin())
+			{
+				if (EditorViewport == InEditorViewport)
+				{
+					AnimViewportToolbar->OnCamSpeedChanged(InNewValue);
+				}
+			}
+		}
+	);
+
+	UE::UnrealEd::OnViewportClientCamSpeedScalarChanged().BindLambda(
+		[AnimViewportToolbarWeak, EditorViewportWeak = InRealViewport.ToWeakPtr()](const TSharedRef<SEditorViewport>& InEditorViewport, float InNewValue)
+		{
+			TSharedPtr<SAnimViewportToolBar> AnimViewportToolbar = AnimViewportToolbarWeak.Pin();
+			if (!AnimViewportToolbar)
+			{
+				return;
+			}
+
+			if (TSharedPtr<SEditorViewport> EditorViewport = EditorViewportWeak.Pin())
+			{
+				if (EditorViewport == InEditorViewport)
+				{
+					AnimViewportToolbar->OnCamSpeedScalarChanged(InNewValue);
+				}
+			}
+		}
+	);
+
 	//@TODO: Need clipping horizontal box: LeftToolbar->AddWrapButton();
 
 	// Create our pinned commands before we bind commands
diff --git a/Engine/Source/Editor/UnrealEd/Private/ViewportToolbar/UnrealEdViewportToolbar.cpp b/Engine/Source/Editor/UnrealEd/Private/ViewportToolbar/UnrealEdViewportToolbar.cpp
--- a/Engine/Source/Editor/UnrealEd/Private/ViewportToolbar/UnrealEdViewportToolbar.cpp
+++ b/Engine/Source/Editor/UnrealEd/Private/ViewportToolbar/UnrealEdViewportToolbar.cpp
@@ -25,6 +25,7 @@
 #include "ToolMenus.h"
 #include "ViewportToolbar/UnrealEdViewportToolbarContext.h"
 #include "Widgets/Input/SNumericEntryBox.h"
+#include "Widgets/Input/SSlider.h"
 #include "Widgets/Input/SSpinBox.h"
 #include "Widgets/Layout/SUniformGridPanel.h"
 #include "Widgets/SToolTip.h"
@@ -549,6 +550,64 @@ bool IsViewModeSupported(EViewModeIndex InViewModeIndex)
 	}
 }
 
+void SetCameraSpeed(const TSharedRef<SEditorViewport>& InEditorViewport, float NewValue)
+{
+	const int32 OldSpeedSetting = InEditorViewport->GetViewportClient()->GetCameraSpeedSetting();
+	const int32 NewSpeedSetting = NewValue * ((float)FEditorViewportClient::MaxCameraSpeeds - 1) + 1;
+
+	if (OldSpeedSetting != NewSpeedSetting)
+	{
+		InEditorViewport->GetViewportClient()->SetCameraSpeedSetting(NewSpeedSetting);
+		UE::UnrealEd::OnViewportClientCamSpeedChanged().ExecuteIfBound(InEditorViewport, NewSpeedSetting);
+	}
+}
+
+float GetCamSpeedSliderPosition(const TSharedRef<SEditorViewport>& InEditorViewport)
+{
+	float SliderPos = 0.0f;
+
+	if (InEditorViewport->GetViewportClient().IsValid())
+	{
+		SliderPos = (InEditorViewport->GetViewportClient()->GetCameraSpeedSetting() - 1)
+				  / ((float)FEditorViewportClient::MaxCameraSpeeds - 1);
+	}
+
+	return SliderPos;
+}
+
+FText GetCameraSpeedLabel(const TSharedRef<SEditorViewport>& InEditorViewport)
+{
+	if (InEditorViewport->GetViewportClient().IsValid())
+	{
+		float CameraSpeed = InEditorViewport->GetViewportClient()->GetCameraSpeed();
+		FNumberFormattingOptions FormattingOptions = FNumberFormattingOptions::DefaultNoGrouping();
+		FormattingOptions.MaximumFractionalDigits = CameraSpeed > 1 ? 1 : 3;
+		return FText::AsNumber(CameraSpeed, &FormattingOptions);
+	}
+
+	return FText();
+}
+
+void SetCameraSpeedScalarValue(const TSharedRef<SEditorViewport>& InEditorViewport, float NewValue)
+{
+	if (InEditorViewport->GetViewportClient().IsValid())
+	{
+		InEditorViewport->GetViewportClient()->SetCameraSpeedScalar(NewValue);
+		UE::UnrealEd::OnViewportClientCamSpeedScalarChanged().ExecuteIfBound(InEditorViewport, NewValue);
+	}
+}
+
+float GetCamSpeedScalarSliderPosition(const TSharedRef<SEditorViewport>& InEditorViewport)
+{
+	float CamSpeedScalar = 1.0f;
+
+	if (InEditorViewport->GetViewportClient().IsValid())
+	{
+		CamSpeedScalar = (InEditorViewport->GetViewportClient()->GetCameraSpeedScalar());
+	}
+
+	return CamSpeedScalar;
+}
 }
 
 static FAutoConsoleVariableRef CVarToolMenusViewportToolbars(
@@ -1854,11 +1913,9 @@ FToolMenuEntry CreateCheckboxSubmenu(
 	return CheckBoxSubmenu;
 }
 
-FToolMenuEntry CreateNumericEntry(
-	const FName InName,
+TSharedRef<SWidget> CreateNumericEntryWidget(
 	const FText& InLabel,
 	const FText& InTooltip,
-	const FCanExecuteAction& InCanExecuteAction,
 	const FNumericEntryExecuteActionDelegate& InOnValueChanged,
 	const TAttribute<float>& InGetValue,
 	float InMinValue,
@@ -1868,11 +1925,8 @@ FToolMenuEntry CreateNumericEntry(
 {
 	const FMargin WidgetsMargin(2.0f, 0.0f, 3.0f, 0.0f);
 
-	FToolMenuEntry NumericEntry = FToolMenuEntry::InitMenuEntry(
-		InName,
-		FUIAction(FExecuteAction(), InCanExecuteAction),
-		// clang-format off
-		SNew(SHorizontalBox)
+	// clang-format off
+	return SNew(SHorizontalBox)
 		+ SHorizontalBox::Slot()
 		.VAlign(VAlign_Center)
 		.HAlign(HAlign_Left)
@@ -1900,11 +1954,35 @@ FToolMenuEntry CreateNumericEntry(
 				.AllowSpin(true)
 				.MaxFractionalDigits(InMaxFractionalDigits)
 				.Font(FAppStyle::GetFontStyle(TEXT("MenuItem.Font")))
-				.OnValueChanged_Lambda([InOnValueChanged](float InValue){ InOnValueChanged.Execute(InValue); })
-				.Value_Lambda([InGetValue](){ return InGetValue.Get(); })
+				.OnValueChanged_Lambda([InOnValueChanged](float InValue)
+				{
+					InOnValueChanged.Execute(InValue);
+				})
+				.Value_Lambda([InGetValue]()
+				{
+					return InGetValue.Get();
+				})
 			]
-		]
-		// clang-format on
+		];
+	// clang-format on
+}
+
+FToolMenuEntry CreateNumericEntry(
+	const FName InName,
+	const FText& InLabel,
+	const FText& InTooltip,
+	const FCanExecuteAction& InCanExecuteAction,
+	const FNumericEntryExecuteActionDelegate& InOnValueChanged,
+	const TAttribute<float>& InGetValue,
+	float InMinValue,
+	float InMaxValue,
+	int32 InMaxFractionalDigits
+)
+{
+	FToolMenuEntry NumericEntry = FToolMenuEntry::InitMenuEntry(
+		InName,
+		FUIAction(FExecuteAction(), InCanExecuteAction),
+		CreateNumericEntryWidget(InLabel, InTooltip, InOnValueChanged, InGetValue, InMinValue, InMaxValue, InMaxFractionalDigits)
 	);
 
 	return NumericEntry;
@@ -2127,6 +2205,71 @@ FName GetCameraSubmenuIconFNameFromViewportType(const ELevelViewportType Viewpor
 	return Icon;
 }
 
+TSharedRef<SWidget> CreateCameraSpeedSliderWidget(const TWeakPtr<SEditorViewport>& InEditorViewportWeak)
+{
+	const FMargin WidgetsMargin(2.0f, 0.0f, 3.0f, 0.0f);
+
+	// clang-format off
+	return
+		SNew(SHorizontalBox)
+		+ SHorizontalBox::Slot()
+		.VAlign(VAlign_Center)
+		.HAlign(HAlign_Left)
+		.Padding(WidgetsMargin)
+		.AutoWidth()
+		[
+			SNew(STextBlock)
+			.Text(LOCTEXT("CameraSpeedLabel", "Camera Speed"))
+		]
+		+ SHorizontalBox::Slot()
+		.Padding(FMargin(0.0f, 2.0f, 3.0f, 2.0f))
+		.FillWidth(1.0f)
+		[
+			SNew(SVerticalBox)
+			+ SVerticalBox::Slot()
+			.HAlign(HAlign_Center)
+			.AutoHeight()
+			[
+				SNew(STextBlock)
+				.Text_Lambda([InEditorViewportWeak]()
+				{
+					if (TSharedPtr<SEditorViewport> LevelViewport = InEditorViewportWeak.Pin())
+					{
+						return UE::UnrealEd::Private::GetCameraSpeedLabel(LevelViewport.ToSharedRef());
+					}
+					return FText();
+				})
+				.Font(FAppStyle::GetFontStyle(TEXT("MenuItem.Font")))
+			]
+			+ SVerticalBox::Slot()
+			.Padding(WidgetsMargin)
+			.HAlign(HAlign_Fill)
+			[
+				SNew(SBox)
+				.MinDesiredWidth(80.0f)
+				[
+					SNew(SSlider)
+					.Value_Lambda([InEditorViewportWeak]()
+					{
+						if (TSharedPtr<SEditorViewport> LevelViewport = InEditorViewportWeak.Pin())
+						{
+							return UE::UnrealEd::Private::GetCamSpeedSliderPosition(LevelViewport.ToSharedRef());
+						}
+						return 0.0f;
+					})
+					.OnValueChanged_Lambda([InEditorViewportWeak](float InNewValue)
+					{
+						if (TSharedPtr<SEditorViewport> LevelViewport = InEditorViewportWeak.Pin())
+						{
+							UE::UnrealEd::Private::SetCameraSpeed(LevelViewport.ToSharedRef(), InNewValue);
+						}
+					})
+				]
+			]
+		];
+	// clang-format on
+}
+
 FToolMenuEntry CreateViewportToolbarCameraSubmenu()
 {
 	return FToolMenuEntry::InitDynamicEntry(
@@ -2335,15 +2478,13 @@ void ExtendCameraSubmenu(FName InCameraOptionsSubmenuName)
 				UnnamedSection.AddEntry(FToolMenuEntry::InitWidget(
 					"CameraFOV",
 					UE::UnrealEd::CreateFOVMenuWidget(EditorViewport.ToSharedRef()),
-					LOCTEXT("CameraSubmenu_FieldOfViewLabel", "Field of View"),
-					true
+					LOCTEXT("CameraSubmenu_FieldOfViewLabel", "Field of View")
 				));
 
 				UnnamedSection.AddEntry(FToolMenuEntry::InitWidget(
 					"CameraFarViewPlane",
 					UE::UnrealEd::CreateFarViewPlaneMenuWidget(EditorViewport.ToSharedRef()),
-					LOCTEXT("CameraSubmenu_FarViewPlaneLabel", "Far View Plane"),
-					true
+					LOCTEXT("CameraSubmenu_FarViewPlaneLabel", "Far View Plane")
 				));
 			}
 		)
@@ -2894,6 +3035,88 @@ FToolMenuEntry CreateToggleRealtimeEntry()
 	);
 }
 
+FOnViewportClientCamSpeedChanged& OnViewportClientCamSpeedChanged()
+{
+	static FOnViewportClientCamSpeedChanged OnViewportClientCamSpeedChangedDelegate;
+	return OnViewportClientCamSpeedChangedDelegate;
+}
+
+FOnViewportClientCamSpeedScalarChanged& OnViewportClientCamSpeedScalarChanged()
+{
+	static FOnViewportClientCamSpeedScalarChanged OnViewportClientCamSpeedScalarChangedDelegate;
+	return OnViewportClientCamSpeedScalarChangedDelegate;
+}
+
+FToolMenuEntry CreateCameraSpeedSlider(const TWeakPtr<::SEditorViewport>& InEditorViewportWeak)
+{
+	FToolMenuEntry CameraSpeedSliderEntry = FToolMenuEntry::InitMenuEntry(
+		"CameraSpeed", FUIAction(FExecuteAction()), CreateCameraSpeedSliderWidget(InEditorViewportWeak)
+	);
+
+	return CameraSpeedSliderEntry;
+}
+
+FToolMenuEntry CreateCameraSpeedScalarSlider(TWeakPtr<::SEditorViewport> InEditorViewportWeak)
+{
+	constexpr float MinSpeed = 1.0f;
+	constexpr float MaxSpeed = 128.0f;
+
+	return UE::UnrealEd::CreateNumericEntry(
+		"CameraSpeedScalar",
+		LOCTEXT("CameraSpeedScalarLabel", "Speed Scalar"),
+		LOCTEXT("CameraSpeedScalarTooltip", "Scalar to increase camera movement range"),
+		FCanExecuteAction(),
+		UE::UnrealEd::FNumericEntryExecuteActionDelegate::CreateLambda(
+			[InEditorViewportWeak](float InValue)
+			{
+				if (TSharedPtr<::SEditorViewport> LevelViewport = InEditorViewportWeak.Pin())
+				{
+					UE::UnrealEd::Private::SetCameraSpeedScalarValue(LevelViewport.ToSharedRef(), InValue);
+				}
+			}
+		),
+		TAttribute<float>::CreateLambda(
+			[InEditorViewportWeak]()
+			{
+				if (TSharedPtr<::SEditorViewport> Viewport = InEditorViewportWeak.Pin())
+				{
+					return UE::UnrealEd::Private::GetCamSpeedScalarSliderPosition(Viewport.ToSharedRef());
+				}
+
+				return MinSpeed;
+			}
+		),
+		MinSpeed,
+		MaxSpeed,
+		1
+	);
+}
+
+void CreateCameraSpeedMenu(UToolMenu* InMenu)
+{
+	if (!InMenu)
+	{
+		return;
+	}
+
+	FToolMenuSection& CameraSpeedSection = InMenu->AddSection("CameraSpeed", LOCTEXT("CameraSpeedLabel", "Camera Speed"));
+
+	UUnrealEdViewportToolbarContext* const EditorViewportContext = InMenu->FindContext<UUnrealEdViewportToolbarContext>();
+	if (!EditorViewportContext)
+	{
+		return;
+	}
+
+	const TSharedPtr<SEditorViewport> EditorViewport = EditorViewportContext->Viewport.Pin();
+	if (!EditorViewport)
+	{
+		return;
+	}
+
+	CameraSpeedSection.AddEntry(CreateCameraSpeedSlider(EditorViewport));
+	CameraSpeedSection.AddEntry(CreateCameraSpeedScalarSlider(EditorViewport));
+}
+
 } // namespace UE::UnrealEd
 
 #undef LOCTEXT_NAMESPACE
diff --git a/Engine/Source/Editor/UnrealEd/Public/ViewportToolbar/UnrealEdViewportToolbar.h b/Engine/Source/Editor/UnrealEd/Public/ViewportToolbar/UnrealEdViewportToolbar.h
--- a/Engine/Source/Editor/UnrealEd/Public/ViewportToolbar/UnrealEdViewportToolbar.h
+++ b/Engine/Source/Editor/UnrealEd/Public/ViewportToolbar/UnrealEdViewportToolbar.h
@@ -70,6 +70,9 @@ DECLARE_DELEGATE_RetVal_OneParam(bool, FScaleGridCheckboxListIsCheckedDelegate,
 
 DECLARE_DELEGATE_OneParam(FNumericEntryExecuteActionDelegate, float);
 
+DECLARE_DELEGATE_TwoParams(FOnViewportClientCamSpeedChanged, const TSharedRef<SEditorViewport>&, int32);
+DECLARE_DELEGATE_TwoParams(FOnViewportClientCamSpeedScalarChanged, const TSharedRef<SEditorViewport>&, float);
+
 UNREALED_API TSharedRef<SWidget> BuildRotationGridCheckBoxList(
 	FName InExtentionHook,
 	const FText& InHeading,
@@ -160,6 +163,11 @@ UNREALED_API void AddDefaultShowFlags(UToolMenu* InMenu);
 
 UNREALED_API FToolMenuEntry CreateToggleRealtimeEntry();
 
+UNREALED_API FOnViewportClientCamSpeedChanged& OnViewportClientCamSpeedChanged();
+UNREALED_API FOnViewportClientCamSpeedScalarChanged& OnViewportClientCamSpeedScalarChanged();
+
+UNREALED_API void CreateCameraSpeedMenu(UToolMenu* InMenu);
+
 // Camera Menu Widgets
 TSharedRef<SWidget> CreateCameraMenuWidget(const TSharedRef<SEditorViewport>& InViewport);
 TSharedRef<SWidget> CreateFOVMenuWidget(const TSharedRef<SEditorViewport>& InViewport);
