Localization:
Fixed a bug in DetectOrphanedLocalizedAssetsCommandlet
- There were a number of false positives being raised because we were only converting the package name of the localized asset to the source asset form and checking if that existed in the asset registry.
- Fix was to convert the full localized asset path to the source version and checking the asset registry.
- We also no longer filter against Audit_InCook as orphaned localized assets to not actually get cooked and by extension not bmake it into Audit_InCook
- Commandlet currently operates on all assets on disk and as such the final report will have assets that might not be in the cook. A TODO down the line.

diff --git a/Engine/Source/Editor/UnrealEd/Private/Commandlets/DetectOrphanedLocalizedAssetsCommandlet.cpp b/Engine/Source/Editor/UnrealEd/Private/Commandlets/DetectOrphanedLocalizedAssetsCommandlet.cpp
--- a/Engine/Source/Editor/UnrealEd/Private/Commandlets/DetectOrphanedLocalizedAssetsCommandlet.cpp
+++ b/Engine/Source/Editor/UnrealEd/Private/Commandlets/DetectOrphanedLocalizedAssetsCommandlet.cpp
@@ -1,11 +1,8 @@
 // Copyright Epic Games, Inc. All Rights Reserved.
 
 #include "Commandlets/DetectOrphanedLocalizedAssetsCommandlet.h"
-#include "CollectionManagerModule.h"
-#include "AssetRegistry/ARFilter.h"
 #include "AssetRegistry/AssetData.h"
 #include "AssetRegistry/IAssetRegistry.h"
-#include "ICollectionManager.h"
 #include "Misc/FileHelper.h"
 #include "Internationalization/PackageLocalizationUtil.h"
 
@@ -36,25 +33,14 @@ int32 UDetectOrphanedLocalizedAssetsCommandlet::Main(const FString& Params)
 		UE_LOG(LogDetectOrphanedLocalizedAssetsCommandlet, Display, TEXT("%s"), *UsageText);
 		return 0;
 	}
-
-	FCollectionManagerModule& CollectionManagerModule = FCollectionManagerModule::GetModule();
-	ICollectionManager& CollectionManager = CollectionManagerModule.Get();
-	FARFilter CollectionFilter;
-	// @TODOLocalization: Don't hardcode and instead turn this into a parameter for the commandlet 
-	bool bSuccess = CollectionManager.GetObjectsInCollection(FName("Audit_InCook"), ECollectionShareType::CST_All, CollectionFilter.SoftObjectPaths, ECollectionRecursionFlags::SelfAndChildren);
-	
 	IAssetRegistry::Get()->SearchAllAssets(true);
+	// We do not filter the asset registry against the Audit_InCook collection because orphaned localized assets will not be cooked and therefore not make it into the Audit_InCook collection 
+
+
 	TArray<FAssetData> AllAssets;
 	const double GetAllAssetsWithFirstPassFilterStartTime = FPlatformTime::Seconds();
-	
-	if (bSuccess)
-	{
-		IAssetRegistry::GetChecked().GetAssets(CollectionFilter, AllAssets);
-	}
-	else
-	{
-		IAssetRegistry::GetChecked().GetAllAssets(AllAssets);
-	}
+	UE_LOG(LogDetectOrphanedLocalizedAssetsCommandlet, Display, TEXT("Detecting orphaned localized assets across all assets on disk."));
+	IAssetRegistry::GetChecked().GetAllAssets(AllAssets);
 	UE_LOG(LogDetectOrphanedLocalizedAssetsCommandlet, Display, TEXT("Getting all assets from asset registry took %.2f seconds."), FPlatformTime::Seconds() - GetAllAssetsWithFirstPassFilterStartTime);
 
 	UE_LOG(LogDetectOrphanedLocalizedAssetsCommandlet, Display, TEXT("Processing %d assets."), AllAssets.Num());
@@ -81,8 +67,8 @@ int32 UDetectOrphanedLocalizedAssetsCommandlet::Main(const FString& Params)
 	for (const FSoftObjectPath& LocalizedAsset: LocalizedAssets)
 	{
 		SourceObjectPath.Reset();
-
-		if (FPackageLocalizationUtil::ConvertLocalizedToSource(LocalizedAsset.GetLongPackageName(), SourceObjectPath))
+		// Despite the namespace, the function can indeed also convert localized assets to source 
+		if (FPackageLocalizationUtil::ConvertLocalizedToSource(LocalizedAsset.GetAssetPathString(), SourceObjectPath))
 		{
 			UE::AssetRegistry::EExists Exists = AssetRegistry.TryGetAssetByObjectPath(FSoftObjectPath(SourceObjectPath), OutAssetData);
 			if (Exists == UE::AssetRegistry::EExists::Exists && !OutAssetData.IsRedirector())
@@ -119,6 +105,23 @@ int32 UDetectOrphanedLocalizedAssetsCommandlet::Main(const FString& Params)
 	float OrphanedPercentage = static_cast<float>(OrphanedLocalizedAssets.Num()) / static_cast<float>(LocalizedAssets.Num()) * 100.0f;
 	UE_LOG(LogDetectOrphanedLocalizedAssetsCommandlet, Display, TEXT("%d out of %d localized assets are orphaned. %.2f%% of all localized assets are orphaned."), OrphanedLocalizedAssets.Num(), LocalizedAssets.Num(), OrphanedPercentage);
 
+	if (OrphanedLocalizedAssets.Num() == 0)
+	{
+		UE_LOG(LogDetectOrphanedLocalizedAssetsCommandlet, Display, TEXT("No orphaned localized assets found."));
+		return 0;
+	}
+
+	// We sort the soft object paths first to cluster all related assets together 
+	OrphanedLocalizedAssets.Sort([](const FSoftObjectPath& A, const FSoftObjectPath& B)
+	{
+		int32 Diff = FCString::Strcmp(*A.GetAssetName(), *B.GetAssetName());
+		// localized variants will share identical asset names, we want to sort by package path to cluster all the locales together in that case 
+		if (Diff == 0)
+		{
+			return A.GetLongPackageName() < B.GetLongPackageName();
+		}
+		return (Diff < 0);
+	});
 	TArray<FString> OrphanedLocalizedAssetsStrings;
 	OrphanedLocalizedAssetsStrings.Reserve(OrphanedLocalizedAssets.Num());
 	for (const FSoftObjectPath& OrphanedAsset : OrphanedLocalizedAssets)
@@ -126,9 +129,6 @@ int32 UDetectOrphanedLocalizedAssetsCommandlet::Main(const FString& Params)
 		OrphanedLocalizedAssetsStrings.Add(OrphanedAsset.ToString());
 	}
 
-	// We sort the strings so it's easier to cluster all assets that are in the same module for fixing.
-	OrphanedLocalizedAssetsStrings.Sort();
-
 	if (ParamVals.Contains("OutputOrphans"))
 	{
 		UE_LOG(LogDetectOrphanedLocalizedAssetsCommandlet, Display, TEXT("An output file was provided. Dumping all found orphaned assets to the file: %s"), *ParamVals["OutputOrphans"]);
