Improve handling of the '-PackageFolder' arg when running ResavePackagesCommandlet.

- Sanitize the path provided by the user by normalizing it.
- Log a warning if the provided path does not contain any assets. It seems reasonable to assume that if the user provided a directory path that they expect it to contain assets.

diff --git a/Engine/Source/Editor/UnrealEd/Private/Commandlets/ContentCommandlets.cpp b/Engine/Source/Editor/UnrealEd/Private/Commandlets/ContentCommandlets.cpp
--- a/Engine/Source/Editor/UnrealEd/Private/Commandlets/ContentCommandlets.cpp
+++ b/Engine/Source/Editor/UnrealEd/Private/Commandlets/ContentCommandlets.cpp
@@ -158,6 +158,8 @@ int32 UResavePackagesCommandlet::InitializeResaveParameters( const TArray<FStrin
 		}
 		else if( FParse::Value( *CurrentSwitch, TEXT( "PACKAGEFOLDER="), PackageFolder ) )
 		{
+			FPaths::NormalizeDirectoryName(PackageFolder);
+
 			TArray<FString> FilesInPackageFolder;
 			FPackageName::FindPackagesInDirectory(FilesInPackageFolder, PackageFolder);
 			for( int32 FileIndex = 0; FileIndex < FilesInPackageFolder.Num(); FileIndex++ )
@@ -166,6 +168,8 @@ int32 UResavePackagesCommandlet::InitializeResaveParameters( const TArray<FStrin
 				FPaths::MakeStandardFilename(PackageFile);
 				PackageNames.Add( *PackageFile );
 			}
+
+			UE_CLOG(PackageNames.IsEmpty(), LogContentCommandlet, Warning, TEXT("Failed to find any packages in folder: '%s'"), *PackageFolder);
 			bExplicitPackages = true;
 		}
 		else if (FParse::Value(*CurrentSwitch, TEXT("MAP="), Maps))
