Fix the tooltip of a number of entries of new viewport toolbars

diff --git a/Engine/Source/Editor/UnrealEd/Private/ViewportToolbar/UnrealEdViewportToolbar.cpp b/Engine/Source/Editor/UnrealEd/Private/ViewportToolbar/UnrealEdViewportToolbar.cpp
--- a/Engine/Source/Editor/UnrealEd/Private/ViewportToolbar/UnrealEdViewportToolbar.cpp
+++ b/Engine/Source/Editor/UnrealEd/Private/ViewportToolbar/UnrealEdViewportToolbar.cpp
@@ -140,7 +140,7 @@ FToolMenuEntry CreateSurfaceSnapCheckboxMenu()
 	FToolMenuEntry Entry = UnrealEd::CreateCheckboxSubmenu(
 		"SurfaceSnapping",
 		LOCTEXT("SurfaceSnapLabel", "Surface"),
-		FEditorViewportCommands::Get().SurfaceSnapping->MakeTooltip()->GetTextTooltip(),
+		FEditorViewportCommands::Get().SurfaceSnapping->GetDescription(),
 		FToolMenuExecuteAction::CreateLambda(
 			[](const FToolMenuContext& InContext)
 			{
@@ -266,7 +266,7 @@ FToolMenuEntry CreateActorSnapCheckboxMenu()
 	FToolMenuEntry ActorSnapping = FToolMenuEntry::InitSubMenu(
 		"ActorSnapping",
 		LOCTEXT("ActorSnapLabel", "Actor"),
-		FLevelEditorCommands::Get().EnableActorSnap->MakeTooltip()->GetTextTooltip(),
+		FLevelEditorCommands::Get().EnableActorSnap->GetDescription(),
 		MakeMenuDelegate,
 		CheckboxMenuAction,
 		EUserInterfaceActionType::ToggleButton
@@ -337,7 +337,7 @@ FToolMenuEntry CreateLocationSnapCheckboxMenu()
 		Entry = UnrealEd::CreateCheckboxSubmenu(
 			"GridSnapping",
 			LocationSnapLabel,
-			FEditorViewportCommands::Get().SurfaceSnapping->MakeTooltip()->GetTextTooltip(),
+			FEditorViewportCommands::Get().SurfaceSnapping->GetDescription(),
 			FToolMenuExecuteAction::CreateLambda(
 				[](const FToolMenuContext& InContext)
 				{
@@ -409,7 +409,7 @@ FToolMenuEntry CreateRotationSnapCheckboxMenu()
 	FToolMenuEntry Entry = UnrealEd::CreateCheckboxSubmenu(
 		"RotationSnapping",
 		RotationSnapLabel,
-		FEditorViewportCommands::Get().SurfaceSnapping->MakeTooltip()->GetTextTooltip(),
+		FEditorViewportCommands::Get().SurfaceSnapping->GetDescription(),
 		FToolMenuExecuteAction::CreateLambda(
 			[](const FToolMenuContext& InContext)
 			{
@@ -513,7 +513,7 @@ FToolMenuEntry CreateScaleSnapCheckboxMenu()
 	FToolMenuEntry Entry = UnrealEd::CreateCheckboxSubmenu(
 		"ScaleSnapping",
 		ScaleSnapLabel,
-		FEditorViewportCommands::Get().SurfaceSnapping->MakeTooltip()->GetTextTooltip(),
+		FEditorViewportCommands::Get().SurfaceSnapping->GetDescription(),
 		FToolMenuExecuteAction::CreateLambda(
 			[](const FToolMenuContext& InContext)
 			{
@@ -811,21 +811,21 @@ FToolMenuEntry CreateViewportToolbarSelectSection()
 					UnnamedSection.AddMenuEntry(
 						FGenericCommands::Get().SelectAll,
 						FGenericCommands::Get().SelectAll->GetLabel(),
-						FGenericCommands::Get().SelectAll->MakeTooltip()->GetTextTooltip(),
+						FGenericCommands::Get().SelectAll->GetDescription(),
 						FSlateIconFinder::FindIcon("FoliageEditMode.SelectAll")
 					);
 
 					UnnamedSection.AddMenuEntry(
 						FLevelEditorCommands::Get().SelectNone,
 						FLevelEditorCommands::Get().SelectNone->GetLabel(),
-						FLevelEditorCommands::Get().SelectNone->MakeTooltip()->GetTextTooltip(),
+						FLevelEditorCommands::Get().SelectNone->GetDescription(),
 						FSlateIconFinder::FindIcon("Cross")
 					);
 
 					UnnamedSection.AddMenuEntry(
 						FLevelEditorCommands::Get().InvertSelection,
 						FLevelEditorCommands::Get().InvertSelection->GetLabel(),
-						FLevelEditorCommands::Get().InvertSelection->MakeTooltip()->GetTextTooltip(),
+						FLevelEditorCommands::Get().InvertSelection->GetDescription(),
 						FSlateIconFinder::FindIcon("FoliageEditMode.DeselectAll")
 					);
 
@@ -863,7 +863,7 @@ FToolMenuEntry CreateViewportToolbarSelectSection()
 					UnnamedSection.AddMenuEntry(
 						FLevelEditorCommands::Get().SelectAllActorsOfSameClass,
 						LOCTEXT("AdvancedSelectAllActorsOfSameClassLabel", "All of Same Class"),
-						FLevelEditorCommands::Get().SelectAllActorsOfSameClass->MakeTooltip()->GetTextTooltip(),
+						FLevelEditorCommands::Get().SelectAllActorsOfSameClass->GetDescription(),
 						FSlateIconFinder::FindIcon("PlacementBrowser.Icons.All")
 					);
 				}
