Improvements for the dump SHK commandlet:
* allow multiple hashes to be specified for the -outputhash argument
* implement dumping only certain objects via the -objname argument

#rnx

diff --git a/Engine/Source/Editor/UnrealEd/Private/Commandlets/ShaderPipelineCacheToolsCommandlet.cpp b/Engine/Source/Editor/UnrealEd/Private/Commandlets/ShaderPipelineCacheToolsCommandlet.cpp
--- a/Engine/Source/Editor/UnrealEd/Private/Commandlets/ShaderPipelineCacheToolsCommandlet.cpp
+++ b/Engine/Source/Editor/UnrealEd/Private/Commandlets/ShaderPipelineCacheToolsCommandlet.cpp
@@ -3,6 +3,7 @@
 #include "Commandlets/ShaderPipelineCacheToolsCommandlet.h"
 
 #include "Algo/Accumulate.h"
+#include "Algo/Find.h"
 #include "Async/ParallelFor.h"
 #include "Async/TaskGraphInterfaces.h"
 #include "Containers/Array.h"
@@ -672,25 +673,50 @@ bool CouldBeUsedTogether(const FStableShaderKeyAndValue& A, const FStableShaderK
 	return true;
 }
 
-int32 DumpStableKeysFile(const FString& Token, const TMap<FString, FString>& ParamVals)
+static TArray<FString> ParseParamStringList(const TMap<FString, FString>& ParamVals, const TCHAR* ParamName)
 {
-	const FStringView File = Token;
-	FString OutputHashToSearch;
-	if (const FString* Param = ParamVals.Find(TEXT("outputhash")))
+	TArray<FString> StrList;
+
+	const FString* Param = ParamVals.Find(ParamName);
+	if (!Param)
 	{
-		OutputHashToSearch = *Param;
+		return StrList;
 	}
 
+	Param->ToLower().ParseIntoArray(StrList, TEXT(","));
+	for (FString& Str : StrList)
+	{
+		Str.TrimStartAndEndInline();
+	}
+	return StrList;
+}
+
+int32 DumpStableKeysFile(const FString& Token, const TMap<FString, FString>& ParamVals)
+{
+	const FStringView File = Token;
+	
+	TArray<FString> OutputHashesToSearch = ParseParamStringList(ParamVals, TEXT("outputhash"));
+	TArray<FString> NamesToSearch = ParseParamStringList(ParamVals, TEXT("objname"));
+
 	TMultiMap<FStableShaderKeyAndValue, FSHAHash> StableMap;
 	LoadStableShaderKeysMultiple(StableMap, MakeArrayView(&File, 1));
 	UE_LOG(LogShaderPipelineCacheTools, Display, TEXT("    %s"), *FStableShaderKeyAndValue::HeaderLine());
 	for (const auto& Pair : StableMap)
 	{
-		if (!OutputHashToSearch.IsEmpty() && Pair.Key.OutputHash.ToString() != OutputHashToSearch)
+		if (!OutputHashesToSearch.IsEmpty() && !OutputHashesToSearch.Contains(Pair.Key.OutputHash.ToString().ToLower()))
 		{
 			continue;
 		}
 
+		if (!NamesToSearch.IsEmpty())
+		{
+			const FString* Found = Algo::FindByPredicate(NamesToSearch, [ObjName = Pair.Key.ClassNameAndObjectPath.ToString().ToLower()](const FString& PartialName) { return ObjName.Contains(PartialName); });
+			if (!Found)
+			{
+				continue;
+			}
+		}
+
 		FStableShaderKeyAndValue Temp(Pair.Key);
 		UE_LOG(LogShaderPipelineCacheTools, Display, TEXT("    %s"), *Temp.ToString());
 	}
@@ -2889,7 +2915,7 @@ int32 UShaderPipelineCacheToolsCommandlet::StaticMain(const FString& Params)
 		}
 	}
 	
-	UE_LOG(LogShaderPipelineCacheTools, Warning, TEXT("Usage: Dump SCLInfo.%s [...] [-outputhash=X] - dumps stable keys file, optionally only dumping keys that match the specified output hash value.\n"), ShaderStableKeysFileExt);
+	UE_LOG(LogShaderPipelineCacheTools, Warning, TEXT("Usage: Dump SCLInfo.%s [...] [-outputhash=LIST] [-objname=LIST] - dumps stable keys file, optionally only dumping entries that match the specified output hashes and/or object names (hash and name lists are comma-separated).\n"), ShaderStableKeysFileExt);
 	UE_LOG(LogShaderPipelineCacheTools, Warning, TEXT("Usage: Dump PSOCache.upipelinecache [SCLInfo.%s] - dumps recorded PSOs. If optional stable keys file, converts shader hashes to readable stable shader descriptions.\n"), ShaderStableKeysFileExt);
 	UE_LOG(LogShaderPipelineCacheTools, Warning, TEXT("Usage: Dump StablePSOCache.csv[.compressed] - dumps content of a stable PSO.\n"));
 	UE_LOG(LogShaderPipelineCacheTools, Warning, TEXT("Usage: Diff ShaderCache1.stablepc.csv ShaderCache1.stablepc.csv [...]]\n"));
