Improve the preview profile selection menu for asset editors

diff --git a/Engine/Source/Editor/EditorStyle/Private/StarshipStyle.cpp b/Engine/Source/Editor/EditorStyle/Private/StarshipStyle.cpp
--- a/Engine/Source/Editor/EditorStyle/Private/StarshipStyle.cpp
+++ b/Engine/Source/Editor/EditorStyle/Private/StarshipStyle.cpp
@@ -1215,6 +1215,8 @@ void FStarshipEditorStyle::FStyle::SetupGeneralStyles()
 		Set( "AssetEditor.ReadOnlyBorder", new FSlateRoundedBoxBrush(FStyleColors::Foreground, 10.0f));
 		Set("AssetEditor.ReadOnlyOpenable", new IMAGE_BRUSH_SVG("Starship/AssetEditors/LockEye", Icon16x16));
 
+		Set("AssetEditor.PreviewSceneSettings",
+			new IMAGE_BRUSH_SVG("Starship/AssetEditors/PreviewSceneSettings_16", Icon16x16));
 	}
 		
 	// Asset Thumbnail
diff --git a/Engine/Source/Editor/UnrealEd/Private/SCommonEditorViewportToolbarBase.cpp b/Engine/Source/Editor/UnrealEd/Private/SCommonEditorViewportToolbarBase.cpp
--- a/Engine/Source/Editor/UnrealEd/Private/SCommonEditorViewportToolbarBase.cpp
+++ b/Engine/Source/Editor/UnrealEd/Private/SCommonEditorViewportToolbarBase.cpp
@@ -4,6 +4,7 @@
 #include "Widgets/SBoxPanel.h"
 #include "Widgets/Layout/SBorder.h"
 #include "Widgets/Layout/SBox.h"
+#include "Widgets/Input/SComboButton.h"
 #include "Widgets/Input/SSpinBox.h"
 #include "Widgets/Input/STextComboBox.h"
 #include "Widgets/SToolTip.h"
@@ -31,96 +32,108 @@
 //////////////////////////////////////////////////////////////////////////
 // SPreviewSceneProfileSelector
 
-SPreviewSceneProfileSelector::~SPreviewSceneProfileSelector()
-{
-	if (PreviewProfileController)
-	{
-		PreviewProfileController->OnPreviewProfileListChanged().RemoveAll(this);
-		PreviewProfileController->OnPreviewProfileChanged().RemoveAll(this);
-	}
-}
-
 void SPreviewSceneProfileSelector::Construct(const FArguments& InArgs)
 {
 	PreviewProfileController = InArgs._PreviewProfileController;
-	if (PreviewProfileController)
-	{
-		PreviewProfileController->OnPreviewProfileListChanged().AddRaw(this, &SPreviewSceneProfileSelector::UpdateAssetViewerProfileList);
-		PreviewProfileController->OnPreviewProfileChanged().AddRaw(this, &SPreviewSceneProfileSelector::UpdateAssetViewerProfileSelection);
-		UpdateAssetViewerProfileList();
-	}
-	
+
+	// clang-format off
+	TSharedRef<SHorizontalBox> ButtonContent = 
+		SNew(SHorizontalBox)
+		+SHorizontalBox::Slot()
+		.VAlign(VAlign_Center)
+		.HAlign(HAlign_Center)
+		.AutoWidth()
+		.Padding(4.0f, 0.0f, 4.f, 0.0f)
+		[
+			SNew(SImage)
+			.Image(FAppStyle::GetBrush("AssetEditor.PreviewSceneSettings"))
+			.ColorAndOpacity(FSlateColor::UseForeground())
+		]
+		+SHorizontalBox::Slot()
+		.Padding(0.0f, 0.0f, 4.f, 0.0f)
+		.AutoWidth()
+		.VAlign(VAlign_Center)
+		[
+			SNew(STextBlock)
+			.Margin(FMargin(0))
+			.Text_Lambda(
+				[this]() -> FText
+				{
+					return FText::FromString(PreviewProfileController->GetActiveProfile());
+				}
+			)
+		];
+	// clang-format on
+
+	// clang-format off
 	ChildSlot
 	[
 		SNew(SVerticalBox)
 		+SVerticalBox::Slot()
 		.AutoHeight()
 		[
-			SAssignNew(AssetViewerProfileComboBox, STextComboBox)//SComboBox<TSharedPtr<FString>>)
-			.OptionsSource(&AssetViewerProfileNames)
+			SAssignNew(AssetViewerProfileComboButton, SComboButton)
 			.ButtonStyle(&FAppStyle::Get().GetWidgetStyle<FButtonStyle>("EditorViewportToolBar.Button"))
-			.ContentPadding(FMargin(2, 0))
-			.ToolTipText(LOCTEXT("AssetViewerProfile_ToolTip", "Select the Preview Scene Profile for this viewport."))
-			.OnSelectionChanged(this, &SPreviewSceneProfileSelector::OnSelectionChanged)
-			.Visibility_Lambda([this]() { return AssetViewerProfileNames.Num() > 1 ? EVisibility::Visible : EVisibility::Collapsed; })
+			.ContentPadding(FMargin(0))
+			.HasDownArrow(false)
+			.OnGetMenuContent(this, &SPreviewSceneProfileSelector::BuildComboMenu)
+			.ButtonContent()
+			[
+				ButtonContent
+			]
 		]
 	];
-
-	UpdateAssetViewerProfileSelection();
+	// clang-format on
 }
 
-void SPreviewSceneProfileSelector::OnSelectionChanged(TSharedPtr<FString> NewSelection, ESelectInfo::Type SelectInfo)
+TSharedRef<SWidget> SPreviewSceneProfileSelector::BuildComboMenu()
 {
-	int32 NewSelectionIndex;
-	if (AssetViewerProfileNames.Find(NewSelection, NewSelectionIndex))
-	{
-		// If that's the user changing the combo box, not an update coming from code to reflect a change that already occurred.
-		if (SelectInfo != ESelectInfo::Direct)
-		{
-			PreviewProfileController->SetActiveProfile(*NewSelection);
-		}
-	}
-}
-
-void SPreviewSceneProfileSelector::UpdateAssetViewerProfileList()
-{
-	if (PreviewProfileController)
-	{
-		// Pull the latest profile list.
-		int32 CurrProfileIndex = 0;
-		TArray<FString> ProfileNames = PreviewProfileController->GetPreviewProfiles(CurrProfileIndex);
+	const bool bShouldCloseWindowAfterMenuSelection = true;
+	TSharedPtr<const FUICommandList> CommandList = nullptr;
+	FMenuBuilder MenuBuilder(bShouldCloseWindowAfterMenuSelection, CommandList);
 
-		// Rebuild the combo box list.
-		AssetViewerProfileNames.Empty();
-		for (const FString& Profile : ProfileNames)
-		{
-			AssetViewerProfileNames.Add(MakeShared<FString>(Profile));
-		}
-
-		// Select the current profile item.
-		if (AssetViewerProfileComboBox)
-		{
-			AssetViewerProfileComboBox->RefreshOptions();
-			AssetViewerProfileComboBox->SetSelectedItem(AssetViewerProfileNames[CurrProfileIndex]);
-		}
-	}
-}
+	MenuBuilder.BeginSection(NAME_None, LOCTEXT("PreviewSceneProfilesSectionLabel", "Preview Scene Profiles"));
 
-void SPreviewSceneProfileSelector::UpdateAssetViewerProfileSelection()
-{
-	if (PreviewProfileController)
+	int32 UnusedActiveIndex;
+	const FName UnusedExtensionHook = NAME_None;
+	const TArray<FString> PreviewProfiles = PreviewProfileController->GetPreviewProfiles(UnusedActiveIndex);
+	for (const FString& ProfileName : PreviewProfiles)
 	{
-		FString ActiveProfileName = PreviewProfileController->GetActiveProfile();
-		if (TSharedPtr<FString>* Match = AssetViewerProfileNames.FindByPredicate(
-			[&ActiveProfileName](const TSharedPtr<FString>& Candidate) { return *Candidate == ActiveProfileName; }))
-		{
-			AssetViewerProfileComboBox->SetSelectedItem(*Match);
-		}
-		else // The profile was likely renamed.
-		{
-			UpdateAssetViewerProfileList();
-		}
+		MenuBuilder.AddMenuEntry(
+			FText::FromString(ProfileName),
+			FText(),
+			FSlateIcon(),
+			FUIAction(
+				FExecuteAction::CreateLambda(
+					[this, WeakController = PreviewProfileController.ToWeakPtr(), ProfileName]()
+					{
+						if (TSharedPtr<IPreviewProfileController> PinnedController = WeakController.Pin())
+						{
+							PinnedController->SetActiveProfile(ProfileName);
+						}
+					}
+				),
+				FCanExecuteAction(),
+				FIsActionChecked::CreateLambda(
+					[WeakController = PreviewProfileController.ToWeakPtr(), ProfileName]()
+					{
+						if (TSharedPtr<IPreviewProfileController> PinnedController = WeakController.Pin())
+						{
+							return ProfileName == PinnedController->GetActiveProfile();
+						}
+
+						return false;
+					}
+				)
+			),
+			UnusedExtensionHook,
+			EUserInterfaceActionType::RadioButton
+		);
 	}
+
+	MenuBuilder.EndSection();
+
+	return MenuBuilder.MakeWidget();
 }
 
 //////////////////////////////////////////////////////////////////////////
diff --git a/Engine/Source/Editor/UnrealEd/Public/SCommonEditorViewportToolbarBase.h b/Engine/Source/Editor/UnrealEd/Public/SCommonEditorViewportToolbarBase.h
--- a/Engine/Source/Editor/UnrealEd/Public/SCommonEditorViewportToolbarBase.h
+++ b/Engine/Source/Editor/UnrealEd/Public/SCommonEditorViewportToolbarBase.h
@@ -10,11 +10,14 @@
 #include "Framework/MultiBox/MultiBoxBuilder.h"
 #include "SViewportToolBar.h"
 #include "IPreviewProfileController.h"
+#include "Templates/SharedPointer.h"
 #include "UObject/Object.h"
 #include "UObject/ObjectMacros.h"
 
 #include "SCommonEditorViewportToolbarBase.generated.h"
 
+class SComboButton;
+
 // This is the interface that the host of a SCommonEditorViewportToolbarBase must implement
 class ICommonEditorViewportToolbarInfoProvider
 {
@@ -79,9 +82,10 @@ public:
 class UNREALED_API SPreviewSceneProfileSelector : public SCompoundWidget
 {
 public:
+	virtual ~SPreviewSceneProfileSelector() override
+	{
+	}
 
-	virtual ~SPreviewSceneProfileSelector() override;
-	
 	SLATE_BEGIN_ARGS(SPreviewSceneProfileSelector)
 	{}
 	SLATE_ARGUMENT(TSharedPtr<IPreviewProfileController>, PreviewProfileController)
@@ -90,13 +94,20 @@ public:
 	void Construct(const FArguments& InArgs);
 
 protected:
-	
-	/** Update the list of asset viewer profiles displayed by the combo box. */
-	void UpdateAssetViewerProfileList();
-	void UpdateAssetViewerProfileSelection();
+	UE_DEPRECATED(5.5, "Unused")
+	void UpdateAssetViewerProfileList()
+	{
+	}
 
-	/** Invoked when the asset viewer profile combo box selection changes. */
-	void OnSelectionChanged(TSharedPtr<FString> NewSelection, ESelectInfo::Type /*SelectInfo*/);
+	UE_DEPRECATED(5.5, "Unused")
+	void UpdateAssetViewerProfileSelection()
+	{
+	}
+
+	UE_DEPRECATED(5.5, "Unused")
+	void OnSelectionChanged(TSharedPtr<FString> NewSelection, ESelectInfo::Type /*SelectInfo*/)
+	{
+	}
 
 	/** Creates and returns the asset viewer profile combo box.*/
 	TSharedRef<SWidget> MakeAssetViewerProfileComboBox();
@@ -105,11 +116,11 @@ private:
 	/** Interface to set/get/list the preview profiles. */
 	TSharedPtr<IPreviewProfileController> PreviewProfileController;
 
-	/** List of advanced preview profiles to fill up the Profiles combo box. */
-	TArray<TSharedPtr<FString>> AssetViewerProfileNames;
-
 	/** Displays/Selects the active advanced viewer profile. */
-	TSharedPtr<STextComboBox> AssetViewerProfileComboBox;
+	TSharedPtr<SComboButton> AssetViewerProfileComboButton;
+
+	/** Builds the drop-down list for selecting a viewer profile. */
+	TSharedRef<SWidget> BuildComboMenu();
 };
 
 /**
