Fix blueprint component data corruption problem where placed instances could lose their property changes when compiling the component class.
This restores the behavior of caching an object archetype during reinstancing to work how it worked in 5.3, which is needed because the normal archetype lookup code failse due to the REINST_ renaming.

diff --git a/Engine/Source/Editor/UnrealEd/Private/Kismet2/KismetReinstanceUtilities.cpp b/Engine/Source/Editor/UnrealEd/Private/Kismet2/KismetReinstanceUtilities.cpp
--- a/Engine/Source/Editor/UnrealEd/Private/Kismet2/KismetReinstanceUtilities.cpp
+++ b/Engine/Source/Editor/UnrealEd/Private/Kismet2/KismetReinstanceUtilities.cpp
@@ -2211,6 +2211,8 @@ static void ReplaceObjectHelper(UObject*& OldObject, UClass* OldClass, UObject*&
 
 	FName OldName(OldObject->GetFName());
 
+	UObject* OldArchetype = nullptr;
+
 	// If the old object is in this table, we've already renamed it away in a previous iteration. Don't rename it again!
 	if (!OldToNewNameMap.Contains(OldObject))
 	{
@@ -2232,6 +2234,7 @@ static void ReplaceObjectHelper(UObject*& OldObject, UClass* OldClass, UObject*&
 		}
 		else
 		{
+			OldArchetype = OldObject->GetArchetype(); // Cache the old object's archetype before renaming
 			OldObject->Rename(nullptr, OldObject->GetOuter(), REN_DoNotDirty | REN_DontCreateRedirectors | REN_AllowPackageLinkerMismatch);
 		}
 	}
@@ -2300,6 +2303,15 @@ static void ReplaceObjectHelper(UObject*& OldObject, UClass* OldClass, UObject*&
 	// We only need to copy properties of the pre-created instances, the rest of the default sub object is done inside the UEditorEngine::CopyPropertiesForUnrelatedObjects
 	for (const auto& Pair : OrderedListOfObjectToCopy)
 	{
+		if (Pair.Key == OldObject)
+		{
+			// If we're copying the object itself, make sure to use the archetype that was cached before
+			Options.SourceObjectArchetype = OldArchetype;
+		}
+		else
+		{
+			Options.SourceObjectArchetype = nullptr;
+		}
 		UEditorEngine::CopyPropertiesForUnrelatedObjects(Pair.Key, Pair.Value, Options);
 	}
 
