Added trace tags to capture expensive editor operations.

diff --git a/Engine/Source/Editor/UnrealEd/Private/GroupActor.cpp b/Engine/Source/Editor/UnrealEd/Private/GroupActor.cpp
--- a/Engine/Source/Editor/UnrealEd/Private/GroupActor.cpp
+++ b/Engine/Source/Editor/UnrealEd/Private/GroupActor.cpp
@@ -407,6 +407,7 @@ void AGroupActor::DrawBracketsForGroups( FPrimitiveDrawInterface* PDI, FViewport
 		UWorld* EditorWorld = GEditor->GetEditorWorldContext().World();
 		if (EditorWorld)
 		{
+			TRACE_CPUPROFILER_EVENT_SCOPE(AGroupActor::DrawBracketsForGroup);
 			TArray<AGroupActor*> GroupsToDraw;
 			
 			for(int32 GroupIndex=0; GroupIndex < EditorWorld->ActiveGroupActors.Num(); ++GroupIndex)
diff --git a/Engine/Source/Editor/UnrealEd/Private/LevelEditorViewport.cpp b/Engine/Source/Editor/UnrealEd/Private/LevelEditorViewport.cpp
--- a/Engine/Source/Editor/UnrealEd/Private/LevelEditorViewport.cpp
+++ b/Engine/Source/Editor/UnrealEd/Private/LevelEditorViewport.cpp
@@ -5066,6 +5066,7 @@ void FLevelEditorViewportClient::Draw(const FSceneView* View,FPrimitiveDrawInter
 void FLevelEditorViewportClient::DrawBrushDetails(const FSceneView* View, FPrimitiveDrawInterface* PDI)
 {
 	// Draw translucent polygons on brushes and volumes
+	TRACE_CPUPROFILER_EVENT_SCOPE(FLevelEditorViewportClient::DrawBrushDetails);
 
 	for (TActorIterator<ABrush> It(GetWorld()); It; ++It)
 	{
