Cooker: Add diagnostics for the error "Package %s is a generated package but its generator does not have a record of it": print out the count of generated packages and the first n elements of the list.
#rnx

diff --git a/Engine/Source/Editor/UnrealEd/Private/CookOnTheFlyServer.cpp b/Engine/Source/Editor/UnrealEd/Private/CookOnTheFlyServer.cpp
--- a/Engine/Source/Editor/UnrealEd/Private/CookOnTheFlyServer.cpp
+++ b/Engine/Source/Editor/UnrealEd/Private/CookOnTheFlyServer.cpp
@@ -2999,6 +2999,26 @@ void UCookOnTheFlyServer::LoadPackageInQueue(UE::Cook::FPackageData& PackageData
 			UE_LOG(LogCook, Error,
 				TEXT("Package %s is a generated package but its generator does not have a record of it. It can not be loaded."),
 				*PackageFileName.ToString());
+			TArray<FString> GeneratedNames;
+			for (FCookGenerationInfo& ExistingInfo : GenerationHelper->GetPackagesToGenerate())
+			{
+				GeneratedNames.Add(ExistingInfo.GetPackageName());
+			}
+			GeneratedNames.Sort();
+			TStringBuilder<1024> GeneratedNamesListStr;
+			constexpr int32 MaxCount = 10;
+			int32 Count = 0;
+			for (const FString& GeneratedName : GeneratedNames)
+			{
+				if (Count++ >= MaxCount)
+				{
+					GeneratedNamesListStr << TEXT("\n\t...");
+					break;
+				}
+				GeneratedNamesListStr << TEXT("\n\t") << GeneratedName;
+			}
+			UE_LOG(LogCook, Display, TEXT("The generator has %d generated packages, but %s is not one of them:%s"),
+				GeneratedNames.Num(), *WriteToString<256>(PackageData.GetPackageName()), *GeneratedNamesListStr);
 			RejectPackageToLoad(PackageData, TEXT("is an orphaned generated package"), ESuppressCookReason::OrphanedGenerated);
 			return;
 		}
