[ContentBrowser] Change Icons for CB 2.0 name to align with the AssetThumbnail way of retrieving them

diff --git a/Engine/Source/Editor/ContentBrowser/Private/AssetViewWidgets.cpp b/Engine/Source/Editor/ContentBrowser/Private/AssetViewWidgets.cpp
--- a/Engine/Source/Editor/ContentBrowser/Private/AssetViewWidgets.cpp
+++ b/Engine/Source/Editor/ContentBrowser/Private/AssetViewWidgets.cpp
@@ -300,7 +300,7 @@ public:
 			SToolTip::FArguments()
 #if UE_CONTENTBROWSER_NEW_STYLE
 			.TextMargin(FMargin(12.f, 8.f, 12.f, 2.f))
-			.BorderImage(FAppStyle::GetBrush("ContentBrowser.Tooltip.Border"))
+			.BorderImage(FAppStyle::GetBrush("AssetThumbnail.Tooltip.Border"))
 #else
 			.TextMargin(1.0f)
 			.BorderImage(UE::ContentBrowser::Private::FContentBrowserStyle::Get().GetBrush("ContentBrowser.TileViewTooltip.ToolTipBorder"))
@@ -1654,7 +1654,7 @@ void SAssetViewItem::CacheDisplayTags()
 const FSlateBrush* SAssetViewItem::GetBorderImage() const
 {
 #if UE_CONTENTBROWSER_NEW_STYLE
-	return bDraggedOver ? FAppStyle::GetBrush("ContentBrowser.AssetTileItem.AssetBorder") : FAppStyle::GetBrush("NoBorder");
+	return bDraggedOver ? FAppStyle::GetBrush("AssetThumbnail.AssetBorder") : FAppStyle::GetBrush("NoBorder");
 #else
 	return bDraggedOver ? FAppStyle::GetBrush("Menu.Background") : FAppStyle::GetBrush("NoBorder");
 #endif
@@ -2539,7 +2539,7 @@ const FSlateBrush* SAssetTileItem::GetAssetAreaOverlayBackgroundImage() const
 	}
 	else if (!IsFolder())
 	{
-		static const FLazyName Normal("ContentBrowser.AssetTileItem.AssetBorder");
+		static const FLazyName Normal("AssetThumbnail.AssetBorder");
 		return FAppStyle::Get().GetBrush(Normal);
 	}
 
diff --git a/Engine/Source/Editor/EditorStyle/Private/StarshipStyle.cpp b/Engine/Source/Editor/EditorStyle/Private/StarshipStyle.cpp
--- a/Engine/Source/Editor/EditorStyle/Private/StarshipStyle.cpp
+++ b/Engine/Source/Editor/EditorStyle/Private/StarshipStyle.cpp
@@ -1232,7 +1232,17 @@ void FStarshipEditorStyle::FStyle::SetupGeneralStyles()
 		Set( "AssetThumbnail.HintShadowOffset", FVector2D(1,1) );
 		Set( "AssetThumbnail.HintShadowColorAndOpacity", FLinearColor(0, 0, 0, 0.5) );
 		Set( "AssetThumbnail.HintBackground", new BOX_BRUSH( "Common/TableViewHeader", FMargin(8.0f/32.0f) ) );
-	
+		Set( "AssetThumbnail.AssetTileViewWhiteChipBorder", new IMAGE_BRUSH_SVG("Starship/ContentBrowser/CornerChipOuter_18", FVector2D(18, 18)));
+		Set( "AssetThumbnail.AssetTileViewWhiteChip", new IMAGE_BRUSH_SVG("Starship/ContentBrowser/CornerChipInner_18", FVector2D(18, 18)));
+		Set( "AssetThumbnail.AssetBorder", new FSlateRoundedBoxBrush(FStyleColors::Transparent, FVector4(4.0f, 2.0f, 4.0f, 4.0f), FStyleColors::Secondary, 2.f));
+		FLinearColor OverlayColorAssetStatusOverlay = FStyleColors::Panel.GetSpecifiedColor();
+		OverlayColorAssetStatusOverlay.A = 0.75f;
+		Set("AssetThumbnail.AssetThumbnailStatusBar", new FSlateRoundedBoxBrush(OverlayColorAssetStatusOverlay, 2.f));
+		FLinearColor OverlayColorAssetThumbnailOverlay = FStyleColors::Panel.GetSpecifiedColor();
+		OverlayColorAssetThumbnailOverlay.A = 0.75f;
+		Set("AssetThumbnail.AssetThumbnailBar", new FSlateRoundedBoxBrush(OverlayColorAssetThumbnailOverlay, 4.f));
+		Set("AssetThumbnail.ToolTip.CommandBorder", new FSlateRoundedBoxBrush(FStyleColors::Transparent, 4.f, FStyleColors::White, 1.f));
+		Set("AssetThumbnail.Tooltip.Border", new FSlateRoundedBoxBrush(FStyleColors::Secondary, 0.f, COLOR("#484848FF"), 2.f));
 	}
 
 	// Open any asset dialog
@@ -6342,17 +6352,12 @@ void FStarshipEditorStyle::FStyle::SetupContentBrowserStyle()
 		);
 
 		// Tile view
-		Set( "ContentBrowser.AssetTileViewWhiteChip", new IMAGE_BRUSH_SVG("Starship/ContentBrowser/CornerChipInner_18", FVector2D(18, 18)));
-		Set( "ContentBrowser.AssetTileViewWhiteChipBorder", new IMAGE_BRUSH_SVG("Starship/ContentBrowser/CornerChipOuter_18", FVector2D(18, 18)));
 		Set( "ContentBrowser.AssetTileViewNameFont", DEFAULT_FONT("Regular", 9));
 		Set( "ContentBrowser.AssetTileViewClassNameFont", DEFAULT_FONT("Regular", 7));
 		Set( "ContentBrowser.AssetTileViewNameFontSmall", DEFAULT_FONT( "VeryLight", 8 ) );
 		Set( "ContentBrowser.AssetTileViewNameFontVerySmall", DEFAULT_FONT( "VeryLight", 7 ) );
 		Set( "ContentBrowser.AssetTileViewNameFontDirty", FStyleFonts::Get().SmallBold);
 
-		Set("ContentBrowser.Tooltip.Border", new FSlateRoundedBoxBrush(FStyleColors::Secondary, 0.f, COLOR("#484848FF"), 2.f));
-		Set("ContentBrowser.ToolTip.CommandBorder", new FSlateRoundedBoxBrush(FStyleColors::Transparent, 4.f, FStyleColors::White, 1.f));
-
 		Set("ContentBrowser.AssetListView.ColumnListTableRow", FTableRowStyle()
 			.SetEvenRowBackgroundBrush(FSlateColorBrush(FStyleColors::Recessed))
 			.SetEvenRowBackgroundHoveredBrush(FSlateColorBrush(FStyleColors::SelectHover))
@@ -6507,13 +6512,6 @@ void FStarshipEditorStyle::FStyle::SetupContentBrowserStyle()
 		// Playback Action
 		Set( "ContentBrowser.AssetAction.PlayIcon", new IMAGE_BRUSH_SVG("Starship/ContentBrowser/PlayIcon", Icon16x16));
 		Set( "ContentBrowser.AssetAction.StopIcon", new IMAGE_BRUSH_SVG("Starship/ContentBrowser/StopIcon", Icon16x16));
-		FLinearColor OverlayColorAssetThumbnailOverlay = FStyleColors::Panel.GetSpecifiedColor();
-		OverlayColorAssetThumbnailOverlay.A = 0.75f;
-		Set("ContentBrowser.AssetTileItem.AssetThumbnailBar", new FSlateRoundedBoxBrush(OverlayColorAssetThumbnailOverlay, 4.f));
-
-		FLinearColor OverlayColorAssetStatusOverlay = FStyleColors::Panel.GetSpecifiedColor();
-		OverlayColorAssetStatusOverlay.A = 0.75f;
-		Set("ContentBrowser.AssetTileItem.AssetThumbnailStatusBar", new FSlateRoundedBoxBrush(OverlayColorAssetStatusOverlay, 2.f));
 
 		Set( "ContentBrowser.ColumnViewAssetIcon", new IMAGE_BRUSH( "Icons/doc_16x", Icon16x16 ) );
 
@@ -6537,7 +6535,6 @@ void FStarshipEditorStyle::FStyle::SetupContentBrowserStyle()
 		Set("ContentBrowser.AssetTileItem.NameAreaSelectedHoverBackground", new FSlateRoundedBoxBrush(FStyleColors::PrimaryHover, FVector4(0.0f, 0.0f, 4.0f, 4.0f)));
 
 		// Asset Thumbnail Border
-		Set("ContentBrowser.AssetTileItem.AssetBorder", new FSlateRoundedBoxBrush(FStyleColors::Transparent, FVector4(4.0f, 2.0f, 4.0f, 4.0f), FStyleColors::Secondary, 2.f));
 		Set("ContentBrowser.AssetTileItem.AssetBorderHoverBackground", new FSlateRoundedBoxBrush(FStyleColors::Transparent, FVector4(4.0f, 2.0f, 4.0f, 4.0f), FStyleColors::Hover, 2.f));
 		Set("ContentBrowser.AssetTileItem.AssetBorderSelectedBackground", new FSlateRoundedBoxBrush(FStyleColors::Transparent, FVector4(4.0f, 2.0f, 4.0f, 4.0f), FStyleColors::Primary, 2.f));
 		Set("ContentBrowser.AssetTileItem.AssetBorderSelectedHoverBackground", new FSlateRoundedBoxBrush(FStyleColors::Transparent, FVector4(4.0f, 2.0f, 4.0f, 4.0f), FStyleColors::PrimaryHover, 2.f));
diff --git a/Engine/Source/Editor/UnrealEd/Private/AssetThumbnail.cpp b/Engine/Source/Editor/UnrealEd/Private/AssetThumbnail.cpp
--- a/Engine/Source/Editor/UnrealEd/Private/AssetThumbnail.cpp
+++ b/Engine/Source/Editor/UnrealEd/Private/AssetThumbnail.cpp
@@ -309,9 +309,9 @@ public:
 		{
 			// The asset color chip
 			LayeredImage = SNew(SLayeredImage)
-				.Image(FAppStyle::GetBrush("ContentBrowser.AssetTileViewWhiteChipBorder"))
+				.Image(FAppStyle::GetBrush(Style, ".AssetTileViewWhiteChipBorder"))
 				.ColorAndOpacity(FStyleColors::Secondary);
-			LayeredImage->AddLayer(FAppStyle::GetBrush("ContentBrowser.AssetTileViewWhiteChip"), AssetColor);
+			LayeredImage->AddLayer(FAppStyle::GetBrush(Style, ".AssetTileViewWhiteChip"), AssetColor);
 
 			OverlayWidget->AddSlot()
 				.HAlign(HAlign_Right)
@@ -320,8 +320,8 @@ public:
 					LayeredImage.ToSharedRef()
 				];
 
-			TSharedRef<SImage> OverlayBorder = SNew(SImage)
-				.Image(FAppStyle::GetBrush("ContentBrowser.AssetTileItem.AssetBorder"))
+			const TSharedRef<SImage> OverlayBorder = SNew(SImage)
+				.Image(FAppStyle::GetBrush(Style, ".AssetBorder"))
 				.Visibility(EVisibility::HitTestInvisible);
 			
 			OverlayWidget->AddSlot()
@@ -331,9 +331,9 @@ public:
 
 			if (AssetChipBorderImageOverride.IsSet())
 			{
-				TSharedRef<SBorder> ThumbnailWidgetBorder = SNew(SBorder)
+				const TSharedRef<SBorder> ThumbnailWidgetBorder = SNew(SBorder)
 					.Padding(0)
-					.BorderImage(FAppStyle::GetBrush("ContentBrowser.AssetTileItem.AssetBorder"))
+					.BorderImage(FAppStyle::GetBrush(Style, ".AssetBorder"))
 					[
 						OverlayWidget
 					];
@@ -407,7 +407,7 @@ public:
 			.Padding(StatusPadding)
 			[
 				SNew(SBorder)
-				.BorderImage(FAppStyle::Get().GetBrush("ContentBrowser.AssetTileItem.AssetThumbnailStatusBar"))
+				.BorderImage(FAppStyle::GetBrush(Style, ".AssetThumbnailStatusBar"))
 				.Visibility(this, &SAssetThumbnail::GetStatusBorderVisibility)
 				[
 					HorizontalBox
@@ -434,7 +434,7 @@ public:
 					.Padding(PaddingFromTopLeftBorder, PaddingFromTopLeftBorder, 0.f, 0.f)
 					[
 						SNew(SBorder)
-						.BorderImage(FAppStyle::GetBrush(TEXT("ContentBrowser.AssetTileItem.AssetThumbnailBar")))
+						.BorderImage(FAppStyle::GetBrush(Style, ".AssetThumbnailBar"))
 						.Visibility_Lambda([this] () { return IsHovered() ? EVisibility::Collapsed : EVisibility::Visible; })
 						[
 							SNew(SBox)
@@ -457,7 +457,7 @@ public:
 					.HAlign(HAlign_Center)
 					[
 						SNew(SBorder)
-						.BorderImage(FAppStyle::GetBrush(TEXT("ContentBrowser.AssetTileItem.AssetThumbnailBar")))
+						.BorderImage(FAppStyle::GetBrush(Style, ".AssetThumbnailBar"))
 						.Visibility_Lambda([this] () { return IsHovered() ? EVisibility::Visible : EVisibility::Collapsed; })
 						[
 							SNew(SBox)
@@ -642,7 +642,6 @@ public:
 					SNew(STextBlock)
 					.Text(NameText)
 					.ColorAndOpacity(FStyleColors::White)
-					.Font(FAppStyle::GetFontStyle("ContentBrowser.Tooltip.AssetNameFont"))
 				]
 
 				+ SVerticalBox::Slot()
@@ -669,7 +668,6 @@ public:
 					.VAlign(VAlign_Center)
 					[
 						SNew(STextBlock)
-						.Font(FAppStyle::GetFontStyle("ContentBrowser.Tooltip.EntryFont"))
 						.Text(ClassNameText)
 					]
 				]
@@ -716,7 +714,6 @@ public:
 						.VAlign(VAlign_Center)
 						[
 							SNew(STextBlock)
-							.Font(FAppStyle::GetFontStyle("ContentBrowser.Tooltip.EntryFont"))
 							.Text(AssetStatusInfo.StatusDescription)
 						]
 					];
@@ -743,7 +740,6 @@ public:
 						const FModifierKeysState ModifierKeys = FSlateApplication::Get().GetModifierKeys();
 						return ModifierKeys.IsAltDown() && ModifierKeys.IsControlDown() ? EVisibility::Collapsed : EVisibility::Visible;
 					})
-					// .Font(FAppStyle::GetFontStyle("ContentBrowser.Tooltip.MoreInfo"))
 					.Justification(ETextJustify::Center)
 					.Text(NSLOCTEXT("AssetThumbnail", "MoreInfoTooltip", "Hold <WrappedCommand/> for more"))
 					+ SRichTextBlock::WidgetDecorator( TEXT("WrappedCommand"), this, &SAssetThumbnail::OnCreateWidgetDecoratorWidget )
@@ -786,14 +782,13 @@ private:
 		TSharedRef<SWidget> CtrlAltWidget = SNew(SBorder)
 			.VAlign(VAlign_Center)
 			.HAlign(HAlign_Center)
-			.BorderImage(FAppStyle::GetBrush("ContentBrowser.ToolTip.CommandBorder"))
+			.BorderImage(FAppStyle::GetBrush(Style, ".ToolTip.CommandBorder"))
 			[
 				SNew(STextBlock)
-				.Font(FAppStyle::GetFontStyle("ContentBrowser.Tooltip.MoreInfo"))
 				.Text(NSLOCTEXT("AssetThumbnail", "CtrlAltLabel", " Ctrl + Alt "))
 			];
 		const TSharedRef<FSlateFontMeasure> FontMeasure = FSlateApplication::Get().GetRenderer()->GetFontMeasureService();
-		const int16 Baseline = FontMeasure->GetBaseline(FAppStyle::GetFontStyle("ContentBrowser.Tooltip.MoreInfo"));
+		const int16 Baseline = FontMeasure->GetBaseline(FStyleDefaults::GetFontInfo());
 
 		return FSlateWidgetRun::FWidgetRunInfo(CtrlAltWidget, Baseline - 2);
 
@@ -826,7 +821,6 @@ private:
 			.Padding(0, 0, 4, 0)
 			[
 				SNew(STextBlock)
-				.Font(FAppStyle::GetFontStyle("ContentBrowser.Tooltip.EntryFont"))
 				.Text(FText::Format(NSLOCTEXT("AssetThumbnailToolTip", "AssetViewTooltipFormat", "{0}:"), Key))
 			]
 
@@ -834,7 +828,6 @@ private:
 			.AutoWidth()
 			[
 				SNew(STextBlock)
-				.Font(FAppStyle::GetFontStyle("ContentBrowser.Tooltip.EntryFont"))
 				.ColorAndOpacity(FStyleColors::White)
 				.Text(Value)
 			]
@@ -1401,7 +1394,7 @@ void SAssetThumbnailToolTip::Construct(const FArguments& InArgs)
 	SToolTip::Construct(
 		SToolTip::FArguments()
 		.TextMargin(FMargin(1.f, -3.f))
-		.BorderImage(FAppStyle::GetBrush("ContentBrowser.Tooltip.Border"))
+		.BorderImage(FAppStyle::GetBrush("AssetThumbnail.Tooltip.Border"))
 		);
 }
 
