[Viewport Toolbar] Make "Wireframe" a checkbox submenu in the new viewport toolbars "View Modes" submenu

diff --git a/Engine/Source/Editor/UnrealEd/Private/ViewportToolbar/UnrealEdViewportToolbar.cpp b/Engine/Source/Editor/UnrealEd/Private/ViewportToolbar/UnrealEdViewportToolbar.cpp
--- a/Engine/Source/Editor/UnrealEd/Private/ViewportToolbar/UnrealEdViewportToolbar.cpp
+++ b/Engine/Source/Editor/UnrealEd/Private/ViewportToolbar/UnrealEdViewportToolbar.cpp
@@ -1196,6 +1196,83 @@ FText GetViewModesSubmenuLabel(TWeakPtr<SEditorViewport> InViewport)
 	return Label;
 }
 
+// Explicitly passing Editor Viewport, since the old toolbar might get an empty Context once calling this
+FToolMenuEntry GetLitWireframeEntryToSection(const TSharedPtr<SEditorViewport>& EditorViewport)
+{
+	FNewToolMenuDelegate MakeMenuDelegate = FNewToolMenuDelegate::CreateLambda(
+		[EditorViewportWeak = EditorViewport.ToWeakPtr()](UToolMenu* Submenu)
+		{
+			const TSharedPtr<SEditorViewport>& EditorViewport = EditorViewportWeak.Pin();
+			if (!EditorViewport)
+			{
+				return;
+			}
+
+			FToolMenuSection& Section =
+				Submenu->AddSection("LitWireframeOpacity", LOCTEXT("WireframeOpacityLabel", "Wireframe Opacity"));
+
+			Section.AddEntry(FToolMenuEntry::InitWidget(
+				"WireframeOpacity", EditorViewport->BuildWireframeMenu(), LOCTEXT("WireframeOpacity", "Opacity")
+			));
+		}
+	);
+
+	FToolUIAction RadialMenuAction;
+	{
+		RadialMenuAction.ExecuteAction = FToolMenuExecuteAction::CreateLambda(
+			[EditorViewportWeak = EditorViewport.ToWeakPtr()](const FToolMenuContext& InContext)
+			{
+				const TSharedPtr<SEditorViewport>& EditorViewport = EditorViewportWeak.Pin();
+				if (!EditorViewport)
+				{
+					return;
+				}
+
+				if (const TSharedPtr<FEditorViewportClient>& ViewportClient = EditorViewport->GetViewportClient())
+				{
+					ViewportClient->SetViewMode(EViewModeIndex::VMI_Lit_Wireframe);
+				}
+			}
+		);
+
+		RadialMenuAction.CanExecuteAction = FToolMenuCanExecuteAction();
+		RadialMenuAction.GetActionCheckState = FToolMenuGetActionCheckState::CreateLambda(
+			[EditorViewportWeak = EditorViewport.ToWeakPtr()](const FToolMenuContext& InContext)
+			{
+				if (const TSharedPtr<SEditorViewport>& EditorViewport = EditorViewportWeak.Pin())
+				{
+					if (const TSharedPtr<FEditorViewportClient>& ViewportClient = EditorViewport->GetViewportClient())
+					{
+						return ViewportClient->IsViewModeEnabled(EViewModeIndex::VMI_Lit_Wireframe)
+								 ? ECheckBoxState::Checked
+								 : ECheckBoxState::Unchecked;
+					}
+				}
+
+				return ECheckBoxState::Unchecked;
+			}
+		);
+	}
+
+	const FEditorViewportCommands& BaseViewportActions = FEditorViewportCommands::Get();
+	const TSharedPtr<FUICommandInfo>& LitWireframeMode = BaseViewportActions.LitWireframeMode;
+
+	const FText& Tooltip = LitWireframeMode->GetDescription();
+	const FSlateIcon& Icon = LitWireframeMode->GetIcon();
+	const FText& Label = UViewModeUtils::GetViewModeDisplayName(EViewModeIndex::VMI_Lit_Wireframe);
+
+	return FToolMenuEntry::InitSubMenu(
+		LitWireframeMode->GetCommandName(),
+		Label,
+		Tooltip,
+		MakeMenuDelegate,
+		RadialMenuAction,
+		EUserInterfaceActionType::RadioButton,
+		false,
+		Icon
+	);
+}
+
 void PopulateViewModesMenu(UToolMenu* InMenu)
 {
 	UUnrealEdViewportToolbarContext* const Context = InMenu->FindContext<UUnrealEdViewportToolbarContext>();
@@ -1254,9 +1331,12 @@ void PopulateViewModesMenu(UToolMenu* InMenu)
 			ViewModesSubmenu::AddModeIfSupported(
 				IsViewModeSupported, Section, BaseViewportActions.WireframeMode, VMI_BrushWireframe
 			);
-			ViewModesSubmenu::AddModeIfSupported(
-				IsViewModeSupported, Section, BaseViewportActions.LitWireframeMode, VMI_Lit_Wireframe
-			);
+
+			if (!IsViewModeSupported.IsBound() || IsViewModeSupported.Execute(VMI_Lit_Wireframe))
+			{
+				Section.AddEntry(GetLitWireframeEntryToSection(EditorViewport));
+			}
+
 			ViewModesSubmenu::AddModeIfSupported(
 				IsViewModeSupported, Section, BaseViewportActions.DetailLightingMode, VMI_Lit_DetailLighting
 			);
@@ -1581,19 +1661,6 @@ void PopulateViewModesMenu(UToolMenu* InMenu)
 		);
 		Section.AddEntry(FToolMenuEntry::InitWidget("FixedEV100", FixedEV100Menu, LOCTEXT("FixedEV100", "EV100")));
 	}
-
-	// TODO: would be nice to make this appear/disappear based on current mode
-	// Wireframe Opacity
-	if (!IsViewModeSupported.IsBound()
-		|| (IsViewModeSupported.Execute(VMI_Wireframe) || IsViewModeSupported.Execute(VMI_BrushWireframe)
-			|| IsViewModeSupported.Execute(VMI_Lit_Wireframe)))
-	{
-		TSharedRef<SWidget> WireOpacityMenu = EditorViewport->BuildWireframeMenu();
-		FToolMenuSection& Section = InMenu->AddSection("Wireframe", LOCTEXT("WireframeHeader", "Wireframe"));
-		Section.AddEntry(
-			FToolMenuEntry::InitWidget("WireframeOpacity", WireOpacityMenu, LOCTEXT("WireframeOpacity", "Opacity"))
-		);
-	}
 }
 
 UUnrealEdViewportToolbarContext* CreateViewportToolbarDefaultContext(const TWeakPtr<SEditorViewport>& InViewport)
