Allow structs in the error state to be contained in other structs in order to prevent cascading errors

[RN] Structure assets in the error state will no longer remove themselves from other structure assets

diff --git a/Engine/Source/Editor/UnrealEd/Private/Kismet2/StructureEditorUtils.cpp b/Engine/Source/Editor/UnrealEd/Private/Kismet2/StructureEditorUtils.cpp
--- a/Engine/Source/Editor/UnrealEd/Private/Kismet2/StructureEditorUtils.cpp
+++ b/Engine/Source/Editor/UnrealEd/Private/Kismet2/StructureEditorUtils.cpp
@@ -118,7 +118,7 @@ FStructureEditorUtils::EStructureError FStructureEditorUtils::IsStructureValid(c
 
 	if (const UUserDefinedStruct* UDStruct = Cast<const UUserDefinedStruct>(Struct))
 	{
-		if (UDStruct->Status != EUserDefinedStructureStatus::UDSS_UpToDate)
+		if (UDStruct->Status != EUserDefinedStructureStatus::UDSS_UpToDate && UDStruct->Status != EUserDefinedStructureStatus::UDSS_Error)
 		{
 			if (OutMsg)
 			{
