GatherTextFromSourceCommandlet now only scans the filesytem once for a given query, rather than once for each file type in that query

#rnx

diff --git a/Engine/Source/Editor/UnrealEd/Private/Commandlets/GatherTextFromSourceCommandlet.cpp b/Engine/Source/Editor/UnrealEd/Private/Commandlets/GatherTextFromSourceCommandlet.cpp
--- a/Engine/Source/Editor/UnrealEd/Private/Commandlets/GatherTextFromSourceCommandlet.cpp
+++ b/Engine/Source/Editor/UnrealEd/Private/Commandlets/GatherTextFromSourceCommandlet.cpp
@@ -321,7 +321,52 @@ void UGatherTextFromSourceCommandlet::GetFilesToProcess(const TArray<FString>& S
 
 	// Search in the root folder for each of the wildcard filters specified and build a list of files
 	{
-		TArray<FString> RootSourceFiles;
+		class FFileMatch : public IPlatformFile::FDirectoryVisitor
+		{
+		public:
+			const TArray<FString>& WildCards;
+
+			UE::FMutex RootSourceFilesLock;
+			TArray<FString> RootSourceFiles;
+
+			FFileMatch(const TArray<FString>& InWildCards)
+				: IPlatformFile::FDirectoryVisitor(EDirectoryVisitorFlags::ThreadSafe)
+				, WildCards(InWildCards)
+			{
+			}
+
+			virtual bool Visit(const TCHAR* FilenameOrDirectory, bool bIsDirectory)
+			{
+				if (!bIsDirectory)
+				{
+					FString FullFilename = FilenameOrDirectory;
+					FString LeafFilename = FPaths::GetCleanFilename(FullFilename);
+
+					bool bMatchesWildCard = false;
+					if (!LeafFilename.EndsWith(TEXTVIEW(".generated.h")) && !LeafFilename.EndsWith(TEXTVIEW(".gen.cpp"))) // Always skip UHT generated files
+					{
+						for (const FString& WildCard : WildCards)
+						{
+							if (LeafFilename.MatchesWildcard(WildCard))
+							{
+								bMatchesWildCard = true;
+								break;
+							}
+						}
+					}
+
+					if (bMatchesWildCard)
+					{
+						UE::TScopeLock ScopeLock(RootSourceFilesLock);
+						RootSourceFiles.Add(MoveTemp(FullFilename));
+					}
+				}
+
+				return true;
+			}
+		};
+
+		FFileMatch Visitor(FileNameFilters);
 		TSet<FString, FLocKeySetFuncs> ProcessedSearchDirectoryPaths;
 		for (const FString& IncludePathFilter : IncludePathFilters)
 		{
@@ -334,6 +379,10 @@ void UGatherTextFromSourceCommandlet::GetFilesToProcess(const TArray<FString>& S
 				SearchDirectoryPath = SearchDirectoryPath.Left(UE_PTRDIFF_TO_INT32(FirstWildcard - *SearchDirectoryPath));
 				SearchDirectoryPath = FPaths::GetPath(MoveTemp(SearchDirectoryPath));
 			}
+			if (FPaths::IsRelative(SearchDirectoryPath))
+			{
+				SearchDirectoryPath = FPaths::ConvertRelativePathToFull(MoveTemp(SearchDirectoryPath));
+			}
 
 			bool bAlreadyProcessed = false;
 			ProcessedSearchDirectoryPaths.Add(SearchDirectoryPath, &bAlreadyProcessed);
@@ -342,21 +391,9 @@ void UGatherTextFromSourceCommandlet::GetFilesToProcess(const TArray<FString>& S
 				continue;
 			}
 
-			for (const FString& FileNameFilter : FileNameFilters)
-			{
-				IFileManager::Get().FindFilesRecursive(RootSourceFiles, *SearchDirectoryPath, *FileNameFilter, true, false, false);
-
-				for (FString& RootSourceFile : RootSourceFiles)
-				{
-					if (FPaths::IsRelative(RootSourceFile))
-					{
-						RootSourceFile = FPaths::ConvertRelativePathToFull(MoveTemp(RootSourceFile));
-					}
-				}
-
-				FilesToProcess.Append(MoveTemp(RootSourceFiles));
-				RootSourceFiles.Reset();
-			}
+			IFileManager::Get().IterateDirectoryRecursively(*SearchDirectoryPath, Visitor);
+			FilesToProcess.Append(MoveTemp(Visitor.RootSourceFiles));
+			Visitor.RootSourceFiles.Reset();
 		}
 	}
 
