[UMG] Sequencer Selectability:
 - Fix selectability not working in UMG editor
 - Allow dragging of designer view when in selectability mode
 - Fix Shift modifier doing the opposite of what it should

diff --git a/Engine/Source/Editor/Sequencer/Private/Sequencer.cpp b/Engine/Source/Editor/Sequencer/Private/Sequencer.cpp
--- a/Engine/Source/Editor/Sequencer/Private/Sequencer.cpp
+++ b/Engine/Source/Editor/Sequencer/Private/Sequencer.cpp
@@ -12202,13 +12202,20 @@ bool FSequencer::IsObjectSelectableInViewport(UObject* const InObject)
 
 	const TSharedRef<UE::MovieScene::FSharedPlaybackState> SharedPlaybackState = GetSharedPlaybackState();
 
+	FMovieSceneEvaluationState* const EvaluationState = SharedPlaybackState->FindCapability<FMovieSceneEvaluationState>();
+	if (!EvaluationState)
+	{
+		return true;
+	}
+
 	const UMovieSceneSequence* OutSequence = nullptr;
 
 	// Early out on first sequence the object is found in
 	ForEachSubSequenceRecursively(FocusedSequence,
-		[this, InObject, &SharedPlaybackState, &OutSequence](UMovieSceneSequence* const InCurrentSequence)
+		[this, InObject, &SharedPlaybackState, EvaluationState, &OutSequence](UMovieSceneSequence* const InCurrentSequence)
 		{
-			const FGuid ObjectGuid = InCurrentSequence->FindBindingFromObject(InObject, SharedPlaybackState);
+			const FMovieSceneSequenceID SequenceID = EvaluationState->FindSequenceId(InCurrentSequence);
+			const FGuid ObjectGuid = EvaluationState->FindCachedObjectId(*InObject, SequenceID, SharedPlaybackState);
 			if (ObjectGuid.IsValid())
 			{
 				OutSequence = InCurrentSequence;
diff --git a/Engine/Source/Editor/UMGEditor/Private/Designer/SDesignerView.cpp b/Engine/Source/Editor/UMGEditor/Private/Designer/SDesignerView.cpp
--- a/Engine/Source/Editor/UMGEditor/Private/Designer/SDesignerView.cpp
+++ b/Engine/Source/Editor/UMGEditor/Private/Designer/SDesignerView.cpp
@@ -1717,43 +1717,44 @@ FReply SDesignerView::OnMouseButtonDown(const FGeometry& MyGeometry, const FPoin
 
 	if (bFoundWidgetUnderCursor)
 	{
-		if (!IsSelectableInSequencer(PendingSelectedWidget.GetPreview()))
-		{
-			PendingSelectedWidget = FWidgetReference();
-			return FReply::Handled();
-		}
-
 		if ( MouseEvent.GetEffectingButton() == EKeys::LeftMouseButton )
 		{
-			const TSet<FWidgetReference>& SelectedWidgets = GetSelectedWidgets();
+			if (IsSelectableInSequencer(PendingSelectedWidget.GetPreview()))
+			{
+				const TSet<FWidgetReference>& SelectedWidgets = GetSelectedWidgets();
 
-			bool bResolvePendingSelectionImmediately = true;
+				bool bResolvePendingSelectionImmediately = true;
 
-			if (SelectedWidgets.Num() > 0)
-			{
-				for (const auto& SelectedWidget : SelectedWidgets)
+				if (SelectedWidgets.Num() > 0)
 				{
-					auto PendingTemplate = PendingSelectedWidget.GetTemplate();
-					auto SelectedTemplate = SelectedWidget.GetTemplate();
-
-					if ( PendingSelectedWidget == SelectedWidget || ( PendingTemplate && SelectedTemplate && PendingTemplate->IsChildOf( SelectedTemplate ) ) )
+					for (const auto& SelectedWidget : SelectedWidgets)
 					{
-						bResolvePendingSelectionImmediately = false;
-						break;
+						auto PendingTemplate = PendingSelectedWidget.GetTemplate();
+						auto SelectedTemplate = SelectedWidget.GetTemplate();
+
+						if ( PendingSelectedWidget == SelectedWidget || ( PendingTemplate && SelectedTemplate && PendingTemplate->IsChildOf( SelectedTemplate ) ) )
+						{
+							bResolvePendingSelectionImmediately = false;
+							break;
+						}
 					}
 				}
-			}
 
-			// If the newly clicked item is a child of the active selection, add it to the pending set of selected 
-			// widgets, if they begin dragging we can just move the parent, but if it's not part of the parent set, 
-			// we want to immediately begin dragging it.  Also if the currently selected widget is the root widget, 
-			// we won't be moving it so just resolve immediately.
-			if ( bResolvePendingSelectionImmediately )
+				// If the newly clicked item is a child of the active selection, add it to the pending set of selected 
+				// widgets, if they begin dragging we can just move the parent, but if it's not part of the parent set, 
+				// we want to immediately begin dragging it.  Also if the currently selected widget is the root widget, 
+				// we won't be moving it so just resolve immediately.
+				if ( bResolvePendingSelectionImmediately )
+				{
+					ResolvePendingSelectedWidgets();
+				}
+
+				DraggingStartPositionScreenSpace = MouseEvent.GetScreenSpacePosition();
+			}
+			else
 			{
-				ResolvePendingSelectedWidgets();
+				PendingSelectedWidget = FWidgetReference();
 			}
-
-			DraggingStartPositionScreenSpace = MouseEvent.GetScreenSpacePosition();
 		}
 	}
 	else
diff --git a/Engine/Source/Editor/UnrealEd/Private/EditorViewportSelectability.cpp b/Engine/Source/Editor/UnrealEd/Private/EditorViewportSelectability.cpp
--- a/Engine/Source/Editor/UnrealEd/Private/EditorViewportSelectability.cpp
+++ b/Engine/Source/Editor/UnrealEd/Private/EditorViewportSelectability.cpp
@@ -275,7 +275,9 @@ bool FEditorViewportSelectability::SelectActorsByPredicate(UWorld* const InWorld
 
 	auto SelectIfPossible = [bInSelect, ActorSelection, &InPredicate, &bSomethingSelected](AActor* const InActor)
 	{
-		if (IsValid(InActor) && ActorSelection->IsSelected(InActor) != bInSelect && InPredicate(InActor) == bInSelect)
+		if (IsValid(InActor)
+			&& ActorSelection->IsSelected(InActor) != bInSelect
+			&& InPredicate(InActor))
 		{
 			bSomethingSelected = true;
 			GEditor->SelectActor(InActor, bInSelect, true);
@@ -309,7 +311,13 @@ bool FEditorViewportSelectability::SelectActorsByPredicate(UWorld* const InWorld
 	return bSomethingSelected;
 }
 
-bool FEditorViewportSelectability::IsActorInLevelHiddenLayer(AActor& InActor, FLevelEditorViewportClient* const InLevelEditorViewportClient)
+bool FEditorViewportSelectability::IsActorSelectableClass(const AActor& InActor)
+{
+	const bool bInvalidClass = InActor.IsA<AWorldSettings>() || InActor.IsA<ABrush>();
+	return !bInvalidClass;
+}
+
+bool FEditorViewportSelectability::IsActorInLevelHiddenLayer(const AActor& InActor, FLevelEditorViewportClient* const InLevelEditorViewportClient)
 {
 	if (!InLevelEditorViewportClient)
 	{
@@ -327,9 +335,9 @@ bool FEditorViewportSelectability::IsActorInLevelHiddenLayer(AActor& InActor, FL
 	return false;
 }
 
-bool FEditorViewportSelectability::DoesActorIntersectBox(AActor& InActor, const FBox& InBox, FEditorViewportClient* const InEditorViewportClient, const bool bInUseStrictSelection)
+bool FEditorViewportSelectability::DoesActorIntersectBox(const AActor& InActor, const FBox& InBox, FEditorViewportClient* const InEditorViewportClient, const bool bInUseStrictSelection)
 {
-	if (InActor.IsHiddenEd())
+	if (InActor.IsHiddenEd() || !IsActorSelectableClass(InActor))
 	{
 		return false;
 	}
@@ -345,22 +353,22 @@ bool FEditorViewportSelectability::DoesActorIntersectBox(AActor& InActor, const
 	for (const UActorComponent* const Component : InActor.GetComponents())
 	{
 		const UPrimitiveComponent* const PrimitiveComponent = Cast<UPrimitiveComponent>(Component);
-		if (PrimitiveComponent && PrimitiveComponent->IsRegistered() && PrimitiveComponent->IsVisibleInEditor())
+		if (PrimitiveComponent
+			&& PrimitiveComponent->IsRegistered()
+			&& PrimitiveComponent->IsVisibleInEditor()
+			&& PrimitiveComponent->IsShown(InEditorViewportClient->EngineShowFlags)
+			&& PrimitiveComponent->ComponentIsTouchingSelectionBox(InBox, false, bInUseStrictSelection))
 		{
-			if (PrimitiveComponent->IsShown(InEditorViewportClient->EngineShowFlags)
-				&& PrimitiveComponent->ComponentIsTouchingSelectionBox(InBox, false, bInUseStrictSelection))
-			{
-				return true;
-			}
+			return true;
 		}
 	}
 
 	return false;
 }
 
-bool FEditorViewportSelectability::DoesActorIntersectFrustum(AActor& InActor, const FConvexVolume& InFrustum, FEditorViewportClient* const InEditorViewportClient, const bool bInUseStrictSelection)
+bool FEditorViewportSelectability::DoesActorIntersectFrustum(const AActor& InActor, const FConvexVolume& InFrustum, FEditorViewportClient* const InEditorViewportClient, const bool bInUseStrictSelection)
 {
-	if (InActor.IsHiddenEd())
+	if (InActor.IsHiddenEd() || !IsActorSelectableClass(InActor))
 	{
 		return false;
 	}
@@ -372,17 +380,16 @@ bool FEditorViewportSelectability::DoesActorIntersectFrustum(AActor& InActor, co
 		return false;
 	}
 
-	// Iterate over all actor components, selecting our primitive components
 	for (const UActorComponent* const Component : InActor.GetComponents())
 	{
 		const UPrimitiveComponent* const PrimitiveComponent = Cast<UPrimitiveComponent>(Component);
-		if (PrimitiveComponent && PrimitiveComponent->IsRegistered() && PrimitiveComponent->IsVisibleInEditor())
+		if (IsValid(PrimitiveComponent)
+			&& PrimitiveComponent->IsRegistered()
+			&& PrimitiveComponent->IsVisibleInEditor()
+			&& PrimitiveComponent->IsShown(InEditorViewportClient->EngineShowFlags)
+			&& PrimitiveComponent->ComponentIsTouchingSelectionFrustum(InFrustum, false, bInUseStrictSelection))
 		{
-			if (PrimitiveComponent->IsShown(InEditorViewportClient->EngineShowFlags)
-				&& PrimitiveComponent->ComponentIsTouchingSelectionFrustum(InFrustum, false, bInUseStrictSelection))
-			{
-				return true;
-			}
+			return true;
 		}
 	}
 
@@ -515,14 +522,15 @@ bool FEditorViewportSelectability::BoxSelectWorldActors(FBox& InBox, FEditorView
 
 	const auto Predicate = [this, &InBox, InEditorViewportClient, bUseStrictSelection](AActor* const InActor) -> bool
 	{
-		return IsObjectSelectableInViewport(InActor)
-			&& DoesActorIntersectBox(*InActor, InBox, InEditorViewportClient, bUseStrictSelection);
+		const bool bSelectable = IsObjectSelectableInViewport(InActor);
+		const bool bIntersects = DoesActorIntersectBox(*InActor, InBox, InEditorViewportClient, bUseStrictSelection);
+		return bSelectable && bIntersects;
 	};
 
 	const bool bShiftDown = InEditorViewportClient->Viewport->KeyState(EKeys::LeftShift)
 		|| InEditorViewportClient->Viewport->KeyState(EKeys::RightShift);
 
-	SelectActorsByPredicate(World, bInSelect, bShiftDown, Predicate);
+	SelectActorsByPredicate(World, bInSelect, !bShiftDown, Predicate);
 
 	return true;
 }
@@ -545,28 +553,29 @@ bool FEditorViewportSelectability::FrustumSelectWorldActors(const FConvexVolume&
 
 	const auto Predicate = [this, &InFrustum, InEditorViewportClient, bUseStrictSelection](AActor* const InActor) -> bool
 	{
-		return IsObjectSelectableInViewport(InActor)
-			&& DoesActorIntersectFrustum(*InActor, InFrustum, InEditorViewportClient, bUseStrictSelection);
+		const bool bSelectable = IsObjectSelectableInViewport(InActor);
+		const bool bIntersects = DoesActorIntersectFrustum(*InActor, InFrustum, InEditorViewportClient, bUseStrictSelection);
+		return bSelectable && bIntersects;
 	};
 
 	const bool bShiftDown = InEditorViewportClient->Viewport->KeyState(EKeys::LeftShift)
 		|| InEditorViewportClient->Viewport->KeyState(EKeys::RightShift);
 
-	SelectActorsByPredicate(World, bInSelect, bShiftDown, Predicate);
+	SelectActorsByPredicate(World, bInSelect, !bShiftDown, Predicate);
 
 	return true;
 }
 
 void FEditorViewportSelectability::DrawEnabledTextNotice(FCanvas* const InCanvas, const FText& InText)
 {
-	const FString HelpString = InText.ToString();
+	const FStringView HelpString = *InText.ToString();
 
 	FTextSizingParameters SizingParameters(GEngine->GetLargeFont(), 1.f, 1.f);
 	UCanvas::CanvasStringSize(SizingParameters, HelpString);
 
 	const float ViewWidth = InCanvas->GetViewRect().Width() / InCanvas->GetDPIScale();
 	const float DrawX = FMath::FloorToFloat((ViewWidth - SizingParameters.DrawXL) * 0.5f);
-	InCanvas->DrawShadowedString(DrawX, 34.f, *HelpString, GEngine->GetLargeFont(), FLinearColor::White);
+	InCanvas->DrawShadowedString(DrawX, 34.f, HelpString, GEngine->GetLargeFont(), FLinearColor::White);
 }
 
 FText FEditorViewportSelectability::GetLimitedSelectionText(const TSharedPtr<FUICommandInfo>& InToggleAction, const FText& InDefaultText)
diff --git a/Engine/Source/Editor/UnrealEd/Public/EditorViewportSelectability.h b/Engine/Source/Editor/UnrealEd/Public/EditorViewportSelectability.h
--- a/Engine/Source/Editor/UnrealEd/Public/EditorViewportSelectability.h
+++ b/Engine/Source/Editor/UnrealEd/Public/EditorViewportSelectability.h
@@ -80,10 +80,12 @@ public:
 	bool FrustumSelectWorldActors(const FConvexVolume& InFrustum, FEditorViewportClient* const InEditorViewportClient, const bool bInSelect);
 
 protected:
-	static bool IsActorInLevelHiddenLayer(AActor& InActor, FLevelEditorViewportClient* const InLevelEditorViewportClient);
+	static bool IsActorSelectableClass(const AActor& InActor);
 
-    static bool DoesActorIntersectBox(AActor& InActor, const FBox& InBox, FEditorViewportClient* const InEditorViewportClient, const bool bInUseStrictSelection);
-    static bool DoesActorIntersectFrustum(AActor& InActor, const FConvexVolume& InFrustum, FEditorViewportClient* const InEditorViewportClient, const bool bInUseStrictSelection);
+	static bool IsActorInLevelHiddenLayer(const AActor& InActor, FLevelEditorViewportClient* const InLevelEditorViewportClient);
+
+    static bool DoesActorIntersectBox(const AActor& InActor, const FBox& InBox, FEditorViewportClient* const InEditorViewportClient, const bool bInUseStrictSelection);
+    static bool DoesActorIntersectFrustum(const AActor& InActor, const FConvexVolume& InFrustum, FEditorViewportClient* const InEditorViewportClient, const bool bInUseStrictSelection);
 
 	/**
 	 * Selects or deselects actors in a world. If no actors are specified, uses all the actors in the level.
