Editor Toolkit: Add a pointer guard if host widget is missing

diff --git a/Engine/Source/Editor/UnrealEd/Private/Toolkits/AssetEditorToolkit.cpp b/Engine/Source/Editor/UnrealEd/Private/Toolkits/AssetEditorToolkit.cpp
--- a/Engine/Source/Editor/UnrealEd/Private/Toolkits/AssetEditorToolkit.cpp
+++ b/Engine/Source/Editor/UnrealEd/Private/Toolkits/AssetEditorToolkit.cpp
@@ -1501,7 +1501,10 @@ void FAssetEditorToolkit::RegenerateMenusAndToolbars()
 
 	PostRegenerateMenusAndToolbars();
 
-	HostWidget->SetToolbar(Toolbar);
+	if (HostWidget)
+	{
+		HostWidget->SetToolbar(Toolbar);
+	}
 }
 
 
