Editor: make viewport FOV computation consistent with the game

#udn 12707945

diff --git a/Engine/Source/Editor/UnrealEd/Private/EditorViewportClient.cpp b/Engine/Source/Editor/UnrealEd/Private/EditorViewportClient.cpp
--- a/Engine/Source/Editor/UnrealEd/Private/EditorViewportClient.cpp
+++ b/Engine/Source/Editor/UnrealEd/Private/EditorViewportClient.cpp
@@ -1063,11 +1063,12 @@ FSceneView* FEditorViewportClient::CalcSceneView(FSceneViewFamily* ViewFamily, c
 		    {
 			    const float MinZ = GetNearClipPlane();
 			    const float MaxZ = MinZ;
-			    // Avoid zero ViewFOV's which cause divide by zero's in projection matrix
-			    const float MatrixFOV = FMath::Max(0.001f, ModifiedViewFOV) * (float)PI / 360.0f;
 
 			    if (bConstrainAspectRatio)
 			    {
+					// Avoid zero ViewFOV's which cause divide by zero's in projection matrix
+					const float MatrixFOV = FMath::Max(0.001f, ModifiedViewFOV) * (float)PI / 360.0f;
+
 				    if ((bool)ERHIZBuffer::IsInverted)
 				    {
 					    ViewInitOptions.ProjectionMatrix = FReversedZPerspectiveMatrix(
@@ -1096,7 +1097,9 @@ FSceneView* FEditorViewportClient::CalcSceneView(FSceneViewFamily* ViewFamily, c
 				    float XAxisMultiplier;
 				    float YAxisMultiplier;
 
-				    if (((ViewportSize.X > ViewportSize.Y) && (AspectRatioAxisConstraint == AspectRatio_MajorAxisFOV)) || (AspectRatioAxisConstraint == AspectRatio_MaintainXFOV))
+					const bool bMaintainXFOV = (((ViewportSize.X > ViewportSize.Y) && (AspectRatioAxisConstraint == AspectRatio_MajorAxisFOV)) || (AspectRatioAxisConstraint == AspectRatio_MaintainXFOV));
+
+				    if (bMaintainXFOV)
 				    {
 					    //if the viewport is wider than it is tall
 					    XAxisMultiplier = 1.0f;
@@ -1109,6 +1112,20 @@ FSceneView* FEditorViewportClient::CalcSceneView(FSceneViewFamily* ViewFamily, c
 					    YAxisMultiplier = 1.0f;
 				    }
 
+					// Here we do something similar to FMinimalViewInfo::CalculateProjectionMatrixGivenViewRectangle
+					// TODO: unify both codebases
+					float MatrixFOV;
+					if (!bMaintainXFOV && AspectRatio != 0.f) // TODO: read CVarUseLegacyMaintainYFOV
+					{
+						const float HalfXFOV = FMath::DegreesToRadians(FMath::Max(0.001f, ModifiedViewFOV) / 2.f);
+						const float HalfYFOV = FMath::Atan(FMath::Tan(HalfXFOV) / AspectRatio);
+						MatrixFOV = HalfYFOV;
+					}
+					else
+					{
+						MatrixFOV = FMath::Max(0.001f, ModifiedViewFOV) * (float)UE_PI / 360.0f;
+					}
+
 				    if ((bool)ERHIZBuffer::IsInverted)
 				    {
 					    ViewInitOptions.ProjectionMatrix = FReversedZPerspectiveMatrix(
