[Viewport Toolbar] Screen Percentage Menu: use UToolMenu version only

- Deprecating SCommonEditorViewportToolbarBase::ConstructScreenPercentageMenu function
- Screen Percentage menus are now using UToolMenu version of function

diff --git a/Engine/Source/Editor/LevelEditor/Private/SLevelViewportToolBar.cpp b/Engine/Source/Editor/LevelEditor/Private/SLevelViewportToolBar.cpp
--- a/Engine/Source/Editor/LevelEditor/Private/SLevelViewportToolBar.cpp
+++ b/Engine/Source/Editor/LevelEditor/Private/SLevelViewportToolBar.cpp
@@ -462,6 +462,12 @@ TSharedRef<SWidget> SLevelViewportToolBar::GenerateOptionsMenu()
 			LevelContextObject->LevelViewport = Viewport;
 			MenuContext.AddObject(LevelContextObject);
 		}
+
+		{
+			UUnrealEdViewportToolbarContext* EdViewportToolbarContext = NewObject<UUnrealEdViewportToolbarContext>();
+			EdViewportToolbarContext->Viewport = Viewport;
+			MenuContext.AddObject(EdViewportToolbarContext);
+		}
 	}
 
 	return UToolMenus::Get()->GenerateWidget(MenuName, MenuContext);
@@ -514,12 +520,11 @@ void SLevelViewportToolBar::FillOptionsMenu(UToolMenu* Menu)
 				Section.AddEntry(UE::LevelEditor::CreateFarViewPlaneMenu(Viewport));
 			}
 
-			FEditorViewportClient& ViewportClient = Viewport.Pin()->GetLevelViewportClient();
 			Section.AddSubMenu(
 				"ScreenPercentageSubMenu",
 				LOCTEXT("ScreenPercentageSubMenu", "Screen Percentage"),
 				LOCTEXT("ScreenPercentageSubMenu_ToolTip", "Customize the viewport's screen percentage"),
-				FNewMenuDelegate::CreateStatic(&SCommonEditorViewportToolbarBase::ConstructScreenPercentageMenu, &ViewportClient));
+				FNewToolMenuDelegate::CreateStatic(&UE::UnrealEd::ConstructScreenPercentageMenu));
 		}
 
 		{
diff --git a/Engine/Source/Editor/LevelEditor/Private/ViewportToolbar/LevelEditorViewportToolbarSections.cpp b/Engine/Source/Editor/LevelEditor/Private/ViewportToolbar/LevelEditorViewportToolbarSections.cpp
--- a/Engine/Source/Editor/LevelEditor/Private/ViewportToolbar/LevelEditorViewportToolbarSections.cpp
+++ b/Engine/Source/Editor/LevelEditor/Private/ViewportToolbar/LevelEditorViewportToolbarSections.cpp
@@ -1641,29 +1641,12 @@ FToolMenuEntry CreateViewportToolbarPerformanceAndScalabilitySubmenu()
 						FNewToolMenuDelegate::CreateLambda(
 							[](UToolMenu* ScreenPercentageSubMenu)
 							{
-								FToolMenuSection& UnnamedSection = ScreenPercentageSubMenu->FindOrAddSection(NAME_None);
-
 								ScreenPercentageSubMenu->AddDynamicSection(
 									NAME_None,
-									FNewToolMenuDelegateLegacy::CreateLambda(
-										[](FMenuBuilder& MenuBuilder, UToolMenu* InMenu) -> void
+									FNewToolMenuDelegate::CreateLambda(
+										[](UToolMenu* InMenu) -> void
 										{
-											ULevelViewportContext* const LevelViewportContext =
-												InMenu->FindContext<ULevelViewportContext>();
-											if (!LevelViewportContext)
-											{
-												return;
-											}
-
-											if (const TSharedPtr<::SLevelViewport> LevelViewport =
-													LevelViewportContext->LevelViewport.Pin())
-											{
-												TSharedPtr<FEditorViewportClient> Client =
-													LevelViewport->GetViewportClient();
-												SCommonEditorViewportToolbarBase::ConstructScreenPercentageMenu(
-													MenuBuilder, Client.Get()
-												);
-											}
+											UE::UnrealEd::ConstructScreenPercentageMenu(InMenu);
 										}
 									)
 								);
diff --git a/Engine/Source/Editor/UnrealEd/Private/SCommonEditorViewportToolbarBase.cpp b/Engine/Source/Editor/UnrealEd/Private/SCommonEditorViewportToolbarBase.cpp
--- a/Engine/Source/Editor/UnrealEd/Private/SCommonEditorViewportToolbarBase.cpp
+++ b/Engine/Source/Editor/UnrealEd/Private/SCommonEditorViewportToolbarBase.cpp
@@ -29,6 +29,62 @@
 
 #define LOCTEXT_NAMESPACE "SCommonEditorViewportToolbarBase"
 
+namespace UE::UnrealEd::Private
+{
+// Used internally:
+// - by old viewport toolbar menu SCommonEditorViewportToolbarBase::GenerateOptionsMenu()
+// - by deprecated API function SCommonEditorViewportToolbarBase::ConstructScreenPercentageMenu(FMenuBuilder& MenuBuilder, FEditorViewportClient* InViewportClient)
+void AddScreenPercentageMenu(FMenuBuilder& InMenuBuilder, const FEditorViewportClient* InViewportClient)
+{
+	if (!InViewportClient)
+	{
+		return;
+	}
+
+	TSharedPtr<SEditorViewport> EditorViewport = InViewportClient->GetEditorViewportWidget();
+	if (!EditorViewport)
+	{
+		return;
+	}
+
+	FName OldScreenPercentageMenuName = "CommonEditorViewport.OldViewportToolbar.ScreenPercentage";
+	if (!UToolMenus::Get()->IsMenuRegistered(OldScreenPercentageMenuName))
+	{
+		UToolMenu* Menu =
+			UToolMenus::Get()->RegisterMenu(OldScreenPercentageMenuName, NAME_None, EMultiBoxType::Menu, false);
+
+		FToolMenuSection& UnnamedSection = Menu->FindOrAddSection(NAME_None);
+
+		UnnamedSection.AddSubMenu(
+			"ScreenPercentage",
+			LOCTEXT("ScreenPercentageSubMenu", "Screen Percentage"),
+			LOCTEXT("ScreenPercentageSubMenu_ToolTip", "Customize the viewport's screen percentage"),
+			FNewToolMenuDelegate::CreateLambda(
+				[](UToolMenu* Submenu) -> void
+				{
+					UE::UnrealEd::ConstructScreenPercentageMenu(Submenu);
+				}
+			)
+		);
+	}
+
+	FToolMenuContext MenuContext;
+	{
+		MenuContext.AppendCommandList(EditorViewport->GetCommandList());
+
+		// Add the UnrealEd viewport toolbar context.
+		{
+			UUnrealEdViewportToolbarContext* const ContextObject =
+				UE::UnrealEd::CreateViewportToolbarDefaultContext(EditorViewport);
+
+			MenuContext.AddObject(ContextObject);
+		}
+	}
+
+	InMenuBuilder.AddWidget(UToolMenus::Get()->GenerateWidget(OldScreenPercentageMenuName, MenuContext), FText());
+}
+} // namespace UE::UnrealEd::Private
+
 //////////////////////////////////////////////////////////////////////////
 // SPreviewSceneProfileSelector
 
@@ -284,65 +340,7 @@ void SCommonEditorViewportToolbarBase::Construct(const FArguments& InArgs, TShar
 
 void SCommonEditorViewportToolbarBase::ConstructScreenPercentageMenu(FMenuBuilder& MenuBuilder, FEditorViewportClient* InViewportClient)
 {
-	FEditorViewportClient& ViewportClient = *InViewportClient;
-
-	FMargin CommonPadding(26.0f, 3.0f);
-
-	const int32 PreviewScreenPercentageMin = ISceneViewFamilyScreenPercentage::kMinTSRResolutionFraction * 100.0f;
-	const int32 PreviewScreenPercentageMax = ISceneViewFamilyScreenPercentage::kMaxTSRResolutionFraction * 100.0f;
-
-	const FEditorViewportCommands& BaseViewportCommands = FEditorViewportCommands::Get();
-
-	MenuBuilder.BeginSection("Summary", LOCTEXT("Summary", "Summary"));
-	{
-		MenuBuilder.AddWidget(
-			UE::UnrealEd::CreateCurrentPercentageWidget(ViewportClient),
-			FText::GetEmpty()
-		);
-
-		MenuBuilder.AddWidget(
-			UE::UnrealEd::CreateResolutionsWidget(ViewportClient),
-			FText::GetEmpty()
-		);
-		MenuBuilder.AddWidget(
-			UE::UnrealEd::CreateActiveViewportWidget(ViewportClient),
-			FText::GetEmpty()
-		);
-		MenuBuilder.AddWidget(
-			UE::UnrealEd::CreateSetFromWidget(ViewportClient),
-			FText::GetEmpty()
-		);
-		MenuBuilder.AddWidget(
-			UE::UnrealEd::CreateCurrentScreenPercentageSettingWidget(ViewportClient),
-			FText::GetEmpty()
-		);
-	}
-	MenuBuilder.EndSection();
-
-	MenuBuilder.BeginSection("ScreenPercentage", LOCTEXT("ScreenPercentage_ViewportOverride", "Viewport Override"));
-	{
-		MenuBuilder.AddMenuEntry(BaseViewportCommands.ToggleOverrideViewportScreenPercentage);
-		MenuBuilder.AddWidget(
-			UE::UnrealEd::CreateCurrentScreenPercentageWidget(ViewportClient),
-			LOCTEXT("ScreenPercentage", "Screen Percentage")
-		);
-	}
-	MenuBuilder.EndSection();
-
-	MenuBuilder.BeginSection("ScreenPercentageSettings", LOCTEXT("ScreenPercentage_ViewportSettings", "Viewport Settings"));
-	{
-		MenuBuilder.AddMenuEntry(BaseViewportCommands.OpenEditorPerformanceProjectSettings,
-			/* InExtensionHook = */ NAME_None,
-			/* InLabelOverride = */ TAttribute<FText>(),
-			/* InToolTipOverride = */ TAttribute<FText>(),
-			FSlateIcon(FAppStyle::GetAppStyleSetName(), "ProjectSettings.TabIcon"));
-		MenuBuilder.AddMenuEntry(BaseViewportCommands.OpenEditorPerformanceEditorPreferences,
-			/* InExtensionHook = */ NAME_None,
-			/* InLabelOverride = */ TAttribute<FText>(),
-			/* InToolTipOverride = */ TAttribute<FText>(),
-			FSlateIcon(FAppStyle::GetAppStyleSetName(), "EditorPreferences.TabIcon"));
-	}
-	MenuBuilder.EndSection();
+	UE::UnrealEd::Private::AddScreenPercentageMenu(MenuBuilder, InViewportClient);
 }
 
 FText SCommonEditorViewportToolbarBase::GetCameraMenuLabel() const
@@ -380,7 +378,7 @@ TSharedRef<SWidget> SCommonEditorViewportToolbarBase::GenerateOptionsMenu() cons
 	TSharedRef<SEditorViewport> ViewportRef = GetInfoProvider().GetViewportWidget();
 
 	const bool bIsPerspective = GetViewportClient().GetViewportType() == LVT_Perspective;
-	
+
 	const bool bInShouldCloseWindowAfterMenuSelection = true;
 	FMenuBuilder OptionsMenuBuilder(bInShouldCloseWindowAfterMenuSelection, ViewportRef->GetCommandList());
 	{
@@ -396,10 +394,7 @@ TSharedRef<SWidget> SCommonEditorViewportToolbarBase::GenerateOptionsMenu() cons
 				OptionsMenuBuilder.AddWidget( UE::UnrealEd::CreateFarViewPlaneMenuWidget(ViewportRef), LOCTEXT("FarViewPlane", "Far View Plane") );
 			}
 
-			OptionsMenuBuilder.AddSubMenu(
-				LOCTEXT("ScreenPercentageSubMenu", "Screen Percentage"),
-				LOCTEXT("ScreenPercentageSubMenu_ToolTip", "Customize the viewport's screen percentage"),
-				FNewMenuDelegate::CreateStatic(&SCommonEditorViewportToolbarBase::ConstructScreenPercentageMenu, &GetViewportClient()));
+			UE::UnrealEd::Private::AddScreenPercentageMenu(OptionsMenuBuilder, &GetViewportClient());
 		}
 		OptionsMenuBuilder.EndSection();
 
diff --git a/Engine/Source/Editor/UnrealEd/Private/ViewportToolbar/UnrealEdViewportToolbar.cpp b/Engine/Source/Editor/UnrealEd/Private/ViewportToolbar/UnrealEdViewportToolbar.cpp
--- a/Engine/Source/Editor/UnrealEd/Private/ViewportToolbar/UnrealEdViewportToolbar.cpp
+++ b/Engine/Source/Editor/UnrealEd/Private/ViewportToolbar/UnrealEdViewportToolbar.cpp
@@ -2835,7 +2835,7 @@ TSharedRef<SWidget> CreateCurrentScreenPercentageWidget(FEditorViewportClient& I
 	// clang-format on
 }
 
-static void ConstructScreenPercentageMenu(UToolMenu* InMenu)
+void ConstructScreenPercentageMenu(UToolMenu* InMenu)
 {
 	UUnrealEdViewportToolbarContext* const LevelViewportContext = InMenu->FindContext<UUnrealEdViewportToolbarContext>();
 	if (!LevelViewportContext)
diff --git a/Engine/Source/Editor/UnrealEd/Public/SCommonEditorViewportToolbarBase.h b/Engine/Source/Editor/UnrealEd/Public/SCommonEditorViewportToolbarBase.h
--- a/Engine/Source/Editor/UnrealEd/Public/SCommonEditorViewportToolbarBase.h
+++ b/Engine/Source/Editor/UnrealEd/Public/SCommonEditorViewportToolbarBase.h
@@ -142,7 +142,7 @@ public:
 
 	UNREALED_API void Construct(const FArguments& InArgs, TSharedPtr<class ICommonEditorViewportToolbarInfoProvider> InInfoProvider);
 
-	/** Build the screen percentage menu. */
+	UE_DEPRECATED(5.6, "Use the version taking UToolMenu as argument: UE::UnrealEd::ConstructScreenPercentageMenu(InMenu)")
 	static UNREALED_API void ConstructScreenPercentageMenu(FMenuBuilder& MenuBuilder, class FEditorViewportClient* ViewportClient);
 
 private:
diff --git a/Engine/Source/Editor/UnrealEd/Public/ViewportToolbar/UnrealEdViewportToolbar.h b/Engine/Source/Editor/UnrealEd/Public/ViewportToolbar/UnrealEdViewportToolbar.h
--- a/Engine/Source/Editor/UnrealEd/Public/ViewportToolbar/UnrealEdViewportToolbar.h
+++ b/Engine/Source/Editor/UnrealEd/Public/ViewportToolbar/UnrealEdViewportToolbar.h
@@ -168,6 +168,8 @@ UNREALED_API FOnViewportClientCamSpeedScalarChanged& OnViewportClientCamSpeedSca
 
 UNREALED_API void CreateCameraSpeedMenu(UToolMenu* InMenu);
 
+UNREALED_API void ConstructScreenPercentageMenu(UToolMenu* InMenu);
+
 // Camera Menu Widgets
 TSharedRef<SWidget> CreateCameraMenuWidget(const TSharedRef<SEditorViewport>& InViewport);
 TSharedRef<SWidget> CreateFOVMenuWidget(const TSharedRef<SEditorViewport>& InViewport);
