Cooker: Fix the preload optimization to support arbitrary ordering of packages during PumpLoads: in addition to preloading the package itself, also preload its graph of imports before calling LoadPackage.
Removed the EPackageState::LoadPrepare, EPackageState::LoadReady, and EPackageStateProperty::State property, replaced them (the return of the original) single EPackageState::Load. Substates within Load are now handled by a separate state variable on FPackagePreloader.
The substate on FPackagePreloader is more complex than a usual substate, because now FPackagePreloaders can be in an active substate of EPreloaderState even when the FPackageData's state is not in load.
#rnx

diff --git a/Engine/Source/Editor/UnrealEd/Classes/CookOnTheSide/CookOnTheFlyServer.h b/Engine/Source/Editor/UnrealEd/Classes/CookOnTheSide/CookOnTheFlyServer.h
--- a/Engine/Source/Editor/UnrealEd/Classes/CookOnTheSide/CookOnTheFlyServer.h
+++ b/Engine/Source/Editor/UnrealEd/Classes/CookOnTheSide/CookOnTheFlyServer.h
@@ -182,6 +182,7 @@ namespace UE::Cook
 	class FCookWorkerServer;
 	class FDiagnostics;
 	class FODSCClientData; 
+	class FPackagePreloader;
 	class FPackageWriterMPCollector;
 	class FRequestCluster;
 	class FRequestQueue;
@@ -383,8 +384,8 @@ private:
 	 * A knob to tune performance - How many packages should be present in the LoadPrepare+LoadReady queues before we
 	 * start processing the LoadQueue. If number is less, we will find other work to do, and load packages only if all
 	 * other work is done.
-	 * This allows us to have enough population in the LoadPrepareQueue to get benefit from the asynchronous work done
-	 * on packages in the LoadPrepareQueue.
+	 * This allows us to have enough population in the LoadQueue to get benefit from the asynchronous work done
+	 * on preloading packages.
 	 */
 	uint32 DesiredLoadQueueLength;
 	/** A knob to tune performance - how many packages to pull off in each call to PumpRequests. */
@@ -518,10 +519,7 @@ private:
 
 	/** Load packages in the LoadQueue until it's time to break. Report the number of loads that were pushed to save. */
 	void PumpLoads(UE::Cook::FTickStackData& StackData, uint32 DesiredQueueLength, int32& OutNumPushed, bool& bOutBusy);
-	/** Move packages from LoadPrepare's entry queue into the PreloadingQueue until we run out of Preload slots. */
-	void PumpPreloadStarts();
-	/** Move preload-completed packages from LoadPrepare->LoadReady until we find one that is not finished preloading. */
-	void PumpPreloadCompletes();
+
 	/**
 	 * Load the given PackageData that was in the load queue and send it on to its next state.
 	 * Report the number of PackageDatas that were pushed to save (0 or 1)
@@ -1644,6 +1642,7 @@ private:
 	friend UE::Cook::FInitializeConfigSettings;
 	friend UE::Cook::FPackageData;
 	friend UE::Cook::FPackageDatas;
+	friend UE::Cook::FPackagePreloader;
 	friend UE::Cook::FPackageTracker;
 	friend UE::Cook::FPackageWriterMPCollector;
 	friend UE::Cook::FPendingCookedPlatformData;
diff --git a/Engine/Source/Editor/UnrealEd/Private/CookOnTheFlyServer.cpp b/Engine/Source/Editor/UnrealEd/Private/CookOnTheFlyServer.cpp
--- a/Engine/Source/Editor/UnrealEd/Private/CookOnTheFlyServer.cpp
+++ b/Engine/Source/Editor/UnrealEd/Private/CookOnTheFlyServer.cpp
@@ -28,6 +28,7 @@
 #include "Cooker/CookGenerationHelper.h"
 #include "Cooker/CookOnTheFlyServerInterface.h"
 #include "Cooker/CookPackageData.h"
+#include "Cooker/CookPackagePreloader.h"
 #include "Cooker/CookPlatformManager.h"
 #include "Cooker/CookProfiling.h"
 #include "Cooker/CookRequestCluster.h"
@@ -1763,23 +1764,13 @@ void UCookOnTheFlyServer::SetLoadBusy(bool bInLoadBusy)
 		{
 			int DisplayCount = 0;
 			const int DisplayMax = 10;
-			FLoadPrepareQueue& LoadPrepareQueue = PackageDatas->GetLoadPrepareQueue();
+			FLoadQueue& LoadQueue = PackageDatas->GetLoadQueue();
 #if !NO_LOGGING
 			FMsg::Logf(__FILE__, __LINE__, LogCook.GetCategoryName(), CookerIdleWarningSeverity,
 				TEXT("Cooker has been blocked from loading the current packages for %.0f seconds. %d packages in the loadqueue:"),
-				(float)(CurrentTime - LoadBusyStartTimeSeconds), LoadPrepareQueue.PreloadingQueue.Num() + LoadPrepareQueue.EntryQueue.Num());
+				(float)(CurrentTime - LoadBusyStartTimeSeconds), LoadQueue.Num());
 #endif
-			for (FPackageData* PackageData : LoadPrepareQueue.PreloadingQueue)
-			{
-				if (DisplayCount == DisplayMax)
-				{
-					UE_LOG(LogCook, Display, TEXT("    ..."));
-					break;
-				}
-				UE_LOG(LogCook, Display, TEXT("    %s"), *PackageData->GetFileName().ToString());
-				++DisplayCount;
-			}
-			for (FPackageData* PackageData : LoadPrepareQueue.EntryQueue)
+			for (FPackageData* PackageData : LoadQueue)
 			{
 				if (DisplayCount == DisplayMax)
 				{
@@ -2374,12 +2365,7 @@ UCookOnTheFlyServer::ECookAction UCookOnTheFlyServer::DecideNextCookAction(UE::C
 			SetIdleStatus(StackData, EIdleStatus::Active);
 			return ECookAction::Save;
 		}
-		else if (!bLoadBusy && Monitor.GetNumUrgent(UE::Cook::EPackageState::LoadPrepare) > 0)
-		{
-			SetIdleStatus(StackData, EIdleStatus::Active);
-			return ECookAction::Load;
-		}
-		else if (!bLoadBusy && Monitor.GetNumUrgent(UE::Cook::EPackageState::LoadReady) > 0)
+		else if (!bLoadBusy && Monitor.GetNumUrgent(UE::Cook::EPackageState::Load) > 0)
 		{
 			SetIdleStatus(StackData, EIdleStatus::Active);
 			return ECookAction::Load;
@@ -2394,12 +2380,7 @@ UCookOnTheFlyServer::ECookAction UCookOnTheFlyServer::DecideNextCookAction(UE::C
 			SetIdleStatus(StackData, EIdleStatus::Active);
 			return ECookAction::Save;
 		}
-		else if (Monitor.GetNumUrgent(UE::Cook::EPackageState::LoadPrepare) > 0)
-		{
-			SetIdleStatus(StackData, EIdleStatus::Active);
-			return ECookAction::Load;
-		}
-		else if (Monitor.GetNumUrgent(UE::Cook::EPackageState::LoadReady) > 0)
+		else if (Monitor.GetNumUrgent(UE::Cook::EPackageState::Load) > 0)
 		{
 			SetIdleStatus(StackData, EIdleStatus::Active);
 			return ECookAction::Load;
@@ -2425,7 +2406,7 @@ UCookOnTheFlyServer::ECookAction UCookOnTheFlyServer::DecideNextCookAction(UE::C
 		return ECookAction::SaveLimited;
 	}
 
-	int32 NumLoads = PackageDatas->GetLoadReadyQueue().Num() + PackageDatas->GetLoadPrepareQueue().Num();
+	int32 NumLoads = PackageDatas->GetLoadQueue().Num();
 	bool bLoadAvailable = ((!bLoadBusy) & (NumLoads > 0)) != 0;
 	if (bLoadAvailable & (NumLoads > static_cast<int32>(DesiredLoadQueueLength)))
 	{
@@ -2536,8 +2517,7 @@ int32 UCookOnTheFlyServer::NumMultiprocessLocalWorkerAssignments() const
 	UE::Cook::FPackageDataMonitor& Monitor = PackageDatas->GetMonitor();
 	int32 Result = WorkerRequests->GetNumExternalRequests() +
 		PackageDatas->GetRequestQueue().Num() +
-		PackageDatas->GetLoadPrepareQueue().Num() +
-		PackageDatas->GetLoadReadyQueue().Num() +
+		PackageDatas->GetLoadQueue().Num() +
 		PackageDatas->GetSaveQueue().Num();
 	for (const FRequestCluster& Cluster : PackageDatas->GetRequestQueue().GetRequestClusters())
 	{
@@ -2700,7 +2680,7 @@ void UCookOnTheFlyServer::PumpRequests(UE::Cook::FTickStackData& StackData, int3
 			DemoteToIdle(*PackageData, ESendFlags::QueueAdd, SuppressCookReason);
 			continue;
 		}
-		PackageData->SendToState(EPackageState::LoadPrepare, ESendFlags::QueueAdd, EStateChangeReason::Requested);
+		PackageData->SendToState(EPackageState::Load, ESendFlags::QueueAdd, EStateChangeReason::Requested);
 		++NumInBatch;
 	}
 	OutNumPushed += NumInBatch;
@@ -2837,107 +2817,113 @@ void UCookOnTheFlyServer::PromoteToSaveComplete(UE::Cook::FPackageData& PackageD
 void UCookOnTheFlyServer::PumpLoads(UE::Cook::FTickStackData& StackData, uint32 DesiredQueueLength, int32& OutNumPushed, bool& bOutBusy)
 {
 	using namespace UE::Cook;
-	FPackageDataQueue& LoadReadyQueue = PackageDatas->GetLoadReadyQueue();
-	FLoadPrepareQueue& LoadPrepareQueue = PackageDatas->GetLoadPrepareQueue();
+	FLoadQueue& LoadQueue = PackageDatas->GetLoadQueue();
 	FPackageDataMonitor& Monitor = PackageDatas->GetMonitor();
 	bool bIsUrgentInProgress = Monitor.GetNumUrgent() > 0;
 	OutNumPushed = 0;
 	bOutBusy = false;
 
+	TSet<FPackageData*>& InProgress = LoadQueue.InProgress;
+	TSet<TRefCountPtr<FPackagePreloader>>& ActivePreloads = LoadQueue.ActivePreloads;
+	TRingBuffer<TRefCountPtr<FPackagePreloader>>& ReadyForLoads = LoadQueue.ReadyForLoads;
+
+	if (bIsUrgentInProgress && !Monitor.GetNumUrgent(EPackageState::Load))
+	{
+		return;
+	}
+
 	// Process loads until we reduce the queue size down to the desired size or we hit the max number of loads per batch
 	// We do not want to load too many packages without saving because if we hit the memory limit and GC every package
 	// we load will have to be loaded again
-	while (LoadReadyQueue.Num() + LoadPrepareQueue.Num() > static_cast<int32>(DesiredQueueLength) &&
-		OutNumPushed < LoadBatchSize)
+	while (InProgress.Num() > static_cast<int32>(DesiredQueueLength) && OutNumPushed < LoadBatchSize)
 	{
 		if (StackData.Timer.IsActionTimeUp())
 		{
 			return;
 		}
-		if (bIsUrgentInProgress && !Monitor.GetNumUrgent(EPackageState::LoadPrepare) && !Monitor.GetNumUrgent(EPackageState::LoadReady))
+		if (bIsUrgentInProgress && !Monitor.GetNumUrgent(EPackageState::Load))
 		{
 			return;
 		}
-		COOK_STAT(DetailedCookStats::PeakLoadQueueSize = FMath::Max(DetailedCookStats::PeakLoadQueueSize, LoadPrepareQueue.Num() + LoadReadyQueue.Num()));
-		PumpPreloadStarts(); // PumpPreloadStarts after every load so that we keep adding preloads ahead of our need for them
+		COOK_STAT(DetailedCookStats::PeakLoadQueueSize = FMath::Max(DetailedCookStats::PeakLoadQueueSize, InProgress.Num()));
 
-		if (LoadReadyQueue.IsEmpty())
+		if (!ReadyForLoads.IsEmpty())
 		{
-			PumpPreloadCompletes();
-			if (LoadReadyQueue.IsEmpty())
+			TRefCountPtr<FPackagePreloader> Preloader = ReadyForLoads.PopFrontValue();
+			FPackageData& PackageData(Preloader->GetPackageData());
+			if (PackageData.GetState() == EPackageState::Load)
+			{
+				// A PackageData is in the load state, and we are done with preloading its imports
+				// and are ready to load it.
+				// Call extra code to add logging and state-transitioning the package out of load.
+				int32 NumPushed;
+				LoadPackageInQueue(PackageData, StackData.ResultFlags, NumPushed);
+				OutNumPushed += NumPushed;
+			}
+			else
 			{
-				if (!LoadPrepareQueue.IsEmpty())
+				// The PackagePreloader is done preloading and needs to be loaded, but its PackageData
+				// is in some other state. Just do minimum amount of work to load the package.
+				// Note that generated packages do not come through here; they are never added because
+				// they are not in any other package's import tree.
+				if (PackageData.IsGenerated())
 				{
-					bOutBusy = true;
+					UE_LOG(LogCook, Warning,
+						TEXT("Package %s is generated but is ReadyForLoad when not in load state. State == %s, PreloaderState == %s, CountFromRequestedLoads == %d."),
+						*WriteToString<256>(PackageData.GetPackageName()), LexToString(PackageData.GetState()),
+						LexToString(Preloader->GetState()), Preloader->GetCountFromRequestedLoads());
 				}
-				break;
+				UPackage* UnusedPackage;
+				LoadPackageForCooking(PackageData, UnusedPackage);
+				Preloader->PumpLoadsMarkLoadAttemptComplete();
 			}
-		}
 
-		FPackageData& PackageData(*LoadReadyQueue.PopFrontValue());
-		FPoppedPackageDataScope Scope(PackageData);
-		if (TryCreateRequestCluster(PackageData))
-		{
-			continue;
-		}
-
-		int32 NumPushed;
-		LoadPackageInQueue(PackageData, StackData.ResultFlags, NumPushed);
-		OutNumPushed += NumPushed;
-		ProcessUnsolicitedPackages(); // May add new packages into the LoadQueue
+			ProcessUnsolicitedPackages(); // May add new packages into LoadInbox
 #if ENABLE_LOW_LEVEL_MEM_TRACKER
-		FLowLevelMemTracker::Get().UpdateStatsPerFrame();
+			FLowLevelMemTracker::Get().UpdateStatsPerFrame();
 #endif
+			if (PumpHasExceededMaxMemory(StackData.ResultFlags))
+			{
+				return;
+			}
+			continue;
+		}
 
-		if (PumpHasExceededMaxMemory(StackData.ResultFlags))
+		// Process all values in the inbox until it is empty or we run out of time. Adding values from the inbox may
+		// change the front of the PendingKicks priority queue.
+		while (FPackagePreloader::PumpLoadsTryStartInboxPackage(*this))
 		{
-			return;
+			// Work was done in while condition
+			if (StackData.Timer.IsActionTimeUp())
+			{
+				return;
+			}
 		}
-	}
-}
 
-void UCookOnTheFlyServer::PumpPreloadCompletes()
-{
-	using namespace UE::Cook;
-
-	FPackageDataQueue& PreloadingQueue = PackageDatas->GetLoadPrepareQueue().PreloadingQueue;
-	const bool bLocalPreloadingEnabled = bPreloadingEnabled;
-	while (!PreloadingQueue.IsEmpty())
-	{
-		FPackageData* PackageData = PreloadingQueue.First();
-		if (!bLocalPreloadingEnabled || PackageData->TryPreload())
+		// Kick preloads until we run out of preload budget
+		while (FPackagePreloader::PumpLoadsTryKickPreload(*this))
 		{
-			// Ready to go
-			PreloadingQueue.PopFront();
-			PackageData->SendToState(EPackageState::LoadReady, ESendFlags::QueueAdd, EStateChangeReason::Loaded);
-			continue;
+			// Work was done in while condition
 		}
-		break;
-	}
-}
-
-void UCookOnTheFlyServer::PumpPreloadStarts()
-{
-	using namespace UE::Cook;
-
-	FPackageDataMonitor& Monitor = PackageDatas->GetMonitor();
-	FLoadPrepareQueue& LoadPrepareQueue = PackageDatas->GetLoadPrepareQueue();
-	FPackageDataQueue& PreloadingQueue = LoadPrepareQueue.PreloadingQueue;
-	FPackageDataQueue& EntryQueue = LoadPrepareQueue.EntryQueue;
 
-	const bool bLocalPreloadingEnabled = bPreloadingEnabled;
-	while (!EntryQueue.IsEmpty() && Monitor.GetNumPreloadAllocated() < static_cast<int32>(MaxPreloadAllocated))
-	{
-		FPackageData* PackageData = EntryQueue.PopFrontValue();
-		if (TryCreateRequestCluster(*PackageData))
+		// Poll all active preloads
+		for (TSet<TRefCountPtr<FPackagePreloader>>::TIterator Iter(ActivePreloads); Iter; ++Iter)
 		{
-			continue;
+			TRefCountPtr<FPackagePreloader>& IterPreloader = *Iter;
+			if (IterPreloader->PumpLoadsIsReadyToLeavePreload())
+			{
+				TRefCountPtr<FPackagePreloader> Preloader(IterPreloader);
+				Iter.RemoveCurrent();
+				Preloader->SendToState(EPreloaderState::ReadyForLoad, ESendFlags::QueueAdd);
+			}
 		}
-		if (bLocalPreloadingEnabled)
+
+		// If we did not find any packages ready to load then report the load queue is busy waiting for preloads
+		if (ReadyForLoads.IsEmpty())
 		{
-			PackageData->TryPreload();
+			bOutBusy = true;
+			break;
 		}
-		PreloadingQueue.Add(PackageData);
 	}
 }
 
@@ -2947,11 +2933,15 @@ void UCookOnTheFlyServer::LoadPackageInQueue(UE::Cook::FPackageData& PackageData
 
 	UPackage* LoadedPackage = nullptr;
 	OutNumPushed = 0;
+	TRefCountPtr<FPackagePreloader> Preloader = PackageData.CreatePackagePreloader();
+	check(Preloader->GetState() == EPreloaderState::ReadyForLoad);
 
 	FName PackageFileName(PackageData.GetFileName());
 	if (!PackageData.IsGenerated())
 	{
 		bool bLoadFullySuccessful = LoadPackageForCooking(PackageData, LoadedPackage);
+		// Mark the load attempt complete before we do any state transition of the PackageData.
+		Preloader->PumpLoadsMarkLoadAttemptComplete();
 		if (!bLoadFullySuccessful)
 		{
 			ResultFlags |= COSR_ErrorLoadingPackage;
@@ -2979,6 +2969,9 @@ void UCookOnTheFlyServer::LoadPackageInQueue(UE::Cook::FPackageData& PackageData
 	}
 	else
 	{
+		// Generated packages do not use the preload, so go ahead and mark it complete now. As with regular packages,
+		// we need to mark it complete before any state transitions of the PackageData.
+		Preloader->PumpLoadsMarkLoadAttemptComplete();
 		TRefCountPtr<FGenerationHelper> GenerationHelper = PackageData.TryCreateValidParentGenerationHelper();
 		if (!GenerationHelper)
 		{
@@ -3036,7 +3029,7 @@ void UCookOnTheFlyServer::LoadPackageInQueue(UE::Cook::FPackageData& PackageData
 	{
 		// Already cooked. This can happen if we needed to load a package that was previously cooked and garbage collected because it is a loaddependency of a new request.
 		// Send the package back to idle, nothing further to do with it.
-		DemoteToIdle(PackageData, ESendFlags::QueueAdd, ESuppressCookReason::AlreadyCooked);
+		DemoteToIdle(PackageData, ESendFlags::QueueAddAndRemove, ESuppressCookReason::AlreadyCooked);
 		return;
 	}
 
@@ -3056,7 +3049,7 @@ void UCookOnTheFlyServer::LoadPackageInQueue(UE::Cook::FPackageData& PackageData
 
 	PostLoadPackageFixup(PackageData, LoadedPackage);
 	PackageData.SetPackage(LoadedPackage);
-	PackageData.SendToState(EPackageState::SaveActive, ESendFlags::QueueAdd, EStateChangeReason::Loaded);
+	PackageData.SendToState(EPackageState::SaveActive, ESendFlags::QueueAddAndRemove, EStateChangeReason::Loaded);
 	++OutNumPushed;
 }
 
@@ -3081,7 +3074,7 @@ void UCookOnTheFlyServer::RejectPackageToLoad(UE::Cook::FPackageData& PackageDat
 			IFileManager::Get().Delete(*SandboxFilename);
 		}
 	}
-	DemoteToIdle(PackageData, UE::Cook::ESendFlags::QueueAdd, Reason);
+	DemoteToIdle(PackageData, UE::Cook::ESendFlags::QueueAddAndRemove, Reason);
 }
 
 EDataValidationResult UCookOnTheFlyServer::ValidateSourcePackage(UE::Cook::FPackageData& PackageData, UPackage* Package)
@@ -5447,15 +5440,7 @@ void UCookOnTheFlyServer::PreGarbageCollect()
 		{
 			AddPackageName(PackageData->GetPackageName());
 		}
-		for (FPackageData* PackageData : PackageDatas->GetLoadPrepareQueue().EntryQueue)
-		{
-			AddPackageName(PackageData->GetPackageName());
-		}
-		for (FPackageData* PackageData : PackageDatas->GetLoadPrepareQueue().PreloadingQueue)
-		{
-			AddPackageName(PackageData->GetPackageName());
-		}
-		for (FPackageData* PackageData : PackageDatas->GetLoadReadyQueue())
+		for (FPackageData* PackageData : PackageDatas->GetLoadQueue())
 		{
 			AddPackageName(PackageData->GetPackageName());
 		}
@@ -5644,9 +5629,20 @@ void UCookOnTheFlyServer::PostGarbageCollect()
 			GenerationHelper->PostGarbageCollect(GenerationHelper, *GCDiagnosticContext);
 		}
 	});
+
+	// Second pass over all PackageDatas, combine a few operations
 	PackageDatas->LockAndEnumeratePackageDatas([](FPackageData* PackageData)
 	{
+		// Mark that the PackageData no longer needs to be keepreferenced.
+		// This can only be done after all GenerationHelper->PostGarbageCollect have been called.
 		PackageData->SetKeepReferencedDuringGC(false);
+
+		// Reset the completion flags for FPreloadPackage, since the UPackage might be no longer loaded.
+		TRefCountPtr<FPackagePreloader> Preloader = PackageData->GetPackagePreloader();
+		if (Preloader)
+		{
+			Preloader->PostGarbageCollect();
+		}
 	});
 
 	// Only after running all possible callbacks that need our links for diagnostics, clear the list of temporary
@@ -10523,8 +10519,7 @@ void UCookOnTheFlyServer::CookByTheBookFinishedInternal()
 	check(PackageDatas->GetRequestQueue().IsEmpty());
 	check(PackageDatas->GetRequestQueue().GetDiscoveryQueue().IsEmpty());
 	check(PackageDatas->GetAssignedToWorkerSet().IsEmpty());
-	check(PackageDatas->GetLoadPrepareQueue().IsEmpty());
-	check(PackageDatas->GetLoadReadyQueue().IsEmpty());
+	check(PackageDatas->GetLoadQueue().IsEmpty());
 	check(PackageDatas->GetSaveQueue().IsEmpty());
 	check(PackageDatas->GetSaveStalledSet().IsEmpty());
 
@@ -11118,34 +11113,34 @@ void UCookOnTheFlyServer::CancelAllQueues()
 	{
 		DemoteToIdle(*SaveQueue.PopFrontValue(), ESendFlags::QueueAdd, ESuppressCookReason::CookCanceled);
 	}
-	FPackageDataQueue& LoadReadyQueue = PackageDatas->GetLoadReadyQueue();
-	while (!LoadReadyQueue.IsEmpty())
-	{
-		DemoteToIdle(*LoadReadyQueue.PopFrontValue(), ESendFlags::QueueAdd, ESuppressCookReason::CookCanceled);
-	}
-	FLoadPrepareQueue& LoadPrepareQueue = PackageDatas->GetLoadPrepareQueue();
-	while (!LoadPrepareQueue.IsEmpty())
+	FLoadQueue& LoadQueue = PackageDatas->GetLoadQueue();
+	FRequestQueue& RequestQueue = PackageDatas->GetRequestQueue();
+	TArray<FPackageData*> DatasInStatesWithTSet;
+	DatasInStatesWithTSet.Reset(LoadQueue.Num()
+		+ PackageDatas->GetAssignedToWorkerSet().Num()
+		+ PackageDatas->GetSaveStalledSet().Num()
+		+ RequestQueue.GetRestartedRequests().Num());
+	for (FPackageData* PackageData : LoadQueue)
 	{
-		DemoteToIdle(*LoadPrepareQueue.PopFront(), ESendFlags::QueueAdd, ESuppressCookReason::CookCanceled);
+		DatasInStatesWithTSet.Add(PackageData);
 	}
 	for (FPackageData* PackageData : PackageDatas->GetAssignedToWorkerSet())
 	{
-		DemoteToIdle(*PackageData, ESendFlags::QueueAdd, ESuppressCookReason::CookCanceled);
+		DatasInStatesWithTSet.Add(PackageData);
 	}
-	PackageDatas->GetAssignedToWorkerSet().Empty();
 	for (FPackageData* PackageData : PackageDatas->GetSaveStalledSet())
 	{
-		DemoteToIdle(*PackageData, ESendFlags::QueueAdd, ESuppressCookReason::CookCanceled);
+		DatasInStatesWithTSet.Add(PackageData);
 	}
-	PackageDatas->GetSaveStalledSet().Empty();
-	FRequestQueue& RequestQueue = PackageDatas->GetRequestQueue();
-	RequestQueue.GetDiscoveryQueue().Empty();
-	FPackageDataSet& RestartedRequests = RequestQueue.GetRestartedRequests();
-	for (FPackageData* PackageData : RestartedRequests)
+	for (FPackageData* PackageData : RequestQueue.GetRestartedRequests())
 	{
-		DemoteToIdle(*PackageData, ESendFlags::QueueAdd, ESuppressCookReason::CookCanceled);
+		DatasInStatesWithTSet.Add(PackageData);
 	}
-	RestartedRequests.Empty();
+	for (FPackageData* PackageData : DatasInStatesWithTSet)
+	{
+		DemoteToIdle(*PackageData, ESendFlags::QueueAddAndRemove, ESuppressCookReason::CookCanceled);
+	}
+	RequestQueue.GetDiscoveryQueue().Empty();
 	TRingBuffer<FRequestCluster>& RequestClusters = RequestQueue.GetRequestClusters();
 	for (FRequestCluster& RequestCluster : RequestClusters)
 	{
@@ -11192,6 +11187,8 @@ void UCookOnTheFlyServer::ClearPlatformCookedData(const ITargetPlatform* TargetP
 
 void UCookOnTheFlyServer::ResetCook(TConstArrayView<TPair<const ITargetPlatform*, bool>> TargetPlatforms)
 {
+	using namespace UE::Cook;
+
 	PackageDatas->LockAndEnumeratePackageDatas([TargetPlatforms](UE::Cook::FPackageData* PackageData)
 	{
 		PackageData->FindOrAddPlatformData(CookerLoadingPlatformKey).ResetReachable();
@@ -11199,7 +11196,7 @@ void UCookOnTheFlyServer::ResetCook(TConstArrayView<TPair<const ITargetPlatform*
 		for (const TPair<const ITargetPlatform*, bool>& Pair : TargetPlatforms)
 		{
 			const ITargetPlatform* TargetPlatform = Pair.Key;
-			UE::Cook::FPackagePlatformData* PlatformData = PackageData->FindPlatformData(TargetPlatform);
+			FPackagePlatformData* PlatformData = PackageData->FindPlatformData(TargetPlatform);
 			if (PlatformData)
 			{
 				bool bResetResults = Pair.Value;
@@ -11210,8 +11207,13 @@ void UCookOnTheFlyServer::ResetCook(TConstArrayView<TPair<const ITargetPlatform*
 				}
 			}
 		}
+
+		PackageData->SetSuppressCookReason(ESuppressCookReason::NotSuppressed);
+		PackageData->SetLeafToRootRank(MAX_uint32);
 	});
 
+	PackageDatas->ResetLeafToRootRank();
+
 	TArray<FName> PackageNames;
 	for (const TPair<const ITargetPlatform*, bool>& Pair : TargetPlatforms)
 	{
@@ -12154,29 +12156,31 @@ void UCookOnTheFlyServer::GetPackagesToRetract(int32 NumToRetract, TArray<FName>
 		return;
 	}
 
-	FLoadPrepareQueue& LoadPrepareQueue = PackageDatas->GetLoadPrepareQueue();
-	for (FPackageData* PackageData : LoadPrepareQueue.EntryQueue)
-	{
-		if (AddPackageIfPossibleAndReportDone(PackageData))
-		{
-			return;
-		}
-	}
-	for (FPackageData* PackageData : LoadPrepareQueue.PreloadingQueue)
+	// Send back loadstate packages that have not started loading before sending back any that have.
+	FLoadQueue& LoadQueue = PackageDatas->GetLoadQueue();
+	for (FPackageData* PackageData : LoadQueue)
 	{
-		if (AddPackageIfPossibleAndReportDone(PackageData))
+		TRefCountPtr<FPackagePreloader> Preloader = PackageData->GetPackagePreloader();
+		if (!(Preloader && (Preloader->IsPackageLoaded() || Preloader->GetState() >= EPreloaderState::ActivePreload)))
 		{
-			return;
+			if (AddPackageIfPossibleAndReportDone(PackageData))
+			{
+				return;
+			}
 		}
 	}
-	for (FPackageData* PackageData : PackageDatas->GetLoadReadyQueue())
+	for (FPackageData* PackageData : LoadQueue)
 	{
-		if (AddPackageIfPossibleAndReportDone(PackageData))
+		TRefCountPtr<FPackagePreloader> Preloader = PackageData->GetPackagePreloader();
+		if (Preloader && (Preloader->IsPackageLoaded() || Preloader->GetState() >= EPreloaderState::ActivePreload))
 		{
-			return;
+			if (AddPackageIfPossibleAndReportDone(PackageData))
+			{
+				return;
+			}
 		}
 	}
-	// Send back all packages that have not started saving before sending back any that have
+	// Send back savestate packages that have not started saving before sending back any that have.
 	for (FPackageData* PackageData : PackageDatas->GetSaveQueue())
 	{
 		if (PackageData->GetSaveSubState() <= ESaveSubState::StartSave)
diff --git a/Engine/Source/Editor/UnrealEd/Private/Cooker/CookPackageData.cpp b/Engine/Source/Editor/UnrealEd/Private/Cooker/CookPackageData.cpp
--- a/Engine/Source/Editor/UnrealEd/Private/Cooker/CookPackageData.cpp
+++ b/Engine/Source/Editor/UnrealEd/Private/Cooker/CookPackageData.cpp
@@ -139,21 +139,6 @@ void FPackageData::ClearReferences()
 	}
 }
 
-const FName& FPackageData::GetPackageName() const
-{
-	return PackageName;
-}
-
-const FName& FPackageData::GetFileName() const
-{
-	return FileName;
-}
-
-void FPackageData::SetFileName(const FName& InFileName)
-{
-	FileName = InFileName;
-}
-
 int32 FPackageData::GetPlatformsNeedingCookingNum() const
 {
 	int32 Result = 0;
@@ -604,11 +589,8 @@ struct FStateProperties
 		case EPackageState::AssignedToWorker:
 			Properties = EPackageStateProperty::InProgress | EPackageStateProperty::AssignedToWorkerProperty;
 			break;
-		case EPackageState::LoadPrepare:
-			Properties = EPackageStateProperty::InProgress | EPackageStateProperty::Loading;
-			break;
-		case EPackageState::LoadReady:
-			Properties = EPackageStateProperty::InProgress | EPackageStateProperty::Loading;
+		case EPackageState::Load:
+			Properties = EPackageStateProperty::InProgress;
 			break;
 		case EPackageState::SaveActive:
 			Properties = EPackageStateProperty::InProgress | EPackageStateProperty::Saving;
@@ -650,19 +632,12 @@ void FPackageData::SendToState(EPackageState NextState, ESendFlags SendFlags, ES
 		}
 		OnExitAssignedToWorker();
 		break;
-	case EPackageState::LoadPrepare:
+	case EPackageState::Load:
 		if (!!(SendFlags & ESendFlags::QueueRemove))
 		{
-			ensure(PackageDatas.GetLoadPrepareQueue().Remove(this) == 1);
+			ensure(PackageDatas.GetLoadQueue().Remove(this) == 1);
 		}
-		OnExitLoadPrepare();
-		break;
-	case EPackageState::LoadReady:
-		if (!!(SendFlags & ESendFlags::QueueRemove))
-		{
-			ensure(PackageDatas.GetLoadReadyQueue().Remove(this) == 1);
-		}
-		OnExitLoadReady();
+		OnExitLoad();
 		break;
 	case EPackageState::SaveActive:
 		if (!!(SendFlags & ESendFlags::QueueRemove))
@@ -707,9 +682,6 @@ void FPackageData::SendToState(EPackageState NextState, ESendFlags SendFlags, ES
 			case EPackageStateProperty::InProgress:
 				OnExitInProgress(ReleaseSaveReason);
 				break;
-			case EPackageStateProperty::Loading:
-				OnExitLoading();
-				break;
 			case EPackageStateProperty::Saving:
 				OnExitSaving(ReleaseSaveReason, NextState);
 				break;
@@ -734,9 +706,6 @@ void FPackageData::SendToState(EPackageState NextState, ESendFlags SendFlags, ES
 			case EPackageStateProperty::InProgress:
 				OnEnterInProgress();
 				break;
-			case EPackageStateProperty::Loading:
-				OnEnterLoading();
-				break;
 			case EPackageStateProperty::Saving:
 				OnEnterSaving();
 				break;
@@ -771,32 +740,11 @@ void FPackageData::SendToState(EPackageState NextState, ESendFlags SendFlags, ES
 			PackageDatas.GetAssignedToWorkerSet().Add(this);
 		}
 		break;
-	case EPackageState::LoadPrepare:
-		OnEnterLoadPrepare();
-		if ((SendFlags & ESendFlags::QueueAdd) != ESendFlags::QueueNone)
-		{
-			if (GetIsUrgent())
-			{
-				PackageDatas.GetLoadPrepareQueue().AddFront(this);
-			}
-			else
-			{
-				PackageDatas.GetLoadPrepareQueue().Add(this);
-			}
-		}
-		break;
-	case EPackageState::LoadReady:
-		OnEnterLoadReady();
+	case EPackageState::Load:
+		OnEnterLoad();
 		if ((SendFlags & ESendFlags::QueueAdd) != ESendFlags::QueueNone)
 		{
-			if (GetIsUrgent())
-			{
-				PackageDatas.GetLoadReadyQueue().AddFront(this);
-			}
-			else
-			{
-				PackageDatas.GetLoadReadyQueue().Add(this);
-			}
+			PackageDatas.GetLoadQueue().Add(this);
 		}
 		break;
 	case EPackageState::SaveActive:
@@ -896,11 +844,8 @@ void FPackageData::CheckInContainer() const
 	case EPackageState::AssignedToWorker:
 		check(PackageDatas.GetAssignedToWorkerSet().Contains(this));
 		break;
-	case EPackageState::LoadPrepare:
-		check(PackageDatas.GetLoadPrepareQueue().Contains(this));
-		break;
-	case EPackageState::LoadReady:
-		check(Algo::Find(PackageDatas.GetLoadReadyQueue(), this) != nullptr);
+	case EPackageState::Load:
+		check(PackageDatas.GetLoadQueue().Contains(this));
 		break;
 	case EPackageState::SaveActive:
 		// The save queue is huge and often pushed at end. Check last element first and then scan.
@@ -989,20 +934,19 @@ void FPackageData::OnExitAssignedToWorker()
 {
 }
 
-void FPackageData::OnEnterLoadPrepare()
-{
-}
-
-void FPackageData::OnExitLoadPrepare()
-{
-}
-
-void FPackageData::OnEnterLoadReady()
+void FPackageData::OnEnterLoad()
 {
+	TRefCountPtr<FPackagePreloader> Local = CreatePackagePreloader();
+	Local->SetSelfReference();
+	check(PackagePreloader);
 }
 
-void FPackageData::OnExitLoadReady()
+void FPackageData::OnExitLoad()
 {
+	check(PackagePreloader); // Guaranteed by OnEnterLoad
+	PackagePreloader->OnPackageLeaveLoadState();
+	PackagePreloader->ClearSelfReference();
+	// PackagePreloader might now be nullptr
 }
 
 void FPackageData::OnEnterSaveActive()
@@ -1051,23 +995,6 @@ void FPackageData::OnExitInProgress(EStateChangeReason StateChangeReason)
 	ClearInProgressData(StateChangeReason);
 }
 
-void FPackageData::OnEnterLoading()
-{
-	TRefCountPtr<FPackagePreloader> Local = CreatePackagePreloader();
-	Local->SetSelfReference();
-
-	check(PackagePreloader);
-	PackagePreloader->CheckPreloadEmpty();
-}
-
-void FPackageData::OnExitLoading()
-{
-	check(PackagePreloader); // Guaranteed by OnEnterLoading
-	PackagePreloader->ClearPreload();
-	PackagePreloader->ClearSelfReference();
-	// PackagePreloader might now be nullptr
-}
-
 void FPackageData::OnEnterSaving()
 {
 	check(GetPackage() != nullptr && GetPackage()->IsFullyLoaded());
@@ -1169,13 +1096,6 @@ void FPackageData::OnPackagePreloaderDestroyed(FPackagePreloader& InPackagePrelo
 	PackagePreloader = nullptr;
 }
 
-bool FPackageData::TryPreload()
-{
-	check(IsInStateProperty(EPackageStateProperty::Loading));
-	check(PackagePreloader != nullptr); // Guaranteed by OnEnterLoading
-	return PackagePreloader->TryPreload();
-}
-
 TArray<FCachedObjectInOuter>& FPackageData::GetCachedObjectsInOuter()
 {
 	return CachedObjectsInOuter;
@@ -2946,49 +2866,6 @@ void FRequestQueue::NotifyRequestFencePassed(FPackageDatas& PackageDatas)
 	RequestFencePackageListeners.Empty();
 }
 
-bool FLoadPrepareQueue::IsEmpty()
-{
-	return Num() == 0;
-}
-
-int32 FLoadPrepareQueue::Num() const
-{
-	return PreloadingQueue.Num() + EntryQueue.Num();
-}
-
-FPackageData* FLoadPrepareQueue::PopFront()
-{
-	if (!PreloadingQueue.IsEmpty())
-	{
-		return PreloadingQueue.PopFrontValue();
-	}
-	else
-	{
-		return EntryQueue.PopFrontValue();
-	}
-}
-
-void FLoadPrepareQueue::Add(FPackageData* PackageData)
-{
-	EntryQueue.Add(PackageData);
-}
-
-void FLoadPrepareQueue::AddFront(FPackageData* PackageData)
-{
-	PreloadingQueue.AddFront(PackageData);
-}
-
-bool FLoadPrepareQueue::Contains(const FPackageData* PackageData) const
-{
-	return (Algo::Find(PreloadingQueue, PackageData) != nullptr) ||
-		(Algo::Find(EntryQueue, PackageData) != nullptr);
-}
-
-uint32 FLoadPrepareQueue::Remove(FPackageData* PackageData)
-{
-	return PreloadingQueue.Remove(PackageData) + EntryQueue.Remove(PackageData);
-}
-
 FPoppedPackageDataScope::FPoppedPackageDataScope(FPackageData& InPackageData)
 #if COOK_CHECKSLOW_PACKAGEDATA
 	: PackageData(InPackageData)
diff --git a/Engine/Source/Editor/UnrealEd/Private/Cooker/CookPackageData.h b/Engine/Source/Editor/UnrealEd/Private/Cooker/CookPackageData.h
--- a/Engine/Source/Editor/UnrealEd/Private/Cooker/CookPackageData.h
+++ b/Engine/Source/Editor/UnrealEd/Private/Cooker/CookPackageData.h
@@ -13,6 +13,7 @@
 #include "HAL/CriticalSection.h"
 #include "HAL/Platform.h"
 #include "IO/IoHash.h"
+#include "Math/NumericLimits.h"
 #include "Misc/EnumClassFlags.h"
 #include "Misc/Optional.h"
 #include "Misc/ScopeRWLock.h"
@@ -330,6 +331,14 @@ public:
 	 */
 	const FName& GetFileName() const;
 
+	/**
+	 * Get the LeafToRootRank that was found for this PackageData, if set. Returns MAX_uint32 if not yet set.
+	 * This value is not replicated to CookWorkers; each CookWorker calculates it individually during Pumploads.
+	 */
+	uint32 GetLeafToRootRank() const;
+	/** Set the LeafToRootRank for this package data. */
+	void SetLeafToRootRank(uint32 Value);
+
 	/** Reset OutPlatforms and copy current set of reachable & cookable & not yet cooked platforms into it. */
 	template <typename ArrayType>
 	void GetPlatformsNeedingCooking(ArrayType& OutPlatforms) const;
@@ -513,9 +522,6 @@ public:
 	bool GetIsVisited() const { return bIsVisited != 0; }
 	void SetIsVisited(bool bValue) { bIsVisited = static_cast<uint32>(bValue); }
 
-	/** Try to preload the file. Return true if preloading is complete (succeeded or failed or was skipped). */
-	bool TryPreload();
-
 	/**
 	 * The list of objects inside the package.  Only non-empty during saving; it is populated on demand by
 	 * TryCreateObjectCache and is cleared when leaving the save state.
@@ -767,10 +773,8 @@ private:
 	void OnExitRequest();
 	void OnEnterAssignedToWorker();
 	void OnExitAssignedToWorker();
-	void OnEnterLoadPrepare();
-	void OnExitLoadPrepare();
-	void OnEnterLoadReady();
-	void OnExitLoadReady();
+	void OnEnterLoad();
+	void OnExitLoad();
 	void OnEnterSaveActive();
 	void OnExitSaveActive();
 	void OnEnterSaveStalledRetracted();
@@ -780,8 +784,6 @@ private:
 	/* Entry/Exit gates for Properties shared between multiple states */
 	void OnExitInProgress(EStateChangeReason StateChangeReason);
 	void OnEnterInProgress();
-	void OnExitLoading();
-	void OnEnterLoading();
 	void OnExitSaving(EStateChangeReason ReleaseSaveReason, EPackageState NewState);
 	void OnEnterSaving();
 	void OnExitAssignedToWorkerProperty();
@@ -805,6 +807,7 @@ private:
 	/** The one-per-CookOnTheFlyServer owner of this PackageData. */
 	FPackageDatas& PackageDatas;
 	FPackagePreloader* PackagePreloader = nullptr;
+	uint32 LeafToRootRank = MAX_uint32;
 	int32 NumPendingCookedPlatformData = 0;
 	int32 CookedPlatformDataNextIndex = -1;
 	int32 NumRetriesBeginCacheOnObject = 0;
@@ -1038,22 +1041,51 @@ private:
 	FPackageDataSet NormalRequests;
 };
 
+/** A wrapper around a TRefCountPtr<FPackagePreloader> that defines operator< for FPackagePreloaderPriorityQueue. */
+struct FPackagePreloaderPriorityWrapper
+{
+	TRefCountPtr<FPackagePreloader> Payload;
+	bool operator<(const FPackagePreloaderPriorityWrapper& Other) const;
+};
+
+/**
+ * Priority queue for FPackagePreloaders in the PendingKick substate, prioritized mostly by LeafToRoot order,
+ * but with various exceptions. Controls which PendingKick preloader will next be kicked.
+ */
+class FPackagePreloaderPriorityQueue
+{
+public:
+	bool IsEmpty() const;
+	void Add(TRefCountPtr<FPackagePreloader> Preloader);
+	void Remove(const TRefCountPtr<FPackagePreloader>& Preloader);
+	TRefCountPtr<FPackagePreloader> PopFront();
+
+private:
+	TArray<FPackagePreloaderPriorityWrapper> Heap;
+};
+
 /**
- * Container for FPackageDatas in the LoadPrepare state. A FIFO with multiple substates.
+ * Container for FPackageDatas in the Load state. Has a single InProgress container, and  multiple subqueues which
+ * contain pointers to PackagePreloaders of packages which might be requested because they're in the load state, or
+ * requested because even though they are in another state (e.g. request or idle), one of the packages in the load
+ * state imports them so we want to preload and load them for better load performance of the referencer package.
  */
-class FLoadPrepareQueue
+class FLoadQueue
 {
 public:
 	bool IsEmpty();
 	int32 Num() const;
-	FPackageData* PopFront();
 	void Add(FPackageData* PackageData);
-	void AddFront(FPackageData* PackageData);
 	bool Contains(const FPackageData* PackageData) const;
 	uint32 Remove(FPackageData* PackageData);
-
-	FPackageDataQueue PreloadingQueue;
-	FPackageDataQueue EntryQueue;
+	TSet<FPackageData*>::TRangedForIterator begin();
+	TSet<FPackageData*>::TRangedForIterator end();
+
+	TRingBuffer<FPackageData*> Inbox;
+	FPackagePreloaderPriorityQueue PendingKicks;
+	TSet<TRefCountPtr<FPackagePreloader>> ActivePreloads;
+	TRingBuffer<TRefCountPtr<FPackagePreloader>> ReadyForLoads;
+	TSet<FPackageData*> InProgress;
 };
 
 /** Data duplicated from FPackageData that is stored separately for read/write from any thread. */
@@ -1101,6 +1133,12 @@ public:
 	FPackageDataMonitor& GetMonitor();
 	/** Return the backpointer to the CookOnTheFlyServer */
 	UCookOnTheFlyServer& GetCookOnTheFlyServer();
+
+	/** Return and increment the RoofToLankRank that should be assigned to the next PackageData needing one. */
+	uint32 GetNextLeafToRootRank();
+	/** Reset the LeafToRootRank index back to 0; called when the cook is reset. */
+	void ResetLeafToRootRank();
+
 	/**
 	 * Return the RequestQueue used by the CookOnTheFlyServer. The RequestQueue is the mostly-FIFO list of
 	 * PackageData that need to be cooked.
@@ -1111,15 +1149,10 @@ public:
 	TFastPointerSet<FPackageData*>& GetAssignedToWorkerSet();
 
 	/**
-	 * Return the LoadPrepareQueue used by the CookOnTheFlyServer. The LoadPrepareQueue is the dependency-ordered
-	 * list of FPackageData that need to be preloaded before they can be loaded.
-	 */
-	FLoadPrepareQueue& GetLoadPrepareQueue();
-	/**
-	 * Return the LoadReadyQueue used by the CookOnTheFlyServer. The LoadReadyQueue is the dependency-ordered list
-	 * of PackageData that need to be loaded.
+	 * Return the LoadQueue used by CookOnTheFlyServer. Container for packages in the Load state and for various
+	 * data that manages their passage through the state.
 	 */
-	FPackageDataQueue& GetLoadReadyQueue();
+	FLoadQueue& GetLoadQueue();
 	/**
 	 * Return the SaveQueue used by the CookOnTheFlyServer. The SaveQueue is the performance-sorted list of
 	 * PackageData that have been loaded and need to start or are only part way through saving.
@@ -1343,12 +1376,12 @@ private:
 	TMap<FName, FThreadsafePackageData> ThreadsafePackageDatas;
 	TRingBuffer<FPendingCookedPlatformDataContainer> PendingCookedPlatformDataLists;
 	TMap<UObject*, FCachedCookedPlatformDataState> CachedCookedPlatformDataObjects;
+	uint32 NextLeafToRootRank = 0;
 	int32 PendingCookedPlatformDataNum = 0;
 	FRequestQueue RequestQueue;
 	TFastPointerSet<FPackageData*> AssignedToWorkerSet;
 	TFastPointerSet<FPackageData*> SaveStalledSet;
-	FLoadPrepareQueue LoadPrepareQueue;
-	FPackageDataQueue LoadReadyQueue;
+	FLoadQueue LoadQueue;
 	FPackageDataQueue SaveQueue;
 	UCookOnTheFlyServer& CookOnTheFlyServer;
 	mutable FRWLock ExistenceLock;
@@ -1358,19 +1391,6 @@ private:
 	static IAssetRegistry* AssetRegistry;
 };
 
-template <typename CallbackType>
-inline void FPackageDatas::LockAndEnumeratePackageDatas(CallbackType&& Callback)
-{
-	FReadScopeLock ExistenceReadLock(ExistenceLock);
-	EnumeratePackageDatasWithinLock(Forward<CallbackType>(Callback));
-}
-
-template <typename CallbackType>
-inline void FPackageDatas::EnumeratePackageDatasWithinLock(CallbackType&& Callback)
-{
-	Allocator.EnumerateAllocations(Forward<CallbackType>(Callback));
-}
-
 /**
  * A debug-only scope class to confirm that each FPackageData removed from a container during a Pump function
  * is added to the container for its new state before leaving the Pump function.
@@ -1392,6 +1412,30 @@ struct FPoppedPackageDataScope
 ///////////////////////////////////////////////////////
 
 
+/** Return and increment the RoofToLankRank that should be assigned to the next PackageData needing one. */
+inline uint32 FPackageDatas::GetNextLeafToRootRank()
+{
+	return NextLeafToRootRank++;
+}
+
+inline void FPackageDatas::ResetLeafToRootRank()
+{
+	NextLeafToRootRank = 0;
+}
+
+template <typename CallbackType>
+inline void FPackageDatas::LockAndEnumeratePackageDatas(CallbackType&& Callback)
+{
+	FReadScopeLock ExistenceReadLock(ExistenceLock);
+	EnumeratePackageDatasWithinLock(Forward<CallbackType>(Callback));
+}
+
+template <typename CallbackType>
+inline void FPackageDatas::EnumeratePackageDatasWithinLock(CallbackType&& Callback)
+{
+	Allocator.EnumerateAllocations(Forward<CallbackType>(Callback));
+}
+
 template<typename CallbackType>
 inline void FPackageDatas::UpdateThreadsafePackageData(FName PackageName, CallbackType&& Callback)
 {
@@ -1433,14 +1477,9 @@ inline TFastPointerSet<FPackageData*>& FPackageDatas::GetAssignedToWorkerSet()
 	return AssignedToWorkerSet;
 }
 
-inline FLoadPrepareQueue& FPackageDatas::GetLoadPrepareQueue()
-{
-	return LoadPrepareQueue;
-}
-
-inline FPackageDataQueue& FPackageDatas::GetLoadReadyQueue()
+inline FLoadQueue& FPackageDatas::GetLoadQueue()
 {
-	return LoadReadyQueue;
+	return LoadQueue;
 }
 
 inline TFastPointerSet<FPackageData*>& FPackageDatas::GetSaveStalledSet()
@@ -1453,6 +1492,32 @@ inline FPackageDataQueue& FPackageDatas::GetSaveQueue()
 	return SaveQueue;
 }
 
+
+inline const FName& FPackageData::GetPackageName() const
+{
+	return PackageName;
+}
+
+inline const FName& FPackageData::GetFileName() const
+{
+	return FileName;
+}
+
+inline void FPackageData::SetFileName(const FName& InFileName)
+{
+	FileName = InFileName;
+}
+
+inline uint32 FPackageData::GetLeafToRootRank() const
+{
+	return LeafToRootRank;
+}
+
+inline void FPackageData::SetLeafToRootRank(uint32 Value)
+{
+	LeafToRootRank = Value;
+}
+
 template <typename ArrayType>
 inline void FPackageData::GetPlatformsNeedingCooking(ArrayType& OutPlatforms) const
 {
diff --git a/Engine/Source/Editor/UnrealEd/Private/Cooker/CookPackagePreloader.cpp b/Engine/Source/Editor/UnrealEd/Private/Cooker/CookPackagePreloader.cpp
--- a/Engine/Source/Editor/UnrealEd/Private/Cooker/CookPackagePreloader.cpp
+++ b/Engine/Source/Editor/UnrealEd/Private/Cooker/CookPackagePreloader.cpp
@@ -2,13 +2,17 @@
 
 #include "Cooker/CookPackagePreloader.h"
 
+#include "AssetRegistry/IAssetRegistry.h"
 #include "CookOnTheSide/CookLog.h"
+#include "CookOnTheSide/CookOnTheFlyServer.h"
 #include "CoreGlobals.h"
 #include "EditorDomain/EditorDomain.h"
+#include "HAL/PlatformMath.h"
 #include "Logging/LogMacros.h"
 #include "Misc/AssertionMacros.h"
 #include "Misc/PackagePath.h"
 #include "Misc/PreloadableFile.h"
+#include "Misc/ScopeExit.h"
 #include "Misc/StringBuilder.h"
 #include "UObject/NameTypes.h"
 #include "UObject/UObjectGlobals.h"
@@ -16,42 +20,253 @@
 namespace UE::Cook
 {
 
+bool FPackagePreloader::bConfigInitialized = false;
+bool FPackagePreloader::bAllowPreloadImports = true;
+bool FPackagePreloader::bDebugPreloadImports = false;
+
+void FPackagePreloader::InitializeConfig()
+{
+	bConfigInitialized = true;
+
+	// bAllowPreloadImports is off by default, because it is currently loading unneeded imports for some packages
+	// and causing poor cook performance
+	bAllowPreloadImports = false;
+	bDebugPreloadImports = false;
+	FParse::Bool(FCommandLine::Get(), TEXT("-CookPreloadImports="), bAllowPreloadImports);
+	if (FParse::Param(FCommandLine::Get(), TEXT("CookDebugPreloadImports")))
+	{
+		bDebugPreloadImports = true;
+		// bAllowPreloadImports needs to be off for us to find which imports are not loaded by the load of the referencer.
+		bAllowPreloadImports = false;
+	}
+}
+
 FPackagePreloader::~FPackagePreloader()
 {
+	// To avoid incorrect behavior, we have to set ESendFlags::QueueNone, otherwise a TRefCountPtr will be created
+	// and invalidly call ~FPackagePreloader when the TRefCountPtr goes out of scope.
+	SendToState(EPreloaderState::Inactive, ESendFlags::QueueNone);
 	PackageData.OnPackagePreloaderDestroyed(*this);
 }
 
 void FPackagePreloader::Shutdown()
 {
 	TRefCountPtr<FPackagePreloader> LocalRef(this);
-	ClearPreload();
+	SendToState(EPreloaderState::Inactive, ESendFlags::QueueAddAndRemove);
 }
 
-bool FPackagePreloader::TryPreload()
+template <typename ShouldKeepFunc, typename ReportAndIsContinueFunc>
+void FPackagePreloader::TraverseImportGraph(ShouldKeepFunc&& ShouldKeep, ReportAndIsContinueFunc&& ReportAndIsContinue,
+	bool bAllowGather)
 {
-	if (bIsPreloadAttempted)
+	// VisitState should only ever be changed from Unvisited during the execution below, and we change it back
+	// before exiting.
+	check(VisitState == EGraphVisitState::Unvisited);
+
+	struct FStackData
 	{
-		return true;
+		FPackagePreloader* Preloader = nullptr;
+		int32 NextImport = 0;
+
+		FStackData(FPackagePreloader* InPreloader)
+			: Preloader(InPreloader)
+		{
+		}
+	};
+	TArray<FStackData, TInlineAllocator<16>> Stack;
+	TArray<FPackagePreloader*, TInlineAllocator<128>> VisitedList;
+
+	// Depth first search over the import graph
+	Stack.Emplace(this);
+	while (!Stack.IsEmpty())
+	{
+		FStackData& Top = Stack.Last();
+
+		// When the stackdata's NextImport is 0, that means we just pushed it onto the stack,
+		// and we need to execute the initial setup
+		if (Top.NextImport == 0)
+		{
+			if (Top.Preloader->VisitState != EGraphVisitState::Unvisited)
+			{
+				// Inprogress or already visited; ignore this link
+				Stack.Pop(EAllowShrinking::No);
+				continue;
+			}
+
+			if (!ShouldKeep(*Top.Preloader))
+			{
+				// Caller does not want us to report or explore this one
+				Top.Preloader->VisitState = EGraphVisitState::Visited;
+				VisitedList.Add(Top.Preloader);
+				Stack.Pop(EAllowShrinking::No);
+				continue;
+			}
+
+			if (!ReportAndIsContinue(*Top.Preloader))
+			{
+				// Caller requested we stop searching; break out of the graph search
+				break;
+			}
+
+			// Gather the imports from the AssetRegistry if not already gathered, keeping only the unloaded ones.
+			if (bAllowGather)
+			{
+				Top.Preloader->GatherUnloadedImports();
+			}
+
+			// Mark that we are on the stack and are traversing imports
+			Top.Preloader->VisitState = EGraphVisitState::InProgress;
+			VisitedList.Add(Top.Preloader);
+
+			// Fallthrough to after the if/else and examine the first import. 
+		}
+
+		// Examine the next import, if we have not yet reached the end
+		if (Top.NextImport < Top.Preloader->UnloadedImports.Num())
+		{
+			FPackagePreloader* Import = Top.Preloader->UnloadedImports[Top.NextImport++];
+			Stack.Emplace(Import);
+			// Top is now possibly a dangling pointer if the stack reallocated; do not access it
+			continue;
+		}
+
+		// Finish the visit of Top
+		Top.Preloader->VisitState = EGraphVisitState::Visited;
+		Stack.Pop(EAllowShrinking::No);
 	}
-	if (FindObjectFast<UPackage>(nullptr, PackageData.GetPackageName()))
+
+	// Clear all the VisitState variables we modified
+	for (FPackagePreloader* PackagePreloader : VisitedList)
 	{
-		if (AsyncRequest && !AsyncRequest->bHasFinished)
+		PackagePreloader->VisitState = EGraphVisitState::Unvisited;
+	}
+}
+
+void FPackagePreloader::GetNeedsLoadPreloadersInImportTree(TArray<TRefCountPtr<FPackagePreloader>>& OutPreloaders)
+{
+	TraverseImportGraph(
+		[](FPackagePreloader& Preloader) -> bool // ShouldKeep
 		{
-			// In case of async loading, the object can be found while still being asynchronously serialized, we need
-			// to wait until the callback is called and the async request is completely done.
-			return false;
+			// Traverse every import; TraverseImportGraph already filters the imports by NeedsLoad 
+			return true;
+		},
+		[&OutPreloaders](FPackagePreloader& Preloader) -> bool // ReportAndIsContinue
+		{
+			OutPreloaders.Add(&Preloader);
+			return true; // Continue iterating through the entire unloaded import graph
+		},
+		true /* bAllowGather */
+	);
+}
+
+void FPackagePreloader::GatherUnloadedImports()
+{
+	if (IsImportsGathered())
+	{
+		return;
+	}
+	SetIsImportsGathered(true);
+
+	IAssetRegistry& AssetRegistry = IAssetRegistry::GetChecked();
+	TArray<FName> ImportNames;
+	FPackageDatas& PackageDatas = PackageData.GetPackageDatas();
+
+	AssetRegistry.GetDependencies(PackageData.GetPackageName(), ImportNames,
+		UE::AssetRegistry::EDependencyCategory::Package, UE::AssetRegistry::EDependencyQuery::Hard);
+	UnloadedImports.Reset(ImportNames.Num());
+	for (FName ImportName : ImportNames)
+	{
+		FPackageData* ImportData = PackageDatas.TryAddPackageDataByPackageName(ImportName);
+		if (!ImportData)
+		{
+			continue;
 		}
 
-		// If the package has already loaded, then there is no point in further preloading
-		ClearPreload();
-		bIsPreloadAttempted = true;
+		TRefCountPtr<FPackagePreloader> ImportPreloader = ImportData->GetPackagePreloader();
+		if (!ImportPreloader)
+		{
+			// Optimization: If the ImportPackage does not have a Packagedata, check whether it is already loaded
+			// before paying the expense of creating a PackagePreloader which we would then immediately delete.
+			if (FPackagePreloader::IsPackageLoaded(*ImportData))
+			{
+				continue;
+			}
+			ImportPreloader = ImportData->CreatePackagePreloader();
+		}
+		else
+		{
+			if (!ImportPreloader->NeedsLoad())
+			{
+				continue;
+			}
+		}
+		UnloadedImports.Add(MoveTemp(ImportPreloader));
+	}
+}
+
+void FPackagePreloader::LogUnusedImports()
+{
+	if (!bDebugPreloadImports || !IsPackageLoaded())
+	{
+		return;
+	}
+
+	TStringBuilder<256> UnusedImportsStr;
+	for (TRefCountPtr<FPackagePreloader>& Import : UnloadedImports)
+	{
+		if (!Import->IsPackageLoaded())
+		{
+			UnusedImportsStr << Import->GetPackageData().GetPackageName() << TEXT(", ");
+		}
+	}
+	if (UnusedImportsStr.Len() != 0)
+	{
+		UnusedImportsStr.RemoveSuffix(2);
+		UE_LOG(LogCook, Warning,
+			TEXT("DebugPreloadImports: the load of %s did not end up loading these expected imports:\n\t%s"),
+			*WriteToString<256>(GetPackageData().GetPackageName()), *UnusedImportsStr);
+	}
+}
+
+void FPackagePreloader::EmptyImports()
+{
+	UnloadedImports.Empty();
+	SetIsImportsGathered(false);
+}
+
+bool FPackagePreloader::TryPreload()
+{
+	bool bTreatPackageAsLoaded = WasLoadAttempted() || IsPackageLoaded();
+	if (bTreatPackageAsLoaded)
+	{
+		if (AsyncRequest || IsPreloaded())
+		{
+			if (AsyncRequest && !AsyncRequest->bHasFinished)
+			{
+				// In case of async loading, the object can be found while still being asynchronously serialized, we need
+				// to wait until the callback is called and the async request is completely done.
+				return false;
+			}
+
+			// If the package has already loaded, then we no longer need the preloaded data
+			ClearPreload();
+		}
+		SetIsPreloadAttempted(true);
 		return true;
 	}
+	else
+	{
+		if (IsPreloadAttempted())
+		{
+			return true;
+		}
+	}
+
 	if (PackageData.IsGenerated())
 	{
 		// Deferred populate generated packages are loaded from their generator, not from disk
 		ClearPreload();
-		bIsPreloadAttempted = true;
+		SetIsPreloadAttempted(true);
 		return true;
 	}
 	if (IsAsyncLoadingMultithreaded())
@@ -72,8 +287,7 @@ bool FPackagePreloader::TryPreload()
 			);
 		}
 
-		// always return false so we continue to check the status of the load until FindObjectFast above finds the
-		// loaded object
+		// always return false so we continue to check the status of the load until IsPackageLoaded().
 		return false;
 	}
 	if (!PreloadableFile.Get())
@@ -94,8 +308,8 @@ bool FPackagePreloader::TryPreload()
 				// not change until the PackageData is destructed, and the destructor does not run and other threads do not read
 				// or write PreloadableFileOpenResult until after PreloadableFile.Get() has finished initialization
 				// and this callback is therefore complete.
-				// The code that accomplishes that waiting is in TryPreload (IsInitialized) and ClearPreload
-				// (ReleaseCache)
+				// The code that accomplishes that waiting is in TryPreload (via IsInitialized) and ClearPreload
+				// (via ReleaseCache)
 				PackageData.GetFileName().ToString(FileNameString);
 				FPackagePath PackagePath = FPackagePath::FromLocalPath(FileNameString);
 				FOpenPackageResult Result = IPackageResourceManager::Get().OpenReadPackage(PackagePath);
@@ -126,7 +340,7 @@ bool FPackagePreloader::TryPreload()
 	{
 		UE_LOG(LogCook, Warning, TEXT("Failed to find file when preloading %s."),
 			*PackageData.GetFileName().ToString());
-		bIsPreloadAttempted = true;
+		SetIsPreloadAttempted(true);
 		PreloadableFile.Reset(PackageData);
 		PreloadableFileOpenResult = FOpenPackageResult();
 		return true;
@@ -139,14 +353,14 @@ bool FPackagePreloader::TryPreload()
 	{
 		UE_LOG(LogCook, Warning, TEXT("Failed to register %s for preload."),
 			*PackageData.GetFileName().ToString());
-		bIsPreloadAttempted = true;
+		SetIsPreloadAttempted(true);
 		PreloadableFile.Reset(PackageData);
 		PreloadableFileOpenResult = FOpenPackageResult();
 		return true;
 	}
 
-	bIsPreloaded = true;
-	bIsPreloadAttempted = true;
+	SetIsPreloaded(true);
+	SetIsPreloadAttempted(true);
 	return true;
 }
 
@@ -164,7 +378,7 @@ void FPackagePreloader::ClearPreload()
 	}
 
 	const TSharedPtr<FPreloadableArchive>& FilePtr = PreloadableFile.Get();
-	if (bIsPreloaded)
+	if (IsPreloaded())
 	{
 		check(FilePtr);
 		TStringBuilder<NAME_SIZE> FileNameString;
@@ -184,16 +398,344 @@ void FPackagePreloader::ClearPreload()
 
 	PreloadableFile.Reset(PackageData);
 	PreloadableFileOpenResult = FOpenPackageResult();
-	bIsPreloaded = false;
-	bIsPreloadAttempted = false;
+	SetIsPreloaded(false);
+	SetIsPreloadAttempted(false);
+}
+
+void FPackagePreloader::PostGarbageCollect()
+{
+	// Reevaluate imports.
+	SetIsImportsGathered(false);
+
+	// Reevaluate variables that depend on whether our package is loaded.
+	SetLoadAttempted(false);
+
+	if (!AsyncRequest && !PreloadableFile.Get())
+	{
+		// If we have no preload data, we might have marked that we are done preloading becuase the package already
+		// exists. Call ClearPreload so we reevaluate whether the package exists
+		ClearPreload();
+	}
+	else
+	{
+		if (AsyncRequest)
+		{
+			// The AsyncRequest should have been flushed (and then either kept in memory or garbage collected),
+			// so clear the preload data
+			ClearPreload();
+		}
+
+		if (PreloadableFile.Get() && IsPreloaded() && !PreloadableFile.Get()->HasValidData())
+		{
+			// If we finished preloading the file, then we registered it, and it might have been consumed by the loaded
+			// package, but then the loaded package GC'd. In that case we need to clear the PreloadableFile data so we
+			// can restart it when necessary the next load of the package. And if the package already exists in memory then
+			// we don't need the preloaded data, so it's okay to free it.
+			// If we didn't register it, or we registered it but it has not yet been consumed, then we don't need to free it.
+			ClearPreload();
+		}
+	}
+
+	// If state is past ActivePreload, move back to ActivePreload to reevaluate whether we're ready.
+	if (State > EPreloaderState::ActivePreload)
+	{
+		SendToState(EPreloaderState::ActivePreload, ESendFlags::QueueAddAndRemove);
+	}
 }
 
-void FPackagePreloader::CheckPreloadEmpty()
+void FPackagePreloader::OnPackageLeaveLoadState()
 {
-	check(!AsyncRequest);
-	check(!bIsPreloadAttempted);
-	check(!PreloadableFile.Get());
-	check(!bIsPreloaded);
+	// Caller guarantees that a ref count is held during this function, so no need for a local refcount.
+
+	if (HasInitializedRequestedLoads())
+	{
+		SetHasInitializedRequestedLoads(false);
+
+		// There should be at least one CountFromRequestedLoads due to a request from *this.
+		check(CountFromRequestedLoads > 0);
+		// Don't allow the triggering of a state transition on this during the for loop.
+		IncrementCountFromRequestedLoads();
+
+		for (TRefCountPtr<FPackagePreloader>& Other : RequestedLoads)
+		{
+			Other->DecrementCountFromRequestedLoads();
+		}
+		RequestedLoads.Empty();
+
+		DecrementCountFromRequestedLoads(); // This decrement might kick *this back to Inactive.
+	}
+}
+
+bool FPackagePreloader::IsPackageLoaded() const
+{
+	return IsPackageLoaded(PackageData);
+}
+
+bool FPackagePreloader::IsPackageLoaded(const FPackageData& InPackageData)
+{
+	return FindObjectFast<UPackage>(nullptr, InPackageData.GetPackageName()) != nullptr;
+}
+
+void FPackagePreloader::IncrementCountFromRequestedLoads()
+{
+	++CountFromRequestedLoads;
+}
+
+void FPackagePreloader::DecrementCountFromRequestedLoads()
+{
+	check(CountFromRequestedLoads > 0); // Assert we do not have an unbalanced decrement.
+	--CountFromRequestedLoads;
+
+	if (CountFromRequestedLoads == 0 && GetState() != EPreloaderState::Inactive)
+	{
+		SendToState(EPreloaderState::Inactive, ESendFlags::QueueAddAndRemove);
+	}
+}
+
+void FPackagePreloader::SetRequestedLoads(TArray<TRefCountPtr<FPackagePreloader>>&& InRequestedLoads)
+{
+	FPackageDatas& PackageDatas = PackageData.GetPackageDatas();
+	check(RequestedLoads.IsEmpty()); // This function is only for setting from empty.
+
+	// Our contract specifies that InRequestedLoads is in RootToLeaf order, so traverse it
+	// backwards to set the LeafToRoot rank
+	for (TRefCountPtr<FPackagePreloader>& NeedsLoadPreloader : ReverseIterate(InRequestedLoads))
+	{
+		FPackageData& NeedsLoadData(NeedsLoadPreloader->GetPackageData());
+		if (NeedsLoadData.GetLeafToRootRank() == MAX_uint32)
+		{
+			NeedsLoadData.SetLeafToRootRank(PackageDatas.GetNextLeafToRootRank());
+		}
+
+		if (NeedsLoadPreloader->GetState() == EPreloaderState::Inactive)
+		{
+			NeedsLoadPreloader->SendToState(EPreloaderState::PendingKick, ESendFlags::QueueAddAndRemove);
+		}
+
+		NeedsLoadPreloader->IncrementCountFromRequestedLoads();
+	}
+	RequestedLoads = MoveTemp(InRequestedLoads);
+}
+
+void FPackagePreloader::OnExitActive()
+{
+	if (bDebugPreloadImports)
+	{
+		LogUnusedImports();
+	}
+
+	ClearPreload();
+	EmptyImports();
+}
+
+bool FPackagePreloader::NeedsLoad()
+{
+	if (WasLoadAttempted())
+	{
+		return false;
+	}
+	if (!IsPackageLoaded(PackageData))
+	{
+		return true;
+	}
+	// We might not be done preloading even if the package exists. Calling TryPreload(false) will ClearPreload
+	// and return true unless we're still waiting on asynchronous post-loads to complete.
+	return !TryPreload();
+}
+
+bool FPackagePreloader::IsHigherPriorityThan(const FPackagePreloader& Other) const
+{
+	if (PackageData.GetIsUrgent() != Other.PackageData.GetIsUrgent())
+	{
+		return PackageData.GetIsUrgent();
+	}
+
+	// Leaves are higher priority because we want them to be already preloaded (or even
+	// better, completely loaded) when we load the package that imports them (and is therefore
+	// more rootwards).
+	return PackageData.GetLeafToRootRank() < Other.PackageData.GetLeafToRootRank();
+}
+
+void FPackagePreloader::SendToState(EPreloaderState NewState, ESendFlags SendFlags)
+{
+	TRefCountPtr<FPackagePreloader> KeepRemovalResident;
+	FLoadQueue& LoadQueue = PackageData.GetPackageDatas().GetLoadQueue();
+	if (EnumHasAnyFlags(SendFlags, ESendFlags::QueueRemove))
+	{
+		KeepRemovalResident = this;
+		switch (State)
+		{
+		case EPreloaderState::Inactive:
+			break;
+		case EPreloaderState::PendingKick:
+			LoadQueue.PendingKicks.Remove(this);
+			break;
+		case EPreloaderState::ActivePreload:
+			LoadQueue.ActivePreloads.Remove(this);
+			break;
+		case EPreloaderState::ReadyForLoad:
+			LoadQueue.ReadyForLoads.Remove(this);
+			break;
+		default:
+			checkNoEntry();
+			break;
+		}
+	}
+	bool bActive = State != EPreloaderState::Inactive;
+	bool bNewActive = NewState != EPreloaderState::Inactive;
+	if (bActive != bNewActive)
+	{
+		if (!bNewActive)
+		{
+			OnExitActive();
+		}
+	}
+
+	State = NewState;
+
+	if (EnumHasAnyFlags(SendFlags, ESendFlags::QueueAdd))
+	{
+		switch (State)
+		{
+		case EPreloaderState::Inactive:
+			break;
+		case EPreloaderState::PendingKick:
+			LoadQueue.PendingKicks.Add(this);
+			break;
+		case EPreloaderState::ActivePreload:
+			LoadQueue.ActivePreloads.Add(this);
+			break;
+		case EPreloaderState::ReadyForLoad:
+			LoadQueue.ReadyForLoads.Add(this);
+			break;
+		default:
+			checkNoEntry();
+			break;
+		}
+	}
+}
+
+bool FPackagePreloader::PumpLoadsTryStartInboxPackage(UCookOnTheFlyServer& COTFS)
+{
+	using namespace UE::Cook;
+
+	TRingBuffer<FPackageData*>& Inbox = COTFS.PackageDatas->GetLoadQueue().Inbox;
+	if (Inbox.IsEmpty())
+	{
+		return false;
+	}
+
+	FPackageData* PackageData = Inbox.PopFrontValue();
+	check(PackageData->GetState() == EPackageState::Load);
+	TRefCountPtr<FPackagePreloader> Preloader = PackageData->GetPackagePreloader();
+	check(Preloader);
+	Preloader->SetIsInInbox(false);
+
+	if (COTFS.TryCreateRequestCluster(*PackageData))
+	{
+		return true;
+	}
+
+	// If the package is already ready for loading, or we otherwise want to skip preloading for it,
+	// skip the preload containers and put in the ReadyLoads container
+	if (!COTFS.bPreloadingEnabled || Preloader->IsPackageLoaded() || PackageData->GetIsUrgent())
+	{
+		if (Preloader->GetState() != EPreloaderState::ReadyForLoad)
+		{
+			Preloader->SendToState(EPreloaderState::ReadyForLoad, ESendFlags::QueueAddAndRemove);
+		}
+		return true;
+	}
+
+	if (!Preloader->HasInitializedRequestedLoads())
+	{
+		Preloader->SetHasInitializedRequestedLoads(true);
+
+		TArray<TRefCountPtr<FPackagePreloader>> NeedsLoadPreloaders;
+		if (bAllowPreloadImports)
+		{
+			Preloader->GetNeedsLoadPreloadersInImportTree(NeedsLoadPreloaders);
+		}
+		else
+		{
+			NeedsLoadPreloaders.Add(Preloader);
+			if (bDebugPreloadImports)
+			{
+				// Gather the imports for *this, but do not activate them. We will review them after load.
+				Preloader->GatherUnloadedImports();
+			}
+		}
+		check(!NeedsLoadPreloaders.IsEmpty() && NeedsLoadPreloaders[0].GetReference() == Preloader);
+		Preloader->SetRequestedLoads(MoveTemp(NeedsLoadPreloaders));
+		// Should have been set to active by SetRequestedLoads
+		check(Preloader->GetState() != EPreloaderState::Inactive);
+	}
+	else if (Preloader->GetState() == EPreloaderState::Inactive)
+	{
+		// Edgecase: we've already initialized loads, but the preloader is inactive and not loaded somehow. Put it
+		// directly into ReadyForLoad since its not clear that it needs preloading.
+		Preloader->SendToState(EPreloaderState::ReadyForLoad, ESendFlags::QueueAddAndRemove);
+	}
+
+	return true;
+}
+
+bool FPackagePreloader::PumpLoadsTryKickPreload(UCookOnTheFlyServer& COTFS)
+{
+	using namespace UE::Cook;
+
+	FPackagePreloaderPriorityQueue& PendingKicks = COTFS.PackageDatas->GetLoadQueue().PendingKicks;
+	if (PendingKicks.IsEmpty())
+	{
+		return false;
+	}
+	FPackageDataMonitor& Monitor = COTFS.PackageDatas->GetMonitor();
+	if (Monitor.GetNumPreloadAllocated() >= static_cast<int32>(COTFS.MaxPreloadAllocated))
+	{
+		return false;
+	}
+
+	TRefCountPtr<FPackagePreloader> Preloader = PendingKicks.PopFront();
+	Preloader->TryPreload();
+	Preloader->SendToState(EPreloaderState::ActivePreload, ESendFlags::QueueAdd);
+
+	return true;
+}
+
+bool FPackagePreloader::PumpLoadsIsReadyToLeavePreload()
+{
+	// Once we are added to ActivePreloads, we stop caring about the preload status of all of our imports
+	// that are not actively preloading.
+	// The imports that we depend on should have been added to ActivePreloads before us, and we either check
+	// them because they are still active, or we don't need to check them because they already finished preloading
+	// and moved past the ActivePreloads state.
+	// In the case of a cycle, or if one of our imports was demoted out of ActivePreloads somehow, we
+	// need to proceed despite the package possibly not being preloaded, so we don't get stuck in a cycle that
+	// we can't preload all elements of.
+
+	bool bAllActivePreloadsComplete = true;
+	TraverseImportGraph(
+		[](FPackagePreloader& Preloader) -> bool // ShouldKeep
+		{
+			// Only look at the imports that are in ActivePreload, per the comment above. Most notably, this includes
+			// ourself, at the root of the import tree.
+			return Preloader.GetState() == EPreloaderState::ActivePreload;
+		},
+		[&bAllActivePreloadsComplete](FPackagePreloader& Preloader) -> bool // ReportAndIsContinue
+		{
+			bAllActivePreloadsComplete = bAllActivePreloadsComplete && Preloader.TryPreload();
+			return bAllActivePreloadsComplete; // Once we find we will return false, stop searching
+		},
+		false /* bAllowGather */
+	);
+	return bAllActivePreloadsComplete;
+}
+
+void FPackagePreloader::PumpLoadsMarkLoadAttemptComplete()
+{
+	SetLoadAttempted(true);
+	// Caller is responsible for having removed *this from the ReadyForLoad container.
+	SendToState(EPreloaderState::Inactive, ESendFlags::QueueNone);
 }
 
 void FPackagePreloader::FTrackedPreloadableFilePtr::Set(TSharedPtr<FPreloadableArchive>&& InPtr, FPackageData& Owner)
@@ -215,4 +757,102 @@ void FPackagePreloader::FTrackedPreloadableFilePtr::Reset(FPackageData& Owner)
 	}
 }
 
+bool FPackagePreloaderPriorityWrapper::operator<(const FPackagePreloaderPriorityWrapper& Other) const
+{
+	check(Payload.IsValid() && Other.Payload.IsValid());
+	// Higher Priority -> comes earlier in the queue -> has lower index -> is less than
+	return Payload->IsHigherPriorityThan(*Other.Payload);
+}
+
+bool FPackagePreloaderPriorityQueue::IsEmpty() const
+{
+	return Heap.IsEmpty();
+}
+
+void FPackagePreloaderPriorityQueue::Add(TRefCountPtr<FPackagePreloader> Preloader)
+{
+	Heap.HeapPush(FPackagePreloaderPriorityWrapper{ MoveTemp(Preloader) });
+}
+
+void FPackagePreloaderPriorityQueue::Remove(const TRefCountPtr<FPackagePreloader>& Preloader)
+{
+	int32 Index = Heap.IndexOfByPredicate([&Preloader](const FPackagePreloaderPriorityWrapper& Element)
+		{
+			return Element.Payload == Preloader;
+		});
+	if (Index != INDEX_NONE)
+	{
+		Heap.HeapRemoveAt(Index, EAllowShrinking::No);
+	}
+}
+
+TRefCountPtr<FPackagePreloader> FPackagePreloaderPriorityQueue::PopFront()
+{
+	FPackagePreloaderPriorityWrapper Wrapper;
+	Heap.HeapPop(Wrapper, EAllowShrinking::No);
+	return MoveTemp(Wrapper.Payload);
+}
+
+bool FLoadQueue::IsEmpty()
+{
+	return InProgress.IsEmpty();
+}
+
+int32 FLoadQueue::Num() const
+{
+	return InProgress.Num();
+}
+
+void FLoadQueue::Add(FPackageData* PackageData)
+{
+	bool bAlreadyInSet;
+	// The Package must be in the LoadState to be added to the container, and OnEnterLoading
+	// guarantees a refcount exists on the Preloader. We rely on this, because we need to store
+	// IsInInbox on the Preloader.
+	TRefCountPtr<FPackagePreloader> Preloader = PackageData->GetPackagePreloader();
+	check(Preloader);
+
+	InProgress.Add(PackageData, &bAlreadyInSet);
+	if (!bAlreadyInSet)
+	{
+		Inbox.Add(PackageData);
+		Preloader->SetIsInInbox(true);
+	}
+}
+
+bool FLoadQueue::Contains(const FPackageData* PackageData) const
+{
+	return InProgress.Contains(PackageData);
+}
+
+uint32 FLoadQueue::Remove(FPackageData* PackageData)
+{
+	uint32 Result = static_cast<uint32>(InProgress.Remove(PackageData));
+	if (Result == 0)
+	{
+		return 0;
+	}
+
+	TRefCountPtr<FPackagePreloader> Preloader = PackageData->GetPackagePreloader();
+	if (Preloader && Preloader->IsInInbox())
+	{
+		Inbox.Remove(PackageData);
+		Preloader->SetIsInInbox(false);
+	}
+	// This Remove function is not responsible for removing the PackageData's Preloader from the sub containers for
+	// preloaders. That responsibility is complicated and the work that needs to be done for it upon leaving the load
+	// state is done by FPackageData::OnExitLoad.
+	return Result;
+}
+
+TSet<FPackageData*>::TRangedForIterator FLoadQueue::begin()
+{
+	return InProgress.begin();
+}
+
+TSet<FPackageData*>::TRangedForIterator FLoadQueue::end()
+{
+	return InProgress.end();
+}
+
 }
\ No newline at end of file
diff --git a/Engine/Source/Editor/UnrealEd/Private/Cooker/CookPackagePreloader.h b/Engine/Source/Editor/UnrealEd/Private/Cooker/CookPackagePreloader.h
--- a/Engine/Source/Editor/UnrealEd/Private/Cooker/CookPackagePreloader.h
+++ b/Engine/Source/Editor/UnrealEd/Private/Cooker/CookPackagePreloader.h
@@ -3,19 +3,21 @@
 #pragma once
 
 #include "Cooker/CookPackageData.h"
+#include "Cooker/CookTypes.h"
 #include "Templates/RefCounting.h"
 #include "Templates/SharedPointer.h"
 #include "UObject/PackageResourceManager.h"
 
 #include <atomic>
 
+class UCookOnTheFlyServer;
+
 namespace UE::Cook
 {
 
 /**
  * Helper class for the Load state on FPackageData. Also supports loads of other packages that depend on the owner
- * package. Handles preloading or asyncloading the package before we put the package into LoadReadyQueue for
- * a call to LoadPackage.
+ * package. Handles preloading or asyncloading the package before we call LoadPackage on it.
  */
 class FPackagePreloader : public FThreadSafeRefCountedObject
 {
@@ -23,6 +25,10 @@ public:
 	FPackagePreloader(FPackageData& InPackageData);
 	~FPackagePreloader();
 
+	FPackageData& GetPackageData();
+
+	EPreloaderState GetState() const;
+
 	/** Free any IO requests or buffers, and remove all references to any other FPackagePreloaders. */
 	void Shutdown();
 
@@ -31,12 +37,27 @@ public:
 	/** Clear the SelfReference; the owner package is done using this object. */
 	void ClearSelfReference();
 
-	/** Try to preload the file. Return true if preloading is complete (succeeded or failed or was skipped). */
-	bool TryPreload();
-	/** Clear any allocated preload data. */
-	void ClearPreload();
-	/** Issue check statements confirming that no preload data is allocated or flags are set. */
-	void CheckPreloadEmpty();
+	/** Clear our completion records; we need to reevaluate them since packages have been removed from memory. */
+	void PostGarbageCollect();
+	/** Called when the owner package exits the loadstate; clears reference counts of requested imports. */
+	void OnPackageLeaveLoadState();
+	/** Return whether the Preloader's package is loaded. */
+	bool IsPackageLoaded() const;
+	/** Return whether the given PackageData's package is loaded. */
+	static bool IsPackageLoaded(const FPackageData& InPackageData);
+	int32 GetCountFromRequestedLoads(); 
+
+	/**
+	 * State transition function that calls OnEnter/OnExit functions and (optionally) removes and adds the
+	 * PackagePreloader to the container for its state.
+	 */
+	void SendToState(EPreloaderState NewState, ESendFlags SendFlags);
+
+	// UCookOnTheFlyServer::PumpLoads helper functions, see comments in PumpLoads.
+	static bool PumpLoadsTryStartInboxPackage(UCookOnTheFlyServer& COTFS);
+	static bool PumpLoadsTryKickPreload(UCookOnTheFlyServer& COTFS);
+	bool PumpLoadsIsReadyToLeavePreload();
+	void PumpLoadsMarkLoadAttemptComplete();
 
 private:
 	/**
@@ -62,29 +83,153 @@ private:
 		std::atomic<bool> bHasFinished{ false };
 	};
 
+	/** Used during DFS of the import graph. */
+	enum class EGraphVisitState : uint8
+	{
+		Unvisited,
+		InProgress,
+		Visited,
+	};
+
+private:
+	bool HasInitializedRequestedLoads() const;
+	void SetHasInitializedRequestedLoads(bool bValue);
+	bool IsInInbox() const;
+	void SetIsInInbox(bool bValue);
+	bool WasLoadAttempted() const;
+	void SetLoadAttempted(bool bValue);
+	bool IsPreloadAttempted() const;
+	void SetIsPreloadAttempted(bool bValue);
+	bool IsPreloaded() const;
+	void SetIsPreloaded(bool bValue);
+	bool IsImportsGathered() const;
+	void SetIsImportsGathered(bool bValue);
+
+	void SetState(EPreloaderState Value);
+	/** Clear any allocated preload data. */
+	void ClearPreload();
+	void EmptyImports();
+	/** Gather imports from the AssetRegistry, but skip Preloads that return true for NeedsLoad. */
+	void GatherUnloadedImports();
+	/** Called after loading when bDebugPreloadImports; report imports we thought we needed but did not load. */
+	void LogUnusedImports();
+
+	/**
+	 * Traverse the import tree of this Preloader, optionally gathering imports from the asset registry if they have not
+	 * already been gathered. For each package found, if loaded skip it, otherwise add it to the list of imports needed
+	 * by the referencer, and make its PackagePreloader active in the loadqueue.
+	 * 
+	 * Caller can specify skipping a given import and not reporting it or its transitive imports.
+	 * Caller provides a Report function that is called for each reported import.
+	 * Packages are reported in root-to-leaf order. If a cycle is present, order in the cycle is arbitrary.
+	 * 
+	 * bool ShouldKeep(FPackagePreloader&) -> should Report be called and the preloader be explored.
+	 * bool ReportAndIsContinue(FPackagePreloader&) -> should the search continue.
+	 * 
+	 * If bAllowGather is true, caller must activate all inactive reported PackagePreloaders, otherwise
+	 * they will hold imports in the inactive state which is not valid.
+	 */
+	template <typename ShouldKeepFunc, typename ReportAndIsContinueFunc>
+	void TraverseImportGraph(ShouldKeepFunc&& ShouldKeep, ReportAndIsContinueFunc&& ReportAndIsContinue,
+		bool bAllowGather);
+	/**
+	 * TraverseImportGraph and return all the discovered needs-loading imports. The output list is
+	 * in root to leaf order.
+	 */
+	void GetNeedsLoadPreloadersInImportTree(TArray<TRefCountPtr<FPackagePreloader>>& OutPreloaders);
+
+	/**
+	 * Try to preload just the file for this PackagePreloader's owner package, does not consider imports.
+	 * Return true if preloading is complete (succeeded or failed or was skipped).
+	 */
+	bool TryPreload();
+	
+	/** Increment the countfromrequestedloads. When it is decremented back to zero, the preload+load is cancelled. */
+	void IncrementCountFromRequestedLoads();
+	/** Decrement the countfromrequestedloads. When it is reaches zero, the preload+load is cancelled. */
+	void DecrementCountFromRequestedLoads();
+	/**
+	 * Set the list of packages that are in the import tree of this one and have not loaded yet.
+	 * Each one will have its CountFromRequestedLoads incremented to keep it in an active state until
+	 * done loading.
+	 * InRequestedLoads must be in Root to Leaf order; that order is used to define the priority of
+	 * packages in PendingKick queue (leafwards packages are kicked earlier).
+	 */
+	void SetRequestedLoads(TArray<TRefCountPtr<FPackagePreloader>>&& InRequestedLoads);
+	/**
+	 * Called when a PackagePreloader leaves the active state (regardless of what state its Package is in) to
+	 * clear the resources used during preloading.
+	 */
+	void OnExitActive();
+	/**
+	 * Report whether the Preloader's package still needs to be loaded.
+	 * Performance sideeffect: ClearPreload if it no longer needs to load.
+	 */
+	bool NeedsLoad();
+	/** Used to decide which PendingKick PackagePreloaders should next be preloaded. */
+	bool IsHigherPriorityThan(const FPackagePreloader& Other) const;
+
+	static void InitializeConfig();
+
 private:
 	FTrackedPreloadableFilePtr PreloadableFile;
 	FOpenPackageResult PreloadableFileOpenResult;
 	FPackageData& PackageData;
 	TRefCountPtr<FPackagePreloader> SelfReference;
 	TSharedPtr<FAsyncRequest> AsyncRequest;
+	TArray<TRefCountPtr<FPackagePreloader>> UnloadedImports;
+	TArray<TRefCountPtr<FPackagePreloader>> RequestedLoads;
+	int32 CountFromRequestedLoads = 0;
+	EPreloaderState State = EPreloaderState::Inactive;
+	EGraphVisitState VisitState = EGraphVisitState::Unvisited;
 	bool bIsPreloadAttempted = false;
 	bool bIsPreloaded = false;
+	bool bImportsGathered = false;
+	bool bIsInInbox = false;
+	bool bHasInitializedRequestedLoads = false;
+	bool bWasLoadAttempted = false;
+
+	static bool bConfigInitialized;
+	static bool bAllowPreloadImports;
+	static bool bDebugPreloadImports;
+
+	friend UE::Cook::FLoadQueue;
+	friend UE::Cook::FPackagePreloaderPriorityWrapper;
 };
 
 }
 
+
 ///////////////////////////////////////////////////////
 // Inline implementations
 ///////////////////////////////////////////////////////
 
+
 namespace UE::Cook
 {
 
 inline FPackagePreloader::FPackagePreloader(FPackageData& InPackageData)
 : PackageData(InPackageData)
 {
+	if (!bConfigInitialized)
+	{
+		InitializeConfig();
+	}
+}
+
+inline FPackageData& FPackagePreloader::GetPackageData()
+{
+	return PackageData;
+}
+
+inline EPreloaderState FPackagePreloader::GetState() const
+{
+	return State;
+}
 
+inline void FPackagePreloader::SetState(EPreloaderState Value)
+{
+	State = Value;
 }
 
 inline void FPackagePreloader::SetSelfReference()
@@ -97,4 +242,69 @@ inline void FPackagePreloader::ClearSelfReference()
 	SelfReference.SafeRelease();
 }
 
+inline int32 FPackagePreloader::GetCountFromRequestedLoads()
+{
+	return CountFromRequestedLoads;
+}
+
+inline bool FPackagePreloader::HasInitializedRequestedLoads() const
+{
+	return bHasInitializedRequestedLoads;
+}
+
+inline void FPackagePreloader::SetHasInitializedRequestedLoads(bool bValue)
+{
+	bHasInitializedRequestedLoads = bValue;
+}
+
+inline bool FPackagePreloader::IsInInbox() const
+{
+	return bIsInInbox;
+}
+
+inline void FPackagePreloader::SetIsInInbox(bool bValue)
+{
+	bIsInInbox = bValue;
+}
+
+inline bool FPackagePreloader::WasLoadAttempted() const
+{
+	return bWasLoadAttempted;
+}
+
+inline void FPackagePreloader::SetLoadAttempted(bool bValue)
+{
+	bWasLoadAttempted = bValue;
+}
+
+inline bool FPackagePreloader::IsPreloadAttempted() const
+{
+	return bIsPreloadAttempted;
+}
+
+inline void FPackagePreloader::SetIsPreloadAttempted(bool bValue)
+{
+	bIsPreloadAttempted = bValue;
+}
+
+inline bool FPackagePreloader::IsPreloaded() const
+{
+	return bIsPreloaded;
+}
+
+inline void FPackagePreloader::SetIsPreloaded(bool bValue)
+{
+	bIsPreloaded = bValue;
+}
+
+inline bool FPackagePreloader::IsImportsGathered() const
+{
+	return bImportsGathered;
+}
+
+inline void FPackagePreloader::SetIsImportsGathered(bool bValue)
+{
+	bImportsGathered = bValue;
+}
+
 }
\ No newline at end of file
diff --git a/Engine/Source/Editor/UnrealEd/Private/Cooker/CookTypes.cpp b/Engine/Source/Editor/UnrealEd/Private/Cooker/CookTypes.cpp
--- a/Engine/Source/Editor/UnrealEd/Private/Cooker/CookTypes.cpp
+++ b/Engine/Source/Editor/UnrealEd/Private/Cooker/CookTypes.cpp
@@ -90,14 +90,13 @@ const TCHAR* LexToString(ESuppressCookReason Reason)
 
 const TCHAR* LexToString(UE::Cook::EPackageState Reason)
 {
-	static_assert(static_cast<int>(EPackageState::Count) == 8);
+	static_assert(static_cast<int>(EPackageState::Count) == 7);
 	switch (Reason)
 	{
 	case EPackageState::Idle: return TEXT("Idle");
 	case EPackageState::Request: return TEXT("Request");
 	case EPackageState::AssignedToWorker: return TEXT("AssignedToWorker");
-	case EPackageState::LoadPrepare: return TEXT("LoadPrepare");
-	case EPackageState::LoadReady: return TEXT("LoadReady");
+	case EPackageState::Load: return TEXT("Load");
 	case EPackageState::SaveActive: return TEXT("SaveActive");
 	case EPackageState::SaveStalledRetracted: return TEXT("SaveStalledRetracted");
 	case EPackageState::SaveStalledAssignedToWorker: return TEXT("SaveStalledAssignedToWorker");
@@ -105,6 +104,19 @@ const TCHAR* LexToString(UE::Cook::EPackageState Reason)
 	}
 }
 
+const TCHAR* LexToString(UE::Cook::EPreloaderState State)
+{
+	static_assert(static_cast<int>(EPreloaderState::Count) == 4);
+	switch (State)
+	{
+	case EPreloaderState::Inactive: return TEXT("Inactive");
+	case EPreloaderState::PendingKick: return TEXT("PendingKick");
+	case EPreloaderState::ActivePreload: return TEXT("ActivePreload");
+	case EPreloaderState::ReadyForLoad: return TEXT("ReadyForLoad");
+	default: return TEXT("Invalid");
+	}
+}
+
 const TCHAR* LexToString(UE::Cook::ESaveSubState State)
 {
 	static_assert(static_cast<int>(ESaveSubState::Count) == 17);
diff --git a/Engine/Source/Editor/UnrealEd/Private/Cooker/CookTypes.h b/Engine/Source/Editor/UnrealEd/Private/Cooker/CookTypes.h
--- a/Engine/Source/Editor/UnrealEd/Private/Cooker/CookTypes.h
+++ b/Engine/Source/Editor/UnrealEd/Private/Cooker/CookTypes.h
@@ -195,13 +195,8 @@ namespace UE::Cook
 		 * had any operations performed on it locally.
 		 */
 		AssignedToWorker,
-		/** The Package is in the LoadPrepareQueue. Preloading is in progress. */
-		LoadPrepare,
-		/**
-		 * The package is in the LoadReadyQueue. Preloading is complete and it will be loaded when its turn
-		 * comes up.
-		 */
-		LoadReady,
+		/** The Package is in the LoadQueue, in one of multiple substates that handle loading and preloading. */
+		Load,
 		/** The Package is in the SaveQueue; it has been fully loaded and some target data may have been calculated. */
 		SaveActive,
 		/**
@@ -247,6 +242,21 @@ namespace UE::Cook
 	};
 	ENUM_CLASS_FLAGS(EPackageStateProperty);
 
+	/**
+	 * A substate of EPackageState::Load; it describes the state of the PackagePreloader in PumpLoads.
+	 * This state is on the PackagePreloader and not the PackageData, and might be active even while the package
+	 * is not in the load state.
+	 */
+	enum class EPreloaderState : uint8
+	{
+		Inactive,
+		PendingKick,
+		ActivePreload,
+		ReadyForLoad,
+		Count,
+	};
+	const TCHAR* LexToString(UE::Cook::EPreloaderState State);
+
 	/** SubState when in a Saving state. */
 	enum class ESaveSubState : uint8
 	{
diff --git a/Engine/Source/Editor/UnrealEd/Private/Cooker/CookWorkerClient.cpp b/Engine/Source/Editor/UnrealEd/Private/Cooker/CookWorkerClient.cpp
--- a/Engine/Source/Editor/UnrealEd/Private/Cooker/CookWorkerClient.cpp
+++ b/Engine/Source/Editor/UnrealEd/Private/Cooker/CookWorkerClient.cpp
@@ -922,7 +922,7 @@ void FCookWorkerClient::AssignPackages(FAssignPackagesMessage& Message)
 							UE_LOG(LogCook, Display,
 								TEXT("Package %s is in the save state, but the CookDirector updated the requested platforms to include the new platform %s. Restarting the package's save."),
 								*PackageData.GetPackageName().ToString(), *TargetPlatform->PlatformName());
-							PackageData.SendToState(EPackageState::LoadReady, ESendFlags::QueueAddAndRemove, EStateChangeReason::DirectorRequest);
+							PackageData.SendToState(EPackageState::Load, ESendFlags::QueueAddAndRemove, EStateChangeReason::DirectorRequest);
 						}
 					}
 				}
diff --git a/Engine/Source/Runtime/Core/Private/Misc/PreloadableFile.cpp b/Engine/Source/Runtime/Core/Private/Misc/PreloadableFile.cpp
--- a/Engine/Source/Runtime/Core/Private/Misc/PreloadableFile.cpp
+++ b/Engine/Source/Runtime/Core/Private/Misc/PreloadableFile.cpp
@@ -291,6 +291,12 @@ FArchive* FPreloadableArchive::DetachLowerLevel()
 	return SynchronousArchive.Release();
 }
 
+bool FPreloadableArchive::HasValidData() const
+{
+	WaitForInitialization();
+	return SynchronousArchive || AsynchronousHandle;
+}
+
 void FPreloadableArchive::PausePreload()
 {
 	bIsPreloadingPaused = true;
diff --git a/Engine/Source/Runtime/Core/Public/Misc/PreloadableFile.h b/Engine/Source/Runtime/Core/Public/Misc/PreloadableFile.h
--- a/Engine/Source/Runtime/Core/Public/Misc/PreloadableFile.h
+++ b/Engine/Source/Runtime/Core/Public/Misc/PreloadableFile.h
@@ -99,6 +99,8 @@ public:
 	CORE_API bool IsCacheAllocated() const;
 	/** Return the LowerLevel FArchive if it has been allocated. May return null, even if the FPreloadableFile is currently active. If return value is non-null, caller is responsible for deleting it. */
 	CORE_API FArchive* DetachLowerLevel();
+	/** Return whether this archive has an inner archive it can provide data out of. Serialize will fail if this returns false. */
+	CORE_API bool HasValidData() const;
 
 	// FArchive
 	CORE_API virtual void Serialize(void* V, int64 Length) final;
diff --git a/Engine/Source/Runtime/CoreUObject/Private/UObject/PackageResourceManager.cpp b/Engine/Source/Runtime/CoreUObject/Private/UObject/PackageResourceManager.cpp
--- a/Engine/Source/Runtime/CoreUObject/Private/UObject/PackageResourceManager.cpp
+++ b/Engine/Source/Runtime/CoreUObject/Private/UObject/PackageResourceManager.cpp
@@ -243,7 +243,21 @@ bool IPackageResourceManager::TryTakePreloadableArchive(const FPackagePath& Pack
 	// it is not preloading bytes, but instead is only providing a pre-opened (and possibly primed) sync handle
 	if (!OutResult.Archive)
 	{
-		// Otherwise the archive is in PreloadBytes mode, and we need to return a proxy to it
+		if (!PreloadableArchive->HasValidData())
+		{
+			// It should be impossible for the PreloadableArchive to not have valid data, because the code that
+			// registers the Archives should not register it in those cases. But we have seen this case come up
+			// once. So check for invalid data and log a warning, and avoid a followup serialization error by returning 
+			// false.
+			UE_LOG(LogPackageResourceManager, Warning,
+				TEXT("Registered PreloadableArchive for package %s does not have valid data; this should be impossible. We will ignore the registered archive and load without preloading.")
+				TEXT("\n\tIsInitialized = %s, Size = %" INT64_FMT "."),
+				*PackageName.ToString(), *LexToString(PreloadableArchive->IsInitialized()), PreloadableArchive->TotalSize());
+			OutResult = FOpenPackageResult();
+			return false;
+		}
+
+		// Otherwise the archive is in PreloadBytes mode, and we need to return a proxy to it.
 		OutResult.Archive = TUniquePtr<FArchive>(new FPreloadableArchiveProxy(PreloadableArchive));
 	}
 	return true;
