Cooker: Bugfix to preloading of imports: fix breaking of the undocumented invariant that packages in the load state have a CountFromRequestedLoads > 0, and add a comment about it. Fix a shadowed variable compiler error.
This fixes the warning "Package %s is generated but is ReadyForLoad when not in load state...."
#rnx

diff --git a/Engine/Source/Editor/UnrealEd/Private/Cooker/CookPackagePreloader.cpp b/Engine/Source/Editor/UnrealEd/Private/Cooker/CookPackagePreloader.cpp
--- a/Engine/Source/Editor/UnrealEd/Private/Cooker/CookPackagePreloader.cpp
+++ b/Engine/Source/Editor/UnrealEd/Private/Cooker/CookPackagePreloader.cpp
@@ -492,7 +492,8 @@ void FPackagePreloader::DecrementCountFromRequestedLoads()
 	}
 }
 
-void FPackagePreloader::SetRequestedLoads(TArray<TRefCountPtr<FPackagePreloader>>&& InRequestedLoads)
+void FPackagePreloader::SetRequestedLoads(TArray<TRefCountPtr<FPackagePreloader>>&& InRequestedLoads,
+	bool bMakeActive)
 {
 	FPackageDatas& PackageDatas = PackageData.GetPackageDatas();
 	check(RequestedLoads.IsEmpty()); // This function is only for setting from empty.
@@ -507,7 +508,7 @@ void FPackagePreloader::SetRequestedLoads(TArray<TRefCountPtr<FPackagePreloader>
 			NeedsLoadData.SetLeafToRootRank(PackageDatas.GetNextLeafToRootRank());
 		}
 
-		if (NeedsLoadPreloader->GetState() == EPreloaderState::Inactive)
+		if (NeedsLoadPreloader->GetState() == EPreloaderState::Inactive && bMakeActive)
 		{
 			NeedsLoadPreloader->SendToState(EPreloaderState::PendingKick, ESendFlags::QueueAddAndRemove);
 		}
@@ -625,23 +626,39 @@ bool FPackagePreloader::PumpLoadsTryStartInboxPackage(UCookOnTheFlyServer& COTFS
 		return false;
 	}
 
-	FPackageData* PackageData = Inbox.PopFrontValue();
-	check(PackageData->GetState() == EPackageState::Load);
-	TRefCountPtr<FPackagePreloader> Preloader = PackageData->GetPackagePreloader();
+	FPackageData* PoppedPackageData = Inbox.PopFrontValue();
+	check(PoppedPackageData->GetState() == EPackageState::Load);
+	TRefCountPtr<FPackagePreloader> Preloader = PoppedPackageData->GetPackagePreloader();
 	check(Preloader);
 	Preloader->SetIsInInbox(false);
 
-	if (COTFS.TryCreateRequestCluster(*PackageData))
+	// A required invariant for any preloader moved into an active state is that it has a count from the packages in
+	// load state that are requesting it. Assert that we satisfy that invariant for *this during this function.
+	ON_SCOPE_EXIT
+	{
+		if (Preloader->GetState() != EPreloaderState::Inactive)
+		{
+			check(Preloader->GetCountFromRequestedLoads() > 0);
+		}
+	};
+
+	if (COTFS.TryCreateRequestCluster(*PoppedPackageData))
 	{
 		return true;
 	}
 
 	// If the package is already ready for loading, or we otherwise want to skip preloading for it,
 	// skip the preload containers and put in the ReadyLoads container
-	if (!COTFS.bPreloadingEnabled || Preloader->IsPackageLoaded() || PackageData->GetIsUrgent())
+	if (!COTFS.bPreloadingEnabled || Preloader->IsPackageLoaded() || PoppedPackageData->GetIsUrgent())
 	{
 		if (Preloader->GetState() != EPreloaderState::ReadyForLoad)
 		{
+			if (!Preloader->HasInitializedRequestedLoads())
+			{
+				Preloader->SetHasInitializedRequestedLoads(true);
+				Preloader->SetRequestedLoads(TArray<TRefCountPtr<FPackagePreloader>>({ Preloader }),
+					false /* bMakeActive */);
+			}
 			Preloader->SendToState(EPreloaderState::ReadyForLoad, ESendFlags::QueueAddAndRemove);
 		}
 		return true;
@@ -674,6 +691,12 @@ bool FPackagePreloader::PumpLoadsTryStartInboxPackage(UCookOnTheFlyServer& COTFS
 	{
 		// Edgecase: we've already initialized loads, but the preloader is inactive and not loaded somehow. Put it
 		// directly into ReadyForLoad since its not clear that it needs preloading.
+		if (!Preloader->HasInitializedRequestedLoads())
+		{
+			Preloader->SetHasInitializedRequestedLoads(true);
+			Preloader->SetRequestedLoads(TArray<TRefCountPtr<FPackagePreloader>>({ Preloader }),
+				false /* bMakeActive */);
+		}
 		Preloader->SendToState(EPreloaderState::ReadyForLoad, ESendFlags::QueueAddAndRemove);
 	}
 
diff --git a/Engine/Source/Editor/UnrealEd/Private/Cooker/CookPackagePreloader.h b/Engine/Source/Editor/UnrealEd/Private/Cooker/CookPackagePreloader.h
--- a/Engine/Source/Editor/UnrealEd/Private/Cooker/CookPackagePreloader.h
+++ b/Engine/Source/Editor/UnrealEd/Private/Cooker/CookPackagePreloader.h
@@ -155,7 +155,7 @@ private:
 	 * InRequestedLoads must be in Root to Leaf order; that order is used to define the priority of
 	 * packages in PendingKick queue (leafwards packages are kicked earlier).
 	 */
-	void SetRequestedLoads(TArray<TRefCountPtr<FPackagePreloader>>&& InRequestedLoads);
+	void SetRequestedLoads(TArray<TRefCountPtr<FPackagePreloader>>&& InRequestedLoads, bool bMakeActive = true);
 	/**
 	 * Called when a PackagePreloader leaves the active state (regardless of what state its Package is in) to
 	 * clear the resources used during preloading.
