Changed MapBuildDataID warning to only apply for StaticMeshComponents who use StaticLighting, and surpress it for computing VLM only.

diff --git a/Engine/Source/Editor/UnrealEd/Private/StaticLightingSystem/StaticLightingSystem.cpp b/Engine/Source/Editor/UnrealEd/Private/StaticLightingSystem/StaticLightingSystem.cpp
--- a/Engine/Source/Editor/UnrealEd/Private/StaticLightingSystem/StaticLightingSystem.cpp
+++ b/Engine/Source/Editor/UnrealEd/Private/StaticLightingSystem/StaticLightingSystem.cpp
@@ -466,7 +466,7 @@ void FStaticLightingManager::FinishLightingBuild()
 	// Verify if new MapBuildDataIDs were created during the build and mark their object's package dirty
 	for (UStaticMeshComponent* StaticMeshComponent : TObjectRange<UStaticMeshComponent>())
 	{
-		if (StaticMeshComponent->IsTemplate())
+		if (StaticMeshComponent->IsTemplate() || !StaticMeshComponent->HasStaticLighting())
 		{
 			continue;
 		}
