IterativeCook: Fix for legacy -iterate cook that fixes bug caused by change to meaning of EDependencyProperty::Build. Build is no longer set for imports from linker table. Legacy iterative should use EDependencyProperty::Hard to identify the imports and use that during dependency walking.
#rnx

diff --git a/Engine/Source/Editor/UnrealEd/Private/Commandlets/AssetRegistryGenerator.cpp b/Engine/Source/Editor/UnrealEd/Private/Commandlets/AssetRegistryGenerator.cpp
--- a/Engine/Source/Editor/UnrealEd/Private/Commandlets/AssetRegistryGenerator.cpp
+++ b/Engine/Source/Editor/UnrealEd/Private/Commandlets/AssetRegistryGenerator.cpp
@@ -1336,6 +1336,8 @@ void FAssetRegistryGenerator::ComputePackageDifferences(const FComputeDifference
 			// Read referencers from the current state. If there are referencers in the old state that are not in the
 			// current state, then they must have changed themselves and so are already in the modified set.
 			TArray<FAssetIdentifier> Referencers;
+			State.GetReferencers(ModifiedPackage, Referencers, UE::AssetRegistry::EDependencyCategory::Package,
+				UE::AssetRegistry::EDependencyQuery::Hard);
 			State.GetReferencers(ModifiedPackage, Referencers, UE::AssetRegistry::EDependencyCategory::Package,
 				UE::AssetRegistry::EDependencyQuery::Build);
 
