Capturing Insights Snapshots is off by default in Editor Performance tool. Some developers have seen this lead to performance degredation. It should be an opt-in

diff --git a/Engine/Source/Editor/UnrealEd/Classes/Editor/EditorPerformanceSettings.h b/Engine/Source/Editor/UnrealEd/Classes/Editor/EditorPerformanceSettings.h
--- a/Engine/Source/Editor/UnrealEd/Classes/Editor/EditorPerformanceSettings.h
+++ b/Engine/Source/Editor/UnrealEd/Classes/Editor/EditorPerformanceSettings.h
@@ -51,7 +51,7 @@ class UEditorPerformanceSettings : public UDeveloperSettings
 
 	/** Whether to create Insights Snapshots */
 	UPROPERTY(EditAnywhere, config, Category = EditorPerformanceTool, meta = (DisplayName = "Enable Snapshots", ConfigRestartRequired = false, EditCondition = "bEnableEditorPeformanceTool"))
-	bool bEnableSnapshots = true;
+	bool bEnableSnapshots = false;
 
 	/** Whether to Show Warnings Only */
 	UPROPERTY(EditAnywhere, config, Category = EditorPerformanceTool, meta = (DisplayName = "Show Warnings Only", ConfigRestartRequired = false, EditCondition = "bEnableEditorPeformanceTool"))
