Cooker: Make generated packages High urgency so that they clear out and free the memory for their generator quickly.
#rnx

diff --git a/Engine/Source/Editor/UnrealEd/Classes/CookOnTheSide/CookOnTheFlyServer.h b/Engine/Source/Editor/UnrealEd/Classes/CookOnTheSide/CookOnTheFlyServer.h
--- a/Engine/Source/Editor/UnrealEd/Classes/CookOnTheSide/CookOnTheFlyServer.h
+++ b/Engine/Source/Editor/UnrealEd/Classes/CookOnTheSide/CookOnTheFlyServer.h
@@ -198,6 +198,7 @@ namespace UE::Cook
 	enum class EStateChangeReason : uint8;
 	enum class ESuppressCookReason : uint8;
 	enum class ESendFlags : uint8;
+	enum class EUrgency : uint8;
 	struct FBeginCookConfigSettings;
 	struct FCachedObjectInOuter;
 	struct FConstructPackageData;
@@ -541,10 +542,12 @@ private:
 	 * @param PackageData			The PackageData to be considered for saving.
 	 */
 	void QueueDiscoveredPackage(UE::Cook::FPackageData& PackageData, UE::Cook::FInstigator&& Instigator, 
-		UE::Cook::FDiscoveredPlatformSet&& ReachablePlatforms, bool bUrgent=false,
+		UE::Cook::FDiscoveredPlatformSet&& ReachablePlatforms, UE::Cook::EUrgency Urgency,
 		UE::Cook::FGenerationHelper* ParentGenerationHelper = nullptr);
+	void QueueDiscoveredPackage(UE::Cook::FPackageData& PackageData, UE::Cook::FInstigator&& Instigator,
+		UE::Cook::FDiscoveredPlatformSet&& ReachablePlatforms);
 	void QueueDiscoveredPackageOnDirector(UE::Cook::FPackageData& PackageData, UE::Cook::FInstigator&& Instigator,
-		UE::Cook::FDiscoveredPlatformSet&& ReachablePlatforms, bool bUrgent);
+		UE::Cook::FDiscoveredPlatformSet&& ReachablePlatforms, UE::Cook::EUrgency Urgency);
 
 	/** Called when a package is cancelled and returned to idle. Notifies CookDirector when on a CookWorker. */
 	void DemoteToIdle(UE::Cook::FPackageData& PackageData, UE::Cook::ESendFlags SendFlags, UE::Cook::ESuppressCookReason Reason);
diff --git a/Engine/Source/Editor/UnrealEd/Private/CookOnTheFlyServer.cpp b/Engine/Source/Editor/UnrealEd/Private/CookOnTheFlyServer.cpp
--- a/Engine/Source/Editor/UnrealEd/Private/CookOnTheFlyServer.cpp
+++ b/Engine/Source/Editor/UnrealEd/Private/CookOnTheFlyServer.cpp
@@ -2332,6 +2332,8 @@ void UCookOnTheFlyServer::WaitForAsync(UE::Cook::FTickStackData& StackData)
 
 UCookOnTheFlyServer::ECookAction UCookOnTheFlyServer::DecideNextCookAction(UE::Cook::FTickStackData& StackData)
 {
+	using namespace UE::Cook;
+
 	if (StackData.ResultFlags & COSR_YieldTick)
 	{
 		// Yielding on demand does not impact idle status
@@ -2350,54 +2352,61 @@ UCookOnTheFlyServer::ECookAction UCookOnTheFlyServer::DecideNextCookAction(UE::C
 		return ECookAction::Poll;
 	}
 
-	UE::Cook::FRequestQueue& RequestQueue = PackageDatas->GetRequestQueue();
+	FRequestQueue& RequestQueue = PackageDatas->GetRequestQueue();
 	if (RequestQueue.HasRequestsToExplore())
 	{
 		SetIdleStatus(StackData, EIdleStatus::Active);
 		return ECookAction::Request;
 	}
 
-	UE::Cook::FPackageDataMonitor& Monitor = PackageDatas->GetMonitor();
-	if (Monitor.GetNumUrgent() > 0)
+	FPackageDataMonitor& Monitor = PackageDatas->GetMonitor();
+
+	// If we have any packages with urgency higher than EUrgency::Normal, pump their states first, if not busy,
+	// before pumping any lower-urgency states.
+	for (EUrgency UrgencyLevel = EUrgency::Max; UrgencyLevel > EUrgency::Normal;
+		UrgencyLevel = static_cast<EUrgency>(static_cast<uint32>(UrgencyLevel) - 1))
 	{
-		if (!bSaveBusy && Monitor.GetNumUrgent(UE::Cook::EPackageState::SaveActive) > 0)
-		{
-			SetIdleStatus(StackData, EIdleStatus::Active);
-			return ECookAction::Save;
-		}
-		else if (!bLoadBusy && Monitor.GetNumUrgent(UE::Cook::EPackageState::Load) > 0)
+		if (Monitor.GetNumUrgent(UrgencyLevel) > 0)
 		{
-			SetIdleStatus(StackData, EIdleStatus::Active);
-			return ECookAction::Load;
-		}
-		else if (Monitor.GetNumUrgent(UE::Cook::EPackageState::Request) > 0)
-		{
-			SetIdleStatus(StackData, EIdleStatus::Active);
-			return ECookAction::Request;
-		}
-		else if (Monitor.GetNumUrgent(UE::Cook::EPackageState::SaveActive) > 0)
-		{
-			SetIdleStatus(StackData, EIdleStatus::Active);
-			return ECookAction::Save;
-		}
-		else if (Monitor.GetNumUrgent(UE::Cook::EPackageState::Load) > 0)
-		{
-			SetIdleStatus(StackData, EIdleStatus::Active);
-			return ECookAction::Load;
-		}
+			static_assert(static_cast<uint32>(EPackageState::Count) == 7, "Need to handle every state here.");
+			bool bSaveHasUrgent = Monitor.GetNumUrgent(EPackageState::SaveActive, UrgencyLevel) > 0;
+			if (!bSaveBusy && bSaveHasUrgent)
+			{
+				SetIdleStatus(StackData, EIdleStatus::Active);
+				return ECookAction::Save;
+			}
+			bool bLoadHasUrgent = Monitor.GetNumUrgent(EPackageState::Load, UrgencyLevel) > 0;
+			if (!bLoadBusy && bLoadHasUrgent)
+			{
+				SetIdleStatus(StackData, EIdleStatus::Active);
+				return ECookAction::Load;
+			}
+			if (Monitor.GetNumUrgent(EPackageState::Request, UrgencyLevel) > 0)
+			{
+				SetIdleStatus(StackData, EIdleStatus::Active);
+				return ECookAction::Request;
+			}
+			if (UrgencyLevel == EUrgency::Blocking && bSaveHasUrgent)
+			{
+				SetIdleStatus(StackData, EIdleStatus::Active);
+				return ECookAction::Save;
+			}
+			if (UrgencyLevel == EUrgency::Blocking && bLoadHasUrgent)
+			{
+				SetIdleStatus(StackData, EIdleStatus::Active);
+				return ECookAction::Load;
+			}
 
-		if (Monitor.GetNumUrgent(UE::Cook::EPackageState::AssignedToWorker) > 0)
-		{
-			// Fall through and do non-urgent while we wait for the Worker to finish
-		}
-		else
-		{
-			// UE::Cook::EPackageState::SaveStalledAssignedToWorker should not be possible for urgent packages
-			// UE::Cook::EPackageState::SaveStalledRetracted should not be possible for urgent packages
-			checkf(false, TEXT("Urgent request is in state not yet handled by DecideNextCookAction"));
+			// For the the remaining states, nothing to do
+			// EPackageState::AssignedToWorker
+			// EPackageState::SaveStalledAssignedToWorker
+			// EPackageState::SaveStalledRetracted
+
+			// fall through and do the next lower level of urgency
 		}
 	}
 
+	static_assert(static_cast<uint32>(EPackageState::Count) == 7, "Need to handle every state here.");
 	int32 NumSaves = PackageDatas->GetSaveQueue().Num();
 	bool bSaveAvailable = ((!bSaveBusy) & (NumSaves > 0)) != 0;
 	if (bSaveAvailable & (NumSaves > static_cast<int32>(DesiredSaveQueueLength)))
@@ -2819,7 +2828,7 @@ void UCookOnTheFlyServer::PumpLoads(UE::Cook::FTickStackData& StackData, uint32
 	using namespace UE::Cook;
 	FLoadQueue& LoadQueue = PackageDatas->GetLoadQueue();
 	FPackageDataMonitor& Monitor = PackageDatas->GetMonitor();
-	bool bIsUrgentInProgress = Monitor.GetNumUrgent() > 0;
+	bool bIsBlockingUrgencyInProgress = Monitor.GetNumUrgent(EUrgency::Blocking) > 0;
 	OutNumPushed = 0;
 	bOutBusy = false;
 
@@ -2827,7 +2836,7 @@ void UCookOnTheFlyServer::PumpLoads(UE::Cook::FTickStackData& StackData, uint32
 	TSet<TRefCountPtr<FPackagePreloader>>& ActivePreloads = LoadQueue.ActivePreloads;
 	TRingBuffer<TRefCountPtr<FPackagePreloader>>& ReadyForLoads = LoadQueue.ReadyForLoads;
 
-	if (bIsUrgentInProgress && !Monitor.GetNumUrgent(EPackageState::Load))
+	if (bIsBlockingUrgencyInProgress && !Monitor.GetNumUrgent(EPackageState::Load, EUrgency::Blocking))
 	{
 		return;
 	}
@@ -2841,7 +2850,7 @@ void UCookOnTheFlyServer::PumpLoads(UE::Cook::FTickStackData& StackData, uint32
 		{
 			return;
 		}
-		if (bIsUrgentInProgress && !Monitor.GetNumUrgent(EPackageState::Load))
+		if (bIsBlockingUrgencyInProgress && !Monitor.GetNumUrgent(EPackageState::Load, EUrgency::Blocking))
 		{
 			return;
 		}
@@ -3245,8 +3254,15 @@ EDataValidationResult UCookOnTheFlyServer::ValidateSourcePackage(UE::Cook::FPack
 }
 
 void UCookOnTheFlyServer::QueueDiscoveredPackage(UE::Cook::FPackageData& PackageData,
-	UE::Cook::FInstigator&& Instigator, UE::Cook::FDiscoveredPlatformSet&& ReachablePlatforms, bool bUrgent,
-	UE::Cook::FGenerationHelper* ParentGenerationHelper)
+	UE::Cook::FInstigator&& Instigator, UE::Cook::FDiscoveredPlatformSet&& ReachablePlatforms)
+{
+	QueueDiscoveredPackage(PackageData, MoveTemp(Instigator), MoveTemp(ReachablePlatforms),
+		UE::Cook::EUrgency::Normal, nullptr /* ParentGenerationHelper */);
+}
+
+void UCookOnTheFlyServer::QueueDiscoveredPackage(UE::Cook::FPackageData& PackageData,
+	UE::Cook::FInstigator&& Instigator, UE::Cook::FDiscoveredPlatformSet&& ReachablePlatforms,
+	UE::Cook::EUrgency Urgency, UE::Cook::FGenerationHelper* ParentGenerationHelper)
 {
 	using namespace UE::Cook;
 
@@ -3275,11 +3291,12 @@ void UCookOnTheFlyServer::QueueDiscoveredPackage(UE::Cook::FPackageData& Package
 		OnDiscoveredPackageDebug(PackageData.GetPackageName(), Instigator);
 	}
 	WorkerRequests->QueueDiscoveredPackage(*this, PackageData, MoveTemp(Instigator), MoveTemp(ReachablePlatforms),
-		bUrgent, ParentGenerationHelper);
+		Urgency, ParentGenerationHelper);
 }
 
 void UCookOnTheFlyServer::QueueDiscoveredPackageOnDirector(UE::Cook::FPackageData& PackageData,
-	UE::Cook::FInstigator&& Instigator, UE::Cook::FDiscoveredPlatformSet&& ReachablePlatforms, bool bUrgent)
+	UE::Cook::FInstigator&& Instigator, UE::Cook::FDiscoveredPlatformSet&& ReachablePlatforms,
+	UE::Cook::EUrgency Urgency)
 {
 	using namespace UE::Cook;
 
@@ -3298,7 +3315,7 @@ void UCookOnTheFlyServer::QueueDiscoveredPackageOnDirector(UE::Cook::FPackageDat
 	if (!CookByTheBookOptions->bSkipHardReferences ||
 		(Instigator.Category == EInstigator::GeneratedPackage))
 	{
-		PackageData.QueueAsDiscovered(MoveTemp(Instigator), MoveTemp(ReachablePlatforms), bUrgent);
+		PackageData.QueueAsDiscovered(MoveTemp(Instigator), MoveTemp(ReachablePlatforms), Urgency);
 	}
 }
 
@@ -3377,13 +3394,14 @@ UE::Cook::EPollStatus UCookOnTheFlyServer::QueueGeneratedPackages(UE::Cook::FGen
 			FPackageData* ChildPackageData = ChildInfo.PackageData;
 			// Set the Instigator now rather than delaying it until the discovery queue is processed.
 			ChildPackageData->SetInstigator(GenerationHelper, FInstigator(EInstigator::GeneratedPackage, OwnerName));
-			// When running -cookfirst, generated packages should also be cooked first, so set the urgency of the
-			// generated packages to match the urgency of the generator
-			bool bUrgent = PackageData.GetIsUrgent();
+			// The urgency of generated packages must be at least as high as the generator to satisfy the contract of 
+			// making the generator urgent. By default they are High urgency rather than Normal so that they are saved 
+			// quickly, so that we release the memory used by their generator for them.
+			EUrgency Urgency = PackageData.GetUrgency() > EUrgency::High ? PackageData.GetUrgency() : EUrgency::High;
 
 			// Queue the package for cooking
 			QueueDiscoveredPackage(*ChildPackageData, FInstigator(ChildPackageData->GetInstigator()),
-				EDiscoveredPlatformSet::CopyFromInstigator, bUrgent, &GenerationHelper);
+				EDiscoveredPlatformSet::CopyFromInstigator, Urgency, &GenerationHelper);
 		}
 		GenerationHelper.EndQueueGeneratedPackages(*this);
 	}
@@ -3834,7 +3852,7 @@ UE::Cook::EPollStatus UCookOnTheFlyServer::PrepareSave(UE::Cook::FPackageData& P
 	if (Result == EPollStatus::Success && PackageData.GetIsCookLast())
 	{
 		// No longer urgent
-		PackageData.ClearCookLastUrgency();
+		PackageData.SetUrgency(EUrgency::Normal, ESendFlags::QueueAddAndRemove);
 		// Mark it as still not ready if there are non-cook-last packages still in progress
 		if (PackageDatas->GetMonitor().GetNumInProgress() - PackageDatas->GetMonitor().GetNumCookLast() > 0)
 		{
@@ -4643,9 +4661,9 @@ void UCookOnTheFlyServer::PumpSaves(UE::Cook::FTickStackData& StackData, uint32
 		bool bShouldExitPump = false;
 		if (IsCookOnTheFlyMode())
 		{
-			if (IsUsingLegacyCookOnTheFlyScheduling() && !PackageData.GetIsUrgent())
+			if (IsUsingLegacyCookOnTheFlyScheduling() && PackageData.GetUrgency() != EUrgency::Blocking)
 			{
-				if (WorkerRequests->HasExternalRequests() || PackageDatas->GetMonitor().GetNumUrgent() > 0)
+				if (WorkerRequests->HasExternalRequests() || PackageDatas->GetMonitor().GetNumUrgent(EUrgency::Blocking) > 0)
 				{
 					bShouldExitPump = true;
 				}
@@ -4718,7 +4736,7 @@ void UCookOnTheFlyServer::PumpSaves(UE::Cook::FTickStackData& StackData, uint32
 			}
 
 			// Can we postpone?
-			if (!PackageData.GetIsUrgent())
+			if (PackageData.GetUrgency() != EUrgency::Blocking)
 			{
 				bool HasCheckedAllPackagesAreCached = HandledCount >= OriginalPackagesToSaveCount;
 				if (!HasCheckedAllPackagesAreCached)
@@ -4728,7 +4746,7 @@ void UCookOnTheFlyServer::PumpSaves(UE::Cook::FTickStackData& StackData, uint32
 				}
 			}
 			// Should we wait?
-			if (PackageData.GetIsUrgent() && !IsRealtimeMode())
+			if (PackageData.GetUrgency() == EUrgency::Blocking && !IsRealtimeMode())
 			{
 				UE_SCOPED_HIERARCHICAL_COOKTIMER(WaitingForCachedCookedPlatformData);
 				do
@@ -4746,7 +4764,8 @@ void UCookOnTheFlyServer::PumpSaves(UE::Cook::FTickStackData& StackData, uint32
 					// Poll the results again and check whether we are now done
 					PackageDatas->PollPendingCookedPlatformDatas(true, LastCookableObjectTickTime);
 					PrepareSaveStatus = PrepareSave(PackageData, StackData.Timer, false /* bPrecaching */);
-				} while (!StackData.Timer.IsActionTimeUp() && PrepareSaveStatus == EPollStatus::Incomplete && PackageData.GetIsUrgent());
+				} while (!StackData.Timer.IsActionTimeUp() && PrepareSaveStatus == EPollStatus::Incomplete
+					&& PackageData.GetUrgency() == EUrgency::Blocking);
 			}
 			// If we couldn't postpone or wait, then we need to exit and try again later
 			if (PrepareSaveStatus != EPollStatus::Success)
@@ -4797,7 +4816,7 @@ void UCookOnTheFlyServer::PumpSaves(UE::Cook::FTickStackData& StackData, uint32
 			check(!PackageData.GetGenerationHelper() && !PackageData.GetParentGenerationHelper());
 			ReleaseCookedPlatformData(PackageData, EStateChangeReason::RecreateObjectCache, EPackageState::SaveActive);
 			PackageData.ClearObjectCache();
-			if (PackageData.GetIsUrgent())
+			if (PackageData.GetUrgency() > EUrgency::Normal)
 			{
 				SaveQueue.AddFront(&PackageData);
 			}
@@ -5588,7 +5607,7 @@ void UCookOnTheFlyServer::PostGarbageCollect()
 			if (PackageData->GetIsCookLast())
 			{
 				// CookLast packages in SaveState have had their urgency removed. Add it back if we need to demote them.
-				PackageData->AddUrgency(true /* bValue */, false /* bAllowUpdateState */);
+				PackageData->SetUrgency(EUrgency::Blocking, ESendFlags::QueueNone);
 			}
 			PackageDatas->GetRequestQueue().AddRequest(PackageData, /* bForceUrgent */ true);
 			break;
@@ -6742,7 +6761,7 @@ void FSaveCookedPackageContext::FinishPackage()
 
 	if (!bHasRetryErrorCode)
 	{
-		if (COTFS.IsCookOnTheFlyMode() && !PackageData.GetIsUrgent() &&
+		if (COTFS.IsCookOnTheFlyMode() && PackageData.GetUrgency() != EUrgency::Blocking &&
 			(!COTFS.CookOnTheFlyRequestManager || COTFS.CookOnTheFlyRequestManager->ShouldUseLegacyScheduling()))
 		{
 			// this is an unsolicited package
diff --git a/Engine/Source/Editor/UnrealEd/Private/Cooker/CookDirector.cpp b/Engine/Source/Editor/UnrealEd/Private/Cooker/CookDirector.cpp
--- a/Engine/Source/Editor/UnrealEd/Private/Cooker/CookDirector.cpp
+++ b/Engine/Source/Editor/UnrealEd/Private/Cooker/CookDirector.cpp
@@ -529,8 +529,8 @@ void FCookDirector::AssignRequests(TArray<FWorkerId>&& InWorkers,
 		{
 			WorkerId = WorkerIdConstraint;
 		}
-		// Override the loadbalancer's assignment to force it local if the Package is urgent
-		else if (RequestPackage->GetIsUrgent())
+		// Override the loadbalancer's assignment to force it local if the Package is blocking urgency
+		else if (RequestPackage->GetUrgency() == EUrgency::Blocking)
 		{
 			WorkerId = FWorkerId::Local();
 		}
diff --git a/Engine/Source/Editor/UnrealEd/Private/Cooker/CookPackageData.cpp b/Engine/Source/Editor/UnrealEd/Private/Cooker/CookPackageData.cpp
--- a/Engine/Source/Editor/UnrealEd/Private/Cooker/CookPackageData.cpp
+++ b/Engine/Source/Editor/UnrealEd/Private/Cooker/CookPackageData.cpp
@@ -92,11 +92,11 @@ bool FPackagePlatformData::NeedsCooking(const ITargetPlatform* PlatformItBelongs
 
 FPackageData::FPackageData(FPackageDatas& PackageDatas, const FName& InPackageName, const FName& InFileName)
 	: ParentGenerationHelper(nullptr), PackageName(InPackageName), FileName(InFileName), PackageDatas(PackageDatas)
-	, Instigator(EInstigator::NotYetRequested), bIsUrgent(0), bIsCookLast(0)
+	, Instigator(EInstigator::NotYetRequested), Urgency(static_cast<uint32>(EUrgency::Normal)), bIsCookLast(0)
 	, bIsVisited(0)
 	, bHasSaveCache(0), bPrepareSaveFailed(0), bPrepareSaveRequiresGC(0)
 	, MonitorCookResult((uint8)ECookResult::NotAttempted)
-	, bCompletedGeneration(0), bGenerated(0), bKeepReferencedDuringGC(0)
+	, bGenerated(0), bKeepReferencedDuringGC(0)
 	, bWasCookedThisSession(0)
 	, DoesGeneratedRequireGeneratorValue(static_cast<uint32>(ICookPackageSplitter::EGeneratedRequiresGenerator::None))
 {
@@ -240,43 +240,40 @@ void FPackageData::AddReachablePlatformsInternal(FPackageData& PackageData,
 }
 
 void FPackageData::QueueAsDiscovered(FInstigator&& InInstigator, FDiscoveredPlatformSet&& ReachablePlatforms,
-	bool bUrgent)
+	EUrgency InUrgency)
 {
-	QueueAsDiscoveredInternal(*this, MoveTemp(InInstigator), MoveTemp(ReachablePlatforms), bUrgent);
+	QueueAsDiscoveredInternal(*this, MoveTemp(InInstigator), MoveTemp(ReachablePlatforms), InUrgency);
 }
 
 void FPackageData::QueueAsDiscoveredInternal(FPackageData& PackageData, FInstigator&& InInstigator,
-	FDiscoveredPlatformSet&& ReachablePlatforms, bool bUrgent)
+	FDiscoveredPlatformSet&& ReachablePlatforms, EUrgency InUrgency)
 {
 	// This is a static helper function to make it impossible to make a typo and use this->Instigator instead of
 	// InInstigator
 	TRingBuffer<FDiscoveryQueueElement>& Queue = PackageData.PackageDatas.GetRequestQueue().GetDiscoveryQueue();
-	Queue.Add(FDiscoveryQueueElement{ &PackageData, MoveTemp(InInstigator), MoveTemp(ReachablePlatforms), bUrgent });
+	Queue.Add(FDiscoveryQueueElement{ &PackageData, MoveTemp(InInstigator), MoveTemp(ReachablePlatforms), InUrgency });
 }
 
-void FPackageData::SetIsUrgent(bool Value)
+void FPackageData::SetUrgency(EUrgency NewUrgency, ESendFlags SendFlags, bool bAllowUrgencyInIdle)
 {
-	bool OldValue = static_cast<bool>(bIsUrgent);
-	if (OldValue != Value)
-	{
-		bIsUrgent = Value != 0;
-		check(IsInProgress() || !bIsUrgent);
-		PackageDatas.GetMonitor().OnUrgencyChanged(*this);
-	}
-}
-
-void FPackageData::AddUrgency(bool bUrgent, bool bAllowUpdateState)
-{
-	if (!bUrgent)
+	if (GetUrgency() == NewUrgency)
 	{
 		return;
 	}
-	bool bWasUrgent = GetIsUrgent();
-	SetIsUrgent(true);
-	if (!bWasUrgent && bAllowUpdateState)
+
+	// It is illegal to SetUrgency to above normal when in the Idle state, unless the caller explicitly takes
+	// responsibility for changing the state immediately afterwards.
+	check(bAllowUrgencyInIdle || IsInProgress() || NewUrgency == EUrgency::Normal);
+	// For SendFlags when setting urgency, only AddAndRemove or None are supported
+	check(SendFlags == ESendFlags::QueueAddAndRemove || SendFlags == ESendFlags::QueueNone);
+
+	EUrgency OldUrgency = GetUrgency();
+	Urgency = static_cast<uint32>(NewUrgency);
+	if (SendFlags == ESendFlags::QueueAddAndRemove)
 	{
-		SendToState(GetState(), ESendFlags::QueueAddAndRemove, EStateChangeReason::UrgencyUpdated);
+		UpdateContainerUrgency(OldUrgency, NewUrgency);
 	}
+	PackageDatas.GetMonitor().OnUrgencyChanged(*this, OldUrgency, NewUrgency);
 }
 
 void FPackageData::SetIsCookLast(bool bValue)
@@ -289,15 +286,6 @@ void FPackageData::SetIsCookLast(bool bValue)
 	}
 }
 
-void FPackageData::ClearCookLastUrgency()
-{
-	if (!GetIsCookLast() || !GetIsUrgent())
-	{
-		return;
-	}
-	SetIsUrgent(false);
-}
-
 void FPackageData::SetInstigator(FRequestCluster& Cluster, FInstigator&& InInstigator)
 {
 	SetInstigatorInternal(MoveTemp(InInstigator));
@@ -323,7 +311,7 @@ void FPackageData::SetInstigatorInternal(FInstigator&& InInstigator)
 
 void FPackageData::ClearInProgressData(EStateChangeReason StateChangeReason)
 {
-	SetIsUrgent(false);
+	SetUrgency(EUrgency::Normal, ESendFlags::QueueNone);
 	CompletionCallback = FCompletionCallback();
 	if (GenerationHelper)
 	{
@@ -751,7 +739,7 @@ void FPackageData::SendToState(EPackageState NextState, ESendFlags SendFlags, ES
 		OnEnterSaveActive();
 		if (((SendFlags & ESendFlags::QueueAdd) != ESendFlags::QueueNone))
 		{
-			if (GetIsUrgent())
+			if (GetUrgency() > EUrgency::Normal)
 			{
 				PackageDatas.GetSaveQueue().AddFront(this);
 			}
@@ -783,6 +771,51 @@ void FPackageData::SendToState(EPackageState NextState, ESendFlags SendFlags, ES
 	PackageDatas.GetMonitor().OnStateChanged(*this, OldState);
 }
 
+void FPackageData::UpdateContainerUrgency(EUrgency OldUrgency, EUrgency NewUrgency)
+{
+	switch (GetState())
+	{
+	case EPackageState::Idle:
+		// Urgency does not affect behavior in the Idle state
+		break;
+	case EPackageState::Request:
+		PackageDatas.GetRequestQueue().UpdateUrgency(this, OldUrgency, NewUrgency);
+		break;
+	case EPackageState::AssignedToWorker:
+		// Urgency does not affect behavior in the AssignedToWorker state
+		break;
+	case EPackageState::Load:
+		PackageDatas.GetLoadQueue().UpdateUrgency(this, OldUrgency, NewUrgency);
+		break;
+	case EPackageState::SaveActive:
+		if (NewUrgency > EUrgency::Normal)
+		{
+			FPackageDataQueue& Queue = PackageDatas.GetSaveQueue();
+			if (Queue.Remove(this) > 0)
+			{
+				Queue.AddFront(this);
+			}
+		}
+		break;
+	case EPackageState::SaveStalledRetracted:
+		// Urgency does not affect behavior in stalled states
+		break;
+	case EPackageState::SaveStalledAssignedToWorker:
+		// Urgency does not affect behavior in stalled states
+		break;
+	default:
+		check(false);
+		break;
+	}
+
+	// The Package preloader can be active in any state, and is contained in the LoadQueue.
+	// If it exists and we did not already call UpdateUrgency on the LoadQueue, then call it.
+	if (GetState() != EPackageState::Load && GetPackagePreloader())
+	{
+		PackageDatas.GetLoadQueue().UpdateUrgency(this, OldUrgency, NewUrgency);
+	}
+}
+
 void FPackageData::Stall(EPackageState TargetState, ESendFlags SendFlags)
 {
 	switch (TargetState)
@@ -1817,14 +1850,17 @@ FPackageDataMonitor::FPackageDataMonitor()
 	FMemory::Memset(NumCookLastInState, 0);
 }
 
-int32 FPackageDataMonitor::GetNumUrgent() const
+int32 FPackageDataMonitor::GetNumUrgent(EUrgency UrgencyLevel) const
 {
+	check(EUrgency::Min <= UrgencyLevel && UrgencyLevel <= EUrgency::Max);
+	int32 UrgencyIndex = static_cast<uint32>(UrgencyLevel) - static_cast<uint32>(EUrgency::Min);
 	int32 NumUrgent = 0;
 	for (EPackageState State = EPackageState::Min;
 		State <= EPackageState::Max;
 		State = static_cast<EPackageState>(static_cast<uint32>(State) + 1))
 	{
-		NumUrgent += NumUrgentInState[static_cast<uint32>(State) - static_cast<uint32>(EPackageState::Min)];
+		int32 StateIndex = static_cast<uint32>(State) - static_cast<uint32>(EPackageState::Min);
+		NumUrgent += NumUrgentInState[StateIndex][UrgencyIndex];
 	}
 	return NumUrgent;
 }
@@ -1841,16 +1877,20 @@ int32 FPackageDataMonitor::GetNumCookLast() const
 	return Num;
 }
 
-int32 FPackageDataMonitor::GetNumUrgent(EPackageState InState) const
+int32 FPackageDataMonitor::GetNumUrgent(EPackageState InState, EUrgency UrgencyLevel) const
 {
+	check(EUrgency::Min <= UrgencyLevel && UrgencyLevel <= EUrgency::Max);
+	int32 UrgencyIndex = static_cast<uint32>(UrgencyLevel) - static_cast<uint32>(EUrgency::Min);
 	check(EPackageState::Min <= InState && InState <= EPackageState::Max);
-	return NumUrgentInState[static_cast<uint32>(InState) - static_cast<uint32>(EPackageState::Min)];
+	int32 StateIndex = static_cast<uint32>(InState) - static_cast<uint32>(EPackageState::Min);
+	return NumUrgentInState[StateIndex][UrgencyIndex];
 }
 
 int32 FPackageDataMonitor::GetNumCookLast(EPackageState InState) const
 {
 	check(EPackageState::Min <= InState && InState <= EPackageState::Max);
-	return NumUrgentInState[static_cast<uint32>(InState) - static_cast<uint32>(EPackageState::Min)];
+	int32 StateIndex = static_cast<uint32>(InState) - static_cast<uint32>(EPackageState::Min);
+	return NumCookLastInState[StateIndex];
 }
 
 int32 FPackageDataMonitor::GetNumPreloadAllocated() const
@@ -1900,10 +1940,10 @@ void FPackageDataMonitor::OnLastCookedPlatformRemoved(FPackageData& PackageData)
 	}
 }
 
-void FPackageDataMonitor::OnUrgencyChanged(FPackageData& PackageData)
+void FPackageDataMonitor::OnUrgencyChanged(FPackageData& PackageData, EUrgency OldUrgency, EUrgency NewUrgency)
 {
-	int32 Delta = PackageData.GetIsUrgent() ? 1 : -1;
-	TrackUrgentRequests(PackageData.GetState(), Delta);
+	TrackUrgentRequests(PackageData.GetState(), OldUrgency, -1);
+	TrackUrgentRequests(PackageData.GetState(), NewUrgency, 1);
 }
 
 void FPackageDataMonitor::OnCookLastChanged(FPackageData& PackageData)
@@ -1915,10 +1955,11 @@ void FPackageDataMonitor::OnCookLastChanged(FPackageData& PackageData)
 void FPackageDataMonitor::OnStateChanged(FPackageData& PackageData, EPackageState OldState)
 {
 	EPackageState NewState = PackageData.GetState();
-	if (PackageData.GetIsUrgent())
+	EUrgency Urgency = PackageData.GetUrgency();
+	if (Urgency > EUrgency::Normal)
 	{
-		TrackUrgentRequests(OldState, -1);
-		TrackUrgentRequests(NewState, 1);
+		TrackUrgentRequests(OldState, Urgency, -1);
+		TrackUrgentRequests(NewState, Urgency, 1);
 	}
 	if (PackageData.GetIsCookLast())
 	{
@@ -1935,11 +1976,20 @@ void FPackageDataMonitor::OnStateChanged(FPackageData& PackageData, EPackageStat
 	}
 }
 
-void FPackageDataMonitor::TrackUrgentRequests(EPackageState State, int32 Delta)
+void FPackageDataMonitor::TrackUrgentRequests(EPackageState State, EUrgency Urgency, int32 Delta)
 {
+	if (State == EPackageState::Idle || Urgency == EUrgency::Normal)
+	{
+		// We don't track urgency count in idle, and we don't track normal urgency count.
+		return;
+	}
 	check(EPackageState::Min <= State && State <= EPackageState::Max);
-	NumUrgentInState[static_cast<uint32>(State) - static_cast<uint32>(EPackageState::Min)] += Delta;
-	check(NumUrgentInState[static_cast<uint32>(State) - static_cast<uint32>(EPackageState::Min)] >= 0);
+	check(EUrgency::Min <= Urgency && Urgency <= EUrgency::Max);
+
+	int32 StateIndex = static_cast<uint32>(State) - static_cast<uint32>(EPackageState::Min);
+	int32 UrgencyIndex = static_cast<uint32>(Urgency) - static_cast<uint32>(EUrgency::Min);
+	NumUrgentInState[StateIndex][UrgencyIndex] += Delta;
+	check(NumUrgentInState[StateIndex][UrgencyIndex] >= 0);
 }
 
 void FPackageDataMonitor::TrackCookLastRequests(EPackageState State, int32 Delta)
@@ -2834,7 +2884,7 @@ void FRequestQueue::AddRequest(FPackageData* PackageData, bool bForceUrgent)
 
 void FRequestQueue::AddReadyRequest(FPackageData* PackageData, bool bForceUrgent)
 {
-	if (bForceUrgent || PackageData->GetIsUrgent())
+	if (bForceUrgent || PackageData->GetUrgency() > EUrgency::Normal)
 	{
 		UrgentRequests.Add(PackageData);
 	}
@@ -2844,6 +2894,25 @@ void FRequestQueue::AddReadyRequest(FPackageData* PackageData, bool bForceUrgent
 	}
 }
 
+void FRequestQueue::UpdateUrgency(FPackageData* PackageData, EUrgency OldUrgency, EUrgency NewUrgency)
+{
+	if (OldUrgency == EUrgency::Normal)
+	{
+		if (NormalRequests.Remove(PackageData) > 0)
+		{
+			UrgentRequests.Add(PackageData);
+		}
+	}
+	else
+	{
+		if (UrgentRequests.Remove(PackageData) > 0)
+		{
+			NormalRequests.Add(PackageData);
+		}
+	}
+	// The other subcontainers do not handle urgency types differently
+}
+
 void FRequestQueue::AddRequestFenceListener(FName PackageName)
 {
 	RequestFencePackageListeners.Add(PackageName);
diff --git a/Engine/Source/Editor/UnrealEd/Private/Cooker/CookPackageData.h b/Engine/Source/Editor/UnrealEd/Private/Cooker/CookPackageData.h
--- a/Engine/Source/Editor/UnrealEd/Private/Cooker/CookPackageData.h
+++ b/Engine/Source/Editor/UnrealEd/Private/Cooker/CookPackageData.h
@@ -361,25 +361,18 @@ public:
 	/** Return whether all Platforms that have been marked reachable have also been explored in an FRequestCluster. */
 	bool AreAllReachablePlatformsVisitedByCluster() const;
 
-	/**
-	 * Get the flag for whether this InProgress PackageData has been marked as an urgent request
-	 * (e.g. because it has been requested from the game client during cookonthefly.).
-	 * Always false for PackageDatas that are not InProgress.
-	 */
-	bool GetIsUrgent() const { return static_cast<bool>(bIsUrgent); }
-	/**
-	 * Mark this PackageData as urgent if bUrgent is true. Noop if bUrgent is false.
-	 * Move it to front of its current container if bAllowUpdateState.
-	 */
-	void AddUrgency(bool bUrgent, bool bAllowUpdateState);
-	void ClearCookLastUrgency();
+	/** Get/Set the urgency of a package. Always EUrgency::Normal for pacakges that are not in progress. */
+	EUrgency GetUrgency() const;
+	void SetUrgency(EUrgency NewUrgency, ESendFlags SendFlags, bool bAllowUrgencyInIdle = false);
+	/** Set the Urgency iff the newurgency is greater than current. */
+	void RaiseUrgency(EUrgency NewUrgency, ESendFlags SendFlags, bool bAllowUrgencyInIdle = false);
 
 	/** Accessor for RequestClusters to add reachable platforms directly without modifying dependent data. */
 	void AddReachablePlatforms(FRequestCluster& RequestCluster, TConstArrayView<const ITargetPlatform*> Platforms,
 		FInstigator&& InInstigator);
 
 	/** Add the given reachable platforms to this PackageData and send it back to Request state for exploration. */
-	void QueueAsDiscovered(FInstigator&& InInstigator, FDiscoveredPlatformSet&& ReachablePlatforms, bool bUrgent);
+	void QueueAsDiscovered(FInstigator&& InInstigator, FDiscoveredPlatformSet&& ReachablePlatforms, EUrgency InUrgency);
 
 	/**
 	 * Clear all the inprogress variables from the current PackageData. It is invalid to call this except when
@@ -632,10 +625,6 @@ public:
 		bool& bOutNeedWaitForIsLoaded);
 	/** Helper function for ~FGenerationHelper: clear the pointer this->GenerationHelper. */
 	void OnGenerationHelperDestroyed(FGenerationHelper& InGenerationHelper);
-	/** Get whether the PackageData has done any necessary Generator steps and is ready for BeginCache calls. */
-	bool HasCompletedGeneration() const;
-	/** Set whether the PackageData has done any necessary Generator steps and is ready for BeginCache calls. */
-	void SetCompletedGeneration(bool Value);
 	/** Return whether the PackageData is a generated package created by a owning generator PackageData. */
 	bool IsGenerated() const;
 	/** Mark that the PackageData is a generated package created by a ParentGenerator. */
@@ -722,12 +711,12 @@ public:
 private:
 	friend struct UE::Cook::FPackageDatas;
 
-	void SetIsUrgent(bool Value);
+	void UpdateContainerUrgency(EUrgency OldUrgency, EUrgency NewUrgency);
 	void SetInstigatorInternal(FInstigator&& InInstigator);
 	static void AddReachablePlatformsInternal(FPackageData& PackageData,
 		TConstArrayView<const ITargetPlatform*> Platforms, FInstigator&& InInstigator);
 	static void QueueAsDiscoveredInternal(FPackageData& PackageData, FInstigator&& InInstigator,
-		FDiscoveredPlatformSet&& ReachablePlatforms, bool bUrgent);
+		FDiscoveredPlatformSet&& ReachablePlatforms, EUrgency InUrgency);
 
 	/**
 	 * Set the FileName of the file that contains the package. This member is private because FPackageDatas
@@ -818,14 +807,13 @@ private:
 	uint32 State : int32(EPackageState::BitCount);
 	uint32 SaveSubState : int32(ESaveSubState::BitCount);
 	uint32 SuppressCookReason : int32(ESuppressCookReason::BitCount);
-	uint32 bIsUrgent : 1;
+	uint32 Urgency : int32(EUrgency::BitCount);
 	uint32 bIsCookLast : 1;
 	uint32 bIsVisited : 1;
 	uint32 bHasSaveCache : 1;
 	uint32 bPrepareSaveFailed : 1;
 	uint32 bPrepareSaveRequiresGC : 1;
 	uint32 MonitorCookResult : (int) ECookResult::NumBits;
-	uint32 bCompletedGeneration : 1;
 	uint32 bGenerated : 1;
 	uint32 bKeepReferencedDuringGC : 1;
 	uint32 bWasCookedThisSession : 1;
@@ -928,16 +916,16 @@ public:
 	/** Report the number of packages that have cooked any platform. Used by CookCommandlet progress reporting. */
 	int32 GetNumCooked(ECookResult CookResult) const;
 	/**
-	 * Report the number of FPackageData that are currently marked as urgent.
+	 * Report the number of FPackageData that are currently set to the given urgency level.
 	 * Used to check if a Pump function needs to exit to handle urgent PackageData in other states.
 	 */
-	int32 GetNumUrgent() const;
+	int32 GetNumUrgent(EUrgency UrgencyLevel) const;
 	/**
-	 * Report the number of FPackageData that are in the given state and have been marked as urgent. Only valid to
-	 * call on states that are in the InProgress set, such as Save.
+	 * Report the number of FPackageData that are in the given state and are set to the given urgency level. Only
+	 * valid to call on states that are in the InProgress set, such as Save.
 	 * Used to prioritize scheduler actions.
 	 */
-	int32 GetNumUrgent(EPackageState InState) const;
+	int32 GetNumUrgent(EPackageState InState, EUrgency UrgencyLevel) const;
 	/** Report the number of CookLast packages. */
 	int32 GetNumCookLast() const;
 	/** Report the number of CookLast packages in the given state. */
@@ -957,7 +945,7 @@ public:
 	 */
 	void OnLastCookedPlatformRemoved(FPackageData& PackageData);
 	/** Callback called from FPackageData when it has changed its urgency. */
-	void OnUrgencyChanged(FPackageData& PackageData);
+	void OnUrgencyChanged(FPackageData& PackageData, EUrgency OldUrgency, EUrgency NewUrgency);
 	/** Callback called from FPackageData when it has changed its value of IsCookLast. */
 	void OnCookLastChanged(FPackageData& PackageData);
 	/** Callback called from FPackageData when it has changed its state. */
@@ -968,14 +956,14 @@ public:
 
 private:
 	/** Increment or decrement the NumUrgent counter for the given state. */
-	void TrackUrgentRequests(EPackageState State, int32 Delta);
+	void TrackUrgentRequests(EPackageState State, EUrgency Urgency, int32 Delta);
 	/** Increment or decrement the NumCookLast counter for the given state. */
 	void TrackCookLastRequests(EPackageState State, int32 Delta);
 
 	int32 NumInProgress = 0;
 	int32 NumCooked[(uint8)ECookResult::Count]{};
 	int32 NumPreloadAllocated = 0;
-	int32 NumUrgentInState[static_cast<uint32>(EPackageState::Count)];
+	int32 NumUrgentInState[static_cast<uint32>(EPackageState::Count)][static_cast<uint32>(EUrgency::Count)];
 	int32 NumCookLastInState[static_cast<uint32>(EPackageState::Count)];
 	int32 MPCookAssignedFenceMarker = 0;
 	int32 MPCookRetiredFenceMarker = 0;
@@ -986,7 +974,7 @@ struct FDiscoveryQueueElement
 	FPackageData* PackageData;
 	FInstigator Instigator;
 	FDiscoveredPlatformSet ReachablePlatforms;
-	bool bUrgent;
+	UE::Cook::EUrgency Urgency;
 };
 
 /**
@@ -1013,6 +1001,7 @@ public:
 	void AddReadyRequest(FPackageData* PackageData, bool bForceUrgent = false);
 	uint32 RemoveRequest(FPackageData* PackageData);
 	uint32 RemoveRequestExceptFromCluster(FPackageData* PackageData, FRequestCluster* ExceptFromCluster);
+	void UpdateUrgency(FPackageData* PackageData, EUrgency bOldUrgency, EUrgency NewUrgency);
 
 	FPackageDataSet& GetRestartedRequests() { return RestartedRequests; }
 	/**
@@ -1078,6 +1067,7 @@ public:
 	void Add(FPackageData* PackageData);
 	bool Contains(const FPackageData* PackageData) const;
 	uint32 Remove(FPackageData* PackageData);
+	void UpdateUrgency(FPackageData* PackageData, EUrgency bOldUrgency, EUrgency NewUrgency);
 	TSet<FPackageData*>::TRangedForIterator begin();
 	TSet<FPackageData*>::TRangedForIterator end();
 
@@ -1572,14 +1562,17 @@ inline void FPackageData::SetSuppressCookReason(ESuppressCookReason Reason)
 	SuppressCookReason = static_cast<uint32>(Reason);
 }
 
-inline bool FPackageData::HasCompletedGeneration() const
+inline EUrgency FPackageData::GetUrgency() const
 {
-	return static_cast<bool>(bCompletedGeneration);
+	return static_cast<EUrgency>(Urgency);
 }
 
-inline void FPackageData::SetCompletedGeneration(bool Value)
+inline void FPackageData::RaiseUrgency(EUrgency NewUrgency, ESendFlags SendFlags, bool bAllowUrgencyInIdle)
 {
-	bCompletedGeneration = Value != 0;
+	if (NewUrgency > GetUrgency())
+	{
+		SetUrgency(NewUrgency, SendFlags, bAllowUrgencyInIdle);
+	}
 }
 
 inline bool FPackageData::IsGenerated() const
diff --git a/Engine/Source/Editor/UnrealEd/Private/Cooker/CookPackagePreloader.cpp b/Engine/Source/Editor/UnrealEd/Private/Cooker/CookPackagePreloader.cpp
--- a/Engine/Source/Editor/UnrealEd/Private/Cooker/CookPackagePreloader.cpp
+++ b/Engine/Source/Editor/UnrealEd/Private/Cooker/CookPackagePreloader.cpp
@@ -324,9 +324,9 @@ bool FPackagePreloader::TryPreload()
 	const TSharedPtr<FPreloadableArchive>& FilePtr = PreloadableFile.Get();
 	if (!FilePtr->IsInitialized())
 	{
-		if (PackageData.GetIsUrgent())
+		if (PackageData.GetUrgency() == EUrgency::Blocking)
 		{
-			// For urgent requests, wait on them to finish preloading rather than letting them run asynchronously
+			// For blocking requests, wait on them to finish preloading rather than letting them run asynchronously
 			// and coming back to them later
 			FilePtr->WaitForInitialization();
 			check(FilePtr->IsInitialized());
@@ -546,9 +546,9 @@ bool FPackagePreloader::NeedsLoad()
 
 bool FPackagePreloader::IsHigherPriorityThan(const FPackagePreloader& Other) const
 {
-	if (PackageData.GetIsUrgent() != Other.PackageData.GetIsUrgent())
+	if (PackageData.GetUrgency() != Other.PackageData.GetUrgency())
 	{
-		return PackageData.GetIsUrgent();
+		return PackageData.GetUrgency() > Other.PackageData.GetUrgency();
 	}
 
 	// Leaves are higher priority because we want them to be already preloaded (or even
@@ -649,7 +649,8 @@ bool FPackagePreloader::PumpLoadsTryStartInboxPackage(UCookOnTheFlyServer& COTFS
 
 	// If the package is already ready for loading, or we otherwise want to skip preloading for it,
 	// skip the preload containers and put in the ReadyLoads container
-	if (!COTFS.bPreloadingEnabled || Preloader->IsPackageLoaded() || PoppedPackageData->GetIsUrgent())
+	if (!COTFS.bPreloadingEnabled || Preloader->IsPackageLoaded()
+		|| PoppedPackageData->GetUrgency() == EUrgency::Blocking)
 	{
 		if (Preloader->GetState() != EPreloaderState::ReadyForLoad)
 		{
@@ -868,6 +869,35 @@ uint32 FLoadQueue::Remove(FPackageData* PackageData)
 	return Result;
 }
 
+void FLoadQueue::UpdateUrgency(FPackageData* PackageData, EUrgency bOldUrgency, EUrgency NewUrgency)
+{
+	TRefCountPtr<FPackagePreloader> Preloader = PackageData->GetPackagePreloader();
+	if (!Preloader)
+	{
+		// Urgency does not impact state for packages that haven't reached a preloader state yet
+		return;
+	}
+	switch (Preloader->GetState())
+	{
+	case EPreloaderState::Inactive:
+		// Urgency does not impact state for packages that are inactive
+		break;
+	case EPreloaderState::PendingKick:
+		PendingKicks.Remove(Preloader);
+		PendingKicks.Add(Preloader);
+		break;
+	case EPreloaderState::ActivePreload:
+		break;
+	case EPreloaderState::ReadyForLoad:
+		ReadyForLoads.Remove(Preloader);
+		ReadyForLoads.AddFront(Preloader);
+		break;
+	default:
+		checkNoEntry();
+		break;
+	}
+}
+
 TSet<FPackageData*>::TRangedForIterator FLoadQueue::begin()
 {
 	return InProgress.begin();
diff --git a/Engine/Source/Editor/UnrealEd/Private/Cooker/CookRequestCluster.cpp b/Engine/Source/Editor/UnrealEd/Private/Cooker/CookRequestCluster.cpp
--- a/Engine/Source/Editor/UnrealEd/Private/Cooker/CookRequestCluster.cpp
+++ b/Engine/Source/Editor/UnrealEd/Private/Cooker/CookRequestCluster.cpp
@@ -153,14 +153,21 @@ FRequestCluster::FRequestCluster(UCookOnTheFlyServer& InCOTFS, TRingBuffer<FDisc
 		{
 			// If there are no new reachable platforms, add it to the cluster for cooking if it needs
 			// it, otherwise let it remain where it is
+			EUrgency Urgency = Discovery->Urgency;
 			DiscoveryQueue.PopFrontValue();
 			Discovery = nullptr;
 			if (!PackageData.IsInProgress() && PackageData.GetPlatformsNeedingCookingNum() == 0)
 			{
 				PackageData.SendToState(EPackageState::Request, ESendFlags::QueueRemove,
 					EStateChangeReason::RequestCluster);
+				PackageData.RaiseUrgency(Urgency, ESendFlags::QueueNone);
+				// SetPackageDataSuppressReason adds it in the proper container of *this
 				SetPackageDataSuppressReason(PackageData, ESuppressCookReason::NotSuppressed);
 			}
+			else if (PackageData.IsInProgress())
+			{
+				PackageData.RaiseUrgency(Urgency, ESendFlags::QueueAddAndRemove);
+			}
 			continue;
 		}
 
@@ -233,7 +240,8 @@ FRequestCluster::FRequestCluster(UCookOnTheFlyServer& InCOTFS, TRingBuffer<FDisc
 		// Send it to the Request state if it's not already there, remove it from its old container
 		// and add it to this cluster.
 		PackageData.SendToState(EPackageState::Request, ESendFlags::QueueRemove, EStateChangeReason::RequestCluster);
-		PackageData.AddUrgency(Discovery->bUrgent, false /* bAllowUpdateState */);
+		PackageData.RaiseUrgency(Discovery->Urgency, ESendFlags::QueueNone);
+		// SetPackageDataSuppressReason adds it in the proper container of *this
 		SetPackageDataSuppressReason(PackageData, ESuppressCookReason::NotSuppressed);
 	}
 }
@@ -299,7 +307,10 @@ void FRequestCluster::FetchPackageNames(const FCookerTimer& CookerTimer, bool& b
 		{
 			PackageData->AddReachablePlatforms(*this, Request.GetPlatforms(), MoveTemp(Request.GetInstigator()));
 			PullIntoCluster(*PackageData);
-			PackageData->AddUrgency(Request.IsUrgent(), false /* bAllowUpdateState */);
+			if (Request.IsUrgent())
+			{
+				PackageData->SetUrgency(EUrgency::Blocking, ESendFlags::QueueNone);
+			}
 		}
 		else
 		{
@@ -307,13 +318,19 @@ void FRequestCluster::FetchPackageNames(const FCookerTimer& CookerTimer, bool& b
 			{
 				// If it's already in progress with no new platforms, we don't need to add it to the cluster, but add
 				// add on our urgency setting
-				PackageData->AddUrgency(Request.IsUrgent(), true /* bAllowUpdateState */);
+				if (Request.IsUrgent())
+				{
+					PackageData->SetUrgency(EUrgency::Blocking, ESendFlags::QueueAddAndRemove);
+				}
 			}
 			else if (PackageData->GetPlatformsNeedingCookingNum() > 0)
 			{
 				// If it's missing cookable platforms and not in progress we need to add it to the cluster for cooking
 				PullIntoCluster(*PackageData);
-				PackageData->AddUrgency(Request.IsUrgent(), true /* bAllowUpdateState */);
+				if (Request.IsUrgent())
+				{
+					PackageData->SetUrgency(EUrgency::Blocking, ESendFlags::QueueNone);
+				}
 			}
 		}
 		// Add on our completion callback, or call it immediately if already done
diff --git a/Engine/Source/Editor/UnrealEd/Private/Cooker/CookTypes.cpp b/Engine/Source/Editor/UnrealEd/Private/Cooker/CookTypes.cpp
--- a/Engine/Source/Editor/UnrealEd/Private/Cooker/CookTypes.cpp
+++ b/Engine/Source/Editor/UnrealEd/Private/Cooker/CookTypes.cpp
@@ -143,6 +143,18 @@ const TCHAR* LexToString(UE::Cook::ESaveSubState State)
 	}
 }
 
+const TCHAR* LexToString(UE::Cook::EUrgency Urgency)
+{
+	static_assert(static_cast<int>(EUrgency::Count) == 3);
+	switch (Urgency)
+	{
+	case EUrgency::Normal: return TEXT("Normal");
+	case EUrgency::High: return TEXT("High");
+	case EUrgency::Blocking: return TEXT("Blocking");
+	default: return TEXT("Invalid");
+	}
+}
+
 EStateChangeReason ConvertToStateChangeReason(ESuppressCookReason Reason)
 {
 	switch (Reason)
diff --git a/Engine/Source/Editor/UnrealEd/Private/Cooker/CookTypes.h b/Engine/Source/Editor/UnrealEd/Private/Cooker/CookTypes.h
--- a/Engine/Source/Editor/UnrealEd/Private/Cooker/CookTypes.h
+++ b/Engine/Source/Editor/UnrealEd/Private/Cooker/CookTypes.h
@@ -288,6 +288,20 @@ namespace UE::Cook
 	};
 	const TCHAR* LexToString(UE::Cook::ESaveSubState State);
 
+	/** How quickly we should push a PackageData through the cook, compared to other PackageDatas. */
+	enum class EUrgency : uint8
+	{
+		Normal = 0,
+		High,
+		Blocking,
+
+		Min = Normal,
+		Max = Blocking,
+		Count = Max + 1,
+		BitCount = FPlatformMath::ConstExprCeilLogTwo(Count),
+	};
+	const TCHAR* LexToString(UE::Cook::EUrgency Urgency);
+
 	/** Used as a helper to timeslice cooker functions. */
 	struct FCookerTimer
 	{
diff --git a/Engine/Source/Editor/UnrealEd/Private/Cooker/CookWorkerClient.cpp b/Engine/Source/Editor/UnrealEd/Private/Cooker/CookWorkerClient.cpp
--- a/Engine/Source/Editor/UnrealEd/Private/Cooker/CookWorkerClient.cpp
+++ b/Engine/Source/Editor/UnrealEd/Private/Cooker/CookWorkerClient.cpp
@@ -312,7 +312,7 @@ void FCookWorkerClient::ReportPackageMessage(FName PackageName, TUniquePtr<FPack
 }
 
 void FCookWorkerClient::ReportDiscoveredPackage(const FPackageData& PackageData, const FInstigator& Instigator,
-	FDiscoveredPlatformSet&& ReachablePlatforms, FGenerationHelper* ParentGenerationHelper)
+	FDiscoveredPlatformSet&& ReachablePlatforms, FGenerationHelper* ParentGenerationHelper, EUrgency Urgency)
 {
 	FDiscoveredPackageReplication& Discovered = PendingDiscoveredPackages.Emplace_GetRef();
 	Discovered.PackageName = PackageData.GetPackageName();
@@ -322,6 +322,7 @@ void FCookWorkerClient::ReportDiscoveredPackage(const FPackageData& PackageData,
 	Discovered.Platforms = MoveTemp(ReachablePlatforms);
 	Discovered.Platforms.ConvertToBitfield(OrderedSessionAndSpecialPlatforms);
 	Discovered.DoesGeneratedRequireGenerator = PackageData.DoesGeneratedRequireGenerator();
+	Discovered.Urgency = Urgency;
 	if (ParentGenerationHelper)
 	{
 		if (FCookGenerationInfo* Info = ParentGenerationHelper->FindInfo(PackageData))
@@ -928,6 +929,7 @@ void FCookWorkerClient::AssignPackages(FAssignPackagesMessage& Message)
 				}
 				// Allow the package to continue in its progress. If it was in a stalled-by-retraction state, return it to active.
 				PackageData.UnStall(ESendFlags::QueueAddAndRemove);
+				PackageData.RaiseUrgency(AssignData.Urgency, ESendFlags::QueueAddAndRemove);
 				continue;
 			}
 
@@ -939,6 +941,7 @@ void FCookWorkerClient::AssignPackages(FAssignPackagesMessage& Message)
 			}
 			PackageData.FindOrAddPlatformData(CookerLoadingPlatformKey).MarkCookableForWorker(*this);
 			PackageData.SetInstigator(*this, FInstigator(AssignData.Instigator));
+			PackageData.RaiseUrgency(AssignData.Urgency, ESendFlags::QueueAddAndRemove, true /* bAllowUrgencyInIdle */);
 			PackageData.SendToState(EPackageState::Request, ESendFlags::QueueAddAndRemove,
 				EStateChangeReason::DirectorRequest);
 		}
diff --git a/Engine/Source/Editor/UnrealEd/Private/Cooker/CookWorkerClient.h b/Engine/Source/Editor/UnrealEd/Private/Cooker/CookWorkerClient.h
--- a/Engine/Source/Editor/UnrealEd/Private/Cooker/CookWorkerClient.h
+++ b/Engine/Source/Editor/UnrealEd/Private/Cooker/CookWorkerClient.h
@@ -68,7 +68,7 @@ public:
 	void ReportPromoteToSaveComplete(FPackageData& PackageData);
 	/** Queue a message to the server that a package was discovered as needed in the cook. Will be sent during Tick. */
 	void ReportDiscoveredPackage(const FPackageData& PackageData, const FInstigator& Instigator,
-		FDiscoveredPlatformSet&& ReachablePlatforms, FGenerationHelper* ParentGenerationHelper);
+		FDiscoveredPlatformSet&& ReachablePlatforms, FGenerationHelper* ParentGenerationHelper, EUrgency Urgency);
 	/**
 	 * Queue a message to the server that a PackageGenerator queued its generated packages for cooking, and will
 	 * keep itself in memory until it receives an EGeneratorEvent::QueuedGeneratedPackagesFencePassed.
diff --git a/Engine/Source/Editor/UnrealEd/Private/Cooker/CookWorkerServer.cpp b/Engine/Source/Editor/UnrealEd/Private/Cooker/CookWorkerServer.cpp
--- a/Engine/Source/Editor/UnrealEd/Private/Cooker/CookWorkerServer.cpp
+++ b/Engine/Source/Editor/UnrealEd/Private/Cooker/CookWorkerServer.cpp
@@ -692,6 +692,7 @@ void FCookWorkerServer::SendPendingPackages()
 		AssignData.ParentGenerator = PackageData->GetParentGenerator();
 		AssignData.DoesGeneratedRequireGenerator = PackageData->DoesGeneratedRequireGenerator();
 		AssignData.Instigator = PackageData->GetInstigator();
+		AssignData.Urgency = PackageData->GetUrgency();
 		SessionPlatformNeedsCook.Init(false, OrderedSessionPlatforms.Num());
 		int32 PlatformIndex = 0;
 		for (const ITargetPlatform* SessionPlatform : OrderedSessionPlatforms)
@@ -1093,7 +1094,7 @@ void FCookWorkerServer::QueueDiscoveredPackage(FDiscoveredPackageReplication&& D
 	Director.ResetFinalIdleHeartbeatFence();
 	Platforms.ConvertFromBitfield(OrderedSessionAndSpecialPlatforms);
 	COTFS.QueueDiscoveredPackageOnDirector(PackageData, MoveTemp(Instigator), MoveTemp(Platforms),
-		false /* bUrgent */);
+		DiscoveredPackage.Urgency);
 }
 
 void FCookWorkerServer::HandleGeneratorMessage(FGeneratorEventMessage& GeneratorMessage)
@@ -1220,6 +1221,8 @@ void FAssignPackageData::Write(FCbWriter& Writer,
 	Writer << ConstructData;
 	Writer << ParentGenerator;
 	Writer << Instigator;
+	Writer << static_cast<uint8>(Urgency);
+	static_assert(sizeof(EUrgency) <= sizeof(uint8), "We are storing it in a uint8");
 	WriteToCompactBinary(Writer, NeedCookPlatforms, OrderedSessionPlatforms);
 	{
 		Writer.BeginArray();
@@ -1245,6 +1248,15 @@ bool FAssignPackageData::TryRead(FCbFieldView Field, TConstArrayView<const ITarg
 	bOk = LoadFromCompactBinary(*It++, ConstructData) & bOk;
 	bOk = LoadFromCompactBinary(*It++, ParentGenerator) & bOk;
 	bOk = LoadFromCompactBinary(*It++, Instigator) & bOk;
+	uint8 UrgencyInt = It->AsUInt8();
+	if (!(It++)->HasError() && UrgencyInt < static_cast<uint8>(EUrgency::Count))
+	{
+		Urgency = static_cast<EUrgency>(UrgencyInt);
+	}
+	else
+	{
+		bOk = false;
+	}
 	bOk = LoadFromCompactBinary(*It++, NeedCookPlatforms, OrderedSessionPlatforms) & bOk;
 	{
 		FCbFieldView ArrayFieldView = *It++;
@@ -1475,13 +1487,15 @@ void FDiscoveredPackageReplication::Write(FCbWriter& Writer,
 	Writer << static_cast<uint8>(Instigator.Category);
 	Writer << Instigator.Referencer;
 	Writer << static_cast<uint8>(DoesGeneratedRequireGenerator);
+	static_assert(sizeof(ICookPackageSplitter::EGeneratedRequiresGenerator) <= sizeof(uint8), "We are storing it in a uint8");
+	Writer << static_cast<uint8>(Urgency);
+	static_assert(sizeof(EUrgency) <= sizeof(uint8), "We are storing it in a uint8");
 	bool bGeneratedPackageHash = !GeneratedPackageHash.IsZero();
 	Writer << bGeneratedPackageHash;
 	if (bGeneratedPackageHash)
 	{
 		Writer << GeneratedPackageHash;
 	}
-	static_assert(sizeof(ICookPackageSplitter::EGeneratedRequiresGenerator) <= sizeof(uint8), "We are storing it in a uint8");
 	WriteToCompactBinary(Writer, Platforms, OrderedSessionAndSpecialPlatforms);
 	Writer.EndArray();
 }
@@ -1522,6 +1536,15 @@ bool FDiscoveredPackageReplication::TryRead(FCbFieldView Field,
 	{
 		bOk = false;
 	}
+	uint8 UrgencyInt = Iter->AsUInt8();
+	if (!(Iter++)->HasError() && UrgencyInt < static_cast<uint8>(EUrgency::Count))
+	{
+		Urgency = static_cast<EUrgency>(UrgencyInt);
+	}
+	else
+	{
+		bOk = false;
+	}
 	bool bGeneratedPackageHash = false;
 	bOk = LoadFromCompactBinary(Iter++, bGeneratedPackageHash) & bOk;
 	if (bGeneratedPackageHash)
diff --git a/Engine/Source/Editor/UnrealEd/Private/Cooker/CookWorkerServer.h b/Engine/Source/Editor/UnrealEd/Private/Cooker/CookWorkerServer.h
--- a/Engine/Source/Editor/UnrealEd/Private/Cooker/CookWorkerServer.h
+++ b/Engine/Source/Editor/UnrealEd/Private/Cooker/CookWorkerServer.h
@@ -254,6 +254,7 @@ struct FAssignPackageData
 	FConstructPackageData ConstructData;
 	FName ParentGenerator;
 	FInstigator Instigator;
+	EUrgency Urgency;
 	FDiscoveredPlatformSet NeedCookPlatforms;
 	TMap<FName, FAssetPackageData> GeneratorPreviousGeneratedPackages;
 	TArray<UE::CompactBinaryTCP::FMarshalledMessage> PerPackageCollectorMessages;
@@ -414,6 +415,7 @@ struct FDiscoveredPackageReplication
 	FDiscoveredPlatformSet Platforms;
 	ICookPackageSplitter::EGeneratedRequiresGenerator DoesGeneratedRequireGenerator =
 		ICookPackageSplitter::EGeneratedRequiresGenerator::None;
+	EUrgency Urgency = EUrgency::Normal;
 
 private:
 	void Write(FCbWriter& Writer, TConstArrayView<const ITargetPlatform*> OrderedSessionAndSpecialPlatforms) const;
diff --git a/Engine/Source/Editor/UnrealEd/Private/Cooker/IWorkerRequests.h b/Engine/Source/Editor/UnrealEd/Private/Cooker/IWorkerRequests.h
--- a/Engine/Source/Editor/UnrealEd/Private/Cooker/IWorkerRequests.h
+++ b/Engine/Source/Editor/UnrealEd/Private/Cooker/IWorkerRequests.h
@@ -62,7 +62,7 @@ public:
 
 	// Writing Packages during Cook
 	virtual void QueueDiscoveredPackage(UCookOnTheFlyServer& COTFS, FPackageData& PackageData,
-		FInstigator&& Instigator, FDiscoveredPlatformSet&& ReachablePlatforms, bool bUrgent,
+		FInstigator&& Instigator, FDiscoveredPlatformSet&& ReachablePlatforms, EUrgency Urgency,
 		FGenerationHelper* ParentGenerationHelper) = 0;
 	virtual void EndQueueGeneratedPackages(UCookOnTheFlyServer& COTFS, FGenerationHelper& GenerationHelper) = 0;
 
diff --git a/Engine/Source/Editor/UnrealEd/Private/Cooker/WorkerRequestsLocal.cpp b/Engine/Source/Editor/UnrealEd/Private/Cooker/WorkerRequestsLocal.cpp
--- a/Engine/Source/Editor/UnrealEd/Private/Cooker/WorkerRequestsLocal.cpp
+++ b/Engine/Source/Editor/UnrealEd/Private/Cooker/WorkerRequestsLocal.cpp
@@ -46,13 +46,13 @@ void FWorkerRequestsLocal::DequeueAllExternal(TArray<FSchedulerCallback>& OutCal
 }
 
 void FWorkerRequestsLocal::QueueDiscoveredPackage(UCookOnTheFlyServer& COTFS, FPackageData& PackageData,
-	FInstigator&& Instigator, FDiscoveredPlatformSet&& ReachablePlatforms, bool bUrgent,
+	FInstigator&& Instigator, FDiscoveredPlatformSet&& ReachablePlatforms, EUrgency Urgency,
 	FGenerationHelper* ParentGenerationHelper)
 {
 	// FWorkerRequestsRemote needs to send ParentGenerationHelper data to the director, but in the local case we
 	// already consumed the data in FGenerationHelper::StartQueueGeneratedPackages, so we don't use the value here.
 	(void)ParentGenerationHelper;
-	COTFS.QueueDiscoveredPackageOnDirector(PackageData, MoveTemp(Instigator), MoveTemp(ReachablePlatforms), bUrgent);
+	COTFS.QueueDiscoveredPackageOnDirector(PackageData, MoveTemp(Instigator), MoveTemp(ReachablePlatforms), Urgency);
 }
 
 void FWorkerRequestsLocal::EndQueueGeneratedPackages(UCookOnTheFlyServer& COTFS,
diff --git a/Engine/Source/Editor/UnrealEd/Private/Cooker/WorkerRequestsLocal.h b/Engine/Source/Editor/UnrealEd/Private/Cooker/WorkerRequestsLocal.h
--- a/Engine/Source/Editor/UnrealEd/Private/Cooker/WorkerRequestsLocal.h
+++ b/Engine/Source/Editor/UnrealEd/Private/Cooker/WorkerRequestsLocal.h
@@ -33,7 +33,7 @@ public:
 	virtual void RemapTargetPlatforms(const TMap<ITargetPlatform*, ITargetPlatform*>& Remap) override;
 	virtual void OnRemoveSessionPlatform(const ITargetPlatform* TargetPlatform) override;
 	virtual void QueueDiscoveredPackage(UCookOnTheFlyServer& COTFS, FPackageData& PackageData,
-		FInstigator&& Instigator, FDiscoveredPlatformSet&& ReachablePlatforms, bool bUrgent,
+		FInstigator&& Instigator, FDiscoveredPlatformSet&& ReachablePlatforms, EUrgency Urgency,
 		FGenerationHelper* ParentGenerationHelper) override;
 	virtual void EndQueueGeneratedPackages(UCookOnTheFlyServer& COTFS, FGenerationHelper& GenerationHelper) override;
 	virtual void AddStartCookByTheBookRequest(FFilePlatformRequest&& Request) override;
diff --git a/Engine/Source/Editor/UnrealEd/Private/Cooker/WorkerRequestsRemote.cpp b/Engine/Source/Editor/UnrealEd/Private/Cooker/WorkerRequestsRemote.cpp
--- a/Engine/Source/Editor/UnrealEd/Private/Cooker/WorkerRequestsRemote.cpp
+++ b/Engine/Source/Editor/UnrealEd/Private/Cooker/WorkerRequestsRemote.cpp
@@ -55,12 +55,11 @@ void FWorkerRequestsRemote::DequeueAllExternal(TArray<FSchedulerCallback>& OutCa
 }
 
 void FWorkerRequestsRemote::QueueDiscoveredPackage(UCookOnTheFlyServer& COTFS, FPackageData& PackageData,
-	FInstigator&& Instigator, FDiscoveredPlatformSet&& ReachablePlatforms, bool bUrgent,
+	FInstigator&& Instigator, FDiscoveredPlatformSet&& ReachablePlatforms, EUrgency Urgency,
 	FGenerationHelper* ParentGenerationHelper)
 {
-	(void)bUrgent; // Tracking urgency on CookWorkers is not supported
 	CookWorkerClient.ReportDiscoveredPackage(PackageData, MoveTemp(Instigator), MoveTemp(ReachablePlatforms),
-		ParentGenerationHelper);
+		ParentGenerationHelper, Urgency);
 }
 
 void FWorkerRequestsRemote::EndQueueGeneratedPackages(UCookOnTheFlyServer& COTFS,
diff --git a/Engine/Source/Editor/UnrealEd/Private/Cooker/WorkerRequestsRemote.h b/Engine/Source/Editor/UnrealEd/Private/Cooker/WorkerRequestsRemote.h
--- a/Engine/Source/Editor/UnrealEd/Private/Cooker/WorkerRequestsRemote.h
+++ b/Engine/Source/Editor/UnrealEd/Private/Cooker/WorkerRequestsRemote.h
@@ -41,7 +41,7 @@ public:
 	virtual void RemapTargetPlatforms(const TMap<ITargetPlatform*, ITargetPlatform*>& Remap) override;
 	virtual void OnRemoveSessionPlatform(const ITargetPlatform* TargetPlatform) override;
 	virtual void QueueDiscoveredPackage(UCookOnTheFlyServer& COTFS, FPackageData& PackageData,
-		FInstigator&& Instigator, FDiscoveredPlatformSet&& ReachablePlatforms, bool bUrgent,
+		FInstigator&& Instigator, FDiscoveredPlatformSet&& ReachablePlatforms, EUrgency Urgency,
 		FGenerationHelper* ParentGenerationHelper) override;
 	virtual void EndQueueGeneratedPackages(UCookOnTheFlyServer& COTFS, FGenerationHelper& GenerationHelper) override;
 	virtual void AddStartCookByTheBookRequest(FFilePlatformRequest&& Request) override;
