Cooker: CookPackageSplitter: Add more feedback for the error "PackageSplitter object was deleted by garbage collection while generation was still ongoing...."
#rnx

diff --git a/Engine/Source/Editor/UnrealEd/Private/Cooker/CookGenerationHelper.cpp b/Engine/Source/Editor/UnrealEd/Private/Cooker/CookGenerationHelper.cpp
--- a/Engine/Source/Editor/UnrealEd/Private/Cooker/CookGenerationHelper.cpp
+++ b/Engine/Source/Editor/UnrealEd/Private/Cooker/CookGenerationHelper.cpp
@@ -1675,13 +1675,32 @@ void FGenerationHelper::PostGarbageCollect(const TRefCountPtr<FGenerationHelper>
 	FPackageData& Owner = GetOwner();
 	if (Owner.IsInStateProperty(EPackageStateProperty::Saving))
 	{
-		// UCookOnTheFlyServer::PreGarbageCollect adds references for the Generator package and all its public
-		// objects, so it should still be loaded
+		// If the package no longer exists, then UpdateSaveAfterGarbageCollect earlier in
+		// UCookOnTheFlyServer::PostGarbageCollect should have demoted the package out of saving.
+		// And if the package exists, then the SplitDataObject, which should be a public object within it, should
+		// have been kept in memory by ConstructSoftGCPackageToObjectList. If the package or split object no longer
+		// exist then we are in an invalid state and the savepackage might behave incorrectly.
+		if (!Owner.GetPackage())
+		{
+			UPackage* FoundPackage = FindObject<UPackage>(nullptr, *Owner.GetPackageName().ToString());
+			if (FoundPackage)
+			{
+				Owner.SetPackage(FoundPackage);
+				UE_LOG(LogCook, Warning,
+					TEXT("CookPackageSplitter's package pointer was unexpectedly set to null by garbage collection while the package is still in the %s state, %s substate, but the package is still in memory.")
+					TEXT("\n\tPackage=%s, Splitter=%s."),
+					LexToString(Owner.GetState()), LexToString(Owner.GetSaveSubState()),
+					*Owner.GetPackageName().ToString(), *GetSplitDataObjectName().ToString());
+			}
+		}
 		if (!Owner.GetPackage() || !GetWeakSplitDataObject())
 		{
 			UE_LOG(LogCook, Error,
-				TEXT("PackageSplitter object was deleted by garbage collection while generation was still ongoing. This will break the generation.")
-				TEXT("\n\tSplitter=%s."), *GetSplitDataObjectName().ToString());
+				TEXT("CookPackageSplitter's %s was deleted by garbage collection while the package is still in the %s state, %s substate. This will break the generation.")
+				TEXT("\n\tPackage=%s, Splitter=%s."),
+				(!Owner.GetPackage() ? TEXT("package") : TEXT("splitter object")),
+				LexToString(Owner.GetState()), LexToString(Owner.GetSaveSubState()),
+				*Owner.GetPackageName().ToString(), *GetSplitDataObjectName().ToString());
 		}
 	}
 	else if (!IsUseInternalReferenceToAvoidGarbageCollect())
