Cooker: CookPackageSplitter: Fix assertion that the ConstructSoftGCPackageToObjectList buffer has already been created when gathering diagnostics in PostGarbageCollect; preserve it if it has already been created.
#rnx

diff --git a/Engine/Source/Editor/UnrealEd/Classes/CookOnTheSide/CookOnTheFlyServer.h b/Engine/Source/Editor/UnrealEd/Classes/CookOnTheSide/CookOnTheFlyServer.h
--- a/Engine/Source/Editor/UnrealEd/Classes/CookOnTheSide/CookOnTheFlyServer.h
+++ b/Engine/Source/Editor/UnrealEd/Classes/CookOnTheSide/CookOnTheFlyServer.h
@@ -1455,6 +1455,7 @@ private:
 
 		UCookOnTheFlyServer& COTFS;
 		TGuardValue<bool> SoftGCGuard;
+		bool bNeedsConstructBuffer;
 	};
 	/** Callback for FGenericCrashContext; provides the current ActivePackage as context. */
 	void DumpCrashContext(FCrashContextExtendedWriter& Writer);
diff --git a/Engine/Source/Editor/UnrealEd/Private/CookOnTheFlyServer.cpp b/Engine/Source/Editor/UnrealEd/Private/CookOnTheFlyServer.cpp
--- a/Engine/Source/Editor/UnrealEd/Private/CookOnTheFlyServer.cpp
+++ b/Engine/Source/Editor/UnrealEd/Private/CookOnTheFlyServer.cpp
@@ -5330,17 +5330,23 @@ static void ConstructSoftGCPackageToObjectList(TArray<UObject*>& PackageToObject
 }
 
 UCookOnTheFlyServer::FScopeFindCookReferences::FScopeFindCookReferences(UCookOnTheFlyServer& InCOTFS)
-	:COTFS(InCOTFS),
-	SoftGCGuard(UPackage::bSupportCookerSoftGC, true)
+	: COTFS(InCOTFS)
+	, SoftGCGuard(UPackage::bSupportCookerSoftGC, true)
+	, bNeedsConstructBuffer(COTFS.SoftGCPackageToObjectListBuffer.IsEmpty())
 {
-	check(COTFS.SoftGCPackageToObjectListBuffer.IsEmpty())
-	ConstructSoftGCPackageToObjectList(COTFS.SoftGCPackageToObjectListBuffer);
+	if (bNeedsConstructBuffer)
+	{
+		ConstructSoftGCPackageToObjectList(COTFS.SoftGCPackageToObjectListBuffer);
+	}
 }
 
 UCookOnTheFlyServer::FScopeFindCookReferences::~FScopeFindCookReferences()
 {
-	UPackage::SoftGCPackageToObjectList.Empty();
-	COTFS.SoftGCPackageToObjectListBuffer.Empty();
+	if (bNeedsConstructBuffer)
+	{
+		UPackage::SoftGCPackageToObjectList.Empty();
+		COTFS.SoftGCPackageToObjectListBuffer.Empty();
+	}
 }
 
 void UCookOnTheFlyServer::PreGarbageCollect()
