IncrementalCook: Fix spurious warning for CookPackageSplitters that were previously cooked but that were not cooked in the current session.
#rnx

diff --git a/Engine/Source/Editor/UnrealEd/Private/CookOnTheFlyServer.cpp b/Engine/Source/Editor/UnrealEd/Private/CookOnTheFlyServer.cpp
--- a/Engine/Source/Editor/UnrealEd/Private/CookOnTheFlyServer.cpp
+++ b/Engine/Source/Editor/UnrealEd/Private/CookOnTheFlyServer.cpp
@@ -10554,7 +10554,14 @@ void UCookOnTheFlyServer::CookByTheBookFinishedInternal()
 			TRefCountPtr<FGenerationHelper> GenerationHelper = PackageData->GetGenerationHelper();
 			if (GenerationHelper)
 			{
-				DanglingGenerationHelpers.Add(PackageData);
+				// One reason it might still be around is the keepforiterativeflag, if it were in the oplog but never
+				// cooked. Clear that flag now and then retest whether it is still referenced.
+				GenerationHelper->ClearKeepForIterative();
+				GenerationHelper.SafeRelease();
+				if (PackageData->GetGenerationHelper())
+				{
+					DanglingGenerationHelpers.Add(PackageData);
+				}
 			}
 		});
 	for (FPackageData* PackageData : DanglingGenerationHelpers)
diff --git a/Engine/Source/Editor/UnrealEd/Private/Cooker/CookGenerationHelper.cpp b/Engine/Source/Editor/UnrealEd/Private/Cooker/CookGenerationHelper.cpp
--- a/Engine/Source/Editor/UnrealEd/Private/Cooker/CookGenerationHelper.cpp
+++ b/Engine/Source/Editor/UnrealEd/Private/Cooker/CookGenerationHelper.cpp
@@ -303,7 +303,8 @@ void FGenerationHelper::DiagnoseWhyNotShutdown()
 			}
 		};
 	TestInfo(GetOwnerInfo());
-	for (FCookGenerationInfo& Info : GetPackagesToGenerate())
+	// Do not call GetPackagesToGenerate as that would initialize.
+	for (FCookGenerationInfo& Info : PackagesToGenerate)
 	{
 		TestInfo(Info);
 	}
