Fix a potential shutdown crash in the bulkdata registry.

#rnx

- When FBulkDataRegistryImpl::Teardown is run we were cancelling FUpdatePayloadWorker by signalling to them that they should abandon their work, but we did not make sure that all of the jobs completed.
- In most cases this was fine but FUpdatePayloadWorker::DoWork will attempt to load editor bulkdata payloads first, before checking for cancellation and even if we did check first, it would be possible for currently running jobs to be mid loading of the payload when Teardown is run. In either of these scenarios we could potentially end up attempting to access the asset registry which would be cleaned up shortly after Teardown is called and end up asserting.
- To fix this we now check if the work should be abandoned before we start trying to read the bulkdata payload, then record how many jobs are currently reading payloads.
- In ::Teardown we can then check how many jobs are currently reading payloads and wait until that value reaches zero before returning.
- Note that this issue will only really show up on machines with very low worker core limits, when the editor is closed very shortly after loading and the project being loaded has manypackages in older pre 5.0 formats.

diff --git a/Engine/Source/Editor/UnrealEd/Private/EditorDomain/BulkDataRegistryImpl.cpp b/Engine/Source/Editor/UnrealEd/Private/EditorDomain/BulkDataRegistryImpl.cpp
--- a/Engine/Source/Editor/UnrealEd/Private/EditorDomain/BulkDataRegistryImpl.cpp
+++ b/Engine/Source/Editor/UnrealEd/Private/EditorDomain/BulkDataRegistryImpl.cpp
@@ -145,6 +145,21 @@ void FBulkDataRegistryImpl::Teardown()
 	{
 		PendingPayloadIdPair.Value->Cancel();
 	}
+
+	// Finally wait for any FUpdatePayloadWorker jobs that are still in flight
+	if (SharedDataLock->UpdatePayloadWorkerInFlight > 0)
+	{
+		const double StartTime = FPlatformTime::Seconds();
+
+		UE_LOG(LogBulkDataRegistry, Log, TEXT("Waiting on %d FUpdatePayloadWorker jobs to complete..."), SharedDataLock->UpdatePayloadWorkerInFlight.load());
+		
+		while (SharedDataLock->UpdatePayloadWorkerInFlight > 0)
+		{
+			FPlatformProcess::SleepNoStats(0.0f);
+		}
+
+		UE_LOG(LogBulkDataRegistry, Log, TEXT("All FUpdatePayloadWorker jobs finished, exit was stalled for %.1f(s)"), FPlatformTime::Seconds() - StartTime);
+	}
 }
 
 void FBulkDataRegistryImpl::OnEndLoadPackage(const FEndLoadPackageContext& Context)
@@ -974,6 +989,17 @@ void FUpdatePayloadWorker::DoWork()
 	bool bValid = true;
 	for (;;)
 	{
+		{
+			FReadScopeLock SharedDataScopeLock(SharedDataLock->ActiveLock);
+			if (!SharedDataLock->bActive)
+			{
+				// The BulkDataRegistry has destructed. Our list of requesters is on the BulkDataRegistry, so there's nothing we can do except exit
+				return;
+			}
+
+			SharedDataLock->UpdatePayloadWorkerInFlight++;
+		}
+
 		BulkData.UpdatePayloadId();
 		Buffer = BulkData.GetCompressedPayload().Get();
 
@@ -981,6 +1007,8 @@ void FUpdatePayloadWorker::DoWork()
 		FName WritePayloadIdPackageName;
 		{
 			FReadScopeLock SharedDataScopeLock(SharedDataLock->ActiveLock);
+			--SharedDataLock->UpdatePayloadWorkerInFlight;
+
 			if (!SharedDataLock->bActive)
 			{
 				// The BulkDataRegistry has destructed. Our list of requesters is on the BulkDataRegistry, so there's nothing we can do except exit
diff --git a/Engine/Source/Editor/UnrealEd/Private/EditorDomain/BulkDataRegistryImpl.h b/Engine/Source/Editor/UnrealEd/Private/EditorDomain/BulkDataRegistryImpl.h
--- a/Engine/Source/Editor/UnrealEd/Private/EditorDomain/BulkDataRegistryImpl.h
+++ b/Engine/Source/Editor/UnrealEd/Private/EditorDomain/BulkDataRegistryImpl.h
@@ -144,6 +144,8 @@ class FTaskSharedDataLock : public FThreadSafeRefCountedObject
 public:
 	FRWLock ActiveLock;
 	bool bActive = true;
+
+	std::atomic<uint32> UpdatePayloadWorkerInFlight = 0;
 };
 
 /** A worker that updates the PayloadId for a BulkData that is missing its RawHash. */
