World Partition - Spline Mesh Builder: Builder should only process spatially loaded splines
* Editor only components are now excluded from HLOD generation
* Moved common conditions a component must fulfill to be considered HLOD relevant into protected utility method UActorComponent::CanBeHLODRelevant()

diff --git a/Engine/Plugins/Experimental/Water/Source/Runtime/Private/WaterBodyComponent.cpp b/Engine/Plugins/Experimental/Water/Source/Runtime/Private/WaterBodyComponent.cpp
--- a/Engine/Plugins/Experimental/Water/Source/Runtime/Private/WaterBodyComponent.cpp
+++ b/Engine/Plugins/Experimental/Water/Source/Runtime/Private/WaterBodyComponent.cpp
@@ -126,7 +126,7 @@ UWaterBodyComponent::UWaterBodyComponent(const FObjectInitializer& ObjectInitial
 
 bool UWaterBodyComponent::IsHLODRelevant() const
 {
-	if (HasAnyFlags(RF_Transient))
+	if (!CanBeHLODRelevant(this))
 	{
 		return false;
 	}
diff --git a/Engine/Source/Editor/UnrealEd/Private/WorldPartition/WorldPartitionLandscapeSplineMeshesBuilder.cpp b/Engine/Source/Editor/UnrealEd/Private/WorldPartition/WorldPartitionLandscapeSplineMeshesBuilder.cpp
--- a/Engine/Source/Editor/UnrealEd/Private/WorldPartition/WorldPartitionLandscapeSplineMeshesBuilder.cpp
+++ b/Engine/Source/Editor/UnrealEd/Private/WorldPartition/WorldPartitionLandscapeSplineMeshesBuilder.cpp
@@ -39,16 +39,11 @@ DEFINE_LOG_CATEGORY_STATIC(LogWorldPartitionLandscapeSplineMeshBuilder, All, All
 
 static const FName CreatedFromBuilderTag(TEXT("CreatedFromWorldPartitionLandscapeSplineMeshBuilder"));
 
-UStaticMesh* UWorldPartitionLandscapeSplineMeshesBuilder::SplineEditorMesh = nullptr;
-
 UWorldPartitionLandscapeSplineMeshesBuilder::UWorldPartitionLandscapeSplineMeshesBuilder(const FObjectInitializer& ObjectInitializer)
 	: Super(ObjectInitializer)
 	, NewGridSize(0)
 {
 	GetParamValue("NewGridSize=", NewGridSize);
-
-	static ConstructorHelpers::FObjectFinder<UStaticMesh> ConstructorStatics(TEXT("/Engine/EditorLandscapeResources/SplineEditorMesh"));
-	SplineEditorMesh = ConstructorStatics.Object;
 }
 
 FName UWorldPartitionLandscapeSplineMeshesBuilder::GetSplineCollisionProfileName(const UStaticMeshComponent* InMeshComponent)
@@ -312,29 +307,20 @@ ALandscapeSplineMeshesActor* UWorldPartitionLandscapeSplineMeshesBuilder::GetOrC
 
 void UWorldPartitionLandscapeSplineMeshesBuilder::FilterStaticMeshComponents(TArray<UStaticMeshComponent*>& InOutComponents)
 {
-	check(SplineEditorMesh);
-
-	// Remove invalid components from the list
-	InOutComponents.RemoveAll([](UStaticMeshComponent* Component)
+	// Remove unwanted components from the list
+	InOutComponents.RemoveAllSwap([](UStaticMeshComponent* Component)
 	{
-		if (!Component || !Component->GetStaticMesh())
+		// We only care about SplineMeshComponents & ControlPointMeshComponents
+		if (USplineMeshComponent* SplineMeshComponent = Cast<USplineMeshComponent>(Component))
 		{
-			return true;
+			return !SplineMeshComponent->IsRelevantForSplinePartitioning();
 		}
-
-		// Remove editor splines
-		if (Component->GetStaticMesh() == SplineEditorMesh)
-		{
-			return true;
-		}
-
-		// We only care for SplineMeshComponents & ControlPointMeshComponents
-		if (!Component->IsA<USplineMeshComponent>() && !Component->IsA<UControlPointMeshComponent>())
+		else if (UControlPointMeshComponent* ControlPointMeshComponent = Cast<UControlPointMeshComponent>(Component))
 		{
-			return true;
+			return !ControlPointMeshComponent->IsRelevantForSplinePartitioning();
 		}
 
-		return false;
+		return true;
 	});
 }
 
@@ -411,6 +397,12 @@ bool UWorldPartitionLandscapeSplineMeshesBuilder::RunInternal(UWorld* InWorld, c
 			{
 				continue;
 			}
+
+			// Only partition splines if they are spatially loaded & not editor only
+			if (SplineActor->ShouldPartitionSpline())
+			{
+				continue;
+			}
 						
 			TArray<UStaticMeshComponent*> Components;
 			SplineActor->GetComponents(Components);
diff --git a/Engine/Source/Editor/UnrealEd/Public/WorldPartition/WorldPartitionLandscapeSplineMeshesBuilder.h b/Engine/Source/Editor/UnrealEd/Public/WorldPartition/WorldPartitionLandscapeSplineMeshesBuilder.h
--- a/Engine/Source/Editor/UnrealEd/Public/WorldPartition/WorldPartitionLandscapeSplineMeshesBuilder.h
+++ b/Engine/Source/Editor/UnrealEd/Public/WorldPartition/WorldPartitionLandscapeSplineMeshesBuilder.h
@@ -36,6 +36,4 @@ private:
 	UNREALED_API ALandscapeSplineMeshesActor* GetOrCreatePartitionActorForComponent(class UWorld* InWorld, const UStaticMeshComponent* InMeshComponent, const FGuid& InLandscapeGuid);
 
 	uint32 NewGridSize;
-	
-	static UNREALED_API class UStaticMesh* SplineEditorMesh;
 };
diff --git a/Engine/Source/Runtime/Engine/Classes/Components/ActorComponent.h b/Engine/Source/Runtime/Engine/Classes/Components/ActorComponent.h
--- a/Engine/Source/Runtime/Engine/Classes/Components/ActorComponent.h
+++ b/Engine/Source/Runtime/Engine/Classes/Components/ActorComponent.h
@@ -1278,6 +1278,9 @@ public:
 	static ENGINE_API FOnMarkRenderStateDirty MarkRenderStateDirtyEvent;
 
 protected:
+	/** Utility method used by subclasses to test basic conditions a component must fulfill to be considered HLOD relevant. */
+	static ENGINE_API bool CanBeHLODRelevant(const UActorComponent* ActorComponent);
+
 	/** Makes sure navigation system has up to date information regarding component's navigation relevancy 
 	 *	and if it can affect navigation at all 
 	 *	@param bForceUpdate by default updating navigation system will take place only if the component has
diff --git a/Engine/Source/Runtime/Engine/Classes/Components/SplineMeshComponent.h b/Engine/Source/Runtime/Engine/Classes/Components/SplineMeshComponent.h
--- a/Engine/Source/Runtime/Engine/Classes/Components/SplineMeshComponent.h
+++ b/Engine/Source/Runtime/Engine/Classes/Components/SplineMeshComponent.h
@@ -184,6 +184,7 @@ public:
 	//Begin UObject Interface
 	ENGINE_API virtual void Serialize(FArchive& Ar) override;
 #if WITH_EDITOR
+	ENGINE_API bool IsRelevantForSplinePartitioning() const;
 	ENGINE_API virtual bool IsEditorOnly() const override;
 	ENGINE_API virtual bool Modify(bool bAlwaysMarkDirty = true) override;
 	ENGINE_API virtual void PostEditChangeProperty(FPropertyChangedEvent& PropertyChangedEvent) override;
diff --git a/Engine/Source/Runtime/Engine/Private/Components/ActorComponent.cpp b/Engine/Source/Runtime/Engine/Private/Components/ActorComponent.cpp
--- a/Engine/Source/Runtime/Engine/Private/Components/ActorComponent.cpp
+++ b/Engine/Source/Runtime/Engine/Private/Components/ActorComponent.cpp
@@ -2860,5 +2860,36 @@ bool UActorComponent::NeedsInitialization() const
 TFunction<bool(const UActorComponent*)> GIsComponentSelectedInEditor;
 #endif
 
+bool UActorComponent::CanBeHLODRelevant(const UActorComponent* ActorComponent)
+{
+	if (!IsValid(ActorComponent))
+	{
+		return false;
+	}
+
+	if (ActorComponent->HasAnyFlags(RF_Transient)
+#if WITH_EDITOR 
+		&& !(ActorComponent->GetOwner() && ActorComponent->GetOwner()->IsInLevelInstance())			// Treat components in LI as HLOD relevant for the sake of visualisation modes
+#endif
+		)
+	{
+		return false;
+	}
+
+#if WITH_EDITORONLY_DATA
+	if (ActorComponent->IsEditorOnly())
+	{
+		return false;
+	}
+
+	if (ActorComponent->IsVisualizationComponent())
+	{
+		return false;
+	}
+#endif
+
+	return true;
+}
+
 #undef LOCTEXT_NAMESPACE
 
diff --git a/Engine/Source/Runtime/Engine/Private/Components/ChildActorComponent.cpp b/Engine/Source/Runtime/Engine/Private/Components/ChildActorComponent.cpp
--- a/Engine/Source/Runtime/Engine/Private/Components/ChildActorComponent.cpp
+++ b/Engine/Source/Runtime/Engine/Private/Components/ChildActorComponent.cpp
@@ -977,6 +977,11 @@ void UChildActorComponent::BeginPlay()
 
 bool UChildActorComponent::IsHLODRelevant() const
 {
+	if (!CanBeHLODRelevant(this))
+	{
+		return false;
+	}
+
 	const bool bIsHLODRelevant = ChildActor && ChildActor->IsHLODRelevant();
 	return bIsHLODRelevant;
 }
diff --git a/Engine/Source/Runtime/Engine/Private/Components/SplineMeshComponent.cpp b/Engine/Source/Runtime/Engine/Private/Components/SplineMeshComponent.cpp
--- a/Engine/Source/Runtime/Engine/Private/Components/SplineMeshComponent.cpp
+++ b/Engine/Source/Runtime/Engine/Private/Components/SplineMeshComponent.cpp
@@ -30,6 +30,7 @@
 #include "IHierarchicalLODUtilities.h"
 #include "HierarchicalLODUtilitiesModule.h"
 #include "LandscapeSplineActor.h"
+#include "LandscapeSplinesComponent.h"
 #endif // WITH_EDITOR
 
 namespace UE::SplineMesh
@@ -806,14 +807,37 @@ void USplineMeshComponent::Serialize(FArchive& Ar)
 }
 
 #if WITH_EDITOR
-bool USplineMeshComponent::IsEditorOnly() const
+bool USplineMeshComponent::IsRelevantForSplinePartitioning() const
 {
+	// Ignore editor only splines
+	// Avoid using USplineMeshComponent::IsEditorOnly() as that would exclude all spline mesh components that are partitioned already
 	if (Super::IsEditorOnly())
+	{
+		return false;
+	}
+
+	if (GetStaticMesh() == nullptr)
+	{
+		return false;
+	}
+
+	// Ignore editor splines
+	if (GetStaticMesh() == GetDefault<ULandscapeSplinesComponent>()->SplineEditorMesh)
+	{
+		return false;
+	}
+
+	return true;
+}
+
+bool USplineMeshComponent::IsEditorOnly() const
+{
+	if (!IsRelevantForSplinePartitioning())
 	{
 		return true;
 	}
 
-	// If Landscape uses generated LandscapeSplineMeshesActors, SplineMeshComponents is removed from cooked build  
+	// If landscape uses generated LandscapeSplineMeshesActors, SplineMeshComponents is removed from cooked build  
 	const ALandscapeSplineActor* SplineActor = Cast<ALandscapeSplineActor>(GetOwner());
 	if (SplineActor && SplineActor->HasGeneratedLandscapeSplineMeshesActors())
 	{
diff --git a/Engine/Source/Runtime/Engine/Private/Components/StaticMeshComponent.cpp b/Engine/Source/Runtime/Engine/Private/Components/StaticMeshComponent.cpp
--- a/Engine/Source/Runtime/Engine/Private/Components/StaticMeshComponent.cpp
+++ b/Engine/Source/Runtime/Engine/Private/Components/StaticMeshComponent.cpp
@@ -3068,11 +3068,7 @@ void UStaticMeshComponent::ApplyComponentInstanceData(FStaticMeshComponentInstan
 
 bool UStaticMeshComponent::IsHLODRelevant() const
 {
-	if (HasAnyFlags(RF_Transient)
-#if WITH_EDITOR 
-		&& !(GetOwner() && GetOwner()->IsInLevelInstance())			// Treat components in LI as HLOD relevant for the sake of visualisation modes
-#endif
-		)
+	if (!CanBeHLODRelevant(this))
 	{
 		return false;
 	}
@@ -3093,11 +3089,6 @@ bool UStaticMeshComponent::IsHLODRelevant() const
 	}
 
 #if WITH_EDITORONLY_DATA
-	if (IsVisualizationComponent())
-	{
-		return false;
-	}
-
 	if (!bEnableAutoLODGeneration)
 	{
 		return false;
diff --git a/Engine/Source/Runtime/Experimental/GeometryCollectionEngine/Private/GeometryCollection/GeometryCollectionComponent.cpp b/Engine/Source/Runtime/Experimental/GeometryCollectionEngine/Private/GeometryCollection/GeometryCollectionComponent.cpp
--- a/Engine/Source/Runtime/Experimental/GeometryCollectionEngine/Private/GeometryCollection/GeometryCollectionComponent.cpp
+++ b/Engine/Source/Runtime/Experimental/GeometryCollectionEngine/Private/GeometryCollection/GeometryCollectionComponent.cpp
@@ -7280,22 +7280,22 @@ const FTransform& UGeometryCollectionComponent::GetPreviousComponentToWorld() co
 
 bool UGeometryCollectionComponent::IsHLODRelevant() const
 {
-	if (RestCollection == nullptr || !IsCustomRendererAvailable())
+	if (!CanBeHLODRelevant(this))
 	{
 		return false;
 	}
 
-	if (!IsVisible())
+	if (RestCollection == nullptr || !IsCustomRendererAvailable())
 	{
 		return false;
 	}
 
-#if WITH_EDITORONLY_DATA
-	if (IsVisualizationComponent())
+	if (!IsVisible())
 	{
 		return false;
 	}
 
+#if WITH_EDITORONLY_DATA
 	if (!bEnableAutoLODGeneration)
 	{
 		return false;
diff --git a/Engine/Source/Runtime/Landscape/Classes/ControlPointMeshComponent.h b/Engine/Source/Runtime/Landscape/Classes/ControlPointMeshComponent.h
--- a/Engine/Source/Runtime/Landscape/Classes/ControlPointMeshComponent.h
+++ b/Engine/Source/Runtime/Landscape/Classes/ControlPointMeshComponent.h
@@ -34,4 +34,8 @@ class UControlPointMeshComponent : public UStaticMeshComponent
 #endif
 	virtual float GetVirtualTextureMainPassMaxDrawDistance() const override { return VirtualTextureMainPassMaxDrawDistance; }
 	//End UPrimitiveComponent Interface
+
+#if WITH_EDITOR
+	LANDSCAPE_API bool IsRelevantForSplinePartitioning() const;
+#endif
 };
diff --git a/Engine/Source/Runtime/Landscape/Private/Landscape.cpp b/Engine/Source/Runtime/Landscape/Private/Landscape.cpp
--- a/Engine/Source/Runtime/Landscape/Private/Landscape.cpp
+++ b/Engine/Source/Runtime/Landscape/Private/Landscape.cpp
@@ -1988,7 +1988,7 @@ void ULandscapeComponent::PropagateLightingScenarioChange()
 
 bool ULandscapeComponent::IsHLODRelevant() const
 {
-	if (HasAnyFlags(RF_Transient))
+	if (!CanBeHLODRelevant(this))
 	{
 		return false;
 	}
diff --git a/Engine/Source/Runtime/Landscape/Private/LandscapeSplineActor.cpp b/Engine/Source/Runtime/Landscape/Private/LandscapeSplineActor.cpp
--- a/Engine/Source/Runtime/Landscape/Private/LandscapeSplineActor.cpp
+++ b/Engine/Source/Runtime/Landscape/Private/LandscapeSplineActor.cpp
@@ -192,9 +192,14 @@ AActor* ALandscapeSplineActor::GetSceneOutlinerParent() const
 	return nullptr;
 }
 
+bool ALandscapeSplineActor::ShouldPartitionSpline() const
+{
+	return !IsEditorOnly() && bIsSpatiallyLoaded;
+}
+
 bool ALandscapeSplineActor::HasGeneratedLandscapeSplineMeshesActors() const
 {
-	return IsValid(LandscapeActor) && LandscapeActor->GetUseGeneratedLandscapeSplineMeshesActors();
+	return ShouldPartitionSpline() && IsValid(LandscapeActor) && LandscapeActor->GetUseGeneratedLandscapeSplineMeshesActors();
 }
 
 #endif
diff --git a/Engine/Source/Runtime/Landscape/Private/LandscapeSplines.cpp b/Engine/Source/Runtime/Landscape/Private/LandscapeSplines.cpp
--- a/Engine/Source/Runtime/Landscape/Private/LandscapeSplines.cpp
+++ b/Engine/Source/Runtime/Landscape/Private/LandscapeSplines.cpp
@@ -1711,14 +1711,37 @@ UControlPointMeshComponent::UControlPointMeshComponent(const FObjectInitializer&
 }
 
 #if WITH_EDITOR
-bool UControlPointMeshComponent::IsEditorOnly() const
+bool UControlPointMeshComponent::IsRelevantForSplinePartitioning() const
 {
+	// Ignore editor only splines
+	// Avoid using USplineMeshComponent::IsEditorOnly() as that would exclude all spline mesh components that are partitioned already
 	if (Super::IsEditorOnly())
+	{
+		return false;
+	}
+
+	if (GetStaticMesh() == nullptr)
+	{
+		return false;
+	}
+
+	// Ignore editor splines
+	if (GetStaticMesh() == GetDefault<ULandscapeSplinesComponent>()->SplineEditorMesh)
+	{
+		return false;
+	}
+
+	return true;
+}
+
+bool UControlPointMeshComponent::IsEditorOnly() const
+{
+	if (!IsRelevantForSplinePartitioning())
 	{
 		return true;
 	}
 
-	// If Landscape uses generated LandscapeSplineMeshesActors, ControlPointMeshComponents is removed from cooked build  
+	// If landscape uses generated LandscapeSplineMeshesActors, ControlPointMeshComponents is removed from cooked build  
 	ALandscapeSplineActor* SplineActor = Cast<ALandscapeSplineActor>(GetOwner());
 	if (SplineActor && SplineActor->HasGeneratedLandscapeSplineMeshesActors())
 	{
