Substrate - fixed lightmass material generation that got broken recently due to the way we access shading models from materials in legacy shading model node.

#rnx
#tests Add a test for CPU lightmass with Substrate (opaque, unlit, translucent, emissive, legacy and slabs with colored transmittance)

diff --git a/Engine/Source/Editor/UnrealEd/Private/Lightmass/Lightmass.cpp b/Engine/Source/Editor/UnrealEd/Private/Lightmass/Lightmass.cpp
--- a/Engine/Source/Editor/UnrealEd/Private/Lightmass/Lightmass.cpp
+++ b/Engine/Source/Editor/UnrealEd/Private/Lightmass/Lightmass.cpp
@@ -1443,7 +1443,7 @@ void FLightmassExporter::GetMaterialHash(const UMaterialInterface* Material, FSH
 
 	if (Lightmass_IsSubstrateEnabled())
 	{
-		uint32 LightmassSubstrateVersion = 0XB6A0D99F; // This can be change when the code/logic for converting material to Substrate for lightmap has changed.
+		uint32 LightmassSubstrateVersion = 0XFA1081D0; // This can be change when the code/logic for converting material to Substrate for lightmap has changed.
 		HashState.Update((const uint8*)&LightmassSubstrateVersion, sizeof(LightmassSubstrateVersion));
 	}
 
diff --git a/Engine/Source/Editor/UnrealEd/Private/Lightmass/LightmassRender.cpp b/Engine/Source/Editor/UnrealEd/Private/Lightmass/LightmassRender.cpp
--- a/Engine/Source/Editor/UnrealEd/Private/Lightmass/LightmassRender.cpp
+++ b/Engine/Source/Editor/UnrealEd/Private/Lightmass/LightmassRender.cpp
@@ -629,12 +629,17 @@ public:
 	virtual enum EBlendMode GetBlendMode() const override					{ return BLEND_Opaque; }
 	virtual enum ERefractionMode GetRefractionMode() const override			{ return Material ? (ERefractionMode)Material->RefractionMethod : RM_None; }
 	virtual bool GetRootNodeOverridesDefaultRefraction()const override		{ return Material ? Material->bRootNodeOverridesDefaultDistortion : false; }
-	virtual FMaterialShadingModelField GetShadingModels() const override	{ return MSM_Unlit; }
-	virtual bool IsShadingModelFromMaterialExpression() const override		{ return false; }
 	virtual float GetOpacityMaskClipValue() const override					{ return 0.5f; }
 	virtual bool GetCastDynamicShadowAsMasked() const override				{ return false; }
 	virtual FString GetFriendlyName() const override { return FString::Printf(TEXT("FLightmassMaterialRenderer %s"), MaterialInterface ? *MaterialInterface->GetName() : TEXT("NULL")); }
 
+	virtual bool IsShadingModelFromMaterialExpression() const override		{ return false; }
+	virtual FMaterialShadingModelField GetShadingModels() const override	
+	{ 
+		// Substrate needs the real material shading model since the expressions access GetMaterialShadingModels() through the compiler to generate the substrate operators, 
+		// and we do not want unlit materials when it is a Slab or a ShadingModel node.
+		return Lightmass_IsSubstrateEnabled() ? Material->GetShadingModels() : MSM_Unlit; 
+	} 
 	/**
 	 * Should shaders compiled for this material be saved to disk?
 	 */
