IncrementalCook: Add ReferencedSet.txt for use when staging from a zen oplog that could include previous-cook entries than were no longer referenced from the most recent cook.
Refactor CookByTheBookFinishedInternal to move all optional writes into the per-platform forloop so we can share data calculated for each platform between ReferencedSet.txt and the AssetRegistry.
Remove the no-longer-used SCOPED_BOOT_TIMING macros; we profile cook times using UE_SCOPED_HIERARCHICAL_COOKTIMER now.
#rnx

diff --git a/Engine/Source/Editor/UnrealEd/Classes/CookOnTheSide/CookOnTheFlyServer.h b/Engine/Source/Editor/UnrealEd/Classes/CookOnTheSide/CookOnTheFlyServer.h
--- a/Engine/Source/Editor/UnrealEd/Classes/CookOnTheSide/CookOnTheFlyServer.h
+++ b/Engine/Source/Editor/UnrealEd/Classes/CookOnTheSide/CookOnTheFlyServer.h
@@ -1274,6 +1274,7 @@ private:
 	// Return the filename to use for the cook metadata file, adjusted for DLC, sandbox, and platform.
 	FString GetCookedCookMetadataFilename(const FString& PlatformName);
 	void WriteCookMetadata(const ITargetPlatform* InTargetPlatform, uint64 InDevelopmentAssetRegistryHash);
+	void WriteReferencedSet(const ITargetPlatform* InTargetPlatform, TArray<FName>&& CookedPackageNames);
 
 	/* @return Full path of the CachedEditorThumbnails.bin file in the sandbox */
 	FString GetSandboxCachedEditorThumbnailsFilename();
diff --git a/Engine/Source/Editor/UnrealEd/Private/CookOnTheFlyServer.cpp b/Engine/Source/Editor/UnrealEd/Private/CookOnTheFlyServer.cpp
--- a/Engine/Source/Editor/UnrealEd/Private/CookOnTheFlyServer.cpp
+++ b/Engine/Source/Editor/UnrealEd/Private/CookOnTheFlyServer.cpp
@@ -10512,9 +10512,38 @@ void UCookOnTheFlyServer::WriteCookMetadata(const ITargetPlatform* InTargetPlatf
 		MetadataState.SetPlatformAndBuildVersion(PlatformNameString, FApp::GetBuildVersion());
 		MetadataState.SetHordeJobId(FPlatformMisc::GetEnvironmentVariable(TEXT("UE_HORDE_JOBID")));
 
-			MetadataState.SaveToFile(GetCookedCookMetadataFilename(PlatformNameString));
+		MetadataState.SaveToFile(GetCookedCookMetadataFilename(PlatformNameString));
+	}
+}
+
+void UCookOnTheFlyServer::WriteReferencedSet(const ITargetPlatform* InTargetPlatform, TArray<FName>&& CookedPackageNames)
+{
+	const FString MetadataPlatformAgnosticFilename = GetMetadataDirectory() / UE::Cook::GetReferencedSetFilename();
+	const FString MetadataFilename = ConvertToFullSandboxPath(*MetadataPlatformAgnosticFilename, true)
+		.Replace(TEXT("[Platform]"), *InTargetPlatform->PlatformName());
+
+	CookedPackageNames.Sort([](FName A, FName B)
+		{
+			return A.LexicalLess(B);
+		});
 
+	FString CombinedString;
+	if (CookedPackageNames.Num() > 0)
+	{
+		int32 CombinedLength = (CookedPackageNames.Num() - 1) * UE_ARRAY_COUNT(LINE_TERMINATOR);
+		for (FName PackageName : CookedPackageNames)
+		{
+			CombinedLength += PackageName.GetStringLength();
+		}
+		CombinedString.Reserve(CombinedLength);
+		CombinedString += WriteToString<256>(CookedPackageNames[0]);
+		for (FName PackageName : TArrayView<FName>(CookedPackageNames).RightChop(1))
+		{
+			CombinedString += LINE_TERMINATOR;
+			CombinedString += WriteToString<256>(PackageName);
+		}
 	}
+	FFileHelper::SaveStringToFile(CombinedString, *MetadataFilename);
 }
 
 void UCookOnTheFlyServer::CookByTheBookFinished()
@@ -10591,235 +10620,261 @@ void UCookOnTheFlyServer::CookByTheBookFinishedInternal()
 	
 	GetDerivedDataCacheRef().WaitForQuiescence(true);
 	
+	bool bSaveAssetRegistry = !FParse::Param(FCommandLine::Get(), TEXT("SkipSaveAssetRegistry"));
+	// if we are cooking DLC, the DevelopmentAR isn't needed - it's used when making DLC against shipping, so there's no need to make it
+	// again, as we don't make DLC against DLC (but allow an override just in case)
+	bool bSaveDevelopmentAssetRegistry = !FParse::Param(FCommandLine::Get(), TEXT("NoSaveDevAR"));
+	bool bForceNoFilterAssetsFromAssetRegistry = IsCookingDLC();
+	bool bSaveManifests = true;
+	bool bSaveIniSettings = !FParse::Param(FCommandLine::Get(), TEXT("SkipSaveCookSettings"));
+	bool bSaveCookerOpenOrder = true;
+	bool bCacheShaderLibraries = IsUsingShaderCodeLibrary();
+	// SkipSaveAssetRegistry skips some other optional artifacts, because it is used as a
+	// "cook for testing purposes quickly" flag. They also have dependencies on each other in the current code.
+	if (!bSaveAssetRegistry)
+	{
+		bCacheShaderLibraries = false;
+		bSaveDevelopmentAssetRegistry = false;
+		bSaveManifests = false;
+		bSaveCookerOpenOrder = false;
+	}
+
 	UCookerSettings const* CookerSettings = GetDefault<UCookerSettings>();
 
 	FString LibraryName = GetProjectShaderLibraryName();
 	check(!LibraryName.IsEmpty());
-	const bool bCacheShaderLibraries = IsUsingShaderCodeLibrary();
 
-	// The hashes of the entire files, per platform.
-	TArray<uint64> DevelopmentAssetRegistryHashes;
+	// Save modified asset registry with all streaming chunk info generated during cook
+	const FString SandboxRegistryFilename = GetSandboxAssetRegistryFilename();
 
+	// Saving the current ini settings. This is only required for iterative cooking and may take seconds.
+	if (bSaveIniSettings)
 	{
-		// Save modified asset registry with all streaming chunk info generated during cook
-		const FString SandboxRegistryFilename = GetSandboxAssetRegistryFilename();
-
-		// previously shader library was saved at this spot, but it's too early to know the chunk assignments, we need to BuildChunkManifest in the asset registry first
-
-		// Saving the current ini settings. This is only required for iterative cooking and may take seconds.
-		if (!FParse::Param(FCommandLine::Get(), TEXT("SkipSaveCookSettings")))
+		UE_SCOPED_HIERARCHICAL_COOKTIMER(SavingCurrentIniSettings)
+		for (const ITargetPlatform* TargetPlatform : PlatformManager->GetSessionPlatforms() )
 		{
-			UE_SCOPED_HIERARCHICAL_COOKTIMER(SavingCurrentIniSettings)
-			for (const ITargetPlatform* TargetPlatform : PlatformManager->GetSessionPlatforms() )
+			if (FindOrCreateSaveContext(TargetPlatform).PackageWriterCapabilities.bReadOnly)
 			{
-				if (FindOrCreateSaveContext(TargetPlatform).PackageWriterCapabilities.bReadOnly)
-				{
-					continue;
-				}
-				SaveCurrentIniSettings(TargetPlatform);
+				continue;
 			}
+			SaveCurrentIniSettings(TargetPlatform);
 		}
+	}
 
-		if (!FParse::Param(FCommandLine::Get(), TEXT("SkipSaveAssetRegistry")))
-		{
-			UE_SCOPED_HIERARCHICAL_COOKTIMER(SavingAssetRegistry);
-			SCOPED_BOOT_TIMING("SavingAssetRegistry");
+	if (bSaveAssetRegistry)
+	{
+		UE_SCOPED_HIERARCHICAL_COOKTIMER(ChunkGeneration);
 
-			RegisterLocalizationChunkDataGenerator();
-			if (bCacheShaderLibraries)
-			{
-				RegisterShaderChunkDataGenerator();
-			}
+		RegisterLocalizationChunkDataGenerator();
+		if (bCacheShaderLibraries)
+		{
+			RegisterShaderChunkDataGenerator();
+		}
 
-			for (auto ChunkGeneratorFactory : IChunkDataGenerator::GetChunkDataGeneratorFactories())
+		for (auto ChunkGeneratorFactory : IChunkDataGenerator::GetChunkDataGeneratorFactories())
+		{
+			for (const ITargetPlatform* TargetPlatform : PlatformManager->GetSessionPlatforms())
 			{
-				for (const ITargetPlatform* TargetPlatform : PlatformManager->GetSessionPlatforms())
-				{
-					FAssetRegistryGenerator& RegistryGenerator = *(PlatformManager->GetPlatformData(TargetPlatform)->RegistryGenerator);
-					RegistryGenerator.RegisterChunkDataGenerator(ChunkGeneratorFactory(*this));
-				}
+				FAssetRegistryGenerator& RegistryGenerator = *(PlatformManager->GetPlatformData(TargetPlatform)->RegistryGenerator);
+				RegistryGenerator.RegisterChunkDataGenerator(ChunkGeneratorFactory(*this));
 			}
-			// if we are cooking DLC, the DevelopmentAR isn't needed - it's used when making DLC against shipping, so there's no need to make it
-			// again, as we don't make DLC against DLC (but allow an override just in case)
-			bool bSaveDevelopmentAssetRegistry = !FParse::Param(FCommandLine::Get(), TEXT("NoSaveDevAR"));
+		}
+	}
 
-			for (const ITargetPlatform* TargetPlatform : PlatformManager->GetSessionPlatforms())
-			{
-				if (FindOrCreateSaveContext(TargetPlatform).PackageWriterCapabilities.bReadOnly)
-				{
-					continue;
-				}
+	for (const ITargetPlatform* TargetPlatform : PlatformManager->GetSessionPlatforms())
+	{
+		if (FindOrCreateSaveContext(TargetPlatform).PackageWriterCapabilities.bReadOnly)
+		{
+			continue;
+		}
 
-				FPlatformData* PlatformData = PlatformManager->GetPlatformData(TargetPlatform);
-				FAssetRegistryGenerator& Generator = *PlatformData->RegistryGenerator;
-				TArray<FPackageData*> CookedPackageDatas;
-				TArray<FPackageData*> IgnorePackageDatas;
+		FPlatformData* PlatformData = PlatformManager->GetPlatformData(TargetPlatform);
+		FAssetRegistryGenerator& Generator = *PlatformData->RegistryGenerator;
+		TArray<FPackageData*> CookedPackageDatas;
+		TArray<FPackageData*> IgnorePackageDatas;
 
-				FString PlatformNameString = TargetPlatform->PlatformName();
-				FName PlatformName(*PlatformNameString);
+		FString PlatformNameString = TargetPlatform->PlatformName();
+		FName PlatformName(*PlatformNameString);
 
-				PackageDatas->GetCookedPackagesForPlatform(TargetPlatform, CookedPackageDatas, IgnorePackageDatas);
+		TSet<FName> CookedPackageNames;
+		TSet<FName> IgnorePackageNames;
+		{
+			UE_SCOPED_HIERARCHICAL_COOKTIMER(CalculateReferencedSet);
 
-				bool bForceNoFilterAssetsFromAssetRegistry = false;
+			PackageDatas->GetCookedPackagesForPlatform(TargetPlatform, CookedPackageDatas, IgnorePackageDatas);
 
-				if (IsCookingDLC())
+			if (IsCookingDLC())
+			{
+				TMap<FName, FPackageData*> CookedPackagesMap;
+				CookedPackagesMap.Reserve(CookedPackageDatas.Num());
+				for (FPackageData* PackageData : CookedPackageDatas)
 				{
-					TMap<FName, FPackageData*> CookedPackagesMap;
-					CookedPackagesMap.Reserve(CookedPackageDatas.Num());
-					for (FPackageData* PackageData : CookedPackageDatas)
-					{
-						CookedPackagesMap.Add(PackageData->GetFileName(), PackageData);
-					}
-					bForceNoFilterAssetsFromAssetRegistry = true;
-					// remove the previous release cooked packages from the new asset registry, add to ignore list
-					UE_SCOPED_HIERARCHICAL_COOKTIMER(RemovingOldManifestEntries);
+					CookedPackagesMap.Add(PackageData->GetFileName(), PackageData);
+				}
+				// remove the previous release cooked packages from the new asset registry, add to ignore list
+				UE_SCOPED_HIERARCHICAL_COOKTIMER(RemovingOldManifestEntries);
 
-					const TArray<FName>* PreviousReleaseCookedPackages = CookByTheBookOptions->BasedOnReleaseCookedPackages.Find(PlatformName);
-					if (PreviousReleaseCookedPackages)
+				const TArray<FName>* PreviousReleaseCookedPackages = CookByTheBookOptions->BasedOnReleaseCookedPackages.Find(PlatformName);
+				if (PreviousReleaseCookedPackages)
+				{
+					for (FName PreviousReleaseCookedPackage : *PreviousReleaseCookedPackages)
 					{
-						for (FName PreviousReleaseCookedPackage : *PreviousReleaseCookedPackages)
+						FPackageData* PackageData;
+						if (!CookedPackagesMap.RemoveAndCopyValue(PreviousReleaseCookedPackage, PackageData))
 						{
-							FPackageData* PackageData;
-							if (!CookedPackagesMap.RemoveAndCopyValue(PreviousReleaseCookedPackage, PackageData))
-							{
-								PackageData = PackageDatas->FindPackageDataByFileName(PreviousReleaseCookedPackage);
-							}
-							if (PackageData)
-							{
-								IgnorePackageDatas.Add(PackageData);
-							}
+							PackageData = PackageDatas->FindPackageDataByFileName(PreviousReleaseCookedPackage);
+						}
+						if (PackageData)
+						{
+							IgnorePackageDatas.Add(PackageData);
 						}
-					}
-					CookedPackageDatas.Reset();
-					for (TPair<FName, FPackageData*>& Pair : CookedPackagesMap)
-					{
-						CookedPackageDatas.Add(Pair.Value);
 					}
 				}
-
-				TSet<FName> CookedPackageNames;
-				for (FPackageData* PackageData : CookedPackageDatas)
+				CookedPackageDatas.Reset();
+				for (TPair<FName, FPackageData*>& Pair : CookedPackagesMap)
 				{
-					CookedPackageNames.Add(PackageData->GetPackageName());
+					CookedPackageDatas.Add(Pair.Value);
 				}
+			}
 
-				TSet<FName> IgnorePackageNames;
-				if (bSaveDevelopmentAssetRegistry)
-				{
-					for (FPackageData* PackageData : IgnorePackageDatas)
-					{
-						IgnorePackageNames.Add(PackageData->GetPackageName());
-					}
+			for (FPackageData* PackageData : CookedPackageDatas)
+			{
+				CookedPackageNames.Add(PackageData->GetPackageName());
+			}
 
-					// ignore packages that weren't cooked because they were only referenced by editor-only properties
-					TSet<FName> UncookedEditorOnlyPackageNames;
-					PackageTracker->UncookedEditorOnlyPackages.GetValues(UncookedEditorOnlyPackageNames);
-					for (FName UncookedEditorOnlyPackage : UncookedEditorOnlyPackageNames)
-					{
-						IgnorePackageNames.Add(UncookedEditorOnlyPackage);
-					}
-				}
-				
-				if (bCacheShaderLibraries)
-				{
-					FinishPopulateShaderLibrary(TargetPlatform, LibraryName);
-				}
+			for (FPackageData* PackageData : IgnorePackageDatas)
+			{
+				IgnorePackageNames.Add(PackageData->GetPackageName());
+			}
 
-				// Add the package hashes to the relevant AssetPackageDatas.
-				// PackageHashes are gated by requiring UPackage::WaitForAsyncFileWrites(), which is called above.
-				FCookSavePackageContext& SaveContext = FindOrCreateSaveContext(TargetPlatform);
-				TMap<FName, TRefCountPtr<FPackageHashes>>& AllPackageHashes = SaveContext.PackageWriter->GetPackageHashes();
-				for (TPair<FName, TRefCountPtr<FPackageHashes>>& HashSet : AllPackageHashes)
-				{
-					FAssetPackageData* AssetPackageData = Generator.GetAssetPackageData(HashSet.Key);
-					TRefCountPtr<FPackageHashes>& PackageHashes = HashSet.Value;
+			// ignore packages that weren't cooked because they were only referenced by editor-only properties
+			TSet<FName> UncookedEditorOnlyPackageNames;
+			PackageTracker->UncookedEditorOnlyPackages.GetValues(UncookedEditorOnlyPackageNames);
+			for (FName UncookedEditorOnlyPackage : UncookedEditorOnlyPackageNames)
+			{
+				IgnorePackageNames.Add(UncookedEditorOnlyPackage);
+			}
+		}
+		
+		if (bCacheShaderLibraries)
+		{
+			UE_SCOPED_HIERARCHICAL_COOKTIMER(FinishPopulateShaderLibrary);
+			FinishPopulateShaderLibrary(TargetPlatform, LibraryName);
+		}
 
-					AssetPackageData->CookedHash = PackageHashes->PackageHash;
-					Move(AssetPackageData->ChunkHashes, PackageHashes->ChunkHashes);
-				}
+		FCookSavePackageContext& SaveContext = FindOrCreateSaveContext(TargetPlatform);
+		if (bSaveManifests || bSaveAssetRegistry)
+		{
+			UE_SCOPED_HIERARCHICAL_COOKTIMER(GeneratorPreSave);
 
-				{
-					Generator.PreSave(CookedPackageNames);
-				}
-				{
-					UE_SCOPED_HIERARCHICAL_COOKTIMER(BuildChunkManifest);
-					Generator.FinalizeChunkIDs(CookedPackageNames, IgnorePackageNames, *SandboxFile,
-						CookByTheBookOptions->bGenerateStreamingInstallManifests);
-				}
-				{
-					UE_SCOPED_HIERARCHICAL_COOKTIMER(SaveManifests);
-					if (!Generator.SaveManifests(*SandboxFile))
-					{
-						UE_LOG(LogCook, Warning, TEXT("Failed to save chunk manifest"));
-					}
+			// Add the package hashes to the relevant AssetPackageDatas.
+			// PackageHashes are gated by requiring UPackage::WaitForAsyncFileWrites(), which is called above.
+			TMap<FName, TRefCountPtr<FPackageHashes>>& AllPackageHashes = SaveContext.PackageWriter->GetPackageHashes();
+			for (TPair<FName, TRefCountPtr<FPackageHashes>>& HashSet : AllPackageHashes)
+			{
+				FAssetPackageData* AssetPackageData = Generator.GetAssetPackageData(HashSet.Key);
+				TRefCountPtr<FPackageHashes>& PackageHashes = HashSet.Value;
 
-					int64 ExtraFlavorChunkSize;
-					if (FParse::Value(FCommandLine::Get(), TEXT("ExtraFlavorChunkSize="), ExtraFlavorChunkSize) && ExtraFlavorChunkSize > 0)
-					{
-						// ExtraFlavor is a legacy term for this override; etymology unknown. Override the chunksize specified by the platform,
-						// and write the manifest files created with that chunksize into a separate subdirectory.
-						const TCHAR* ManifestSubDir = TEXT("ExtraFlavor");
-						if (!Generator.SaveManifests(*SandboxFile, ExtraFlavorChunkSize, ManifestSubDir))
-						{
-							UE_LOG(LogCook, Warning, TEXT("Failed to save chunk manifest"));
-						}
-					}
-				}
-				{
-					UE_SCOPED_HIERARCHICAL_COOKTIMER(SaveRealAssetRegistry);
-					uint64 DevArHash = 0;
-					Generator.SaveAssetRegistry(SandboxRegistryFilename, bSaveDevelopmentAssetRegistry, bForceNoFilterAssetsFromAssetRegistry, DevArHash);
-					DevelopmentAssetRegistryHashes.Add(DevArHash);
-				}
+				AssetPackageData->CookedHash = PackageHashes->PackageHash;
+				Move(AssetPackageData->ChunkHashes, PackageHashes->ChunkHashes);
+			}
+
+			Generator.PreSave(CookedPackageNames);
+		}
+
+		if (bSaveManifests)
+		{
+			{
+				UE_SCOPED_HIERARCHICAL_COOKTIMER(BuildChunkManifest);
+				Generator.FinalizeChunkIDs(CookedPackageNames, IgnorePackageNames, *SandboxFile,
+					CookByTheBookOptions->bGenerateStreamingInstallManifests);
+			}
+			{
+				UE_SCOPED_HIERARCHICAL_COOKTIMER(SaveManifests);
+				if (!Generator.SaveManifests(*SandboxFile))
 				{
-					Generator.PostSave();
+					UE_LOG(LogCook, Warning, TEXT("Failed to save chunk manifest"));
 				}
+
+				int64 ExtraFlavorChunkSize;
+				if (FParse::Value(FCommandLine::Get(), TEXT("ExtraFlavorChunkSize="), ExtraFlavorChunkSize) && ExtraFlavorChunkSize > 0)
 				{
-					UE_SCOPED_HIERARCHICAL_COOKTIMER(WriteCookerOpenOrder);
-					if (!IsCookFlagSet(ECookInitializationFlags::Iterative))
+					// ExtraFlavor is a legacy term for this override; etymology unknown. Override the chunksize specified by the platform,
+					// and write the manifest files created with that chunksize into a separate subdirectory.
+					const TCHAR* ManifestSubDir = TEXT("ExtraFlavor");
+					if (!Generator.SaveManifests(*SandboxFile, ExtraFlavorChunkSize, ManifestSubDir))
 					{
-						Generator.WriteCookerOpenOrder(*SandboxFile);
+						UE_LOG(LogCook, Warning, TEXT("Failed to save chunk manifest"));
 					}
 				}
-				if (bCacheShaderLibraries)
-				{
-					// now that we have the asset registry and cooking open order, we have enough information to split the shader library
-					// into parts for each chunk and (possibly) lay out the code in accordance with the file order
-					// Save shader code map
-					SaveShaderLibrary(TargetPlatform, LibraryName);
-					CreatePipelineCache(TargetPlatform, LibraryName);
-				}
-				if (FParse::Param(FCommandLine::Get(), TEXT("fastcook")))
-				{
-					FFileHelper::SaveStringToFile(FString(), *(GetSandboxDirectory(PlatformNameString) / TEXT("fastcook.txt")));
-				}
-				if (IsCreatingReleaseVersion())
-				{
-					const FString VersionedRegistryPath = GetCreateReleaseVersionAssetRegistryPath(CookByTheBookOptions->CreateReleaseVersion, PlatformNameString);
-					IFileManager::Get().MakeDirectory(*VersionedRegistryPath, true);
-					const FString VersionedRegistryFilename = VersionedRegistryPath / GetAssetRegistryFilename();
-					const FString CookedAssetRegistryFilename = SandboxRegistryFilename.Replace(TEXT("[Platform]"), *PlatformNameString);
-					IFileManager::Get().Copy(*VersionedRegistryFilename, *CookedAssetRegistryFilename, true, true);
+			}
+		}
 
-					// Also copy development registry if it exists
-					FString DevelopmentAssetRegistryRelativePath = FString::Printf(TEXT("Metadata/%s"), GetDevelopmentAssetRegistryFilename());
-					const FString DevVersionedRegistryFilename = VersionedRegistryFilename.Replace(TEXT("AssetRegistry.bin"), *DevelopmentAssetRegistryRelativePath);
-					const FString DevCookedAssetRegistryFilename = CookedAssetRegistryFilename.Replace(TEXT("AssetRegistry.bin"), *DevelopmentAssetRegistryRelativePath);
-					IFileManager::Get().Copy(*DevVersionedRegistryFilename, *DevCookedAssetRegistryFilename, true, true);
-				}
+		uint64 DevelopmentAssetRegistryHash = 0; // The hashes of the entire files for the platform
+		if (bSaveAssetRegistry)
+		{
+			UE_SCOPED_HIERARCHICAL_COOKTIMER(SaveAssetRegistry);
+			Generator.SaveAssetRegistry(SandboxRegistryFilename, bSaveDevelopmentAssetRegistry, bForceNoFilterAssetsFromAssetRegistry, DevelopmentAssetRegistryHash);
+		}
+
+		if (bSaveManifests || bSaveAssetRegistry)
+		{
+			Generator.PostSave();
+		}
+
+		if (bSaveCookerOpenOrder)
+		{
+			UE_SCOPED_HIERARCHICAL_COOKTIMER(WriteCookerOpenOrder);
+			if (!IsCookFlagSet(ECookInitializationFlags::Iterative))
+			{
+				Generator.WriteCookerOpenOrder(*SandboxFile);
 			}
 		}
-	}
 
-	// Write cook metadata file for each platform
-	int32 PlatformIndex = 0;
-	for (const ITargetPlatform* TargetPlatform : PlatformManager->GetSessionPlatforms())
-	{
-		if (!FindOrCreateSaveContext(TargetPlatform).PackageWriterCapabilities.bReadOnly)
+		if (bCacheShaderLibraries)
+		{
+			// now that we have the asset registry and cooking open order, we have enough information to split the shader library
+			// into parts for each chunk and (possibly) lay out the code in accordance with the file order
+			// Assert that the other saves are enabled because we depend on those files being written.
+			check(bSaveCookerOpenOrder && bSaveAssetRegistry);
+			// Save shader code map
+			SaveShaderLibrary(TargetPlatform, LibraryName);
+			CreatePipelineCache(TargetPlatform, LibraryName);
+		}
+
+		if (FParse::Param(FCommandLine::Get(), TEXT("fastcook")))
+		{
+			FFileHelper::SaveStringToFile(FString(), *(GetSandboxDirectory(PlatformNameString) / TEXT("fastcook.txt")));
+		}
+
+		if (bSaveAssetRegistry && IsCreatingReleaseVersion())
+		{
+			const FString VersionedRegistryPath = GetCreateReleaseVersionAssetRegistryPath(CookByTheBookOptions->CreateReleaseVersion, PlatformNameString);
+			IFileManager::Get().MakeDirectory(*VersionedRegistryPath, true);
+			const FString VersionedRegistryFilename = VersionedRegistryPath / GetAssetRegistryFilename();
+			const FString CookedAssetRegistryFilename = SandboxRegistryFilename.Replace(TEXT("[Platform]"), *PlatformNameString);
+			IFileManager::Get().Copy(*VersionedRegistryFilename, *CookedAssetRegistryFilename, true, true);
+
+			// Also copy development registry if it exists
+			FString DevelopmentAssetRegistryRelativePath = FString::Printf(TEXT("Metadata/%s"), GetDevelopmentAssetRegistryFilename());
+			const FString DevVersionedRegistryFilename = VersionedRegistryFilename.Replace(TEXT("AssetRegistry.bin"), *DevelopmentAssetRegistryRelativePath);
+			const FString DevCookedAssetRegistryFilename = CookedAssetRegistryFilename.Replace(TEXT("AssetRegistry.bin"), *DevelopmentAssetRegistryRelativePath);
+			IFileManager::Get().Copy(*DevVersionedRegistryFilename, *DevCookedAssetRegistryFilename, true, true);
+		}
+
+		// Write cook metadata file for each platform
 		{
-			WriteCookMetadata(TargetPlatform, DevelopmentAssetRegistryHashes[PlatformIndex]);
+			UE_SCOPED_HIERARCHICAL_COOKTIMER(WriteCookMetadata);
+			WriteCookMetadata(TargetPlatform, DevelopmentAssetRegistryHash);
 		}
-		PlatformIndex++;
+
+		// Write ReferencedSet for use by staging and zen commands on incremental cook oplogs: they use only the ops
+		// referenced by the most recent cook.
+		{
+			UE_SCOPED_HIERARCHICAL_COOKTIMER(WriteReferencedSet);
+			WriteReferencedSet(TargetPlatform, CookedPackageNames.Array());
+		}
+
 	}
 
 	FString ActualLibraryName = GenerateShaderCodeLibraryName(LibraryName, IsCookFlagSet(ECookInitializationFlags::IterateSharedBuild));
@@ -12585,7 +12640,6 @@ void UCookOnTheFlyServer::RecordDLCPackagesFromBaseGame(FBeginCookContext& Begin
 	bool bFirstAddExistingPackageDatas = true;
 	for (const ITargetPlatform* TargetPlatform : BeginContext.TargetPlatforms)
 	{
-		SCOPED_BOOT_TIMING("AddCookedPlatforms");
 		TArray<UE::Cook::FConstructPackageData> PackageList;
 		FString PlatformNameString = TargetPlatform->PlatformName();
 		FName PlatformName(*PlatformNameString);
@@ -12605,7 +12659,6 @@ void UCookOnTheFlyServer::RecordDLCPackagesFromBaseGame(FBeginCookContext& Begin
 		TArray<UE::Cook::FConstructPackageData>& ActivePackageList = OverridePackageList.Num() > 0 ? OverridePackageList : PackageList;
 		if (ActivePackageList.Num() > 0)
 		{
-			SCOPED_BOOT_TIMING("AddPackageDataByFileNamesForPlatform");
 			PackageDatas->AddExistingPackageDatasForPlatform(ActivePackageList, TargetPlatform,
 				bFirstAddExistingPackageDatas, PackageDataFromBaseGameNum);
 		}
@@ -13036,7 +13089,6 @@ bool UCookOnTheFlyServer::GetAllPackageFilenamesFromAssetRegistry(const FString&
 	using namespace UE::Cook;
 
 	UE_SCOPED_COOKTIMER(GetAllPackageFilenamesFromAssetRegistry);
-	SCOPED_BOOT_TIMING("GetAllPackageFilenamesFromAssetRegistry");
 	TUniquePtr<FArchive> Reader(IFileManager::Get().CreateFileReader(*AssetRegistryPath));
 	if (Reader)
 	{
diff --git a/Engine/Source/Runtime/CoreUObject/Private/UObject/ICookInfo.cpp b/Engine/Source/Runtime/CoreUObject/Private/UObject/ICookInfo.cpp
--- a/Engine/Source/Runtime/CoreUObject/Private/UObject/ICookInfo.cpp
+++ b/Engine/Source/Runtime/CoreUObject/Private/UObject/ICookInfo.cpp
@@ -53,6 +53,11 @@ FCookInfoEvent FDelegates::CookStarted;
 FCookInfoEvent FDelegates::CookFinished;
 FValidateSourcePackage FDelegates::ValidateSourcePackage;
 
+const TCHAR* GetReferencedSetFilename()
+{
+	return TEXT("ReferencedSet.txt");
+}
+
 }
 
 static thread_local ECookLoadType GCookLoadType = ECookLoadType::Unexpected;
diff --git a/Engine/Source/Runtime/CoreUObject/Public/UObject/ICookInfo.h b/Engine/Source/Runtime/CoreUObject/Public/UObject/ICookInfo.h
--- a/Engine/Source/Runtime/CoreUObject/Public/UObject/ICookInfo.h
+++ b/Engine/Source/Runtime/CoreUObject/Public/UObject/ICookInfo.h
@@ -183,6 +183,16 @@ public:
 	static COREUOBJECT_API FValidateSourcePackage ValidateSourcePackage;
 };
 
+
+/**
+ * Return the relative path under the cook output MetaData folder to the ReferencedSet file.
+ * The ReferencedSet file is a text file list of package names, one per line, that were referenced from the cook
+ * session.
+ * It does not include previously cooked files in an incremental cook that are no longer referenced.
+ * It does not include packages from the base game for a DLC cook.
+ */
+COREUOBJECT_API const TCHAR* GetReferencedSetFilename();
+
 } // namespace UE::Cook
 
 /**
