Cooker Memory: Save packages in root to leaf order rather than leaf to root order, so that we reduce the number of packages we've already saved that we need to reload as dependencies when saving later packages in the cook.
Enable the preloading of imports so that dependent loads still proceed in leaf to root order (and are preloaded) despite the root to leaf order of requests and saves.
Remove debugging code for preloading of imports; the guess that some imports was not required turned out to be false. The extra loads when enabling preloading of imports were instead coming from an increased number of garbage collects, and that has been fixed by a previous change that increased the urgency of generated packages.
#rnx

diff --git a/Engine/Source/Editor/UnrealEd/Private/CookOnTheFlyServer.cpp b/Engine/Source/Editor/UnrealEd/Private/CookOnTheFlyServer.cpp
--- a/Engine/Source/Editor/UnrealEd/Private/CookOnTheFlyServer.cpp
+++ b/Engine/Source/Editor/UnrealEd/Private/CookOnTheFlyServer.cpp
@@ -2714,7 +2714,11 @@ void UCookOnTheFlyServer::AssignRequests(TArrayView<UE::Cook::FPackageData*> Req
 		TArray<FWorkerId> Assignments;
 		CookDirector->AssignRequests(Requests, Assignments, MoveTemp(RequestGraph));
 		check(Assignments.Num() == NumRequests);
-		for (int32 Index = 0; Index < NumRequests; ++Index)
+
+		// The Input RequestQueue is in LeafToRoot order, but we want to save in RootToLeaf order,
+		// so reverse iterate. This is only important for the Assignment.IsLocal case; the other two
+		// cases go into order-independent containers.
+		for (int32 Index = NumRequests - 1; Index >= 0; --Index)
 		{
 			FPackageData* PackageData = Requests[Index];
 			FWorkerId Assignment = Assignments[Index];
@@ -2746,7 +2750,9 @@ void UCookOnTheFlyServer::AssignRequests(TArrayView<UE::Cook::FPackageData*> Req
 			Requests = Shuffled;
 		}
 
-		for (FPackageData* PackageData : Requests)
+		// The Input RequestQueue is in LeafToRoot order, but we want to save in RootToLeaf order,
+		// so reverse iterate.
+		for (FPackageData* PackageData : ReverseIterate(Requests))
 		{
 			RequestQueue.AddReadyRequest(PackageData);
 		}
diff --git a/Engine/Source/Editor/UnrealEd/Private/Cooker/CookPackagePreloader.cpp b/Engine/Source/Editor/UnrealEd/Private/Cooker/CookPackagePreloader.cpp
--- a/Engine/Source/Editor/UnrealEd/Private/Cooker/CookPackagePreloader.cpp
+++ b/Engine/Source/Editor/UnrealEd/Private/Cooker/CookPackagePreloader.cpp
@@ -22,23 +22,13 @@ namespace UE::Cook
 
 bool FPackagePreloader::bConfigInitialized = false;
 bool FPackagePreloader::bAllowPreloadImports = true;
-bool FPackagePreloader::bDebugPreloadImports = false;
 
 void FPackagePreloader::InitializeConfig()
 {
 	bConfigInitialized = true;
 
-	// bAllowPreloadImports is off by default, because it is currently loading unneeded imports for some packages
-	// and causing poor cook performance
-	bAllowPreloadImports = false;
-	bDebugPreloadImports = false;
+	bAllowPreloadImports = true;
 	FParse::Bool(FCommandLine::Get(), TEXT("-CookPreloadImports="), bAllowPreloadImports);
-	if (FParse::Param(FCommandLine::Get(), TEXT("CookDebugPreloadImports")))
-	{
-		bDebugPreloadImports = true;
-		// bAllowPreloadImports needs to be off for us to find which imports are not loaded by the load of the referencer.
-		bAllowPreloadImports = false;
-	}
 }
 
 FPackagePreloader::~FPackagePreloader()
@@ -204,30 +194,6 @@ void FPackagePreloader::GatherUnloadedImports()
 	}
 }
 
-void FPackagePreloader::LogUnusedImports()
-{
-	if (!bDebugPreloadImports || !IsPackageLoaded())
-	{
-		return;
-	}
-
-	TStringBuilder<256> UnusedImportsStr;
-	for (TRefCountPtr<FPackagePreloader>& Import : UnloadedImports)
-	{
-		if (!Import->IsPackageLoaded())
-		{
-			UnusedImportsStr << Import->GetPackageData().GetPackageName() << TEXT(", ");
-		}
-	}
-	if (UnusedImportsStr.Len() != 0)
-	{
-		UnusedImportsStr.RemoveSuffix(2);
-		UE_LOG(LogCook, Warning,
-			TEXT("DebugPreloadImports: the load of %s did not end up loading these expected imports:\n\t%s"),
-			*WriteToString<256>(GetPackageData().GetPackageName()), *UnusedImportsStr);
-	}
-}
-
 void FPackagePreloader::EmptyImports()
 {
 	UnloadedImports.Empty();
@@ -520,11 +486,6 @@ void FPackagePreloader::SetRequestedLoads(TArray<TRefCountPtr<FPackagePreloader>
 
 void FPackagePreloader::OnExitActive()
 {
-	if (bDebugPreloadImports)
-	{
-		LogUnusedImports();
-	}
-
 	ClearPreload();
 	EmptyImports();
 }
@@ -677,11 +638,6 @@ bool FPackagePreloader::PumpLoadsTryStartInboxPackage(UCookOnTheFlyServer& COTFS
 		else
 		{
 			NeedsLoadPreloaders.Add(Preloader);
-			if (bDebugPreloadImports)
-			{
-				// Gather the imports for *this, but do not activate them. We will review them after load.
-				Preloader->GatherUnloadedImports();
-			}
 		}
 		check(!NeedsLoadPreloaders.IsEmpty() && NeedsLoadPreloaders[0].GetReference() == Preloader);
 		Preloader->SetRequestedLoads(MoveTemp(NeedsLoadPreloaders));
diff --git a/Engine/Source/Editor/UnrealEd/Private/Cooker/CookPackagePreloader.h b/Engine/Source/Editor/UnrealEd/Private/Cooker/CookPackagePreloader.h
--- a/Engine/Source/Editor/UnrealEd/Private/Cooker/CookPackagePreloader.h
+++ b/Engine/Source/Editor/UnrealEd/Private/Cooker/CookPackagePreloader.h
@@ -111,8 +111,6 @@ private:
 	void EmptyImports();
 	/** Gather imports from the AssetRegistry, but skip Preloads that return true for NeedsLoad. */
 	void GatherUnloadedImports();
-	/** Called after loading when bDebugPreloadImports; report imports we thought we needed but did not load. */
-	void LogUnusedImports();
 
 	/**
 	 * Traverse the import tree of this Preloader, optionally gathering imports from the asset registry if they have not
@@ -191,7 +189,6 @@ private:
 
 	static bool bConfigInitialized;
 	static bool bAllowPreloadImports;
-	static bool bDebugPreloadImports;
 
 	friend UE::Cook::FLoadQueue;
 	friend UE::Cook::FPackagePreloaderPriorityWrapper;
