Added the rendered object's name as metadata to the RenderThumbnail CPU trace, in order to quickly diagnose the costliest thumbnail renders

diff --git a/Engine/Source/Editor/UnrealEd/Private/ObjectTools.cpp b/Engine/Source/Editor/UnrealEd/Private/ObjectTools.cpp
--- a/Engine/Source/Editor/UnrealEd/Private/ObjectTools.cpp
+++ b/Engine/Source/Editor/UnrealEd/Private/ObjectTools.cpp
@@ -109,6 +109,7 @@
 #include "AssetCompilingManager.h"
 #include "ObjectEditorUtils.h"
 #include "Settings/EditorStyleSettings.h"
+#include "ProfilingDebugging/AssetMetadataTrace.h"
 
 DEFINE_LOG_CATEGORY_STATIC(LogObjectTools, Log, All);
 
@@ -5113,8 +5114,9 @@ namespace ObjectTools
 }
 
 
-
-
+UE_TRACE_EVENT_BEGIN(Cpu, RenderThumbnail, NoSync)
+UE_TRACE_EVENT_FIELD(UE::Trace::WideString, ObjectPath)
+UE_TRACE_EVENT_END()
 
 namespace ThumbnailTools
 {
@@ -5127,7 +5129,10 @@ namespace ThumbnailTools
 			return;
 		}
 		
-		TRACE_CPUPROFILER_EVENT_SCOPE(ThumbnailTools::RenderThumbnail);
+#if CPUPROFILERTRACE_ENABLED
+		UE_TRACE_LOG_SCOPED_T(Cpu, RenderThumbnail, CpuChannel)
+			<< RenderThumbnail.ObjectPath(*InObject->GetPathName());
+#endif // CPUPROFILERTRACE_ENABLED
 
 		// Renderer must be initialized before generating thumbnails
 		check( GIsRHIInitialized );
