[Viewport Toolbar] Add list of viewport types to new viewport toolbar Cameras submenu

- Creates a menu shared between old and new Level Viewport Toolbars.
- Adds Icon to the Default and Cinematic viewport types.

diff --git a/Engine/Plugins/MovieScene/LevelSequenceEditor/Source/LevelSequenceEditor/Private/LevelSequenceEditorCommands.cpp b/Engine/Plugins/MovieScene/LevelSequenceEditor/Source/LevelSequenceEditor/Private/LevelSequenceEditorCommands.cpp
--- a/Engine/Plugins/MovieScene/LevelSequenceEditor/Source/LevelSequenceEditor/Private/LevelSequenceEditorCommands.cpp
+++ b/Engine/Plugins/MovieScene/LevelSequenceEditor/Source/LevelSequenceEditor/Private/LevelSequenceEditorCommands.cpp
@@ -2,6 +2,9 @@
 
 #include "LevelSequenceEditorCommands.h"
 
+#include "CineCameraActor.h"
+#include "Styling/SlateIconFinder.h"
+
 #define LOCTEXT_NAMESPACE "LevelSequenceEditorCommands"
 
 FLevelSequenceEditorCommands::FLevelSequenceEditorCommands()
@@ -15,6 +18,17 @@ void FLevelSequenceEditorCommands::RegisterCommands()
 	UI_COMMAND(CreateNewLevelSequenceWithShotsInLevel, "Add Level Sequence with Shots", "Create a new level sequence with shots asset, and place an instance of it in this level", EUserInterfaceActionType::Button, FInputChord());
 	UI_COMMAND(ToggleCinematicViewportCommand, "Cinematic Viewport", "A viewport layout tailored to cinematic preview", EUserInterfaceActionType::RadioButton, FInputChord());
 
+	FUICommandInfo::MakeCommandInfo(
+		AsShared(),
+		ToggleCinematicViewportCommand,
+		TEXT("Cinematic Viewport"),
+		LOCTEXT("CinematicViewportLabel", "Cinematic Viewport"),
+		LOCTEXT("CinematicViewport_Tooltip", "A viewport layout tailored to cinematic preview"),
+		FSlateIconFinder::FindIconForClass(ACineCameraActor::StaticClass()),
+		EUserInterfaceActionType::RadioButton,
+		FInputChord()
+	);
+
 	UI_COMMAND(SnapSectionsToTimelineUsingSourceTimecode, "Snap Sections to Timeline using Source Timecode", "Snap selected sections to the position in the timeline matching their source timecode", EUserInterfaceActionType::Button, FInputChord());
 	UI_COMMAND(SyncSectionsUsingSourceTimecode, "Sync Sections using Source Timecode", "Synchronize sections to the first selected section using source timecode", EUserInterfaceActionType::Button, FInputChord());
 
diff --git a/Engine/Source/Editor/LevelEditor/Private/LevelViewportActions.cpp b/Engine/Source/Editor/LevelEditor/Private/LevelViewportActions.cpp
--- a/Engine/Source/Editor/LevelEditor/Private/LevelViewportActions.cpp
+++ b/Engine/Source/Editor/LevelEditor/Private/LevelViewportActions.cpp
@@ -16,6 +16,8 @@
 #include "VirtualShadowMapVisualizationData.h"
 #include "Bookmarks/BookmarkUI.h"
 #include "UnrealEdGlobals.h"
+#include "Camera/CameraActor.h"
+#include "Styling/SlateIconFinder.h"
 
 #define LOCTEXT_NAMESPACE "LevelViewportActions"
 
@@ -62,7 +64,16 @@ void FLevelViewportCommands::RegisterCommands()
 	UI_COMMAND( ViewportConfig_FourPanesBottom, "Layout Four Panes (one bottom, three top)", "Changes the viewport arrangement to four panes, one on the bottom, three on the top", EUserInterfaceActionType::ToggleButton, FInputChord() );
 	UI_COMMAND( ViewportConfig_FourPanes2x2, "Layout Four Panes (2x2)", "Changes the viewport arrangement to four panes, in a 2x2 grid", EUserInterfaceActionType::ToggleButton, FInputChord() );
 
-	UI_COMMAND( SetDefaultViewportType, "Default Viewport", "Reconfigures this viewport to the default arrangement", EUserInterfaceActionType::RadioButton, FInputChord() );
+	FUICommandInfo::MakeCommandInfo(
+		AsShared(),
+		SetDefaultViewportType,
+		TEXT("Default Viewport"),
+		LOCTEXT("DefaultViewportLabel", "Default Viewport"),
+		LOCTEXT("DefaultViewport_Tooltip", "Reconfigures this viewport to the default arrangement"),
+		FSlateIconFinder::FindIconForClass(ACameraActor::StaticClass()),
+		EUserInterfaceActionType::RadioButton,
+		FInputChord()
+	);
 
 	UI_COMMAND( ToggleViewportToolbar, "Show Toolbar", "Defines whether a toolbar should be displayed on this viewport", EUserInterfaceActionType::ToggleButton, FInputChord(EModifierKey::Control | EModifierKey::Shift, EKeys::T) );
 
diff --git a/Engine/Source/Editor/LevelEditor/Private/SLevelViewportToolBar.cpp b/Engine/Source/Editor/LevelEditor/Private/SLevelViewportToolBar.cpp
--- a/Engine/Source/Editor/LevelEditor/Private/SLevelViewportToolBar.cpp
+++ b/Engine/Source/Editor/LevelEditor/Private/SLevelViewportToolBar.cpp
@@ -832,25 +832,7 @@ void SLevelViewportToolBar::FillCameraMenu(UToolMenu* Menu) const
 		GeneratePlacedCameraMenuEntries(Section, LookThroughActors);
 	}
 
-	{
-		FLevelEditorModule& LevelEditorModule = FModuleManager::GetModuleChecked<FLevelEditorModule>(TEXT("LevelEditor"));
-
-		int32 NumCustomViewportTypes = 0;
-		LevelEditorModule.IterateViewportTypes([&](FName, const FViewportTypeDefinition&){ ++NumCustomViewportTypes; });
-
-		FText ViewportTypesHeading = LOCTEXT("ViewportTypes", "Viewport Type");
-		const uint32 MaxViewportTypesInTopLevelMenu = 4;
-		if( NumCustomViewportTypes > MaxViewportTypesInTopLevelMenu )
-		{
-			FToolMenuSection& Section = Menu->AddSection("ViewportTypes");
-			Section.AddSubMenu("ViewportTypes", ViewportTypesHeading, FText(), FNewToolMenuDelegate::CreateSP(this, &SLevelViewportToolBar::GenerateViewportTypeMenu) );
-		}
-		else
-		{
-			FToolMenuSection& Section = Menu->AddSection("ViewportTypes", ViewportTypesHeading);
-			GenerateViewportTypeMenu(Section);
-		}
-	}
+	UE::UnrealEd::GenerateViewportTypeMenu(Menu);
 }
 
 TSharedRef<SWidget> SLevelViewportToolBar::GenerateShowMenu() const
diff --git a/Engine/Source/Editor/LevelEditor/Private/ViewportToolbar/LevelEditorViewportToolbarSections.cpp b/Engine/Source/Editor/LevelEditor/Private/ViewportToolbar/LevelEditorViewportToolbarSections.cpp
--- a/Engine/Source/Editor/LevelEditor/Private/ViewportToolbar/LevelEditorViewportToolbarSections.cpp
+++ b/Engine/Source/Editor/LevelEditor/Private/ViewportToolbar/LevelEditorViewportToolbarSections.cpp
@@ -2345,6 +2345,9 @@ void ExtendCameraSubmenu(FName InCameraOptionsSubmenuName)
 							}
 						));
 					}
+
+					UE::UnrealEd::GenerateViewportTypeMenu(InDynamicMenu);
+
 				}
 
 				// Create Section
diff --git a/Engine/Source/Editor/UnrealEd/Private/ViewportToolbar/UnrealEdViewportToolbar.cpp b/Engine/Source/Editor/UnrealEd/Private/ViewportToolbar/UnrealEdViewportToolbar.cpp
--- a/Engine/Source/Editor/UnrealEd/Private/ViewportToolbar/UnrealEdViewportToolbar.cpp
+++ b/Engine/Source/Editor/UnrealEd/Private/ViewportToolbar/UnrealEdViewportToolbar.cpp
@@ -11,6 +11,7 @@
 #include "GPUSkinCache.h"
 #include "GPUSkinCacheVisualizationMenuCommands.h"
 #include "IPreviewProfileController.h"
+#include "LevelEditor.h"
 #include "LevelEditorActions.h"
 #include "RayTracingDebugVisualizationMenuCommands.h"
 #include "SEditorViewport.h"
@@ -2561,8 +2562,7 @@ void ExtendCameraSubmenu(FName InCameraOptionsSubmenuName)
 
 				FToolMenuInsert InsertPosition("LevelViewportCameraType_Ortho", EToolMenuInsertType::After);
 
-				FToolMenuSection& UnnamedSection =
-					InDynamicMenu->FindOrAddSection(NAME_None, FText(), InsertPosition);
+				FToolMenuSection& UnnamedSection = InDynamicMenu->FindOrAddSection(NAME_None, FText(), InsertPosition);
 				UnnamedSection.AddSeparator("CameraSubmenuSeparator");
 
 				UnnamedSection.AddEntry(FToolMenuEntry::InitWidget(
@@ -2581,6 +2581,59 @@ void ExtendCameraSubmenu(FName InCameraOptionsSubmenuName)
 	);
 }
 
+void GenerateViewportTypeMenu(UToolMenu* InMenu)
+{
+	FLevelEditorModule& LevelEditorModule = FModuleManager::GetModuleChecked<FLevelEditorModule>(TEXT("LevelEditor"));
+
+	int32 NumCustomViewportTypes = 0;
+	LevelEditorModule.IterateViewportTypes(
+		[&NumCustomViewportTypes](FName, const FViewportTypeDefinition&)
+		{
+			++NumCustomViewportTypes;
+		}
+	);
+
+	FText ViewportTypesHeading = LOCTEXT("ViewportTypes", "Viewport Type");
+	constexpr uint32 MaxViewportTypesInTopLevelMenu = 4;
+
+	auto PupoulateSection = [&LevelEditorModule](FToolMenuSection& InSection)
+	{
+		LevelEditorModule.IterateViewportTypes(
+			[&InSection](FName ViewportTypeName, const FViewportTypeDefinition& InDefinition)
+			{
+				if (InDefinition.ActivationCommand.IsValid())
+				{
+					InSection.AddMenuEntry(
+						*FString::Printf(TEXT("ViewportType_%s"), *ViewportTypeName.ToString()), InDefinition.ActivationCommand
+					);
+				}
+			}
+		);
+	};
+
+	if (NumCustomViewportTypes > MaxViewportTypesInTopLevelMenu)
+	{
+		FToolMenuSection& Section = InMenu->AddSection("ViewportTypes");
+		Section.AddSubMenu(
+			"ViewportTypes",
+			ViewportTypesHeading,
+			FText(),
+			FNewToolMenuDelegate::CreateLambda(
+				[PupoulateSection](UToolMenu* InSubmenu)
+				{
+					FToolMenuSection& Section = InSubmenu->AddSection(NAME_None);
+					PupoulateSection(Section);
+				}
+			)
+		);
+	}
+	else
+	{
+		FToolMenuSection& Section = InMenu->AddSection("ViewportTypes", ViewportTypesHeading);
+		PupoulateSection(Section);
+	}
+}
+
 static FFormatNamedArguments GetScreenPercentageFormatArguments(const FEditorViewportClient& ViewportClient)
 {
 	const UEditorPerformanceProjectSettings* EditorProjectSettings = GetDefault<UEditorPerformanceProjectSettings>();
diff --git a/Engine/Source/Editor/UnrealEd/Public/ViewportToolbar/UnrealEdViewportToolbar.h b/Engine/Source/Editor/UnrealEd/Public/ViewportToolbar/UnrealEdViewportToolbar.h
--- a/Engine/Source/Editor/UnrealEd/Public/ViewportToolbar/UnrealEdViewportToolbar.h
+++ b/Engine/Source/Editor/UnrealEd/Public/ViewportToolbar/UnrealEdViewportToolbar.h
@@ -149,6 +149,8 @@ UNREALED_API void PopulateCameraMenu(UToolMenu* InMenu);
  */
 UNREALED_API void ExtendCameraSubmenu(FName InCameraOptionsSubmenuName);
 
+UNREALED_API void GenerateViewportTypeMenu(UToolMenu* InMenu);
+
 UNREALED_API bool ShouldShowViewportRealtimeWarning(const FEditorViewportClient& ViewportClient);
 
 UNREALED_API FToolMenuEntry CreatePerformanceAndScalabilitySubmenu();
