Fixed regression brought on by deprecating SetAnimClass

DebugSkeletalMeshComponent ended up skipping the logic for overriding the anim instance class due to refactored virtual function

diff --git a/Engine/Source/Editor/UnrealEd/Classes/Animation/DebugSkelMeshComponent.h b/Engine/Source/Editor/UnrealEd/Classes/Animation/DebugSkelMeshComponent.h
--- a/Engine/Source/Editor/UnrealEd/Classes/Animation/DebugSkelMeshComponent.h
+++ b/Engine/Source/Editor/UnrealEd/Classes/Animation/DebugSkelMeshComponent.h
@@ -397,7 +397,6 @@ class UDebugSkelMeshComponent : public USkeletalMeshComponent
 	//~ Begin SkeletalMeshComponent Interface
 	UNREALED_API virtual void InitAnim(bool bForceReinit) override;
 	virtual bool IsWindEnabled() const override { return true; }
-	UNREALED_API virtual void SetAnimInstanceClass(class UClass* NewClass) override;
 	UNREALED_API virtual void OnClearAnimScriptInstance() override;
 	UNREALED_API virtual void SetSkeletalMesh(USkeletalMesh* InSkelMesh, bool bReinitPose = true) override;
 	//~ End SkeletalMeshComponent Interface
diff --git a/Engine/Source/Editor/UnrealEd/Private/Animation/DebugSkelMeshComponent.cpp b/Engine/Source/Editor/UnrealEd/Private/Animation/DebugSkelMeshComponent.cpp
--- a/Engine/Source/Editor/UnrealEd/Private/Animation/DebugSkelMeshComponent.cpp
+++ b/Engine/Source/Editor/UnrealEd/Private/Animation/DebugSkelMeshComponent.cpp
@@ -612,12 +612,6 @@ void UDebugSkelMeshComponent::InitAnim(bool bForceReinit)
 	}
 }
 
-void UDebugSkelMeshComponent::SetAnimInstanceClass(class UClass* NewClass)
-{
-	// Override this to do nothing and warn the user
-	UE_LOG(LogAnimation, Warning, TEXT("Attempting to destroy an animation preview actor, skipping."));
-}
-
 void UDebugSkelMeshComponent::OnClearAnimScriptInstance()
 {
 	// call to super not strictly necessary (since it is empty)
