Change BP reinstancing to use the archetype caching system that was previously only for override serialization.
This fixes blueprint component data corruption problem where placed instances and subclass templates could lose their property changes when compiling the component class, and replaces the partial fix in 36403142.

diff --git a/Engine/Source/Editor/Kismet/Private/BlueprintCompilationManager.cpp b/Engine/Source/Editor/Kismet/Private/BlueprintCompilationManager.cpp
--- a/Engine/Source/Editor/Kismet/Private/BlueprintCompilationManager.cpp
+++ b/Engine/Source/Editor/Kismet/Private/BlueprintCompilationManager.cpp
@@ -52,6 +52,7 @@
 #include "Stats/StatsHierarchical.h"
 #include "UObject/PropertyBagRepository.h"
 #include "UObject/OverridableManager.h"
+#include "UObject/UObjectArchetypeHelper.h"
 
 extern UNREALED_API UUnrealEdEngine* GUnrealEd;
 
@@ -2738,6 +2739,8 @@ void FBlueprintCompilationManagerImpl::ReinstanceBatch(TArray<FReinstancingJob>&
 
 			UClass* NewClass = ReinstancingJob.OldToNew.Value;
 
+			TSet<UObject*> CachedArchetypeObjects;
+			FEditorCacheArchetypeManager& CacheManager = FEditorCacheArchetypeManager::Get();
 			TMap<UObject*, UObject*>* OldToNewTemplatesForClass = OldToNewTemplates ? &OldToNewTemplates->FindOrAdd(OldClass) : nullptr;
 
 			for(const FReinstancingJob::FArchetypeInfo& OldArchetypeInfo : ReinstancingJob.OldArchetypeObjects)
@@ -2751,19 +2754,19 @@ void FBlueprintCompilationManagerImpl::ReinstanceBatch(TArray<FReinstancingJob>&
 
 				// We need to cache the archetype of the instances of this archetype
 				// as it will not be possible to get them afterwards as it gets renamed
-				if (FOverridableManager::Get().IsEnabled(*OldInstance))
+				// These cached archetypes will not be updated if they were set earlier
+				TArray<UObject*> ArchetypeInstances;
+				OldInstance->GetArchetypeInstances(ArchetypeInstances);
+				for (UObject* ArchetypeInstance : ArchetypeInstances)
 				{
-					TArray<UObject*> ArchetypeInstances;
-					OldInstance->GetArchetypeInstances(ArchetypeInstances);
-					for (UObject* ArchetypeInstance : ArchetypeInstances)
-					{
-						FOverridableManager::Get().CacheArchetype(*ArchetypeInstance);
+					CacheManager.CacheArchetype(ArchetypeInstance);
+					CachedArchetypeObjects.Add(ArchetypeInstance);
 
-						ForEachObjectWithOuter(ArchetypeInstance, [](UObject* SubObject)
-						{
-							FOverridableManager::Get().CacheArchetype(*SubObject);
-						});
-					}
+					ForEachObjectWithOuter(ArchetypeInstance, [&CachedArchetypeObjects, &CacheManager](UObject* SubObject)
+					{
+						CacheManager.CacheArchetype(SubObject);
+						CachedArchetypeObjects.Add(SubObject);
+					});
 				}
 
 				UObject* Destination = GetTransientOuterForRename(OldInstance->GetClass());
@@ -2831,6 +2834,13 @@ void FBlueprintCompilationManagerImpl::ReinstanceBatch(TArray<FReinstancingJob>&
 					OldInstance->MarkAsGarbage();
 				}
 			}
+
+			// Reset any cached archetypes
+			for (UObject* CachedArchetypeObject : CachedArchetypeObjects)
+			{
+				CacheManager.ResetCacheArchetype(CachedArchetypeObject);
+			}
+
 		}
 	}
 
diff --git a/Engine/Source/Editor/UnrealEd/Private/Kismet2/KismetReinstanceUtilities.cpp b/Engine/Source/Editor/UnrealEd/Private/Kismet2/KismetReinstanceUtilities.cpp
--- a/Engine/Source/Editor/UnrealEd/Private/Kismet2/KismetReinstanceUtilities.cpp
+++ b/Engine/Source/Editor/UnrealEd/Private/Kismet2/KismetReinstanceUtilities.cpp
@@ -47,6 +47,7 @@
 #include "UObject/OverridableManager.h"
 #include "UObject/PropertyOptional.h"
 #include "UObject/PropertyBagRepository.h"
+#include "UObject/UObjectArchetypeHelper.h"
 #include "ProfilingDebugging/LoadTimeTracker.h"
 
 DECLARE_CYCLE_STAT(TEXT("Replace Instances"), EKismetReinstancerStats_ReplaceInstancesOfClass, STATGROUP_KismetReinstancer );
@@ -2211,8 +2212,6 @@ static void ReplaceObjectHelper(UObject*& OldObject, UClass* OldClass, UObject*&
 
 	FName OldName(OldObject->GetFName());
 
-	UObject* OldArchetype = nullptr;
-
 	// If the old object is in this table, we've already renamed it away in a previous iteration. Don't rename it again!
 	if (!OldToNewNameMap.Contains(OldObject))
 	{
@@ -2234,7 +2233,6 @@ static void ReplaceObjectHelper(UObject*& OldObject, UClass* OldClass, UObject*&
 		}
 		else
 		{
-			OldArchetype = OldObject->GetArchetype(); // Cache the old object's archetype before renaming
 			OldObject->Rename(nullptr, OldObject->GetOuter(), REN_DoNotDirty | REN_DontCreateRedirectors | REN_AllowPackageLinkerMismatch);
 		}
 	}
@@ -2303,15 +2301,6 @@ static void ReplaceObjectHelper(UObject*& OldObject, UClass* OldClass, UObject*&
 	// We only need to copy properties of the pre-created instances, the rest of the default sub object is done inside the UEditorEngine::CopyPropertiesForUnrelatedObjects
 	for (const auto& Pair : OrderedListOfObjectToCopy)
 	{
-		if (Pair.Key == OldObject)
-		{
-			// If we're copying the object itself, make sure to use the archetype that was cached before
-			Options.SourceObjectArchetype = OldArchetype;
-		}
-		else
-		{
-			Options.SourceObjectArchetype = nullptr;
-		}
 		UEditorEngine::CopyPropertiesForUnrelatedObjects(Pair.Key, Pair.Value, Options);
 	}
 
@@ -2652,6 +2641,8 @@ void FBlueprintCompileReinstancer::ReplaceInstancesOfClass_Inner(const TMap<UCla
 
 	{
 		TArray<UObject*> ObjectsToReplace;
+		TSet<UObject*> CachedArchetypeObjects;
+		FEditorCacheArchetypeManager& CacheManager = FEditorCacheArchetypeManager::Get();
 
 		BP_SCOPED_COMPILER_EVENT_STAT(EKismetReinstancerStats_ReplaceInstancesOfClass);
 
@@ -2727,23 +2718,26 @@ void FBlueprintCompileReinstancer::ReplaceInstancesOfClass_Inner(const TMap<UCla
 						}
 						return Dependencies;
 					});
+				}
 
-					// We need to cache the archetype of the objects about to be replaced 
-					// as it will not be possible to get them during this process as it renames the objects
-					for (UObject* OldObject : ObjectsToReplace)
+				// We need to cache the archetype of the objects about to be replaced 
+				// as it will not be possible to get them during this process as it renames the objects
+				// These cached archetypes will not be updated if they were set earlier
+				for (UObject* OldObject : ObjectsToReplace)
+				{
+					if(!IsValid(OldObject))
 					{
-						if(!IsValid(OldObject))
-						{
-							continue;
-						}
+						continue;
+					}
 
-						FOverridableManager::Get().CacheArchetype(*OldObject);
+					CacheManager.CacheArchetype(OldObject);
+					CachedArchetypeObjects.Add(OldObject);
 
-						ForEachObjectWithOuter(OldObject, [](UObject* SubObject)
-						{
-							FOverridableManager::Get().CacheArchetype(*SubObject);
-						});
-					}
+					ForEachObjectWithOuter(OldObject, [&CacheManager, &CachedArchetypeObjects](UObject* SubObject)
+					{
+						CacheManager.CacheArchetype(SubObject);
+						CachedArchetypeObjects.Add(SubObject);
+					});
 				}
 				
 				// Then fix 'real' (non archetype) instances of the class
@@ -2896,6 +2890,13 @@ void FBlueprintCompileReinstancer::ReplaceInstancesOfClass_Inner(const TMap<UCla
 				}
 			}
 		}
+
+		// Reset any cached archetypes
+		for (UObject* CachedArchetypeObject : CachedArchetypeObjects)
+		{
+			CacheManager.ResetCacheArchetype(CachedArchetypeObject);
+		}
+
 		if (GEngine)
 		{
 			GEngine->OnLevelActorDeleted().Remove(OnLevelActorDeletedHandle);
diff --git a/Engine/Source/Runtime/CoreUObject/Internal/UObject/UObjectArchetypeHelper.h b/Engine/Source/Runtime/CoreUObject/Internal/UObject/UObjectArchetypeHelper.h
--- a/Engine/Source/Runtime/CoreUObject/Internal/UObject/UObjectArchetypeHelper.h
+++ b/Engine/Source/Runtime/CoreUObject/Internal/UObject/UObjectArchetypeHelper.h
@@ -3,6 +3,8 @@
 #pragma once
 
 #include "CoreFwd.h"
+#include "UObject/Object.h"
+#include "UObject/UObjectAnnotation.h"
 
 class FObjectArchetypeHelper
 {
@@ -22,3 +24,58 @@ private:
 #endif
 	COREUOBJECT_API static UObject* GetArchetype(const UObject* InObject, const IObjectArchetypePolicy* InPolicy);
 };
+
+#if WITH_EDITOR
+
+struct FCacheArchetype
+{
+	bool IsDefault() const
+	{
+		return CachedArchetype == nullptr;
+	}
+
+	/** Archetype is cache when it is or its archetype is about to be replaced **/
+	TObjectPtr<UObject> CachedArchetype = nullptr;
+};
+
+class FEditorCacheArchetypeManager
+{
+public:
+	/**
+	 * @return the static instance managing the reinstantiation */
+	COREUOBJECT_API static FEditorCacheArchetypeManager& Get();
+
+	/**
+	 * Caches the archetype for this object during reinstantiation has we rename object in that phase.
+	 * This is only a utility method to ensure the archetype is not lost during reinstantiation
+	 * as it renames and changes outers which makes impossible to retrieve it during the process.
+	 * @param Object top cache its archetype
+	 * @param Archetype to use as a cache, if null it will use GetArchetype on the owner of the specified object */
+	void CacheArchetype(gsl::not_null<const UObject*> Object, UObject* Archetype = nullptr)
+	{
+		FCacheArchetype Annotation = ObjectCachedArchetypeAnnotations.GetAnnotation(Object);
+		if (Annotation.CachedArchetype == nullptr)
+		{
+			Annotation.CachedArchetype = Archetype ? Archetype : Object->GetArchetype();
+			ObjectCachedArchetypeAnnotations.AddAnnotation(Object, Annotation);
+		}
+	}
+
+	void ResetCacheArchetype(gsl::not_null<const UObject*> Object)
+	{
+		ObjectCachedArchetypeAnnotations.RemoveAnnotation(Object);
+	}
+
+	/**
+	 * Gets the archetype that was cached when the CacheArchetype method was called
+	 * @return the cached Archetype if any */
+	UObject* GetCachedArchetype(gsl::not_null<const UObject*> Object)
+	{
+		return ObjectCachedArchetypeAnnotations.GetAnnotation(Object).CachedArchetype;
+	}
+
+private:
+	FUObjectAnnotationSparse<FCacheArchetype, true/*bAutoRemove*/> ObjectCachedArchetypeAnnotations;
+};
+
+#endif // WITH_EDITOR
diff --git a/Engine/Source/Runtime/CoreUObject/Private/UObject/OverridableManager.cpp b/Engine/Source/Runtime/CoreUObject/Private/UObject/OverridableManager.cpp
--- a/Engine/Source/Runtime/CoreUObject/Private/UObject/OverridableManager.cpp
+++ b/Engine/Source/Runtime/CoreUObject/Private/UObject/OverridableManager.cpp
@@ -54,16 +54,6 @@ void FOverridableManager::InheritEnabledFrom(UObject& Object, const UObject* Def
 #endif // WITH_EDITORONLY_DATA
 }
 
-void FOverridableManager::CacheArchetype(UObject& Object, UObject* Archetype /*= nullptr*/)
-{
-#if WITH_EDITORONLY_DATA
-	if (FOverriddenPropertySet* OverriddenProperties = OverriddenObjectAnnotations.Find(Object))
-	{
-		OverriddenProperties->CacheArchetype(Archetype);
-	}
-#endif // WITH_EDITORONLY_DATA
-}
-
 bool FOverridableManager::NeedSubObjectTemplateInstantiation(const UObject& Object)
 {
 #if WITH_EDITORONLY_DATA
diff --git a/Engine/Source/Runtime/CoreUObject/Private/UObject/OverriddenPropertySet.cpp b/Engine/Source/Runtime/CoreUObject/Private/UObject/OverriddenPropertySet.cpp
--- a/Engine/Source/Runtime/CoreUObject/Private/UObject/OverriddenPropertySet.cpp
+++ b/Engine/Source/Runtime/CoreUObject/Private/UObject/OverriddenPropertySet.cpp
@@ -7,6 +7,7 @@
 #include "Serialization/ArchiveSerializedPropertyChain.h"
 #include "UObject/PropertyOptional.h"
 #include "Misc/ScopeExit.h"
+#include "UObject/UObjectArchetypeHelper.h"
 #include "UObject/UObjectThreadContext.h"
 
 /*
@@ -1221,14 +1222,6 @@ void FOverriddenPropertySet::RemoveOverriddenSubProperties(FOverriddenPropertyNo
 	PropertyNode.SubPropertyNodeKeys.Empty();
 }
 
-void FOverriddenPropertySet::CacheArchetype(UObject* Archetype /*= nullptr*/)
-{
-	if (!CachedArchetype)
-	{
-		CachedArchetype = Archetype ? Archetype : Owner->GetArchetype();
-	}
-}
-
 EOverriddenPropertyOperation FOverriddenPropertySet::GetOverriddenPropertyOperation(FPropertyVisitorPath::Iterator PropertyIterator, bool* bOutInheritedOperation /*= nullptr*/) const
 {
 	if (const FOverriddenPropertyNode* RootNode = OverriddenPropertyNodes.Find(RootNodeID))
@@ -1392,12 +1385,14 @@ void FOverriddenPropertySet::Reset()
 
 void FOverriddenPropertySet::HandleObjectsReInstantiated(const TMap<UObject*, UObject*>& Map)
 {
-	// When the saved archetype is set, it is an indicator this object is about to be replaced
+#if WITH_EDITOR
+	// When their is a cached archetype, it is an indicator this object is about to be replaced
 	// So no need to replace any ptr, otherwise we might not be able to reconstitute the right information
-	if (CachedArchetype)
+	if (FEditorCacheArchetypeManager::Get().GetCachedArchetype(Owner))
 	{
 		return;
 	}
+#endif // WITH_EDITOR
 
 	for (FOverriddenPropertyNode& Node : OverriddenPropertyNodes)
 	{
diff --git a/Engine/Source/Runtime/CoreUObject/Private/UObject/UObjectArchetype.cpp b/Engine/Source/Runtime/CoreUObject/Private/UObject/UObjectArchetype.cpp
--- a/Engine/Source/Runtime/CoreUObject/Private/UObject/UObjectArchetype.cpp
+++ b/Engine/Source/Runtime/CoreUObject/Private/UObject/UObjectArchetype.cpp
@@ -71,6 +71,14 @@ FAutoConsoleVariableRef CVarEnableArchetypeCache(
 }
 #endif // UE_CACHE_ARCHETYPE
 
+#if WITH_EDITOR
+FEditorCacheArchetypeManager& FEditorCacheArchetypeManager::Get()
+{
+	static FEditorCacheArchetypeManager Manager;
+	return Manager;
+}
+#endif // WITH_EDITOR
+
 UObject* GetArchetypeImpl(const UObject* InObject, const FObjectArchetypeHelper::IObjectArchetypePolicy* Policy);
 
 UObject* GetArchetypeFromRequiredInfoImpl(const UClass* Class, const UObject* Outer, FName Name, EObjectFlags ObjectFlags, bool bUseUpToDateClass, const FObjectArchetypeHelper::IObjectArchetypePolicy* Policy)
@@ -205,13 +213,10 @@ UObject* GetArchetypeImpl(const UObject* InObject, const FObjectArchetypeHelper:
 		}
 	}
 
-	if (const FOverriddenPropertySet* OverriddenProperties = FOverridableManager::Get().GetOverriddenProperties(*InObject))
+	if (UObject* CacheArchetype = FEditorCacheArchetypeManager::Get().GetCachedArchetype(InObject))
 	{
 		// Use the cached archetype if set
-		if (UObject* CacheArchetype = OverriddenProperties->GetCachedArchetype())
-		{
-			return CacheArchetype;
-		}
+		return CacheArchetype;
 	}
 
 #endif
diff --git a/Engine/Source/Runtime/CoreUObject/Public/UObject/OverridableManager.h b/Engine/Source/Runtime/CoreUObject/Public/UObject/OverridableManager.h
--- a/Engine/Source/Runtime/CoreUObject/Public/UObject/OverridableManager.h
+++ b/Engine/Source/Runtime/CoreUObject/Public/UObject/OverridableManager.h
@@ -128,15 +128,6 @@ public:
 	 * @param DefaultData its default object */
 	COREUOBJECT_API void InheritEnabledFrom(UObject& Object, const UObject* DefaultData);
 
-	/**
-	 * Caches the archetype for this object if it has overridable serialization enabled.
-	 * NOTE: ONLY CACHE THE ARCHETYPE POINTER ON OBJECTS THAT ARE ABOUT TO BE REINSTANTIATED.
-	 * This is only a utility method to ensure the archetype is not lost during reinstantiation
-	 * as it renames and changes outers which makes impossible to retrieve it during the process.
-	 * @param Object top cache its archetype
-	 * @param Archetype to use as a cache, if null it will use GetArchetype on the owner of the specified object */
-	COREUOBJECT_API void CacheArchetype(UObject& Object, UObject* Archetype = nullptr);
-
 	/**
 	 * Return true if this object needs subobject template instantiation
 	 * @param Object to be querying about */
diff --git a/Engine/Source/Runtime/CoreUObject/Public/UObject/OverriddenPropertySet.h b/Engine/Source/Runtime/CoreUObject/Public/UObject/OverriddenPropertySet.h
--- a/Engine/Source/Runtime/CoreUObject/Public/UObject/OverriddenPropertySet.h
+++ b/Engine/Source/Runtime/CoreUObject/Public/UObject/OverriddenPropertySet.h
@@ -379,20 +379,6 @@ public:
 	void Reset();
 	void HandleObjectsReInstantiated(const TMap<UObject*, UObject*>& Map);
 
-	/**
-	 * Caches the archetype of this object for future reference
-	 * @param Archetype to use as a cache, if null it will use GetArchetype on the owner of this object
-	 */
-	void CacheArchetype(UObject* Archetype = nullptr);
-
-	/**
-	 * Gets the archetype that was cached when the CacheArchetype method was called
-	 * @return the cached Archetype if any */
-	UObject* GetCachedArchetype() const
-	{
-		return CachedArchetype;
-	}
-
 protected:
 
 	FOverriddenPropertyNode& FindOrAddNode(FOverriddenPropertyNode& ParentPropertyNode, FOverriddenPropertyNodeID NodeID);
@@ -417,8 +403,6 @@ private:
 
 	static inline FOverriddenPropertyNodeID RootNodeID = FOverriddenPropertyNodeID::RootNodeId();
 
-	/** Archetype is cache when it is or its archetype is about to be replaced **/
-	TObjectPtr<UObject> CachedArchetype = nullptr;
 public:
 	bool bNeedsSubobjectTemplateInstantiation = false;
 };
