MPCook: CookPackageSplitter: Fix retraction edge case of LocalWorker Saving generated package -> retracted to remote worker -> retracted back to local worker. When a package is assigned to the local worker and is in the save state and is retracted from local worker -> remote worker, or this happened previously and we stalled it and we are now retracting it back to local worker, we need to keep it in the Saving states, moving it from Save -> SaveStalled or SaveStalled -> Save, rather than kicking it back to request. Kicking it back to request will cause it to drop the data that keeps it from being garbage collected, and being garbage collected causes an error and is the reason we created the SaveStalled states. This fixes another case of the error "PackageSplitter specified generated package is map in GetGenerateList results, but then in PreSaveGeneratedPackage created it as uasset".
#rnx

diff --git a/Engine/Source/Editor/UnrealEd/Private/Cooker/CookDirector.cpp b/Engine/Source/Editor/UnrealEd/Private/Cooker/CookDirector.cpp
--- a/Engine/Source/Editor/UnrealEd/Private/Cooker/CookDirector.cpp
+++ b/Engine/Source/Editor/UnrealEd/Private/Cooker/CookDirector.cpp
@@ -2088,7 +2088,10 @@ FCookDirector::FRetractionHandler::ReassignPackages(const FWorkerId& FromWorker,
 		}
 
 		AssignmentPackages.Add(PackageData);
-		PackageData->SendToState(EPackageState::Request, ESendFlags::QueueRemove, EStateChangeReason::Retraction);
+		// The package is no longer assigned to the worker, and the worker already knows it, so remove its assignment.
+		// We have to remove the assignment so that we can reassign it down below; it is invalid to call
+		// SetWorkerAssignment to a new worker without clearing the old worker first.
+		PackageData->SetWorkerAssignment(FWorkerId::Invalid());
 	}
 	if (AssignmentPackages.IsEmpty())
 	{
@@ -2147,7 +2150,16 @@ FCookDirector::FRetractionHandler::ReassignPackages(const FWorkerId& FromWorker,
 		}
 		else if (Assignment.IsLocal())
 		{
-			RequestQueue.AddReadyRequest(PackageData);
+			// UnStall the package to handle the case where it was previously in the save state on the LocalWorker
+			// and was retracted to a remote worker, so we stalled it, but now we've pulled it back to the LocalWorker.
+			PackageData->UnStall(ESendFlags::QueueAddAndRemove);
+			// If the packagedata was unstalled back into the Save State, then keep it there, but otherwise it is in
+			// the AssignedToWorker state so kick it back to readyrequest.
+			if (PackageData->GetState() != EPackageState::SaveActive)
+			{
+				PackageData->SendToState(EPackageState::Request, ESendFlags::QueueRemove, EStateChangeReason::Retraction);
+				RequestQueue.AddReadyRequest(PackageData);
+			}
 			bAssignedToLocal = true;
 		}
 		else
