Show both the Play and Simulate buttons in the Level Editor toolbar

diff --git a/Engine/Source/Editor/EditorStyle/Private/StarshipStyle.cpp b/Engine/Source/Editor/EditorStyle/Private/StarshipStyle.cpp
--- a/Engine/Source/Editor/EditorStyle/Private/StarshipStyle.cpp
+++ b/Engine/Source/Editor/EditorStyle/Private/StarshipStyle.cpp
@@ -5633,11 +5633,19 @@ void FStarshipEditorStyle::FStyle::SetupPersonaStyle()
 		Set( "PlayWorld.PlayInVulkanPreview", new IMAGE_BRUSH_SVG("Starship/MainToolbar/PlayMobilePreview", Icon20x20));
 		Set( "PlayWorld.PlayInNewProcess", new IMAGE_BRUSH_SVG("Starship/MainToolbar/PlayStandaloneGame", Icon20x20));
 
+		Set("PlayWorld.PlayPossess", new IMAGE_BRUSH_SVG("Starship/MainToolbar/PlayPossess_20", Icon20x20));
+
+		Set("PlayWorld.SimulatePauseEject_PlaySimulate", new IMAGE_BRUSH_SVG("Starship/MainToolbar/PlaySimulate_20", Icon20x20));
+		Set("PlayWorld.SimulatePauseEject_Eject", new IMAGE_BRUSH_SVG("Starship/MainToolbar/eject", Icon40x40));
+		Set("PlayWorld.SimulatePauseEject_Pause", new IMAGE_BRUSH_SVG("Starship/MainToolbar/pause", Icon40x40));
+
 		Set( "PlayWorld.RepeatLastLaunch", new IMAGE_BRUSH_SVG( "Starship/MainToolbar/launch", Icon20x20 ) );
 
 		Set( "PlayWorld.PlayInCameraLocation", new IMAGE_BRUSH_SVG( "Starship/AssetIcons/CameraActor_16", Icon20x20 ) );
 		Set( "PlayWorld.PlayInDefaultPlayerStart", new IMAGE_BRUSH_SVG( "Starship/AssetIcons/PlayerStart_16", Icon20x20 ) );
 
+		Set("PlayWorld.ResumePIESession", new IMAGE_BRUSH_SVG("Starship/Common/play", Icon40x40));
+		Set("PlayWorld.ResumePIESession.Small", new IMAGE_BRUSH_SVG("Starship/Common/play", Icon20x20));
 		Set("PlayWorld.ResumePlaySession", new IMAGE_BRUSH_SVG("Starship/MainToolbar/simulate", Icon40x40));
 		Set("PlayWorld.ResumePlaySession.Small", new IMAGE_BRUSH_SVG("Starship/MainToolbar/simulate", Icon20x20));
 		Set( "PlayWorld.PausePlaySession", new IMAGE_BRUSH_SVG("Starship/MainToolbar/pause", Icon40x40));
diff --git a/Engine/Source/Editor/UnrealEd/Private/Kismet2/DebuggerCommands.cpp b/Engine/Source/Editor/UnrealEd/Private/Kismet2/DebuggerCommands.cpp
--- a/Engine/Source/Editor/UnrealEd/Private/Kismet2/DebuggerCommands.cpp
+++ b/Engine/Source/Editor/UnrealEd/Private/Kismet2/DebuggerCommands.cpp
@@ -174,8 +174,6 @@ public:
 
 	static void HandleShowSDKTutorial(FString PlatformName, FString NotInstalledDocLink);
 
-	static FSlateIcon GetResumePlaySessionImage();
-	static FText GetResumePlaySessionToolTip();
 	static void StopPlaySession_Clicked();
 	static void LateJoinSession_Clicked();
 	static void SingleFrameAdvance_Clicked();
@@ -192,12 +190,21 @@ public:
 	// Mouse control
 	static void GetMouseControlExecute();
 
+	static void PossessPlayer_Clicked();
+	static bool CanPossessPlayer();
+
 	static void PossessEjectPlayer_Clicked();
 	static bool CanPossessEjectPlayer();
 	static FText GetPossessEjectLabel();
 	static FText GetPossessEjectTooltip();
 	static FSlateIcon GetPossessEjectImage();
 
+	static void SimulatePauseEject_Clicked();
+	static bool CanSimulatePauseEject();
+	static FText GetSimulatePauseEjectLabel();
+	static FText GetSimulatePauseEjectTooltip();
+	static FSlateIcon GetSimulatePauseEjectIcon();
+
 	static bool CanLateJoin();
 	static bool CanShowLateJoinButton();
 
@@ -363,11 +370,14 @@ void FPlayWorldCommands::RegisterCommands()
 	UI_COMMAND(StopPlaySession, "Stop", "Stop simulation", EUserInterfaceActionType::Button, FInputChord(EKeys::Escape));
 	UI_COMMAND(ResumePlaySession, "Resume", "Resume simulation", EUserInterfaceActionType::Button, FInputChord());
 	UI_COMMAND(PausePlaySession, "Pause", "Pause simulation", EUserInterfaceActionType::Button, FInputChord());
+	UI_COMMAND(ResumePIESession, "Resume", "Resume play-in-editor simulation", EUserInterfaceActionType::Button, FInputChord());
 	UI_COMMAND(GetMouseControl, "Mouse Control", "Get mouse cursor while in PIE", EUserInterfaceActionType::Button, FInputChord(EModifierKey::Shift, EKeys::F1));
 	UI_COMMAND(LateJoinSession, "Add Client", "Add another client", EUserInterfaceActionType::Button, FInputChord());
 	UI_COMMAND(SingleFrameAdvance, "Skip", "Advances a single frame", EUserInterfaceActionType::Button, FInputChord());
 	UI_COMMAND(TogglePlayPauseOfPlaySession, "Toggle Play/Pause", "Resume playing if paused, or pause if playing", EUserInterfaceActionType::Button, FInputChord(EKeys::Pause));
+	UI_COMMAND(PlayPossess, "Possess Player", "Possesses the player character", EUserInterfaceActionType::Button, FInputChord(EKeys::F8));
 	UI_COMMAND(PossessEjectPlayer, "Possess or Eject Player", "Possesses or ejects the player from the camera", EUserInterfaceActionType::Button, FInputChord(EKeys::F8));
+	UI_COMMAND(SimulatePauseEject, "Simulate, Pause Simulation or Eject Player", "Starts or pauses simulation or eject the player from the camera", EUserInterfaceActionType::Button, FInputChord());
 	UI_COMMAND(ShowCurrentStatement, "Locate", "Locate the currently active node", EUserInterfaceActionType::Button, FInputChord());
 	UI_COMMAND(AbortExecution, "Abort", "Abort the execution", EUserInterfaceActionType::Button, FInputChord());
 	UI_COMMAND(ContinueExecution, "Continue", "Continue the execution", EUserInterfaceActionType::Button, FInputChord(EModifierKey::Alt | EModifierKey::Shift, EKeys::F10));
@@ -432,10 +442,19 @@ void FPlayWorldCommands::BindGlobalPlayWorldCommands()
 		FIsActionButtonVisible::CreateStatic(&FInternalPlayWorldCommandCallbacks::CanShowNonPlayWorldOnlyActions)
 	);
 
-	ActionList.MapAction(Commands.PlayInViewport,
+	ActionList.MapAction(
+		Commands.PlayInViewport,
 		FExecuteAction::CreateStatic(&FInternalPlayWorldCommandCallbacks::PlayInViewport_Clicked),
 		FCanExecuteAction::CreateStatic(&FInternalPlayWorldCommandCallbacks::PlayInViewport_CanExecute),
-		FIsActionChecked::CreateStatic(&FInternalPlayWorldCommandCallbacks::PlayInModeIsChecked, PlayMode_InViewPort),
+		FIsActionChecked::CreateLambda(
+			[]() -> bool
+			{
+				// The toolbar play button launches PIE whenever we last ran SIE or PIE.
+				// This logic makes sure the checked state in the ... submenu also indicates the same.
+				return FInternalPlayWorldCommandCallbacks::PlayInModeIsChecked(PlayMode_InViewPort)
+					|| FInternalPlayWorldCommandCallbacks::PlayInModeIsChecked(PlayMode_Simulate);
+			}
+		),
 		FIsActionButtonVisible::CreateStatic(&FInternalPlayWorldCommandCallbacks::CanShowNonPlayWorldOnlyActions)
 	);
 
@@ -524,6 +543,14 @@ void FPlayWorldCommands::BindGlobalPlayWorldCommands()
 		FIsActionButtonVisible::CreateStatic(&FPlayWorldCommandCallbacks::HasPlayWorldAndRunning)
 	);
 
+	ActionList.MapAction(
+		Commands.ResumePIESession,
+		FExecuteAction::CreateStatic(&FPlayWorldCommandCallbacks::ResumePlaySession_Clicked),
+		FCanExecuteAction::CreateStatic(&FPlayWorldCommandCallbacks::IsInPIE_AndPaused),
+		FIsActionChecked(),
+		FIsActionButtonVisible::CreateStatic(&FPlayWorldCommandCallbacks::IsInPIE_AndPaused)
+	);
+
 	ActionList.MapAction(Commands.SingleFrameAdvance,
 		FExecuteAction::CreateStatic(&FInternalPlayWorldCommandCallbacks::SingleFrameAdvance_Clicked),
 		FCanExecuteAction::CreateStatic(&FPlayWorldCommandCallbacks::HasPlayWorldAndPaused),
@@ -546,6 +573,15 @@ void FPlayWorldCommands::BindGlobalPlayWorldCommands()
 		FIsActionButtonVisible::CreateStatic(&FInternalPlayWorldCommandCallbacks::HasPlayWorld)
 	);
 
+	// Possess the player (SIE->PIE).
+	ActionList.MapAction(
+		Commands.PlayPossess,
+		FExecuteAction::CreateStatic(&FInternalPlayWorldCommandCallbacks::PossessPlayer_Clicked),
+		FCanExecuteAction::CreateStatic(&FInternalPlayWorldCommandCallbacks::CanPossessPlayer),
+		FIsActionChecked(),
+		FIsActionButtonVisible::CreateStatic(&FInternalPlayWorldCommandCallbacks::CanPossessPlayer)
+	);
+
 	// Toggle PIE/SIE, Eject (PIE->SIE), and Possess (SIE->PIE)
 	ActionList.MapAction(Commands.PossessEjectPlayer,
 		FExecuteAction::CreateStatic(&FInternalPlayWorldCommandCallbacks::PossessEjectPlayer_Clicked),
@@ -554,6 +590,15 @@ void FPlayWorldCommands::BindGlobalPlayWorldCommands()
 		FIsActionButtonVisible::CreateStatic(&DebuggerCommands::AllowPlayWorldFeature)
 	);
 
+	// Starts SIE if we're not in PIE or SIE, toggles pause/resume if we're in SIE, or ejects (PIE -> SIE) if we're in PIE.
+	ActionList.MapAction(
+		Commands.SimulatePauseEject,
+		FExecuteAction::CreateStatic(&FInternalPlayWorldCommandCallbacks::SimulatePauseEject_Clicked),
+		FCanExecuteAction::CreateStatic(&FInternalPlayWorldCommandCallbacks::CanSimulatePauseEject),
+		FIsActionChecked(),
+		FIsActionButtonVisible::CreateStatic(&DebuggerCommands::AllowPlayWorldFeature)
+	);
+
 	// Breakpoint-only commands
 	ActionList.MapAction(Commands.ShowCurrentStatement,
 		FExecuteAction::CreateStatic(&FInternalPlayWorldCommandCallbacks::ShowCurrentStatement_Clicked),
@@ -634,19 +679,46 @@ void FPlayWorldCommands::BuildToolbar(FToolMenuSection& InSection, bool bInclude
 	// Play
 	InSection.AddEntry(PlayMenuEntry);
 
-	// Resume/pause toggle (only one will be visible, and only in PIE/SIE)
-	FToolMenuEntry ResumeEntry = FToolMenuEntry::InitToolBarButton(FPlayWorldCommands::Get().ResumePlaySession, TAttribute<FText>(),
-		TAttribute<FText>::Create(TAttribute<FText>::FGetter::CreateStatic(&FInternalPlayWorldCommandCallbacks::GetResumePlaySessionToolTip)),
-		TAttribute<FSlateIcon>::Create(TAttribute<FSlateIcon>::FGetter::CreateStatic(&FInternalPlayWorldCommandCallbacks::GetResumePlaySessionImage)),
-		FName(TEXT("ResumePlaySession")));
-	ResumeEntry.StyleNameOverride = FName("Toolbar.BackplateLeftPlay");
+	// Resume PIE, pause PIE/SIE, or possess player (SIE -> PIE). (Only one will be visible, and only in PIE/SIE.)
+	{
+		{
+			FToolMenuEntry ResumeEntry = FToolMenuEntry::InitToolBarButton(
+				FPlayWorldCommands::Get().ResumePIESession,
+				TAttribute<FText>(),
+				LOCTEXT("ResumePIE", "Resume play-in-editor session"),
+				FSlateIcon(FAppStyle::GetAppStyleSetName(), "PlayWorld.ResumePIESession"),
+				FName(TEXT("ResumePIESession"))
+			);
+			ResumeEntry.StyleNameOverride = FName("Toolbar.BackplateLeft");
+			InSection.AddEntry(ResumeEntry);
+		}
 
-	FToolMenuEntry PauseEntry = FToolMenuEntry::InitToolBarButton(FPlayWorldCommands::Get().PausePlaySession, TAttribute<FText>(), TAttribute<FText>(), TAttribute<FSlateIcon>(), FName(TEXT("PausePlaySession")));
-	PauseEntry.StyleNameOverride = FName("Toolbar.BackplateLeft");
+		{
+			FToolMenuEntry PauseEntry = FToolMenuEntry::InitToolBarButton(
+				FPlayWorldCommands::Get().PausePlaySession,
+				TAttribute<FText>(),
+				TAttribute<FText>(),
+				TAttribute<FSlateIcon>(),
+				FName(TEXT("PausePlaySession"))
+			);
+			PauseEntry.StyleNameOverride = FName("Toolbar.BackplateLeft");
+
+			InSection.AddEntry(PauseEntry);
+		}
+
+		{
+			FToolMenuEntry PossessEntry = FToolMenuEntry::InitToolBarButton(
+				FPlayWorldCommands::Get().PlayPossess,
+				TAttribute<FText>(),
+				FPlayWorldCommands::Get().PlayPossess->GetDescription(),
+				FPlayWorldCommands::Get().PlayPossess->GetIcon(),
+				FName(TEXT("PlayPossess"))
+			);
+			PossessEntry.StyleNameOverride = FName("Toolbar.BackplateLeft");
+			InSection.AddEntry(PossessEntry);
+		}
+	}
 
-	InSection.AddEntry(ResumeEntry);
-	InSection.AddEntry(PauseEntry);
-	
 	FToolMenuEntry SingleFrameAdvanceEntry = FToolMenuEntry::InitToolBarButton(FPlayWorldCommands::Get().SingleFrameAdvance, TAttribute<FText>(), TAttribute<FText>(), TAttribute<FSlateIcon>(), FName(TEXT("SingleFrameAdvance")));
 	SingleFrameAdvanceEntry.StyleNameOverride = FName("Toolbar.BackplateCenter");
 	InSection.AddEntry(SingleFrameAdvanceEntry);
@@ -662,24 +734,25 @@ void FPlayWorldCommands::BuildToolbar(FToolMenuSection& InSection, bool bInclude
 	LateJoinEntry.StyleNameOverride = FName("Toolbar.BackplateCenter");
 	InSection.AddEntry(LateJoinEntry);
 
-	// Eject/possess toggle
-
-	FToolMenuEntry EjectMenuEntry =
-	FToolMenuEntry::InitToolBarButton(FPlayWorldCommands::Get().PossessEjectPlayer,
-		TAttribute<FText>::Create(TAttribute<FText>::FGetter::CreateStatic(&FInternalPlayWorldCommandCallbacks::GetPossessEjectLabel)),
-		TAttribute<FText>::Create(TAttribute<FText>::FGetter::CreateStatic(&FInternalPlayWorldCommandCallbacks::GetPossessEjectTooltip)),
-		TAttribute<FSlateIcon>::Create(TAttribute<FSlateIcon>::FGetter::CreateStatic(&FInternalPlayWorldCommandCallbacks::GetPossessEjectImage)),
-		FName(TEXT("PossessEjectPlayer")));
-	EjectMenuEntry.StyleNameOverride = FName("Toolbar.BackplateCenter");
-
-	// Play combo box
-	/*EjectMenuEntry.AddOptionsDropdown(
-		SpecialPIEOptionsMenuAction,
-		FOnGetContent::CreateStatic(&GeneratePlayMenuContent, GlobalPlayWorldActions.ToSharedRef()),
-		LOCTEXT("PIEComboToolTip", "Change Play Mode and Play Settings")
-	);*/
+	// Start SIE / pause SIE/PIE / eject (PIE -> SIE).
+	{
+		FToolMenuEntry SimulatePauseEjectMenuEntry = FToolMenuEntry::InitToolBarButton(
+			FPlayWorldCommands::Get().SimulatePauseEject,
+			TAttribute<FText>::Create(
+				TAttribute<FText>::FGetter::CreateStatic(&FInternalPlayWorldCommandCallbacks::GetSimulatePauseEjectLabel)
+			),
+			TAttribute<FText>::Create(
+				TAttribute<FText>::FGetter::CreateStatic(&FInternalPlayWorldCommandCallbacks::GetSimulatePauseEjectTooltip)
+			),
+			TAttribute<FSlateIcon>::Create(TAttribute<FSlateIcon>::FGetter::CreateStatic(
+				&FInternalPlayWorldCommandCallbacks::GetSimulatePauseEjectIcon
+			)),
+			FName(TEXT("SimulatePauseEject"))
+		);
+		SimulatePauseEjectMenuEntry.StyleNameOverride = FName("Toolbar.BackplateCenter");
 
-	InSection.AddEntry(EjectMenuEntry);
+		InSection.AddEntry(SimulatePauseEjectMenuEntry);
+	}
 
 	FUIAction SpecialPIEOptionsMenuAction;
 	SpecialPIEOptionsMenuAction.CanExecuteAction = FCanExecuteAction::CreateStatic(&FInternalPlayWorldCommandCallbacks::CanShowNonPlayWorldOnlyActions);
@@ -833,10 +906,6 @@ TSharedRef< SWidget > FPlayWorldCommands::GeneratePlayMenuContent(TSharedRef<FUI
 				case PlayMode_InVR:
 					PlayModeCommand = FPlayWorldCommands::Get().PlayInVR;
 					break;
-
-				case PlayMode_Simulate:
-					PlayModeCommand = FPlayWorldCommands::Get().Simulate;
-					break;
 				}
 
 				if (PlayModeCommand.IsValid())
@@ -867,7 +936,26 @@ TSharedRef< SWidget > FPlayWorldCommands::GeneratePlayMenuContent(TSharedRef<FUI
 			FLocal::AddPlayModeMenuEntry(Section, PlayMode_InMobilePreview);
 			FLocal::AddPlayModeMenuEntry(Section, PlayMode_InVulkanPreview);
 			FLocal::AddPlayModeMenuEntry(Section, PlayMode_InNewProcess);
-			FLocal::AddPlayModeMenuEntry(Section, PlayMode_Simulate);
+
+			Section.AddSeparator("SimulateSeparator");
+
+			if (FPlayWorldCommands::Get().Simulate)
+			{
+				Section.AddMenuEntry(
+					"Simulate",
+					FPlayWorldCommands::Get().Simulate->GetLabel(),
+					FPlayWorldCommands::Get().Simulate->GetDescription(),
+					FPlayWorldCommands::Get().Simulate->GetIcon(),
+					FExecuteAction::CreateLambda(
+						[]()
+						{
+							FPlayWorldCommands::GlobalPlayWorldActions->ExecuteAction(
+								FPlayWorldCommands::Get().Simulate.ToSharedRef()
+							);
+						}
+					)
+				);
+			}
 		}
 
 #if UE_WITH_TURNKEY_SUPPORT
@@ -1088,6 +1176,10 @@ bool FPlayWorldCommandCallbacks::IsInPIE_AndRunning()
 	return IsInPIE() && ((GEditor->PlayWorld == NULL) || !(GEditor->PlayWorld->bDebugPauseExecution));
 }
 
+bool FPlayWorldCommandCallbacks::IsInPIE_AndPaused()
+{
+	return GUnrealEd && IsInPIE() && HasPlayWorld() && GUnrealEd->PlayWorld->bDebugPauseExecution;
+}
 
 bool FPlayWorldCommandCallbacks::HasPlayWorld()
 {
@@ -1103,7 +1195,7 @@ bool FPlayWorldCommandCallbacks::HasPlayWorldAndPaused()
 
 bool FPlayWorldCommandCallbacks::HasPlayWorldAndRunning()
 {
-	return HasPlayWorld() && !GUnrealEd->PlayWorld->bDebugPauseExecution;
+	return IsInPIE() && HasPlayWorld() && !GUnrealEd->PlayWorld->bDebugPauseExecution;
 }
 
 
@@ -1156,6 +1248,74 @@ FSlateIcon FInternalPlayWorldCommandCallbacks::GetPossessEjectImage()
 	}
 }
 
+FText FInternalPlayWorldCommandCallbacks::GetSimulatePauseEjectLabel()
+{
+	if (IsInPIE())
+	{
+		return LOCTEXT("EjectLabel", "Eject");
+	}
+	else if (IsInSIE())
+	{
+		if (FPlayWorldCommandCallbacks::HasPlayWorldAndPaused())
+		{
+			return LOCTEXT("SimulatePauseEject_Resume", "Resume");
+		}
+		else
+		{
+			return LOCTEXT("SimulatePauseEject_Pause", "Pause");
+		}
+	}
+	else
+	{
+		return FPlayWorldCommands::Get().Simulate->GetLabel();
+	}
+}
+
+FText FInternalPlayWorldCommandCallbacks::GetSimulatePauseEjectTooltip()
+{
+	if (IsInPIE())
+	{
+		return LOCTEXT("EjectToolTip", "Detaches from the player controller, allowing regular editor controls");
+	}
+	else if (IsInSIE())
+	{
+		if (FPlayWorldCommandCallbacks::HasPlayWorldAndPaused())
+		{
+			return LOCTEXT("SimulatePauseEject_Resume", "Resumes the simulation");
+		}
+		else
+		{
+			return LOCTEXT("SimulatePauseEject_Pause", "Pauses the simulation");
+		}
+	}
+	else
+	{
+		return FPlayWorldCommands::Get().Simulate->GetDescription();
+	}
+}
+
+FSlateIcon FInternalPlayWorldCommandCallbacks::GetSimulatePauseEjectIcon()
+{
+	if (IsInPIE())
+	{
+		return FSlateIcon(FAppStyle::GetAppStyleSetName(), "PlayWorld.SimulatePauseEject_Eject");
+	}
+	else if (IsInSIE())
+	{
+		if (FPlayWorldCommandCallbacks::HasPlayWorldAndPaused())
+		{
+			return FSlateIcon(FAppStyle::GetAppStyleSetName(), "PlayWorld.SimulatePauseEject_PlaySimulate");
+		}
+		else
+		{
+			return FSlateIcon(FAppStyle::GetAppStyleSetName(), "PlayWorld.SimulatePauseEject_Pause");
+		}
+	}
+	else
+	{
+		return FSlateIcon(FAppStyle::GetAppStyleSetName(), "PlayWorld.SimulatePauseEject_PlaySimulate");
+	}
+}
 
 bool FInternalPlayWorldCommandCallbacks::CanLateJoin()
 {
@@ -1240,6 +1400,7 @@ const TSharedRef < FUICommandInfo > GetLastPlaySessionCommand()
 
 	switch (PlaySettings->LastExecutedPlayModeType)
 	{
+	case PlayMode_Simulate:
 	case PlayMode_InViewPort:
 		Command = Commands.PlayInViewport.ToSharedRef();
 		break;
@@ -1278,9 +1439,6 @@ const TSharedRef < FUICommandInfo > GetLastPlaySessionCommand()
 	case PlayMode_InVR:
 		Command = Commands.PlayInVR.ToSharedRef();
 		break;
-
-	case PlayMode_Simulate:
-		Command = Commands.Simulate.ToSharedRef();
 	}
 
 	return Command;
@@ -1791,40 +1949,6 @@ void FInternalPlayWorldCommandCallbacks::GetMouseControlExecute()
 	}
 }
 
-FSlateIcon FInternalPlayWorldCommandCallbacks::GetResumePlaySessionImage()
-{
-	if (IsInPIE())
-	{
-		return FSlateIcon(FAppStyle::GetAppStyleSetName(), "PlayWorld.ResumePlaySession");
-	}
-	else if (IsInSIE())
-	{
-		return FSlateIcon(FAppStyle::GetAppStyleSetName(), "PlayWorld.Simulate");
-	}
-	else
-	{
-		return FSlateIcon();
-	}
-}
-
-
-FText FInternalPlayWorldCommandCallbacks::GetResumePlaySessionToolTip()
-{
-	if (IsInPIE())
-	{
-		return LOCTEXT("ResumePIE", "Resume play-in-editor session");
-	}
-	else if (IsInSIE())
-	{
-		return LOCTEXT("ResumeSIE", "Resume simulation");
-	}
-	else
-	{
-		return FText();
-	}
-}
-
-
 void FInternalPlayWorldCommandCallbacks::SingleFrameAdvance_Clicked()
 {
 	if (IsStoppedAtBreakpoint_InEngineMode())
@@ -2010,6 +2134,22 @@ bool FInternalPlayWorldCommandCallbacks::IsStoppedAtBreakpoint_InEngineMode()
 	return GIntraFrameDebuggingGameThread && FKismetDebugUtilities::GetCurrentDebuggingWorld() == nullptr;
 }
 
+void FInternalPlayWorldCommandCallbacks::PossessPlayer_Clicked()
+{
+	GEditor->RequestToggleBetweenPIEandSIE();
+}
+
+bool FInternalPlayWorldCommandCallbacks::CanPossessPlayer()
+{
+	if (IsInSIE() && !IsStoppedAtBreakpoint())
+	{
+		for (auto It = GUnrealEd->SlatePlayInEditorMap.CreateIterator(); It; ++It)
+		{
+			return It.Value().DestinationSlateViewport.IsValid();
+		}
+	}
+	return false;
+}
 
 void FInternalPlayWorldCommandCallbacks::PossessEjectPlayer_Clicked()
 {
@@ -2029,6 +2169,44 @@ bool FInternalPlayWorldCommandCallbacks::CanPossessEjectPlayer()
 	return false;
 }
 
+void FInternalPlayWorldCommandCallbacks::SimulatePauseEject_Clicked()
+{
+	if (IsInPIE())
+	{
+		GEditor->RequestToggleBetweenPIEandSIE();
+	}
+	else if (IsInSIE())
+	{
+		if (FPlayWorldCommandCallbacks::HasPlayWorldAndPaused())
+		{
+			FPlayWorldCommandCallbacks::ResumePlaySession_Clicked();
+		}
+		else
+		{
+			FPlayWorldCommandCallbacks::PausePlaySession_Clicked();
+		}
+	}
+	else
+	{
+		FPlayWorldCommands::GlobalPlayWorldActions->ExecuteAction(FPlayWorldCommands::Get().Simulate.ToSharedRef());
+	}
+}
+
+bool FInternalPlayWorldCommandCallbacks::CanSimulatePauseEject()
+{
+	if (IsInPIE())
+	{
+		return FInternalPlayWorldCommandCallbacks::CanPossessEjectPlayer();
+	}
+	else if (IsInSIE())
+	{
+		return HasPlayWorld();
+	}
+	else
+	{
+		return FPlayWorldCommands::GlobalPlayWorldActions->CanExecuteAction(FPlayWorldCommands::Get().Simulate.ToSharedRef());
+	}
+}
 
 void FInternalPlayWorldCommandCallbacks::AddMessageLog(const FText& Text, const FText& Detail, const FString& TutorialLink, const FString& DocumentationLink)
 {
@@ -2049,7 +2227,6 @@ void FInternalPlayWorldCommandCallbacks::AddMessageLog(const FText& Text, const
 }
 
 
-
 EPlayModeLocations FInternalPlayWorldCommandCallbacks::GetPlayModeLocation()
 {
 	// We can't use PlayLocation_DefaultPlayerStart without a player start position
@@ -2058,22 +2235,4 @@ EPlayModeLocations FInternalPlayWorldCommandCallbacks::GetPlayModeLocation()
 		: PlayLocation_CurrentCameraLocation;
 }
 
-
-
-
-
-
-
-
-
-
-
-
-
-
-
-
-
-
 #undef LOCTEXT_NAMESPACE
-
diff --git a/Engine/Source/Editor/UnrealEd/Public/Kismet2/DebuggerCommands.h b/Engine/Source/Editor/UnrealEd/Public/Kismet2/DebuggerCommands.h
--- a/Engine/Source/Editor/UnrealEd/Public/Kismet2/DebuggerCommands.h
+++ b/Engine/Source/Editor/UnrealEd/Public/Kismet2/DebuggerCommands.h
@@ -98,13 +98,16 @@ public:
 	/** SIE & PIE controls */
 	TSharedPtr<FUICommandInfo> ResumePlaySession;
 	TSharedPtr<FUICommandInfo> PausePlaySession;
+	TSharedPtr<FUICommandInfo> ResumePIESession;
 	TSharedPtr<FUICommandInfo> SingleFrameAdvance;
 	TSharedPtr<FUICommandInfo> TogglePlayPauseOfPlaySession;
 	TSharedPtr<FUICommandInfo> StopPlaySession;
 	TSharedPtr<FUICommandInfo> LateJoinSession;
+	TSharedPtr<FUICommandInfo> PlayPossess;
 	TSharedPtr<FUICommandInfo> PossessEjectPlayer;
 	TSharedPtr<FUICommandInfo> ShowCurrentStatement;
 	TSharedPtr<FUICommandInfo> GetMouseControl;
+	TSharedPtr<FUICommandInfo> SimulatePauseEject;
 
 	/** BP Debugging controls */
 	TSharedPtr<FUICommandInfo> AbortExecution;
@@ -154,6 +157,7 @@ public:
 
 	static UNREALED_API bool IsInSIE_AndRunning();
 	static UNREALED_API bool IsInPIE_AndRunning();
+	static UNREALED_API bool IsInPIE_AndPaused();
 
 	static UNREALED_API bool HasPlayWorld();
 	static UNREALED_API bool HasPlayWorldAndPaused();
