Fix 100% repro crash on right click for Text Render / Dynamic mesh actors in outliner.
Add check class valid AnimationAsset.

#tests Right click text render actor in editor in outliner

diff --git a/Engine/Source/Editor/UnrealEd/Private/Factories/ActorFactory.cpp b/Engine/Source/Editor/UnrealEd/Private/Factories/ActorFactory.cpp
--- a/Engine/Source/Editor/UnrealEd/Private/Factories/ActorFactory.cpp
+++ b/Engine/Source/Editor/UnrealEd/Private/Factories/ActorFactory.cpp
@@ -935,8 +935,9 @@ UActorFactoryAnimationAsset::UActorFactoryAnimationAsset(const FObjectInitialize
 
 bool UActorFactoryAnimationAsset::CanCreateActorFrom( const FAssetData& AssetData, FText& OutErrorMsg )
 { 
-	if ( !AssetData.IsValid() || 
-		( !AssetData.GetClass()->IsChildOf( UAnimSequenceBase::StaticClass() ) )) 
+	if ( !AssetData.IsValid() 
+		|| !AssetData.GetClass()
+		|| !AssetData.GetClass()->IsChildOf(UAnimSequenceBase::StaticClass()))
 	{
 		OutErrorMsg = NSLOCTEXT("CanCreateActor", "NoAnimData", "A valid anim data must be specified.");
 		return false;
