Fix FEditorModeTools::ActivateMode removing editor incompatible modes in a loop that affects the modes iterator.

diff --git a/Engine/Source/Editor/UnrealEd/Private/EditorModeManager.cpp b/Engine/Source/Editor/UnrealEd/Private/EditorModeManager.cpp
--- a/Engine/Source/Editor/UnrealEd/Private/EditorModeManager.cpp
+++ b/Engine/Source/Editor/UnrealEd/Private/EditorModeManager.cpp
@@ -926,6 +926,7 @@ void FEditorModeTools::ActivateMode(FEditorModeID InID, bool bToggle)
 		FGCObjectScopeGuard ScriptModeGuard(ScriptableMode);
 
 		// Remove anything that isn't compatible with this mode
+		TArray<FEditorModeID> EditorModesToRemove;
 		const bool bIsVisibleMode = ScriptableMode->GetModeInfo().IsVisible();
 		for (int32 ModeIndex = ActiveScriptableModes.Num() - 1; ModeIndex >= 0; ModeIndex--)
 		{
@@ -933,9 +934,14 @@ void FEditorModeTools::ActivateMode(FEditorModeID InID, bool bToggle)
 			const bool bModesAreCompatible = ScriptableMode->IsCompatibleWith(Mode->GetID()) || Mode->IsCompatibleWith(ScriptableMode->GetID());
 			if (!bModesAreCompatible || (bIsVisibleMode && Mode->GetModeInfo().IsVisible()))
 			{
-				DeactivateMode(Mode->GetID());
+				EditorModesToRemove.Add(Mode->GetID());
 			}
 		}
+
+		for (FEditorModeID EditorModeToRemove : EditorModesToRemove)
+		{
+			DeactivateMode(EditorModeToRemove);
+		}
 	}
 
 	ActiveScriptableModes.Add(ScriptableMode);
