Fix an edge case in worker retraction. Packages in SaveStalledAssignedToWorker on the CookDirector must not be demoted to idle; they should be demoted to AssignedToWorker instead.
#rnx

diff --git a/Engine/Source/Editor/UnrealEd/Private/CookOnTheFlyServer.cpp b/Engine/Source/Editor/UnrealEd/Private/CookOnTheFlyServer.cpp
--- a/Engine/Source/Editor/UnrealEd/Private/CookOnTheFlyServer.cpp
+++ b/Engine/Source/Editor/UnrealEd/Private/CookOnTheFlyServer.cpp
@@ -5612,6 +5612,8 @@ void UCookOnTheFlyServer::PostGarbageCollect()
 	}
 	for (FPackageData* PackageData : Demotes)
 	{
+		FGenerationHelper::ValidateSaveStalledState(*this, *PackageData, TEXT("PostGarbageCollect"));
+
 		switch (PackageData->GetState())
 		{
 		case EPackageState::SaveActive:
diff --git a/Engine/Source/Editor/UnrealEd/Private/Cooker/CookGenerationHelper.cpp b/Engine/Source/Editor/UnrealEd/Private/Cooker/CookGenerationHelper.cpp
--- a/Engine/Source/Editor/UnrealEd/Private/Cooker/CookGenerationHelper.cpp
+++ b/Engine/Source/Editor/UnrealEd/Private/Cooker/CookGenerationHelper.cpp
@@ -142,7 +142,8 @@ void FGenerationHelper::Uninitialize()
 	// Demote stalled packages; we will be garbage collecting so they no longer need to be preserved.
 	// And we need to demote them so that they drop their references to the generation helper and allow it to be
 	// deleted if no longer referenced
-	DemoteStalledPackages(OwnerInfo.PackageData->GetPackageDatas().GetCookOnTheFlyServer());
+	DemoteStalledPackages(OwnerInfo.PackageData->GetPackageDatas().GetCookOnTheFlyServer(),
+		false /* bFromAllSavesCompleted */);
 
 	NotifyCompletion(ICookPackageSplitter::ETeardown::Complete);
 	check(!CookPackageSplitterInstance);
@@ -217,19 +218,44 @@ void FGenerationHelper::OnAllSavesCompleted(UCookOnTheFlyServer& COTFS)
 	ClearKeepForAllSavedOrGC();
 
 	// Demote stalled packages; we will no longer need to come back to them
-	DemoteStalledPackages(COTFS);
+	DemoteStalledPackages(COTFS, true /* bFromAllSavesCompleted */);
 }
 
-void FGenerationHelper::DemoteStalledPackages(UCookOnTheFlyServer& COTFS)
+void FGenerationHelper::DemoteStalledPackages(UCookOnTheFlyServer& COTFS, bool bFromAllSavesCompleted)
 {
 	// For any packages that we stalled because they were retracted and assigned to another worker,
 	// demote them now. But don't demote non-stalled packages, because doing so could demote the final
 	// package that we just saved locally and still needs to finish its work in PumpSaves.
-	auto ConditionalDemote = [&COTFS](FCookGenerationInfo& Info)
+	auto ConditionalDemote = [&COTFS, &bFromAllSavesCompleted](FCookGenerationInfo& Info)
 		{
 			if (Info.PackageData->IsStalled())
 			{
-				COTFS.DemoteToIdle(*Info.PackageData, ESendFlags::QueueAddAndRemove, ESuppressCookReason::RetractedByCookDirector);
+				FGenerationHelper::ValidateSaveStalledState(COTFS, *Info.PackageData, TEXT("DemoteStalledPackages"));
+				if (Info.PackageData->GetState() == EPackageState::SaveStalledAssignedToWorker)
+				{
+					// If called from OnAllSavesCompleted on the CookDirector, then we should have no stalled packages;
+					// they all should have been unstalled and demoted when saved. If we do have any, then log an error
+					// and demote them to idle.
+					if (bFromAllSavesCompleted)
+					{
+						UE_LOG(LogCook, Error,
+							TEXT("Package %s is still stalled on the CookDirector during FGenerationHelper::OnAllSavesCompleted. ")
+							TEXT("This is unexpected; all stalled packages on a GenerationHelper should have completed saving and therefore unstalled before OnAllSavesCompleted is called."),
+							*Info.PackageData->GetPackageName().ToString());
+
+						COTFS.DemoteToIdle(*Info.PackageData, ESendFlags::QueueAddAndRemove, ESuppressCookReason::RetractedByCookDirector);
+					}
+					else
+					{
+						// Otherwise, when called on the CookDirector, demote them out of save but keep them in the assigned to worker state.
+						Info.PackageData->SendToState(EPackageState::AssignedToWorker, ESendFlags::QueueAddAndRemove, EStateChangeReason::GarbageCollected);
+					}
+				}
+				else
+				{
+					// Demoting stalled packages on a Client just returns them to Idle.
+					COTFS.DemoteToIdle(*Info.PackageData, ESendFlags::QueueAddAndRemove, ESuppressCookReason::RetractedByCookDirector);
+				}
 			}
 		};
 	ConditionalDemote(OwnerInfo);
@@ -239,6 +265,25 @@ void FGenerationHelper::DemoteStalledPackages(UCookOnTheFlyServer& COTFS)
 	}
 }
 
+void FGenerationHelper::ValidateSaveStalledState(UCookOnTheFlyServer& COTFS, FPackageData& PackageData, const TCHAR* Caller)
+{
+	if (!PackageData.IsStalled())
+	{
+		return;
+	}
+	EPackageState ExpectedState = COTFS.CookDirector.IsValid() ? EPackageState::SaveStalledAssignedToWorker : EPackageState::SaveStalledRetracted;
+	if (PackageData.GetState() != ExpectedState)
+	{
+		const TCHAR* Which = COTFS.CookDirector.IsValid() ? TEXT("the CookDirector") : TEXT("a CookWorker");
+		UE_LOG(LogCook, Error,
+			TEXT("In %s, package %s is stalled on %s, but is in state %s. We expect stalled packages on %s to be in state %s."),
+			Caller,
+			*PackageData.GetPackageName().ToString(), Which, LexToString(PackageData.GetState()),
+			Which, LexToString(ExpectedState));
+		PackageData.SendToState(ExpectedState, ESendFlags::QueueNone, EStateChangeReason::Retraction);
+	}
+}
+
 void FGenerationHelper::DiagnoseWhyNotShutdown()
 {
 	TStringBuilder<256> Lines;
diff --git a/Engine/Source/Editor/UnrealEd/Private/Cooker/CookGenerationHelper.h b/Engine/Source/Editor/UnrealEd/Private/Cooker/CookGenerationHelper.h
--- a/Engine/Source/Editor/UnrealEd/Private/Cooker/CookGenerationHelper.h
+++ b/Engine/Source/Editor/UnrealEd/Private/Cooker/CookGenerationHelper.h
@@ -415,6 +415,8 @@ public:
 		bool& bOutNeedWaitForIsLoaded);
 	/** Helper function for Initialize and for TryCreateValidGenerationHelper. */
 	static UPackage* FindOrLoadPackage(UCookOnTheFlyServer& COTFS, FPackageData& OwnerPackageData);
+	/** If PackageData is in a stalled state, verify that its state is valid for CookDirector or CookWorker. */
+	static void ValidateSaveStalledState(UCookOnTheFlyServer& COTFS, FPackageData& PackageData, const TCHAR* Caller);
 
 	/** Initialize settings from commandline and config files. */
 	static void SetBeginCookConfigSettings();
@@ -440,7 +442,7 @@ private:
 	void Uninitialize();
 	void ModifyNumSaved(int32 Delta);
 	void VerifyGeneratorPackageGarbageCollected(FCookGCDiagnosticContext& Context);
-	void DemoteStalledPackages(UCookOnTheFlyServer& COTFS);
+	void DemoteStalledPackages(UCookOnTheFlyServer& COTFS, bool bFromAllSavesCompleted);
 
 private:
 	// When adding a new variable, add it to Uninitialize as well
diff --git a/Engine/Source/Editor/UnrealEd/Private/Cooker/CookWorkerServer.cpp b/Engine/Source/Editor/UnrealEd/Private/Cooker/CookWorkerServer.cpp
--- a/Engine/Source/Editor/UnrealEd/Private/Cooker/CookWorkerServer.cpp
+++ b/Engine/Source/Editor/UnrealEd/Private/Cooker/CookWorkerServer.cpp
@@ -1001,6 +1001,16 @@ void FCookWorkerServer::RecordResults(FPackageResultsMessage& Message)
 			}
 			COTFS.PromoteToSaveComplete(*PackageData, ESendFlags::QueueAddAndRemove);
 		}
+		else if (Result.GetSuppressCookReason() == ESuppressCookReason::RetractedByCookDirector)
+		{
+			UE_LOG(LogCook, Error,
+				TEXT("Package %s was retracted by CookWorker %d, but it still sent a RecordResults message for the package which is supposed to be omitted for RetractedByCookDirector suppressions."),
+				*PackageData->GetPackageName().ToString(), ProfileId);
+			if (PackageData->GetWorkerAssignment() == GetWorkerId())
+			{
+				PackageData->SendToState(EPackageState::Request, ESendFlags::QueueAddAndRemove, EStateChangeReason::Retraction);
+			}
+		}
 		else
 		{
 			COTFS.DemoteToIdle(*PackageData, ESendFlags::QueueAddAndRemove, Result.GetSuppressCookReason());
