[BUGFIX] Fbx Importer: Need to clamp incoming vertex color

diff --git a/Engine/Plugins/Interchange/Runtime/Source/Parsers/Fbx/Private/FbxMesh.cpp b/Engine/Plugins/Interchange/Runtime/Source/Parsers/Fbx/Private/FbxMesh.cpp
--- a/Engine/Plugins/Interchange/Runtime/Source/Parsers/Fbx/Private/FbxMesh.cpp
+++ b/Engine/Plugins/Interchange/Runtime/Source/Parsers/Fbx/Private/FbxMesh.cpp
@@ -1098,11 +1098,12 @@ bool FMeshDescriptionImporter::FillMeshDescriptionFromFbxMesh(FbxMesh* Mesh, con
 
 						FbxColor VertexColor = LayerElementVertexColor->GetDirectArray().GetAt(VectorColorIndex);
 
+						// Clamp manually to avoid wrapping around because of truncation for values outside of the 0.0 to 1.0 range
 						FColor VertexInstanceColor(
-							uint8(255.f * VertexColor.mRed),
-							uint8(255.f * VertexColor.mGreen),
-							uint8(255.f * VertexColor.mBlue),
-							uint8(255.f * VertexColor.mAlpha)
+							uint8(255.f * FMath::Clamp(VertexColor.mRed, 0.0, 1.0)),
+							uint8(255.f * FMath::Clamp(VertexColor.mGreen, 0.0, 1.0)),
+							uint8(255.f * FMath::Clamp(VertexColor.mBlue,0.0, 1.0)),
+							uint8(255.f * FMath::Clamp(VertexColor.mAlpha, 0.0, 1.0))
 						);
 						VertexInstanceColors[AddedVertexInstanceId] = FVector4f(FLinearColor(VertexInstanceColor));
 					}
diff --git a/Engine/Source/Editor/UnrealEd/Private/Fbx/FbxSkeletalMeshImport.cpp b/Engine/Source/Editor/UnrealEd/Private/Fbx/FbxSkeletalMeshImport.cpp
--- a/Engine/Source/Editor/UnrealEd/Private/Fbx/FbxSkeletalMeshImport.cpp
+++ b/Engine/Source/Editor/UnrealEd/Private/Fbx/FbxSkeletalMeshImport.cpp
@@ -3699,6 +3699,15 @@ bool UnFbx::FFbxImporter::FillSkelMeshImporterFromFbx( FSkeletalMeshImportData&
 		// Read vertex colors if they exist.
 		if( LayerElementVertexColor && ImportOptions->VertexColorImportOption == EVertexColorImportOption::Replace)
 		{
+			auto ConvertFbxColorToFColor = [](const FbxColor& VertexColor)
+				{
+					// Clamp manually to avoid wrapping around because of truncation for values outside of the 0.0 to 1.0 range
+					return FColor(uint8(255.f * FMath::Clamp(VertexColor.mRed, 0.0, 1.0)),
+						uint8(255.f * FMath::Clamp(VertexColor.mGreen, 0.0, 1.0)),
+						uint8(255.f * FMath::Clamp(VertexColor.mBlue, 0.0, 1.0)),
+						uint8(255.f * FMath::Clamp(VertexColor.mAlpha, 0.0, 1.0)));
+				};
+
 			switch(VertexColorMappingMode)
 			{
 			case FbxLayerElement::eByControlPoint:
@@ -3711,10 +3720,8 @@ bool UnFbx::FFbxImporter::FillSkelMeshImporterFromFbx( FSkeletalMeshImportData&
 							?	LayerElementVertexColor->GetDirectArray().GetAt(Mesh->GetPolygonVertex(LocalIndex,VertexIndex))
 							:	LayerElementVertexColor->GetDirectArray().GetAt(LayerElementVertexColor->GetIndexArray().GetAt(Mesh->GetPolygonVertex(LocalIndex,VertexIndex)));
 
-						TmpWedges[UnrealVertexIndex].Color =	FColor(	uint8(255.f*VertexColor.mRed),
-																uint8(255.f*VertexColor.mGreen),
-																uint8(255.f*VertexColor.mBlue),
-																uint8(255.f*VertexColor.mAlpha));
+						
+						TmpWedges[UnrealVertexIndex].Color = ConvertFbxColorToFColor(VertexColor);
 					}
 				}
 				break;
@@ -3728,10 +3735,7 @@ bool UnFbx::FFbxImporter::FillSkelMeshImporterFromFbx( FSkeletalMeshImportData&
 							?	LayerElementVertexColor->GetDirectArray().GetAt(LocalIndex*3+VertexIndex)
 							:	LayerElementVertexColor->GetDirectArray().GetAt(LayerElementVertexColor->GetIndexArray().GetAt(LocalIndex*3+VertexIndex));
 
-						TmpWedges[UnrealVertexIndex].Color =	FColor(	uint8(255.f*VertexColor.mRed),
-																uint8(255.f*VertexColor.mGreen),
-																uint8(255.f*VertexColor.mBlue),
-																uint8(255.f*VertexColor.mAlpha));
+						TmpWedges[UnrealVertexIndex].Color = ConvertFbxColorToFColor(VertexColor);
 					}
 				}
 				break;
diff --git a/Engine/Source/Editor/UnrealEd/Private/Fbx/FbxStaticMeshImport.cpp b/Engine/Source/Editor/UnrealEd/Private/Fbx/FbxStaticMeshImport.cpp
--- a/Engine/Source/Editor/UnrealEd/Private/Fbx/FbxStaticMeshImport.cpp
+++ b/Engine/Source/Editor/UnrealEd/Private/Fbx/FbxStaticMeshImport.cpp
@@ -943,11 +943,12 @@ bool UnFbx::FFbxImporter::BuildStaticMeshFromGeometry(FbxNode* Node, UStaticMesh
 
 							FbxColor VertexColor = LayerElementVertexColor->GetDirectArray().GetAt(VectorColorIndex);
 
+							// Clamp manually to avoid wrapping around because of truncation for values outside of the 0.0 to 1.0 range
 							FColor VertexInstanceColor(
-								uint8(255.f*VertexColor.mRed),
-								uint8(255.f*VertexColor.mGreen),
-								uint8(255.f*VertexColor.mBlue),
-								uint8(255.f*VertexColor.mAlpha)
+								uint8(255.f * FMath::Clamp(VertexColor.mRed, 0.0, 1.0)),
+								uint8(255.f * FMath::Clamp(VertexColor.mGreen, 0.0, 1.0)),
+								uint8(255.f * FMath::Clamp(VertexColor.mBlue, 0.0, 1.0)),
+								uint8(255.f * FMath::Clamp(VertexColor.mAlpha, 0.0, 1.0))
 							);
 							VertexInstanceColors[AddedVertexInstanceId] = FVector4f(FLinearColor(VertexInstanceColor));
 						}
