[Viewport Toolbar] Add a button or menu item to select the currently viewed through camera

- Adds "Select Active Camera" entry
- Renames Unnamed section containing Field Of View and Far View Plane to avoid it from being accidentaly moved to another Unnamed section

diff --git a/Engine/Source/Editor/LevelEditor/Private/ViewportToolbar/LevelEditorViewportToolbarSections.cpp b/Engine/Source/Editor/LevelEditor/Private/ViewportToolbar/LevelEditorViewportToolbarSections.cpp
--- a/Engine/Source/Editor/LevelEditor/Private/ViewportToolbar/LevelEditorViewportToolbarSections.cpp
+++ b/Engine/Source/Editor/LevelEditor/Private/ViewportToolbar/LevelEditorViewportToolbarSections.cpp
@@ -2316,6 +2316,13 @@ void ExtendCameraSubmenu(FName InCameraOptionsSubmenuName)
 				ULevelViewportContext* LevelViewportContext = InDynamicMenu->FindContext<ULevelViewportContext>();
 				TWeakPtr<::SLevelViewport> LevelViewportWeak = LevelViewportContext->LevelViewport;
 
+				FToolMenuSection& SelectPilotedActorSection = InDynamicMenu->AddSection("SelectPilotedActor");
+
+				SelectPilotedActorSection.InsertPosition =
+					FToolMenuInsert("LevelViewportCameraType_Perspective", EToolMenuInsertType::Before);
+
+				SelectPilotedActorSection.AddMenuEntry(FLevelViewportCommands::Get().SelectPilotedActor);
+
 				// Camera Selection elements
 				{
 					AddCameraActorSelectSection(InDynamicMenu);
diff --git a/Engine/Source/Editor/UnrealEd/Private/ViewportToolbar/UnrealEdViewportToolbar.cpp b/Engine/Source/Editor/UnrealEd/Private/ViewportToolbar/UnrealEdViewportToolbar.cpp
--- a/Engine/Source/Editor/UnrealEd/Private/ViewportToolbar/UnrealEdViewportToolbar.cpp
+++ b/Engine/Source/Editor/UnrealEd/Private/ViewportToolbar/UnrealEdViewportToolbar.cpp
@@ -2525,7 +2525,7 @@ void PopulateCameraMenu(UToolMenu* InMenu)
 	}
 
 	FToolMenuSection& PerspectiveCameraSection =
-		InMenu->FindOrAddSection("LevelViewportCameraType_Perspective");
+		InMenu->FindOrAddSection("LevelViewportCameraType_Perspective", LOCTEXT("PerspectiveLabel", "Perspective"));
 	PerspectiveCameraSection.AddMenuEntry(FEditorViewportCommands::Get().Perspective);
 
 	FToolMenuSection& OrthographicCameraSection =
@@ -2562,16 +2562,16 @@ void ExtendCameraSubmenu(FName InCameraOptionsSubmenuName)
 
 				FToolMenuInsert InsertPosition("LevelViewportCameraType_Ortho", EToolMenuInsertType::After);
 
-				FToolMenuSection& UnnamedSection = InDynamicMenu->FindOrAddSection(NAME_None, FText(), InsertPosition);
-				UnnamedSection.AddSeparator("CameraSubmenuSeparator");
+				FToolMenuSection& PostOrthoSection = InDynamicMenu->FindOrAddSection("CameraSubmenuPostOrtho", FText(), InsertPosition);
+				PostOrthoSection.AddSeparator("CameraSubmenuSeparator");
 
-				UnnamedSection.AddEntry(FToolMenuEntry::InitWidget(
+				PostOrthoSection.AddEntry(FToolMenuEntry::InitWidget(
 					"CameraFOV",
 					UE::UnrealEd::CreateFOVMenuWidget(EditorViewport.ToSharedRef()),
 					LOCTEXT("CameraSubmenu_FieldOfViewLabel", "Field of View")
 				));
 
-				UnnamedSection.AddEntry(FToolMenuEntry::InitWidget(
+				PostOrthoSection.AddEntry(FToolMenuEntry::InitWidget(
 					"CameraFarViewPlane",
 					UE::UnrealEd::CreateFarViewPlaneMenuWidget(EditorViewport.ToSharedRef()),
 					LOCTEXT("CameraSubmenu_FarViewPlaneLabel", "Far View Plane")
