[PCG] Fix HLSL syntax highlighter doing partial matches of keywords and add HLSL functions and other symbols.

* Fixed partial keyword matching ('f' matching 'float')
* Add in / out / inout to keywords
* Added all the "DX10 and later" functions from the HLSL docs
* Added a bunch of system variables like SV_GroupID and others
* Identify identifier chars properly (SV_GroupID would only correctly style the 'SV' part, as underscore split it)

#rnx

diff --git a/Engine/Source/Editor/UnrealEd/Private/Text/HLSLSyntaxHighlighterMarshaller.cpp b/Engine/Source/Editor/UnrealEd/Private/Text/HLSLSyntaxHighlighterMarshaller.cpp
--- a/Engine/Source/Editor/UnrealEd/Private/Text/HLSLSyntaxHighlighterMarshaller.cpp
+++ b/Engine/Source/Editor/UnrealEd/Private/Text/HLSLSyntaxHighlighterMarshaller.cpp
@@ -160,6 +160,9 @@ const TCHAR* HlslKeywords[] =
 	TEXT("bool1"),
 	TEXT("bool"),
 	TEXT("Buffer"),
+	TEXT("in"),
+	TEXT("out"),
+	TEXT("inout"),
 };
 
 const TCHAR* HlslOperators[] =
@@ -228,6 +231,177 @@ const TCHAR* HlslPreProcessorKeywords[] =
 	TEXT("#undef"),
 };
 
+const TCHAR* HlslSymbols[] =
+{
+	TEXT("abort"),
+	TEXT("abs"),
+	TEXT("acos"),
+	TEXT("all"),
+	TEXT("AllMemoryBarrier"),
+	TEXT("AllMemoryBarrierWithGroupSync"),
+	TEXT("any"),
+	TEXT("asdouble"),
+	TEXT("asfloat"),
+	TEXT("asin"),
+	TEXT("asint"),
+	TEXT("asuint"),
+	TEXT("asuint"),
+	TEXT("atan"),
+	TEXT("atan2"),
+	TEXT("ceil"),
+	TEXT("CheckAccessFullyMapped"),
+	TEXT("clamp"),
+	TEXT("clip"),
+	TEXT("cos"),
+	TEXT("cosh"),
+	TEXT("countbits"),
+	TEXT("cross"),
+	TEXT("D3DCOLORtoUBYTE4"),
+	TEXT("ddx"),
+	TEXT("ddx_coarse"),
+	TEXT("ddx_fine"),
+	TEXT("ddy"),
+	TEXT("ddy_coarse"),
+	TEXT("ddy_fine"),
+	TEXT("degrees"),
+	TEXT("determinant"),
+	TEXT("DeviceMemoryBarrier"),
+	TEXT("DeviceMemoryBarrierWithGroupSync"),
+	TEXT("distance"),
+	TEXT("dot"),
+	TEXT("dst"),
+	TEXT("errorf"),
+	TEXT("EvaluateAttributeCentroid"),
+	TEXT("EvaluateAttributeAtSample"),
+	TEXT("EvaluateAttributeSnapped"),
+	TEXT("exp"),
+	TEXT("exp2"),
+	TEXT("f16tof32"),
+	TEXT("f32tof16"),
+	TEXT("faceforward"),
+	TEXT("firstbithigh"),
+	TEXT("firstbitlow"),
+	TEXT("floor"),
+	TEXT("fma"),
+	TEXT("fmod"),
+	TEXT("frac"),
+	TEXT("frexp"),
+	TEXT("fwidth"),
+	TEXT("GetRenderTargetSampleCount"),
+	TEXT("GetRenderTargetSamplePosition"),
+	TEXT("GroupMemoryBarrier"),
+	TEXT("GroupMemoryBarrierWithGroupSync"),
+	TEXT("InterlockedAdd"),
+	TEXT("InterlockedAnd"),
+	TEXT("InterlockedCompareExchange"),
+	TEXT("InterlockedCompareStore"),
+	TEXT("InterlockedExchange"),
+	TEXT("InterlockedMax"),
+	TEXT("InterlockedMin"),
+	TEXT("InterlockedOr"),
+	TEXT("InterlockedXor"),
+	TEXT("isfinite"),
+	TEXT("isinf"),
+	TEXT("isnan"),
+	TEXT("ldexp"),
+	TEXT("length"),
+	TEXT("lerp"),
+	TEXT("lit"),
+	TEXT("log"),
+	TEXT("log10"),
+	TEXT("log2"),
+	TEXT("mad"),
+	TEXT("max"),
+	TEXT("min"),
+	TEXT("modf"),
+	TEXT("msad4"),
+	TEXT("mul"),
+	TEXT("noise"),
+	TEXT("normalize"),
+	TEXT("pow"),
+	TEXT("printf"),
+	TEXT("Process2DQuadTessFactorsAvg"),
+	TEXT("Process2DQuadTessFactorsMax"),
+	TEXT("Process2DQuadTessFactorsMin"),
+	TEXT("ProcessIsolineTessFactors"),
+	TEXT("ProcessQuadTessFactorsAvg"),
+	TEXT("ProcessQuadTessFactorsMax"),
+	TEXT("ProcessQuadTessFactorsMin"),
+	TEXT("ProcessTriTessFactorsAvg"),
+	TEXT("ProcessTriTessFactorsMax"),
+	TEXT("ProcessTriTessFactorsMin"),
+	TEXT("radians"),
+	TEXT("rcp"),
+	TEXT("reflect"),
+	TEXT("refract"),
+	TEXT("reversebits"),
+	TEXT("round"),
+	TEXT("rsqrt"),
+	TEXT("saturate"),
+	TEXT("sign"),
+	TEXT("sin"),
+	TEXT("sincos"),
+	TEXT("sinh"),
+	TEXT("smoothstep"),
+	TEXT("sqrt"),
+	TEXT("step"),
+	TEXT("tan"),
+	TEXT("tanh"),
+	TEXT("tex1D"),
+	TEXT("tex1D"),
+	TEXT("tex1Dbias"),
+	TEXT("tex1Dgrad"),
+	TEXT("tex1Dlod"),
+	TEXT("tex1Dproj"),
+	TEXT("tex2D"),
+	TEXT("tex2D"),
+	TEXT("tex2Dbias"),
+	TEXT("tex2Dgrad"),
+	TEXT("tex2Dlod"),
+	TEXT("tex2Dproj"),
+	TEXT("tex3D"),
+	TEXT("tex3D"),
+	TEXT("tex3Dbias"),
+	TEXT("tex3Dgrad"),
+	TEXT("tex3Dlod"),
+	TEXT("tex3Dproj"),
+	TEXT("texCUBE"),
+	TEXT("texCUBE"),
+	TEXT("texCUBEbias"),
+	TEXT("texCUBEgrad"),
+	TEXT("texCUBElod"),
+	TEXT("texCUBEproj"),
+	TEXT("transpose"),
+	TEXT("trunc"),
+	TEXT("SV_ClipDistance"),
+	TEXT("SV_CullDistance"),
+	TEXT("SV_Coverage"),
+	TEXT("SV_Depth"),
+	TEXT("SV_DepthGreaterEqual"),
+	TEXT("SV_DepthLessEqual"),
+	TEXT("SV_DispatchThreadID"),
+	TEXT("SV_DomainLocation"),
+	TEXT("SV_GroupID"),
+	TEXT("SV_GroupIndex"),
+	TEXT("SV_GroupThreadID"),
+	TEXT("SV_GSInstanceID"),
+	TEXT("SV_InnerCoverage"),
+	TEXT("SV_InsideTessFactor"),
+	TEXT("SV_InstanceID"),
+	TEXT("SV_IsFrontFace"),
+	TEXT("SV_OutputControlPointID"),
+	TEXT("SV_Position"),
+	TEXT("SV_PrimitiveID"),
+	TEXT("SV_RenderTargetArrayIndex"),
+	TEXT("SV_SampleIndex"),
+	TEXT("SV_StencilRef"),
+	TEXT("SV_Target"),
+	TEXT("SV_TessFactor"),
+	TEXT("SV_VertexID"),
+	TEXT("SV_ViewportArrayIndex"),
+	TEXT("SV_ShadingRate"),
+};
+
 /**
  * Tokenize the text based on Hlsl tokens
  */
@@ -259,21 +433,6 @@ public:
 	};
 
 private:
-	static FORCEINLINE bool IsAlpha(TCHAR Char)
-	{
-		return (Char >= 'a' && Char <= 'z') || (Char >= 'A' && Char <= 'Z');
-	}
-
-	static FORCEINLINE bool IsDigit(TCHAR Char)
-	{
-		return Char >= '0' && Char <= '9';
-	}
-
-	static FORCEINLINE bool IsAlphaOrDigit(TCHAR Char)
-	{
-		return IsAlpha(Char) || IsDigit(Char);
-	}
-
 	FHlslSyntaxTokenizer()
 	{
 		// operators
@@ -294,6 +453,11 @@ private:
 			Keywords.Emplace(PreProcessorKeyword);
 		}
 		
+		// Symbols
+		for (const auto& Function : HlslSymbols)
+		{
+			Keywords.Emplace(Function);
+		}
 	};
 
 	void TokenizeLineRanges(const FString& Input, const TArray<FTextRange>& LineRanges, TArray<FTokenizedLine>& OutTokenizedLines)
@@ -349,7 +513,7 @@ private:
 						{
 							const TCHAR PeekChar = Input[PeekOffset];
 
-							if (!IsAlpha(PeekChar))
+							if (!TChar<TCHAR>::IsAlpha(PeekChar))
 							{
 								break;
 							}
@@ -357,7 +521,7 @@ private:
 							PeekOffset++;
 						}
 					}
-					else if (IsAlpha(CurrentChar))
+					else if (TChar<TCHAR>::IsAlpha(CurrentChar))
 					{
 						// Match Identifiers,
 						// They start with a letter and contain
@@ -366,7 +530,7 @@ private:
 						{
 							const TCHAR PeekChar = Input[PeekOffset];
 
-							if (!IsAlphaOrDigit(PeekChar))
+							if (!TChar<TCHAR>::IsIdentifier(PeekChar))
 							{
 								break;
 							}
@@ -380,7 +544,7 @@ private:
 					// Check if it is an reserved keyword
 					for(const FString& Keyword : Keywords)
 					{
-						if(FCString::Strncmp(CurrentString, *Keyword, CurrentStringLength) == 0)
+						if (FCString::Strncmp(CurrentString, *Keyword, FMath::Max(CurrentStringLength, Keyword.Len())) == 0)
 						{
 							const int32 SyntaxTokenEnd = CurrentOffset + CurrentStringLength;
 							TokenizedLine.Tokens.Emplace(FToken(ETokenType::Syntax, FTextRange(CurrentOffset, SyntaxTokenEnd)));
@@ -391,7 +555,7 @@ private:
 							CurrentOffset = SyntaxTokenEnd;
 							break;
 						}
-					}	
+					}
 
 					if (bHasMatchedSyntax)
 					{
@@ -586,13 +750,13 @@ FTextLayout::FNewLineData FHLSLSyntaxHighlighterMarshaller::ProcessTokenizedLine
 					TextBlockStyle = SyntaxTextStyle.CommentTextStyle;
 					ParseState = EParseState::None;
 				}
-				else if(ParseState == EParseState::None && TChar<WIDECHAR>::IsAlpha(TokenText[0]))
+				else if(ParseState == EParseState::None && TChar<TCHAR>::IsIdentifier(TokenText[0]))
 				{
 					RunInfo.Name = TEXT("SyntaxHighlight.HLSL.Keyword");
 					TextBlockStyle = SyntaxTextStyle.KeywordTextStyle;
 					ParseState = EParseState::None;
 				}
-				else if(ParseState == EParseState::None && !TChar<WIDECHAR>::IsAlpha(TokenText[0]))
+				else if(ParseState == EParseState::None && !TChar<TCHAR>::IsIdentifier(TokenText[0]))
 				{
 					RunInfo.Name = TEXT("SyntaxHighlight.HLSL.Operator");
 					TextBlockStyle = SyntaxTextStyle.OperatorTextStyle;
