[shaders] block on global shader map completion after each format in DDC & compileshaders commandlets. this is required since jobs for multiple different global shader maps cannot be in flight at the same time (this results in conflicting IDs, mixing up results, and failed cache puts).

diff --git a/Engine/Source/Editor/UnrealEd/Private/Commandlets/CompileShadersTestBedCommandlet.cpp b/Engine/Source/Editor/UnrealEd/Private/Commandlets/CompileShadersTestBedCommandlet.cpp
--- a/Engine/Source/Editor/UnrealEd/Private/Commandlets/CompileShadersTestBedCommandlet.cpp
+++ b/Engine/Source/Editor/UnrealEd/Private/Commandlets/CompileShadersTestBedCommandlet.cpp
@@ -162,9 +162,9 @@ int32 UCompileShadersTestBedCommandlet::Main(const FString& Params)
 			{
 				const EShaderPlatform ShaderPlatform = ShaderFormatToLegacyShaderPlatform(DesiredShaderFormats[FormatIndex]);
 				CompileGlobalShaderMap(ShaderPlatform, Platform, false);
+				GShaderCompilingManager->ProcessAsyncResults(bLimitExecutationTime, true /* bBlockOnGlobalShaderCompilation */);
+				check(GGlobalShaderMap[ShaderPlatform]->IsComplete(Platform));
 			}
-
-			GShaderCompilingManager->ProcessAsyncResults(bLimitExecutationTime, true /* bBlockOnGlobalShaderCompilation */);
 		}
 
 		TSet<UMaterialInterface*> MaterialsToCompile;
diff --git a/Engine/Source/Editor/UnrealEd/Private/Commandlets/DerivedDataCacheCommandlet.cpp b/Engine/Source/Editor/UnrealEd/Private/Commandlets/DerivedDataCacheCommandlet.cpp
--- a/Engine/Source/Editor/UnrealEd/Private/Commandlets/DerivedDataCacheCommandlet.cpp
+++ b/Engine/Source/Editor/UnrealEd/Private/Commandlets/DerivedDataCacheCommandlet.cpp
@@ -719,6 +719,11 @@ int32 UDerivedDataCacheCommandlet::Main( const FString& Params )
 				const EShaderPlatform ShaderPlatform = ShaderFormatToLegacyShaderPlatform(DesiredShaderFormats[FormatIndex]);
 				// Kick off global shader compiles for each target platform. Note that shader platform alone is not sufficient to distinguish between WindowsEditor and WindowsClient, which after UE 4.25 have different DDC
 				CompileGlobalShaderMap(ShaderPlatform, Platforms[Index], false);
+				// Block to wait for global shaders for this platform to complete before submitting new jobs.
+				// Only a single platform/target's global shader jobs can be in progress at the same time (otherwise these conflict with each other)
+				// so block after each call to CompileGlobalShaderMap.
+				GShaderCompilingManager->ProcessAsyncResults(false, true);
+				check(GGlobalShaderMap[ShaderPlatform]->IsComplete(Platforms[Index]));
 			}
 		}
 
