MPCook: CookPackageSplitters: Fix the order of executing the AllSavesCompleted event on the generator versus transitioning the package to idle. Transitioning to idle sets some state that is needed by the AllSavesCompleted function so it has to come first.
Add some more diagnostics to cook errors about invalid package state transitions, in case they still come up.
#rnx

diff --git a/Engine/Source/Editor/UnrealEd/Private/CookOnTheFlyServer.cpp b/Engine/Source/Editor/UnrealEd/Private/CookOnTheFlyServer.cpp
--- a/Engine/Source/Editor/UnrealEd/Private/CookOnTheFlyServer.cpp
+++ b/Engine/Source/Editor/UnrealEd/Private/CookOnTheFlyServer.cpp
@@ -2824,7 +2824,13 @@ void UCookOnTheFlyServer::DemoteToIdle(UE::Cook::FPackageData& PackageData, UE::
 
 void UCookOnTheFlyServer::PromoteToSaveComplete(UE::Cook::FPackageData& PackageData, UE::Cook::ESendFlags SendFlags)
 {
-	check(PackageData.IsInProgress());
+	if (!PackageData.IsInProgress())
+	{
+		UE_LOG(LogCook, Error, TEXT("Package %s is in PromoteToSaveComplete but is not in progress."),
+			*PackageData.GetPackageName().ToString());
+		FDebug::DumpStackTraceToLog(ELogVerbosity::Warning);
+		return;
+	}
 	WorkerRequests->ReportPromoteToSaveComplete(PackageData);
 	PackageData.SendToState(UE::Cook::EPackageState::Idle, SendFlags, UE::Cook::EStateChangeReason::Saved);
 }
diff --git a/Engine/Source/Editor/UnrealEd/Private/Cooker/CookGenerationHelper.cpp b/Engine/Source/Editor/UnrealEd/Private/Cooker/CookGenerationHelper.cpp
--- a/Engine/Source/Editor/UnrealEd/Private/Cooker/CookGenerationHelper.cpp
+++ b/Engine/Source/Editor/UnrealEd/Private/Cooker/CookGenerationHelper.cpp
@@ -185,6 +185,11 @@ void FGenerationHelper::Uninitialize()
 void FGenerationHelper::ModifyNumSaved(int32 Delta)
 {
 	NumSaved += Delta;
+	if (bDeferEvents)
+	{
+		return;
+	}
+
 	int32 NumAllSaved = PackagesToGenerate.Num() + 1;
 	check(0 <= NumSaved && NumSaved <= NumAllSaved);
 	if (NumSaved == NumAllSaved)
@@ -210,6 +215,23 @@ void FGenerationHelper::ModifyNumSaved(int32 Delta)
 	}
 }
 
+FGenerationHelper::FScopeDeferEvents::FScopeDeferEvents(const TRefCountPtr<FGenerationHelper>& InGenerationHelper)
+	: GenerationHelper(InGenerationHelper)
+{
+	check(GenerationHelper);
+	bOldDeferEvents = GenerationHelper->bDeferEvents;
+	GenerationHelper->bDeferEvents = true;
+}
+
+FGenerationHelper::FScopeDeferEvents::~FScopeDeferEvents()
+{
+	GenerationHelper->bDeferEvents = bOldDeferEvents;
+	if (!bOldDeferEvents)
+	{
+		GenerationHelper->ModifyNumSaved(0);
+	}
+}
+
 void FGenerationHelper::OnAllSavesCompleted(UCookOnTheFlyServer& COTFS)
 {
 	// Caller is responsible for holding a reference that keeps *this from destructing if it clears
@@ -226,7 +248,7 @@ void FGenerationHelper::DemoteStalledPackages(UCookOnTheFlyServer& COTFS, bool b
 	// For any packages that we stalled because they were retracted and assigned to another worker,
 	// demote them now. But don't demote non-stalled packages, because doing so could demote the final
 	// package that we just saved locally and still needs to finish its work in PumpSaves.
-	auto ConditionalDemote = [&COTFS, &bFromAllSavesCompleted](FCookGenerationInfo& Info)
+	auto ConditionalDemote = [&COTFS, &bFromAllSavesCompleted, this](FCookGenerationInfo& Info)
 		{
 			if (Info.PackageData->IsStalled())
 			{
@@ -240,8 +262,11 @@ void FGenerationHelper::DemoteStalledPackages(UCookOnTheFlyServer& COTFS, bool b
 					{
 						UE_LOG(LogCook, Error,
 							TEXT("Package %s is still stalled on the CookDirector during FGenerationHelper::OnAllSavesCompleted. ")
-							TEXT("This is unexpected; all stalled packages on a GenerationHelper should have completed saving and therefore unstalled before OnAllSavesCompleted is called."),
-							*Info.PackageData->GetPackageName().ToString());
+							TEXT("This is unexpected; all stalled packages on a GenerationHelper should have completed saving and therefore unstalled before OnAllSavesCompleted is called.")
+							TEXT("\n\tNumAllSaved == %d. NumSaved == %d. Info.HasSaved == %s."),
+							*Info.PackageData->GetPackageName().ToString(), PackagesToGenerate.Num() + 1, NumSaved,
+							Info.HasSaved() ? TEXT("true") : TEXT("false"));
+						FDebug::DumpStackTraceToLog(ELogVerbosity::Warning);
 
 						COTFS.DemoteToIdle(*Info.PackageData, ESendFlags::QueueAddAndRemove, ESuppressCookReason::RetractedByCookDirector);
 					}
diff --git a/Engine/Source/Editor/UnrealEd/Private/Cooker/CookGenerationHelper.h b/Engine/Source/Editor/UnrealEd/Private/Cooker/CookGenerationHelper.h
--- a/Engine/Source/Editor/UnrealEd/Private/Cooker/CookGenerationHelper.h
+++ b/Engine/Source/Editor/UnrealEd/Private/Cooker/CookGenerationHelper.h
@@ -388,6 +388,16 @@ public:
 	 * Used to manage KeepForGCOrAllSaved lifetime. Does not call Initialize.
 	 */
 	void MarkPackageSavedRemotely(UCookOnTheFlyServer& COTFS, FPackageData& PackageData, FWorkerId SourceWorkerId);
+	/** Scope object that blocks calls to OnAllSavesCompleted. Checks and calls it if necessary when destructed. */
+	struct FScopeDeferEvents
+	{
+	public:
+		FScopeDeferEvents(const TRefCountPtr<FGenerationHelper>& InGenerationHelper);
+		~FScopeDeferEvents();
+	private:
+		TRefCountPtr<FGenerationHelper> GenerationHelper;
+		bool bOldDeferEvents = false;
+	};
 	/**
 	 * Called on the director when the generator or one of the generated packages was found to be iteratively skipped
 	 * and marked already cooked in an incremental cook. Does not call Initialize.
@@ -478,6 +488,7 @@ private:
 	bool bKeepForCompletedAllSavesMessage = false;
 	bool bNeedConfirmGeneratorPackageDestroyed = false;
 	bool bHasFinishedQueueGeneratedPackages = false;
+	bool bDeferEvents = false;
 
 	friend FCookGenerationInfo;
 };
diff --git a/Engine/Source/Editor/UnrealEd/Private/Cooker/CookWorkerServer.cpp b/Engine/Source/Editor/UnrealEd/Private/Cooker/CookWorkerServer.cpp
--- a/Engine/Source/Editor/UnrealEd/Private/Cooker/CookWorkerServer.cpp
+++ b/Engine/Source/Editor/UnrealEd/Private/Cooker/CookWorkerServer.cpp
@@ -21,6 +21,7 @@
 #include "Misc/Char.h"
 #include "Misc/FeedbackContext.h"
 #include "Misc/FileHelper.h"
+#include "Misc/Optional.h"
 #include "Misc/Parse.h"
 #include "Misc/ScopeLock.h"
 #include "PackageResultsMessage.h"
@@ -927,8 +928,63 @@ void FCookWorkerServer::RecordResults(FPackageResultsMessage& Message)
 			continue;
 		}
 		bRetiredAnyPackages = true;
-		PackageData->SetWorkerAssignment(FWorkerId::Invalid(), ESendFlags::QueueNone);
 
+		// MPCOOKTODO: Refactor FSaveCookedPackageContext::FinishPlatform and ::FinishPackage so we can call them from
+		// here to reduce duplication
+		if (Result.GetSuppressCookReason() == ESuppressCookReason::NotSuppressed)
+		{
+			int32 NumPlatforms = OrderedSessionPlatforms.Num();
+			if (Result.GetPlatforms().Num() != NumPlatforms)
+			{
+				UE_LOG(LogCook, Error,
+					TEXT("CookWorkerServer %d received FPackageResultsMessage for package %s with an invalid number of platform results: expected %d, actual %d. Ignoring it."),
+					ProfileId, *Result.GetPackageName().ToString(), NumPlatforms, Result.GetPlatforms().Num());
+			}
+			else
+			{
+				HandleReceivedPackagePlatformMessages(*PackageData, nullptr, Result.ReleaseMessages());
+				for (int32 PlatformIndex = 0; PlatformIndex < NumPlatforms; ++PlatformIndex)
+				{
+					ITargetPlatform* TargetPlatform = OrderedSessionPlatforms[PlatformIndex];
+					FPackageRemoteResult::FPlatformResult& PlatformResult = Result.GetPlatforms()[PlatformIndex];
+					FPackagePlatformData& ExistingData = PackageData->FindOrAddPlatformData(TargetPlatform);
+					if (!ExistingData.NeedsCooking(TargetPlatform))
+					{
+						if (PlatformResult.GetCookResults() != ECookResult::Invalid)
+						{
+							UE_LOG(LogCook, Display,
+								TEXT("CookWorkerServer %d received FPackageResultsMessage for package %s, platform %s, but that platform has already been cooked. Ignoring the results for that platform."),
+								ProfileId, *Result.GetPackageName().ToString(), *TargetPlatform->PlatformName());
+						}
+					}
+					else
+					{
+						if (PlatformResult.GetCookResults() != ECookResult::Invalid)
+						{
+							PackageData->SetPlatformCooked(TargetPlatform, PlatformResult.GetCookResults());
+						}
+						HandleReceivedPackagePlatformMessages(*PackageData, TargetPlatform,
+							PlatformResult.ReleaseMessages());
+					}
+				}
+				COTFS.RecordExternalActorDependencies(Result.GetExternalActorDependencies());
+				if (Result.IsReferencedOnlyByEditorOnlyData())
+				{
+					COTFS.PackageTracker->UncookedEditorOnlyPackages.AddUnique(Result.GetPackageName());
+				}
+			}
+		}
+
+		// For generator and generated packages, after we handle all of their messages above, execute their state
+		// changes in the required order:
+		// *) Defer the GenerationHelper's events so that we don't yet complete it if this was the last save.
+		// *) Mark saved on the generator, so that the generator has full context for the save.
+		// *) Transition the PackageData state to complete. The code to automatically mark generated as saved with the
+		//    generator will be skipped since we already did it in the step above.
+		// *) Unfreeze the GenerationHelper's events and call OnAllSavesCompleted if necessary.
+		TOptional<FGenerationHelper::FScopeDeferEvents> DeferGenerationHelperEvents;
+
+		// If generated or generator, defer events and mark saved
 		if (PackageData->IsGenerated())
 		{
 			TRefCountPtr<FGenerationHelper> ParentGenerationHelper = PackageData->GetOrFindParentGenerationHelper();
@@ -940,6 +996,7 @@ void FCookWorkerServer::RecordResults(FPackageResultsMessage& Message)
 			}
 			else
 			{
+				DeferGenerationHelperEvents.Emplace(ParentGenerationHelper);
 				ParentGenerationHelper->MarkPackageSavedRemotely(COTFS, *PackageData, GetWorkerId());
 				EStateChangeReason StateChangeReason =
 					Result.GetSuppressCookReason() == ESuppressCookReason::NotSuppressed
@@ -951,54 +1008,15 @@ void FCookWorkerServer::RecordResults(FPackageResultsMessage& Message)
 		TRefCountPtr<FGenerationHelper> GenerationHelper = PackageData->GetGenerationHelper();
 		if (GenerationHelper)
 		{
+			DeferGenerationHelperEvents.Emplace(GenerationHelper);
 			GenerationHelper->MarkPackageSavedRemotely(COTFS, *PackageData, GetWorkerId());
 			GenerationHelper.SafeRelease();
 		}
 
-		// MPCOOKTODO: Refactor FSaveCookedPackageContext::FinishPlatform and ::FinishPackage so we can call them from
-		// here to reduce duplication
+		// For all packages, transition them to idle
+		PackageData->SetWorkerAssignment(FWorkerId::Invalid(), ESendFlags::QueueNone);
 		if (Result.GetSuppressCookReason() == ESuppressCookReason::NotSuppressed)
 		{
-			int32 NumPlatforms = OrderedSessionPlatforms.Num();
-			if (Result.GetPlatforms().Num() != NumPlatforms)
-			{
-				UE_LOG(LogCook, Warning,
-					TEXT("CookWorkerServer %d received FPackageResultsMessage for package %s with an invalid number of platform results: expected %d, actual %d. Ignoring it."),
-					ProfileId, *Result.GetPackageName().ToString(), NumPlatforms, Result.GetPlatforms().Num());
-				continue;
-			}
-
-			HandleReceivedPackagePlatformMessages(*PackageData, nullptr, Result.ReleaseMessages());
-			for (int32 PlatformIndex = 0; PlatformIndex < NumPlatforms; ++PlatformIndex)
-			{
-				ITargetPlatform* TargetPlatform = OrderedSessionPlatforms[PlatformIndex];
-				FPackageRemoteResult::FPlatformResult& PlatformResult = Result.GetPlatforms()[PlatformIndex];
-				FPackagePlatformData& ExistingData = PackageData->FindOrAddPlatformData(TargetPlatform);
-				if (!ExistingData.NeedsCooking(TargetPlatform))
-				{
-					if (PlatformResult.GetCookResults() != ECookResult::Invalid)
-					{
-						UE_LOG(LogCook, Display,
-							TEXT("CookWorkerServer %d received FPackageResultsMessage for package %s, platform %s, but that platform has already been cooked. Ignoring the results for that platform."),
-							ProfileId, *Result.GetPackageName().ToString(), *TargetPlatform->PlatformName());
-					}
-					continue;
-				}
-				else
-				{
-					if (PlatformResult.GetCookResults() != ECookResult::Invalid)
-					{
-						PackageData->SetPlatformCooked(TargetPlatform, PlatformResult.GetCookResults());
-					}
-					HandleReceivedPackagePlatformMessages(*PackageData, TargetPlatform,
-						PlatformResult.ReleaseMessages());
-				}
-			}
-			COTFS.RecordExternalActorDependencies(Result.GetExternalActorDependencies());
-			if (Result.IsReferencedOnlyByEditorOnlyData())
-			{
-				COTFS.PackageTracker->UncookedEditorOnlyPackages.AddUnique(Result.GetPackageName());
-			}
 			COTFS.PromoteToSaveComplete(*PackageData, ESendFlags::QueueAddAndRemove);
 		}
 		else if (Result.GetSuppressCookReason() == ESuppressCookReason::RetractedByCookDirector)
@@ -1015,7 +1033,11 @@ void FCookWorkerServer::RecordResults(FPackageResultsMessage& Message)
 		{
 			COTFS.DemoteToIdle(*PackageData, ESendFlags::QueueAddAndRemove, Result.GetSuppressCookReason());
 		}
+
+		// For generated packages, undefer events and process AllSavesCompleted if necessary.
+		DeferGenerationHelperEvents.Reset();
 	}
+
 	Director.ResetFinalIdleHeartbeatFence();
 	if (bRetiredAnyPackages)
 	{
