Fix LOCTEXT conflict

diff --git a/Engine/Source/Editor/UnrealEd/Private/Kismet2/DebuggerCommands.cpp b/Engine/Source/Editor/UnrealEd/Private/Kismet2/DebuggerCommands.cpp
--- a/Engine/Source/Editor/UnrealEd/Private/Kismet2/DebuggerCommands.cpp
+++ b/Engine/Source/Editor/UnrealEd/Private/Kismet2/DebuggerCommands.cpp
@@ -1258,11 +1258,11 @@ FText FInternalPlayWorldCommandCallbacks::GetSimulatePauseEjectLabel()
 	{
 		if (FPlayWorldCommandCallbacks::HasPlayWorldAndPaused())
 		{
-			return LOCTEXT("SimulatePauseEject_Resume", "Resume");
+			return LOCTEXT("SimulatePauseEjectLabel_Resume", "Resume");
 		}
 		else
 		{
-			return LOCTEXT("SimulatePauseEject_Pause", "Pause");
+			return LOCTEXT("SimulatePauseEjectLabel_Pause", "Pause");
 		}
 	}
 	else
@@ -1281,11 +1281,11 @@ FText FInternalPlayWorldCommandCallbacks::GetSimulatePauseEjectTooltip()
 	{
 		if (FPlayWorldCommandCallbacks::HasPlayWorldAndPaused())
 		{
-			return LOCTEXT("SimulatePauseEject_Resume", "Resumes the simulation");
+			return LOCTEXT("SimulatePauseEjectTooltip_Resume", "Resumes the simulation");
 		}
 		else
 		{
-			return LOCTEXT("SimulatePauseEject_Pause", "Pauses the simulation");
+			return LOCTEXT("SimulatePauseEjectTooltip_Pause", "Pauses the simulation");
 		}
 	}
 	else
