ResaveContentCommandlet: when fixing up redirectors, make sure to only resave external object files that points to redirectors, not all of them.

diff --git a/Engine/Source/Editor/UnrealEd/Private/Commandlets/ContentCommandlets.cpp b/Engine/Source/Editor/UnrealEd/Private/Commandlets/ContentCommandlets.cpp
--- a/Engine/Source/Editor/UnrealEd/Private/Commandlets/ContentCommandlets.cpp
+++ b/Engine/Source/Editor/UnrealEd/Private/Commandlets/ContentCommandlets.cpp
@@ -5,7 +5,6 @@
 =============================================================================*/
 
 #include "Algo/RemoveIf.h"
-#include "Algo/Unique.h"
 #include "AssetCompilingManager.h"
 #include "AssetRegistry/AssetData.h"
 #include "CollectionManagerModule.h"
@@ -435,26 +434,6 @@ int32 UResavePackagesCommandlet::InitializeResaveParameters( const TArray<FStrin
 				TArray<FName> Referencers;
 				AssetRegistry.GetReferencers(AssetData.PackageName, Referencers);
 
-				// For external objects referencers, also add the object's outer package as a referencer so it can be handled by PerformAdditionalOperations.
-				FARFilter Filter;
-				Filter.bIncludeOnlyOnDiskAssets = true;
-				Filter.PackageNames = Referencers;
-
-				TArray<FAssetData> AssetReferencers;
-				AssetRegistry.GetAssets(Filter, AssetReferencers);
-
-				TArray<FName> ReferencerOuters;
-				for (const FAssetData& AssetReferencer : AssetReferencers)
-				{
-					if (!AssetReferencer.GetOptionalOuterPathName().IsNone())
-					{
-						Referencers.Add(FSoftObjectPath(AssetReferencer.GetOptionalOuterPathName().ToString()).GetLongPackageFName());
-					}
-				}
-
-				Referencers.Sort(FNameFastLess());
-				Referencers.SetNum(Algo::Unique(Referencers));
-
 				for (FName Referencer : Referencers)
 				{
 					FString ReferencerFile;
@@ -475,9 +454,29 @@ int32 UResavePackagesCommandlet::InitializeResaveParameters( const TArray<FStrin
 				RedirectorsToFixup.Add(PackageName);
 			}
 
-			if (ReferencerPackages.Contains(PackageName))
+			for (const FString& ReferencerPackage : ReferencerPackages)
 			{
-				PackageNames.Add(PackageName);
+				if (ReferencerPackage == PackageName)
+				{
+					PackageNames.Add(PackageName);
+				}
+				else if (ReferencerPackage.Contains(FPackagePath::GetExternalActorsFolderName()) || ReferencerPackage.Contains(FPackagePath::GetExternalObjectsFolderName()))
+				{
+					FString CleanPackageName = PackageName;
+					if (CleanPackageName.RemoveFromEnd(TEXT(".umap")))
+					{
+						FString WorldReferencerPackage = 
+							ReferencerPackage
+								.Replace(FPackagePath::GetExternalActorsFolderName(), TEXT(""))
+								.Replace(FPackagePath::GetExternalObjectsFolderName(), TEXT(""))
+								.Replace(TEXT("//"), TEXT("/"));
+						WorldReferencerPackage.LeftInline(CleanPackageName.Len());
+						if (WorldReferencerPackage == CleanPackageName)
+						{
+							PackageNames.Add(ReferencerPackage);
+						}
+					}
+				}
 			}
 		}
 	}
@@ -1907,6 +1906,7 @@ void UResavePackagesCommandlet::PerformAdditionalOperations(class UWorld* World,
 
 	bool bRevertCheckedOutFilesIfNotSaving = true;
 
+	const bool bFixupRedirects = (Switches.Contains(TEXT("FixupRedirects")) || Switches.Contains(TEXT("FixupRedirectors")));
 	const bool bShouldBuildTextureStreamingForWorld = bShouldBuildTextureStreaming && !bShouldBuildTextureStreamingForAll;
 	const bool bBuildingNonHLODData = (bShouldBuildLighting || bShouldBuildTextureStreamingForWorld || bShouldBuildReflectionCaptures);
 
@@ -1918,7 +1918,7 @@ void UResavePackagesCommandlet::PerformAdditionalOperations(class UWorld* World,
 	const int32 DefaultExternalActorGCFreq = 2048;
 
 	// Load and Save Level's external packages
- 	if (!bResaveWorldPartitionExternalPackages)
+ 	if (!bResaveWorldPartitionExternalPackages && !bFixupRedirects)
 	{
 		// Use a default GC frequency for external actors if GarbageCollectionFrequency is 0.
 		TGuardValue<int32> ScopedGCFreq(GarbageCollectionFrequency, GarbageCollectionFrequency ? GarbageCollectionFrequency : DefaultExternalActorGCFreq);
@@ -1956,35 +1956,38 @@ void UResavePackagesCommandlet::PerformAdditionalOperations(class UWorld* World,
 		// Use a default GC frequency for external actors if GarbageCollectionFrequency is 0.
 		TGuardValue<int32> ScopedGCFreq(GarbageCollectionFrequency, GarbageCollectionFrequency ? GarbageCollectionFrequency : DefaultExternalActorGCFreq);
 
-		auto ResaveExternalPackage = [this](const UPackage* Package)
+		if (!bFixupRedirects)
 		{
-			++TotalPackagesForResave;
-			if (Package == nullptr)
+			auto ResaveExternalPackage = [this](const UPackage* Package)
 			{
-				check(bCanIgnoreFails);
-				return;
-			}
-			const FString PackageFilename = Package->GetLoadedPath().GetLocalFullPath();
-			check(FLinkerLoad::FindExistingLinkerForPackage(Package));
-			LoadAndSaveOnePackage(PackageFilename);
-		};
+				++TotalPackagesForResave;
+				if (Package == nullptr)
+				{
+					check(bCanIgnoreFails);
+					return;
+				}
+				const FString PackageFilename = Package->GetLoadedPath().GetLocalFullPath();
+				check(FLinkerLoad::FindExistingLinkerForPackage(Package));
+				LoadAndSaveOnePackage(PackageFilename);
+			};
 
-		// Resave all external actors packages
-		FWorldPartitionHelpers::ForEachActorDescInstance(WorldPartition, [this, WorldPartition, ResaveExternalPackage](const FWorldPartitionActorDescInstance* ActorDescInstance)
-		{
-			// Load & Register World Partition Actor
-			FWorldPartitionReference LoadedActor(WorldPartition, ActorDescInstance->GetGuid());
-			AActor* Actor = LoadedActor.GetActor();
-			UPackage* Package = Actor ? Actor->GetExternalPackage() : nullptr;
-			ResaveExternalPackage(Package);
-			return true;
-		});
+			// Resave all external actors packages
+			FWorldPartitionHelpers::ForEachActorDescInstance(WorldPartition, [this, WorldPartition, ResaveExternalPackage](const FWorldPartitionActorDescInstance* ActorDescInstance)
+			{
+				// Load & Register World Partition Actor
+				FWorldPartitionReference LoadedActor(WorldPartition, ActorDescInstance->GetGuid());
+				AActor* Actor = LoadedActor.GetActor();
+				UPackage* Package = Actor ? Actor->GetExternalPackage() : nullptr;
+				ResaveExternalPackage(Package);
+				return true;
+			});
 
-		// Resave all external objects packages
-		FExternalPackageHelper::LoadObjectsFromExternalPackages<UObject>(World, [this, ResaveExternalPackage](UObject* ExternalObject)
-		{
-			ResaveExternalPackage(ExternalObject->GetPackage());
-		});
+			// Resave all external objects packages
+			FExternalPackageHelper::LoadObjectsFromExternalPackages<UObject>(World, [this, ResaveExternalPackage](UObject* ExternalObject)
+			{
+				ResaveExternalPackage(ExternalObject->GetPackage());
+			});
+		}
 	}
 
 	if (bBuildingNonHLODData || bShouldBuildHLOD || bShouldBuildNavigationData)
