Remove the "Select" submenu from the new level editor viewport toolbar

Also move a few selection-related toggles from the removed "Select" submenu to the "Transform" menu.

diff --git a/Engine/Source/Editor/LevelEditor/Private/SLevelViewport.cpp b/Engine/Source/Editor/LevelEditor/Private/SLevelViewport.cpp
--- a/Engine/Source/Editor/LevelEditor/Private/SLevelViewport.cpp
+++ b/Engine/Source/Editor/LevelEditor/Private/SLevelViewport.cpp
@@ -1966,13 +1966,6 @@ TSharedPtr<SWidget> SLevelViewport::BuildViewportToolbar()
 					LeftSection.AddEntry(TransformsSubmenu);
 				}
 
-				// Add the "Selection" sub menu.
-				{
-					FToolMenuEntry SelectionSubmenu = UE::UnrealEd::CreateViewportToolbarSelectSection();
-					SelectionSubmenu.InsertPosition.Position = EToolMenuInsertType::First;
-					LeftSection.AddEntry(SelectionSubmenu);
-				}
-
 				// Add the "Snapping" sub menu.
 				{
 					FToolMenuEntry SnappingSubmenu = UE::UnrealEd::CreateViewportToolbarSnappingSubmenu();
diff --git a/Engine/Source/Editor/UnrealEd/Private/ViewportToolbar/UnrealEdViewportToolbar.cpp b/Engine/Source/Editor/UnrealEd/Private/ViewportToolbar/UnrealEdViewportToolbar.cpp
--- a/Engine/Source/Editor/UnrealEd/Private/ViewportToolbar/UnrealEdViewportToolbar.cpp
+++ b/Engine/Source/Editor/UnrealEd/Private/ViewportToolbar/UnrealEdViewportToolbar.cpp
@@ -788,6 +788,31 @@ FToolMenuEntry CreateViewportToolbarTransformsSection()
 						"GizmoScale", GizmoScaleWidget, LOCTEXT("GizmoScaleLabel", "Gizmo Scale")
 					));
 				}
+
+				{
+					FToolMenuSection& SelectionSection =
+						Submenu->FindOrAddSection("Selection", LOCTEXT("SelectionLabel", "Selection"));
+
+					SelectionSection.AddMenuEntry(
+						FLevelEditorCommands::Get().AllowTranslucentSelection,
+						LOCTEXT("SelectionAllowTranslucentSelectionLabel", "Translucent Objects")
+					);
+
+					SelectionSection.AddMenuEntry(
+						FLevelEditorCommands::Get().AllowGroupSelection,
+						LOCTEXT("SelectionAllowGroupSelectionLabel", "Select Groups")
+					);
+
+					SelectionSection.AddMenuEntry(
+						FLevelEditorCommands::Get().StrictBoxSelect,
+						LOCTEXT("SelectionStrictBoxSelectLabel", "Strict Marquee Selection")
+					);
+
+					SelectionSection.AddMenuEntry(
+						FLevelEditorCommands::Get().TransparentBoxSelect,
+						LOCTEXT("SelectionTransparentBoxSelectLabel", "Marquee Select Occluded")
+					);
+				}
 			}
 		)
 	);
@@ -1076,31 +1101,6 @@ FToolMenuEntry CreateViewportToolbarSelectSection()
 						FSlateIconFinder::FindIcon("ShowFlagsMenu.StaticMeshes")
 					);
 				}
-
-				{
-					FToolMenuSection& OptionsSection =
-						Submenu->FindOrAddSection("Options", LOCTEXT("OptionsLabel", "Options"));
-
-					OptionsSection.AddMenuEntry(
-						FLevelEditorCommands::Get().AllowTranslucentSelection,
-						LOCTEXT("OptionsAllowTranslucentSelectionLabel", "Translucent Objects")
-					);
-
-					OptionsSection.AddMenuEntry(
-						FLevelEditorCommands::Get().AllowGroupSelection,
-						LOCTEXT("OptionsAllowGroupSelectionLabel", "Select Groups")
-					);
-
-					OptionsSection.AddMenuEntry(
-						FLevelEditorCommands::Get().StrictBoxSelect,
-						LOCTEXT("OptionsStrictBoxSelectLabel", "Strict Marquee Selection")
-					);
-
-					OptionsSection.AddMenuEntry(
-						FLevelEditorCommands::Get().TransparentBoxSelect,
-						LOCTEXT("OptionsTransparentBoxSelectLabel", "Marquee Select Occluded")
-					);
-				}
 			}
 		)
 	);
