Fix regression where F8 could no longer toggle back from PIE to SIE

diff --git a/Engine/Source/Editor/UnrealEd/Private/Kismet2/DebuggerCommands.cpp b/Engine/Source/Editor/UnrealEd/Private/Kismet2/DebuggerCommands.cpp
--- a/Engine/Source/Editor/UnrealEd/Private/Kismet2/DebuggerCommands.cpp
+++ b/Engine/Source/Editor/UnrealEd/Private/Kismet2/DebuggerCommands.cpp
@@ -375,7 +375,7 @@ void FPlayWorldCommands::RegisterCommands()
 	UI_COMMAND(LateJoinSession, "Add Client", "Add another client", EUserInterfaceActionType::Button, FInputChord());
 	UI_COMMAND(SingleFrameAdvance, "Skip", "Advances a single frame", EUserInterfaceActionType::Button, FInputChord());
 	UI_COMMAND(TogglePlayPauseOfPlaySession, "Toggle Play/Pause", "Resume playing if paused, or pause if playing", EUserInterfaceActionType::Button, FInputChord(EKeys::Pause));
-	UI_COMMAND(PlayPossess, "Possess Player", "Possesses the player character", EUserInterfaceActionType::Button, FInputChord(EKeys::F8));
+	UI_COMMAND(PlayPossess, "Possess Player", "Possesses the player character", EUserInterfaceActionType::Button, FInputChord());
 	UI_COMMAND(PossessEjectPlayer, "Possess or Eject Player", "Possesses or ejects the player from the camera", EUserInterfaceActionType::Button, FInputChord(EKeys::F8));
 	UI_COMMAND(SimulatePauseEject, "Simulate, Pause Simulation or Eject Player", "Starts or pauses simulation or eject the player from the camera", EUserInterfaceActionType::Button, FInputChord());
 	UI_COMMAND(ShowCurrentStatement, "Locate", "Locate the currently active node", EUserInterfaceActionType::Button, FInputChord());
