Add TranslateRotateMode and TranslateRotate2DMode to the new level editor viewport toolbar

* Add TranslateRotateMode and TranslateRotate2DMode as top-level raised buttons of the Transform menu of the new level editor viewport toolbar. These can be enabled in the Editor Preferences and Project Settings, respectively.
* Make mode commands radio buttons instead of togglebuttons. This applies to SelectMode, TranslatleMode, RotateMode, ScaleMode, TranslateRotateMode, and TranslateRotate2DMod. I'm making this change because these commands activate states that are supposed to be mutually exclusive and I want them indicate this to users by showing up as radio buttons when added to ToolMenus, so their types have to change.

diff --git a/Engine/Source/Editor/UnrealEd/Private/EditorViewportCommands.cpp b/Engine/Source/Editor/UnrealEd/Private/EditorViewportCommands.cpp
--- a/Engine/Source/Editor/UnrealEd/Private/EditorViewportCommands.cpp
+++ b/Engine/Source/Editor/UnrealEd/Private/EditorViewportCommands.cpp
@@ -143,13 +143,12 @@ void FEditorViewportCommands::RegisterCommands()
 	UI_COMMAND( IncrementRotationGridSize, "Grid Size (Rotation): Increment", "Increases the rotation grid size setting by one", EUserInterfaceActionType::Button, FInputChord( EModifierKey::Shift, EKeys::RightBracket ) );
 	UI_COMMAND( DecrementRotationGridSize, "Grid Size (Rotation): Decrement", "Decreases the rotation grid size setting by one", EUserInterfaceActionType::Button, FInputChord( EModifierKey::Shift, EKeys::LeftBracket ) );
 
-	
-	UI_COMMAND(SelectMode, "Select Mode", "Select objects", EUserInterfaceActionType::ToggleButton, FInputChord(EKeys::Q));
-	UI_COMMAND( TranslateMode, "Translate Mode", "Select and translate objects", EUserInterfaceActionType::ToggleButton, FInputChord(EKeys::W) );
-	UI_COMMAND( RotateMode, "Rotate Mode", "Select and rotate objects", EUserInterfaceActionType::ToggleButton, FInputChord(EKeys::E) );
-	UI_COMMAND( ScaleMode, "Scale Mode", "Select and scale objects", EUserInterfaceActionType::ToggleButton, FInputChord(EKeys::R) );
-	UI_COMMAND( TranslateRotateMode, "Combined Translate and Rotate Mode", "Select and translate or rotate objects", EUserInterfaceActionType::ToggleButton, FInputChord() );
-	UI_COMMAND( TranslateRotate2DMode, "2D Mode", "Select and translate or rotate objects in 2D", EUserInterfaceActionType::ToggleButton, FInputChord());
+	UI_COMMAND(SelectMode, "Select Mode", "Select objects", EUserInterfaceActionType::RadioButton, FInputChord(EKeys::Q));
+	UI_COMMAND( TranslateMode, "Translate Mode", "Select and translate objects", EUserInterfaceActionType::RadioButton, FInputChord(EKeys::W) );
+	UI_COMMAND( RotateMode, "Rotate Mode", "Select and rotate objects", EUserInterfaceActionType::RadioButton, FInputChord(EKeys::E) );
+	UI_COMMAND( ScaleMode, "Scale Mode", "Select and scale objects", EUserInterfaceActionType::RadioButton, FInputChord(EKeys::R) );
+	UI_COMMAND( TranslateRotateMode, "Combined Translate and Rotate Mode", "Select and translate or rotate objects", EUserInterfaceActionType::RadioButton, FInputChord() );
+	UI_COMMAND( TranslateRotate2DMode, "2D Mode", "Select and translate or rotate objects in 2D", EUserInterfaceActionType::RadioButton, FInputChord());
 
 	UI_COMMAND( ShrinkTransformWidget, "Shrink Transform Widget", "Shrink the level editor transform widget", EUserInterfaceActionType::Button, FInputChord(EModifierKey::Alt, EKeys::LeftBracket) );
 	UI_COMMAND( ExpandTransformWidget, "Expand Transform Widget", "Expand the level editor transform widget", EUserInterfaceActionType::Button, FInputChord(EModifierKey::Alt, EKeys::RightBracket) );
diff --git a/Engine/Source/Editor/UnrealEd/Private/ViewportToolbar/UnrealEdViewportToolbar.cpp b/Engine/Source/Editor/UnrealEd/Private/ViewportToolbar/UnrealEdViewportToolbar.cpp
--- a/Engine/Source/Editor/UnrealEd/Private/ViewportToolbar/UnrealEdViewportToolbar.cpp
+++ b/Engine/Source/Editor/UnrealEd/Private/ViewportToolbar/UnrealEdViewportToolbar.cpp
@@ -666,26 +666,32 @@ FToolMenuEntry CreateViewportToolbarTransformsSection()
 						Submenu->FindOrAddSection("TransformTools", LOCTEXT("TransformToolsLabel", "Transform Tools"));
 
 					FToolMenuEntry SelectMode = FToolMenuEntry::InitMenuEntry(FEditorViewportCommands::Get().SelectMode);
-					SelectMode.UserInterfaceActionType = EUserInterfaceActionType::RadioButton;
 					SelectMode.SetShowInToolbarTopLevel(true);
 					TransformToolsSection.AddEntry(SelectMode);
 
 					FToolMenuEntry TranslateMode =
 						FToolMenuEntry::InitMenuEntry(FEditorViewportCommands::Get().TranslateMode);
-					TranslateMode.UserInterfaceActionType = EUserInterfaceActionType::RadioButton;
 					TranslateMode.SetShowInToolbarTopLevel(true);
 					TransformToolsSection.AddEntry(TranslateMode);
 
 					FToolMenuEntry RotateMode = FToolMenuEntry::InitMenuEntry(FEditorViewportCommands::Get().RotateMode);
-					RotateMode.UserInterfaceActionType = EUserInterfaceActionType::RadioButton;
 					RotateMode.SetShowInToolbarTopLevel(true);
 					TransformToolsSection.AddEntry(RotateMode);
 
 					FToolMenuEntry ScaleMode = FToolMenuEntry::InitMenuEntry(FEditorViewportCommands::Get().ScaleMode);
-					ScaleMode.UserInterfaceActionType = EUserInterfaceActionType::RadioButton;
 					ScaleMode.SetShowInToolbarTopLevel(true);
 					TransformToolsSection.AddEntry(ScaleMode);
 
+					FToolMenuEntry TranslateRotateMode =
+						FToolMenuEntry::InitMenuEntry(FEditorViewportCommands::Get().TranslateRotateMode);
+					TranslateRotateMode.SetShowInToolbarTopLevel(true);
+					TransformToolsSection.AddEntry(TranslateRotateMode);
+
+					FToolMenuEntry TranslateRotate2DMode =
+						FToolMenuEntry::InitMenuEntry(FEditorViewportCommands::Get().TranslateRotate2DMode);
+					TranslateRotate2DMode.SetShowInToolbarTopLevel(true);
+					TransformToolsSection.AddEntry(TranslateRotate2DMode);
+
 					// Build a submenu for selecting the coordinate system to use.
 					{
 						TransformToolsSection.AddSeparator("CoordinateSystemSeparator");
