Auto-save of external packages no longer changes the dirty flag on the package (avoids spamming the dirty package trackers).

diff --git a/Engine/Source/Editor/UnrealEd/Private/FileHelpers.cpp b/Engine/Source/Editor/UnrealEd/Private/FileHelpers.cpp
--- a/Engine/Source/Editor/UnrealEd/Private/FileHelpers.cpp
+++ b/Engine/Source/Editor/UnrealEd/Private/FileHelpers.cpp
@@ -3153,7 +3153,7 @@ EAutosaveContentPackagesResult::Type FEditorFileUtils::AutosaveMapEx(const FStri
 					for (UPackage* ExternalPackage : ExternalPackagesToSave)
 					{
 						const FString AutosaveFilename = GetAutoSaveFilename(ExternalPackage, AbsoluteAutosaveDir, AutosaveIndex, FPackageName::GetAssetPackageExtension());
-						if (!GEditor->Exec(nullptr, *FString::Printf(TEXT("OBJ SAVEPACKAGE PACKAGE=\"%s\" FILE=\"%s\" SILENT=false AUTOSAVING=true"), *ExternalPackage->GetName(), *AutosaveFilename)))
+						if (!GEditor->Exec(nullptr, *FString::Printf(TEXT("OBJ SAVEPACKAGE PACKAGE=\"%s\" FILE=\"%s\" SILENT=false AUTOSAVING=true KEEPDIRTY=true"), *ExternalPackage->GetName(), *AutosaveFilename)))
 						{
 							return EAutosaveContentPackagesResult::Failure;
 						}
@@ -3161,8 +3161,8 @@ EAutosaveContentPackagesResult::Type FEditorFileUtils::AutosaveMapEx(const FStri
 						// We saved an actor
 						bResult = true;
 
-						// Re-mark the package as dirty, because autosaving it will have cleared the dirty flag
-						ExternalPackage->MarkPackageDirty();
+						// Verify that the package is still dirty after the save
+						check(ExternalPackage->IsDirty());
 					}
 
 					FEditorDelegates::PostSaveExternalActors.Broadcast(World);
