Plugin encryption fixes
 - Instead of scanning the file system for plugin assets, use the cooked files list and bring in anything that lives under the encrypted plugin mount point
 - When generating the asset registry at the end of cook, if a package has an encrypted chunk assigned, ignore any generator connections.

diff --git a/Engine/Source/Editor/UnrealEd/Private/Commandlets/AssetRegistryGenerator.cpp b/Engine/Source/Editor/UnrealEd/Private/Commandlets/AssetRegistryGenerator.cpp
--- a/Engine/Source/Editor/UnrealEd/Private/Commandlets/AssetRegistryGenerator.cpp
+++ b/Engine/Source/Editor/UnrealEd/Private/Commandlets/AssetRegistryGenerator.cpp
@@ -856,17 +856,34 @@ void FAssetRegistryGenerator::CalculateChunkIdsAndAssignToManifest(const FName&
 	{
 		FName PackageNameThatDefinesChunks = PackageFName;
 
-		// Generated packages use the chunks defined by their Generator
-		FName GeneratorName = GetGeneratorPackage(PackageFName, this->State);
-		if (!GeneratorName.IsNone())
+		UAssetManager& AssetManager = UAssetManager::Get();
+
+		TArray<int32> PackageChunkIDs;
+		
+		for (int32 ChunkID : AssetManager.GetEncryptedChunkIDsForPackage(PackageFName))
 		{
-			PackageNameThatDefinesChunks = GeneratorName;
+			PackageChunkIDs.Add(ChunkID);
 		}
 
-		TArray<int32> PackageChunkIDs = GetExplicitChunkIDs(PackageNameThatDefinesChunks);
-		ExistingChunkIDs = GetExistingPackageChunkAssignments(PackageNameThatDefinesChunks);
-		PackageChunkIDs.Append(ExistingChunkIDs);
-		UAssetManager::Get().GetPackageChunkIds(PackageNameThatDefinesChunks, TargetPlatform, PackageChunkIDs, TargetChunks);
+		// We only want to override the package name 
+		if (PackageChunkIDs.Num() == 0)
+		{
+			// Generated packages use the chunks defined by their Generator
+			FName GeneratorName = GetGeneratorPackage(PackageFName, this->State);
+			if (!GeneratorName.IsNone())
+			{
+				PackageNameThatDefinesChunks = GeneratorName;
+			}
+
+			PackageChunkIDs = GetExplicitChunkIDs(PackageNameThatDefinesChunks);
+			ExistingChunkIDs = GetExistingPackageChunkAssignments(PackageNameThatDefinesChunks);
+			PackageChunkIDs.Append(ExistingChunkIDs);
+			AssetManager.GetPackageChunkIds(PackageNameThatDefinesChunks, TargetPlatform, PackageChunkIDs, TargetChunks);
+		}
+		else
+		{
+			TargetChunks.Append(PackageChunkIDs);
+		}
 	}
 
 	// Add the package to the manifest for every chunk the AssetManager found it should belong to
diff --git a/Engine/Source/Runtime/Engine/Classes/Engine/AssetManager.h b/Engine/Source/Runtime/Engine/Classes/Engine/AssetManager.h
--- a/Engine/Source/Runtime/Engine/Classes/Engine/AssetManager.h
+++ b/Engine/Source/Runtime/Engine/Classes/Engine/AssetManager.h
@@ -616,6 +616,10 @@ public:
 	 */
 	virtual FGuid GetChunkEncryptionKeyGuid(int32 InChunkId) const { return FGuid(); }
 
+
+	/** Get the encrypted chunk assignments for a give package. Only valid between PreSaveAssetRegistry and PostSaveAssetRegistry */
+	ENGINE_API virtual TSet<int32> GetEncryptedChunkIDsForPackage(FName InPackageName) const { return TSet<int32>(); }
+
 	/**
 	 * Determine if we should separate the asset registry for this chunk out into its own file and return the unique name that identifies it
 	 * @param InChunkIndex Chunk index to check
diff --git a/Engine/Source/Runtime/Engine/Classes/Engine/ContentEncryptionConfig.h b/Engine/Source/Runtime/Engine/Classes/Engine/ContentEncryptionConfig.h
--- a/Engine/Source/Runtime/Engine/Classes/Engine/ContentEncryptionConfig.h
+++ b/Engine/Source/Runtime/Engine/Classes/Engine/ContentEncryptionConfig.h
@@ -28,6 +28,10 @@ public:
 	{
 		TSet<FName> PackageNames;
 		TSet<FString> NonAssetFiles;
+
+		// Allow this group to track all assets under a specific set of mount points
+		TSet<FName> MountPoints;
+
 		EGroupType GroupType = EGroupType::Root;
 		EAllowedReferences AllowedReferences = EAllowedReferences::None;
 	};
@@ -44,6 +48,11 @@ public:
 		PackageGroups.FindOrAdd(InGroupName).NonAssetFiles.Add(InFilename);
 	}
 
+	void AddMountPoint(FName InGroupName, FName InMountPoint)
+	{
+		PackageGroups.FindOrAdd(InGroupName).MountPoints.Add(InMountPoint);
+	}
+
 	void SetGroupType(FName InGroupName, EGroupType InGroupType)
 	{
 		PackageGroups.FindOrAdd(InGroupName).GroupType = InGroupType;
