IncrementalCook: Normalize the package names in the ReferencedSet (lower case) to match normalization done when creating the oplog key for each package. We do not want the zenserver to need to handle that normalization; it should be a client responsibility.
Add a # Version line to the ReferencedSet so the code that interprets it knows to ignore the older non-normalized versions.
#rnx

diff --git a/Engine/Source/Editor/UnrealEd/Private/CookOnTheFlyServer.cpp b/Engine/Source/Editor/UnrealEd/Private/CookOnTheFlyServer.cpp
--- a/Engine/Source/Editor/UnrealEd/Private/CookOnTheFlyServer.cpp
+++ b/Engine/Source/Editor/UnrealEd/Private/CookOnTheFlyServer.cpp
@@ -10536,26 +10536,32 @@ void UCookOnTheFlyServer::WriteReferencedSet(const ITargetPlatform* InTargetPlat
 	const FString MetadataFilename = ConvertToFullSandboxPath(*MetadataPlatformAgnosticFilename, true)
 		.Replace(TEXT("[Platform]"), *InTargetPlatform->PlatformName());
 
-	CookedPackageNames.Sort([](FName A, FName B)
-		{
-			return A.LexicalLess(B);
-		});
+	CookedPackageNames.Sort(FNameLexicalLess());
 
+	TStringBuilder<256> OplogKeyStr;
+	auto NormalizeOplogKey = [&OplogKeyStr](FName PackageName)
+		{
+			OplogKeyStr.Reset();
+			OplogKeyStr << PackageName;
+			for (TCHAR& C : MakeArrayView(OplogKeyStr))
+			{
+				C = FChar::ToLower(C);
+			}
+			return OplogKeyStr.ToView();
+		};
+	constexpr FStringView VersionStr(TEXTVIEW("# Version 1"));
+	int32 CombinedLength = VersionStr.Len() + CookedPackageNames.Num() * UE_ARRAY_COUNT(LINE_TERMINATOR);
+	for (FName PackageName : CookedPackageNames)
+	{
+		CombinedLength += PackageName.GetStringLength();
+	}
 	FString CombinedString;
-	if (CookedPackageNames.Num() > 0)
+	CombinedString.Reserve(CombinedLength);
+	CombinedString += VersionStr;
+	for (FName PackageName : CookedPackageNames)
 	{
-		int32 CombinedLength = (CookedPackageNames.Num() - 1) * UE_ARRAY_COUNT(LINE_TERMINATOR);
-		for (FName PackageName : CookedPackageNames)
-		{
-			CombinedLength += PackageName.GetStringLength();
-		}
-		CombinedString.Reserve(CombinedLength);
-		CombinedString += WriteToString<256>(CookedPackageNames[0]);
-		for (FName PackageName : TArrayView<FName>(CookedPackageNames).RightChop(1))
-		{
-			CombinedString += LINE_TERMINATOR;
-			CombinedString += WriteToString<256>(PackageName);
-		}
+		CombinedString += LINE_TERMINATOR;
+		CombinedString += NormalizeOplogKey(PackageName);
 	}
 	FFileHelper::SaveStringToFile(CombinedString, *MetadataFilename);
 }
