GatherTextFromAssets IncludePathFilters optimization

When IncludePathFilters contains only StartsWith paths (a single wildcard argument at the end of the path) then we can convert these to PackagePaths filters for the first pass asset registry filter.

Moving this to the first pass filter can be significantly faster when only gathering a subset of the assets that would otherwise be returned by the first pass filter, as the second pass filter performs string wildcard matching so is quite slow.

This change also adds FGatherTextDelegates::GetAdditionalGatherPathsForContext (which will replace FGatherTextDelegates::GetAdditionalGatherPaths), to provide extra context to the callback about what kind of paths the gather implementation prefers, eg) to only return valid content paths that would be compatible with PackagePaths.

diff --git a/Engine/Source/Editor/UnrealEd/Classes/Commandlets/GatherTextCommandletBase.h b/Engine/Source/Editor/UnrealEd/Classes/Commandlets/GatherTextCommandletBase.h
--- a/Engine/Source/Editor/UnrealEd/Classes/Commandlets/GatherTextCommandletBase.h
+++ b/Engine/Source/Editor/UnrealEd/Classes/Commandlets/GatherTextCommandletBase.h
@@ -10,11 +10,30 @@
 #include "LocalizedAssetUtil.h"
 #include "GatherTextCommandletBase.generated.h"
 
+struct FGatherTextContext
+{
+	enum class EPreferredPathType : uint8
+	{
+		/** Generic "root" path, eg) the root folder of a plugin */
+		Root,
+		/** Content path, eg) the content folder of a plugin */
+		Content,
+	};
+
+	FTopLevelAssetPath CommandletClass;
+	EPreferredPathType PreferredPathType = EPreferredPathType::Root;
+};
+
 struct FGatherTextDelegates
 {
 	/** Delegate called during a localization gather to allow code to inject new gather and exclude paths for the given localization target */
-	DECLARE_MULTICAST_DELEGATE_ThreeParams(FGetAdditionalGatherPaths, const FString& /*InLocalizationTargetName*/, TArray<FString>& /*InOutIncludePathFilters*/, TArray<FString>& /*InOutExcludePathFilters*/);
+	// TODO: Deprecate in favor of GetAdditionalGatherPathsForContext
+	using FGetAdditionalGatherPaths = TMulticastDelegate<void(const FString& InLocalizationTargetName, TArray<FString>& InOutIncludePathFilters, TArray<FString>& InOutExcludePathFilters)>;
 	static UNREALED_API FGetAdditionalGatherPaths GetAdditionalGatherPaths;
+
+	/** Delegate called during a localization gather to allow code to inject new gather and exclude paths for the given localization target (with context) */
+	using FGetAdditionalGatherPathsForContext = TMulticastDelegate<void(const FString& InLocalizationTargetName, const FGatherTextContext& InContext, TArray<FString>& InOutIncludePathFilters, TArray<FString>& InOutExcludePathFilters)>;
+	static UNREALED_API FGetAdditionalGatherPathsForContext GetAdditionalGatherPathsForContext;
 };
 
 /** Performs fuzzy path matching against a set of include and exclude paths */
@@ -33,16 +52,9 @@ public:
 
 	UNREALED_API EPathMatch TestPath(const FString& InPathToTest) const;
 
-private:
-	enum class EPathType : uint8
-	{
-		Include,
-		Exclude,
-	};
-
 	/**
-	* * The algorithm used to test path matches for a fuzzy path. Defaults to FString::MatchesWildcard
-	*/
+	 * The algorithm used to test path matches for a fuzzy path. Defaults to FString::MatchesWildcard
+	 */
 	enum class EPathTestPolicy : uint8
 	{
 		/** Performs the path test with FString::MatchesWildcard. This is the default algorithm to use for fuzzy paths that can't be optimized with FString::StartsWith.*/
@@ -51,6 +63,15 @@ private:
 		StartsWith
 	};
 
+	static UNREALED_API EPathTestPolicy CalculatePolicyForPath(const FString& InPath);
+
+private:
+	enum class EPathType : uint8
+	{
+		Include,
+		Exclude,
+	};
+
 	struct FFuzzyPath
 	{
 		FFuzzyPath(FString InPathFilter, const EPathType InPathType);
diff --git a/Engine/Source/Editor/UnrealEd/Classes/Commandlets/GatherTextFromAssetsCommandlet.h b/Engine/Source/Editor/UnrealEd/Classes/Commandlets/GatherTextFromAssetsCommandlet.h
--- a/Engine/Source/Editor/UnrealEd/Classes/Commandlets/GatherTextFromAssetsCommandlet.h
+++ b/Engine/Source/Editor/UnrealEd/Classes/Commandlets/GatherTextFromAssetsCommandlet.h
@@ -57,6 +57,7 @@ private:
 	void ApplyFirstPassFilter(const FARFilter& InFilter, TArray<FAssetData>& InOutAssetDataArray) const;
 	bool BuildFirstPassFilter(FARFilter& InOutFilter) const;
 	bool BuildCollectionFilter(FARFilter& InOutFilter) const;
+	bool BuildPackagePathsFilter(FARFilter& InOutFilter) const;
 	bool BuildExcludeDerivedClassesFilter(FARFilter& InOutFilter) const;
 	bool PerformExcludeExactClassesFilter(TArray<FAssetData>& InOutAssetDataArray) const;
 	bool BuildExcludeExactClassesFilter(FARFilter& InOutFilter) const;
diff --git a/Engine/Source/Editor/UnrealEd/Private/Commandlets/GatherTextCommandletBase.cpp b/Engine/Source/Editor/UnrealEd/Private/Commandlets/GatherTextCommandletBase.cpp
--- a/Engine/Source/Editor/UnrealEd/Private/Commandlets/GatherTextCommandletBase.cpp
+++ b/Engine/Source/Editor/UnrealEd/Private/Commandlets/GatherTextCommandletBase.cpp
@@ -19,6 +19,7 @@
 DEFINE_LOG_CATEGORY_STATIC(LogGatherTextCommandletBase, Log, All);
 
 FGatherTextDelegates::FGetAdditionalGatherPaths FGatherTextDelegates::GetAdditionalGatherPaths;
+FGatherTextDelegates::FGetAdditionalGatherPathsForContext FGatherTextDelegates::GetAdditionalGatherPathsForContext;
 
 //////////////////////////////////////////////////////////////////////////
 //UGatherTextCommandletBase
@@ -34,6 +35,15 @@ UGatherTextCommandletBase::UGatherTextCommandletBase(const FObjectInitializer& O
 	: Super(ObjectInitializer)
 {
 	ShowErrorCount = false;
+
+	if (HasAnyFlags(RF_ClassDefaultObject))
+	{
+		// Wrapper to call the older GetAdditionalGatherPaths callback
+		FGatherTextDelegates::GetAdditionalGatherPathsForContext.AddStatic([](const FString& InLocalizationTargetName, const FGatherTextContext& InContext, TArray<FString>& InOutIncludePathFilters, TArray<FString>& InOutExcludePathFilters)
+		{
+			FGatherTextDelegates::GetAdditionalGatherPaths.Broadcast(InLocalizationTargetName, InOutIncludePathFilters, InOutExcludePathFilters);
+		});
+	}
 }
 
 void UGatherTextCommandletBase::Initialize( const TSharedRef< FLocTextHelper >& InGatherManifestHelper, const TSharedPtr< FLocalizationSCC >& InSourceControlInfo )
@@ -180,18 +190,8 @@ const FString& UGatherTextCommandletBase::GetProjectBasePath()
 FFuzzyPathMatcher::FFuzzyPath::FFuzzyPath(FString InPathFilter, const EPathType InPathType)
 	: PathFilter(MoveTemp(InPathFilter))
 	, PathType(InPathType)
-	, PathTestPolicy(EPathTestPolicy::MatchesWildcard)
+	, PathTestPolicy(FFuzzyPathMatcher::CalculatePolicyForPath(PathFilter))
 {
-	const int32 PathLastIndex = PathFilter.Len() - 1;
-	bool bEndsWithWildcard = PathFilter[PathLastIndex] == TEXT('*');
-	// If we know we end with a * but searching from the front, the index of the first * is not the end of the string, there's more than 1 *
-	int32 FirstAsteriskIndex = INDEX_NONE;
-	PathFilter.FindChar(TEXT('*'), FirstAsteriskIndex);
-	bool bOnlyContainsOneWildcard = FirstAsteriskIndex == PathLastIndex;
-	if (bEndsWithWildcard && bOnlyContainsOneWildcard)
-	{
-		PathTestPolicy = EPathTestPolicy::StartsWith;
-	}
 }
 
 FFuzzyPathMatcher::FFuzzyPathMatcher(const TArray<FString>& InIncludePathFilters, const TArray<FString>& InExcludePathFilters)
@@ -275,3 +275,16 @@ FFuzzyPathMatcher::EPathMatch FFuzzyPathMatcher::TestPath(const FString& InPathT
 
 	return EPathMatch::NoMatch;
 }
+
+FFuzzyPathMatcher::EPathTestPolicy FFuzzyPathMatcher::CalculatePolicyForPath(const FString& InPath)
+{
+	// We can perform a StartsWith policy if there's a single * at the end of the string
+	// FindChar looks from the start of the string, so if we find the match at the last character then we know there's a single *
+	int32 FirstAsteriskIndex = INDEX_NONE;
+	InPath.FindChar(TEXT('*'), FirstAsteriskIndex);
+
+	bool bHasSingleWildcardAtEnd = FirstAsteriskIndex == InPath.Len() - 1;
+	return bHasSingleWildcardAtEnd
+		? EPathTestPolicy::StartsWith
+		: EPathTestPolicy::MatchesWildcard;
+}
diff --git a/Engine/Source/Editor/UnrealEd/Private/Commandlets/GatherTextFromAssetsCommandlet.cpp b/Engine/Source/Editor/UnrealEd/Private/Commandlets/GatherTextFromAssetsCommandlet.cpp
--- a/Engine/Source/Editor/UnrealEd/Private/Commandlets/GatherTextFromAssetsCommandlet.cpp
+++ b/Engine/Source/Editor/UnrealEd/Private/Commandlets/GatherTextFromAssetsCommandlet.cpp
@@ -525,6 +525,12 @@ bool UGatherTextFromAssetsCommandlet::BuildFirstPassFilter(FARFilter& InOutFilte
 		return false;
 	}
 
+	// Filter object paths to those in IncludePathFilters, if possible
+	if (!BuildPackagePathsFilter(InOutFilter))
+	{
+		return false;
+	}
+
 	// Filter out any objects of the specified classes and their children at this point.
 	if (bShouldExcludeDerivedClasses)
 	{
@@ -533,6 +539,7 @@ bool UGatherTextFromAssetsCommandlet::BuildFirstPassFilter(FARFilter& InOutFilte
 			return false;
 		}
 	}
+
 	return true;
 }
 
@@ -554,6 +561,43 @@ bool UGatherTextFromAssetsCommandlet::BuildCollectionFilter(FARFilter& InOutFilt
 	return !bHasFailedToGetACollection;
 }
 
+/** Builds a filter to include assets based on the current IncludePathFilters, if those filters can be represented as an asset registry filter */
+bool UGatherTextFromAssetsCommandlet::BuildPackagePathsFilter(FARFilter& InOutFilter) const
+{
+	TArray<FName> IncludePackagePaths;
+
+	for (const FString& IncludePath : IncludePathFilters)
+	{
+		FString AbsoluteIncludePath = FPaths::ConvertRelativePathToFull(IncludePath);
+		if (FFuzzyPathMatcher::CalculatePolicyForPath(AbsoluteIncludePath) != FFuzzyPathMatcher::EPathTestPolicy::StartsWith)
+		{
+			// Not valid to use as an asset registry filter, but not an error
+			UE_LOG(LogGatherTextFromAssetsCommandlet, Display, TEXT("Skipping first pass PackagePaths filter due to complex IncludePathFilters: %s"), *IncludePath);
+			return true;
+		}
+
+		FNameBuilder IncludePackagePath;
+		AbsoluteIncludePath.LeftChopInline(1);
+		if (!FPackageName::TryConvertFilenameToLongPackageName(AbsoluteIncludePath, IncludePackagePath))
+		{
+			// Check if we're just missing the Content folder (eg, "Plugins/Foo" rather than "Plugins/Foo/Content")
+			AbsoluteIncludePath /= TEXT("Content");
+			if (!FPackageName::TryConvertFilenameToLongPackageName(AbsoluteIncludePath, IncludePackagePath))
+			{
+				// Not valid to use as an asset registry filter, but not an error
+				UE_LOG(LogGatherTextFromAssetsCommandlet, Display, TEXT("Skipping first pass PackagePaths filter due to non-content IncludePathFilters: %s"), *IncludePath);
+				return true;
+			}
+		}
+
+		IncludePackagePaths.Add(FName(IncludePackagePath));
+	}
+
+	InOutFilter.bRecursivePaths = true;
+	InOutFilter.PackagePaths.Append(MoveTemp(IncludePackagePaths));
+	return true;
+}
+
 /** Builds a filter to remove classes and derived classes of ExactClassFilter.*/
 bool UGatherTextFromAssetsCommandlet::BuildExcludeDerivedClassesFilter(FARFilter& InOutFilter) const
 {
@@ -912,7 +956,13 @@ bool UGatherTextFromAssetsCommandlet::ParseCommandLineHelper(const FString& InCo
 		return false;
 	}
 
-	FGatherTextDelegates::GetAdditionalGatherPaths.Broadcast(GatherManifestHelper->GetTargetName(), IncludePathFilters, ExcludePathFilters);
+	{
+		FGatherTextContext Context;
+		Context.CommandletClass = GetClass()->GetClassPathName();
+		Context.PreferredPathType = FGatherTextContext::EPreferredPathType::Content;
+
+		FGatherTextDelegates::GetAdditionalGatherPathsForContext.Broadcast(GatherManifestHelper->GetTargetName(), Context, IncludePathFilters, ExcludePathFilters);
+	}
 
 	// Get destination path
 	if (!GetPathFromConfig(*SectionName, TEXT("DestinationPath"), DestinationPath, GatherTextConfigPath))
diff --git a/Engine/Source/Editor/UnrealEd/Private/Commandlets/GatherTextFromMetadataCommandlet.cpp b/Engine/Source/Editor/UnrealEd/Private/Commandlets/GatherTextFromMetadataCommandlet.cpp
--- a/Engine/Source/Editor/UnrealEd/Private/Commandlets/GatherTextFromMetadataCommandlet.cpp
+++ b/Engine/Source/Editor/UnrealEd/Private/Commandlets/GatherTextFromMetadataCommandlet.cpp
@@ -126,7 +126,13 @@ int32 UGatherTextFromMetaDataCommandlet::Main( const FString& Params )
 		}
 	}
 
-	FGatherTextDelegates::GetAdditionalGatherPaths.Broadcast(GatherManifestHelper->GetTargetName(), IncludePathFilters, ExcludePathFilters);
+	{
+		FGatherTextContext Context;
+		Context.CommandletClass = GetClass()->GetClassPathName();
+		Context.PreferredPathType = FGatherTextContext::EPreferredPathType::Root;
+
+		FGatherTextDelegates::GetAdditionalGatherPathsForContext.Broadcast(GatherManifestHelper->GetTargetName(), Context, IncludePathFilters, ExcludePathFilters);
+	}
 
 	// Get whether we should gather editor-only data. Typically only useful for the localization of UE itself.
 	if (!GetBoolFromConfig(*SectionName, TEXT("ShouldGatherFromEditorOnlyData"), ShouldGatherFromEditorOnlyData, GatherTextConfigPath))
diff --git a/Engine/Source/Editor/UnrealEd/Private/Commandlets/GatherTextFromSourceCommandlet.cpp b/Engine/Source/Editor/UnrealEd/Private/Commandlets/GatherTextFromSourceCommandlet.cpp
--- a/Engine/Source/Editor/UnrealEd/Private/Commandlets/GatherTextFromSourceCommandlet.cpp
+++ b/Engine/Source/Editor/UnrealEd/Private/Commandlets/GatherTextFromSourceCommandlet.cpp
@@ -316,7 +316,11 @@ void UGatherTextFromSourceCommandlet::GetFilesToProcess(const TArray<FString>& S
 
 	if (bAdditionalGatherPaths)
 	{
-		FGatherTextDelegates::GetAdditionalGatherPaths.Broadcast(GatherManifestHelper->GetTargetName(), IncludePathFilters, ExcludePathFilters);
+		FGatherTextContext Context;
+		Context.CommandletClass = GetClass()->GetClassPathName();
+		Context.PreferredPathType = FGatherTextContext::EPreferredPathType::Root;
+
+		FGatherTextDelegates::GetAdditionalGatherPathsForContext.Broadcast(GatherManifestHelper->GetTargetName(), Context, IncludePathFilters, ExcludePathFilters);
 	}
 
 	// Search in the root folder for each of the wildcard filters specified and build a list of files
