Build fixes

#rnx

diff --git a/Engine/Source/Editor/UnrealEd/Private/EditorState/ActorEditorContextEditorState.cpp b/Engine/Source/Editor/UnrealEd/Private/EditorState/ActorEditorContextEditorState.cpp
--- a/Engine/Source/Editor/UnrealEd/Private/EditorState/ActorEditorContextEditorState.cpp
+++ b/Engine/Source/Editor/UnrealEd/Private/EditorState/ActorEditorContextEditorState.cpp
@@ -6,7 +6,7 @@
 #include "Engine/World.h"
 #include "GameFramework/Actor.h"
 #include "Subsystems/ActorEditorContextSubsystem.h"
-#include "UObject\Package.h"
+#include "UObject/Package.h"
 
 #include UE_INLINE_GENERATED_CPP_BY_NAME(ActorEditorContextEditorState)
 
diff --git a/Engine/Source/Editor/WorldBookmark/Private/WorldBookmark/AssetDefinition_WorldBookmark.cpp b/Engine/Source/Editor/WorldBookmark/Private/WorldBookmark/AssetDefinition_WorldBookmark.cpp
--- a/Engine/Source/Editor/WorldBookmark/Private/WorldBookmark/AssetDefinition_WorldBookmark.cpp
+++ b/Engine/Source/Editor/WorldBookmark/Private/WorldBookmark/AssetDefinition_WorldBookmark.cpp
@@ -116,6 +116,4 @@ namespace MenuExtension_Bookmark
 	});
 }
 
-
-
-
+#undef LOCTEXT_NAMESPACE
