Cooker: ReferencedSet: Write the file as utf8 without a ByteObjectMarker, rather than defaulting to unicode when it has non-ascii characters. Utf8 without a BOM is readable by textreaders and is easy to interpret in zenserver.
#rnx

diff --git a/Engine/Source/Editor/UnrealEd/Private/CookOnTheFlyServer.cpp b/Engine/Source/Editor/UnrealEd/Private/CookOnTheFlyServer.cpp
--- a/Engine/Source/Editor/UnrealEd/Private/CookOnTheFlyServer.cpp
+++ b/Engine/Source/Editor/UnrealEd/Private/CookOnTheFlyServer.cpp
@@ -10563,7 +10563,8 @@ void UCookOnTheFlyServer::WriteReferencedSet(const ITargetPlatform* InTargetPlat
 		CombinedString += LINE_TERMINATOR;
 		CombinedString += NormalizeOplogKey(PackageName);
 	}
-	FFileHelper::SaveStringToFile(CombinedString, *MetadataFilename);
+	FFileHelper::SaveStringToFile(CombinedString, *MetadataFilename,
+		FFileHelper::EEncodingOptions::ForceUTF8WithoutBOM);
 }
 
 void UCookOnTheFlyServer::CookByTheBookFinished()
