[fbx scene importer] Fix crash when re-importing scene staticmesh.

#rnx

diff --git a/Engine/Source/Editor/UnrealEd/Private/Fbx/FbxStaticMeshImport.cpp b/Engine/Source/Editor/UnrealEd/Private/Fbx/FbxStaticMeshImport.cpp
--- a/Engine/Source/Editor/UnrealEd/Private/Fbx/FbxStaticMeshImport.cpp
+++ b/Engine/Source/Editor/UnrealEd/Private/Fbx/FbxStaticMeshImport.cpp
@@ -1329,6 +1329,14 @@ UStaticMesh* UnFbx::FFbxImporter::ReimportSceneStaticMesh(uint64 FbxNodeUniqueId
 				// import LOD meshes
 				for (int32 LODIndex = 1; LODIndex < NodeParent->GetChildCount(); LODIndex++)
 				{
+					if (LODIndex >= MAX_STATIC_MESH_LODS)
+					{
+						AddTokenizedErrorMessage(FTokenizedMessage::Create(EMessageSeverity::Warning, FText::Format(
+							LOCTEXT("ImporterLimits_MaximumStaticMeshLODReach_ReimportSceneStaticMesh", "Reached the maximum number of LODs for a Static Mesh({0}) - discarding {1} LOD meshes.")
+								, FText::AsNumber(MAX_STATIC_MESH_LODS), FText::AsNumber(NodeParent->GetChildCount() - MAX_STATIC_MESH_LODS))
+						), FFbxErrors::Generic_Mesh_TooManyLODs);
+						break;
+					}
 					AllNodeInLod.Empty();
 					FindAllLODGroupNode(AllNodeInLod, NodeParent, LODIndex);
 					if (AllNodeInLod.Num() > 0)
