Fix incorrect spelling. Courtesy to Olivier for spotting the issue

#tests trivial

diff --git a/Engine/Source/Editor/UnrealEd/Classes/Factories/TextureFactory.h b/Engine/Source/Editor/UnrealEd/Classes/Factories/TextureFactory.h
--- a/Engine/Source/Editor/UnrealEd/Classes/Factories/TextureFactory.h
+++ b/Engine/Source/Editor/UnrealEd/Classes/Factories/TextureFactory.h
@@ -274,6 +274,6 @@ class UUDIMTextureFunctionLibrary : public UBlueprintFunctionLibrary
 	* @param bCheckOutAndSave		Whether to check out and save the UDIM texture
 	* @return UTexture2D*			Pointer to the UDIM texture or null if failed
 	*/
-	UFUNCTION(BlueprintCallable, Category = "Utilities", meta = (DispalyName = "Make UDIM Texture from Texture2Ds"))
+	UFUNCTION(BlueprintCallable, Category = "Utilities", meta = (DisplayName = "Make UDIM Texture from Texture2Ds"))
 	static UTexture2D* MakeUDIMVirtualTextureFromTexture2Ds(FString OutputPathName, const TArray<UTexture2D*>& SourceTextures, const TArray<FIntPoint>& BlockCoords, bool bKeepExistingSettings = false, bool bCheckOutAndSave = false);
 };
