Removed some implicit dependencies on instanced property flags during reference property value iteration.

Change summary:
- Revised CPFUO to return the inner Duplicate() path to an internal impl for "editinline" subobject instancing.
- Simplified PatchCDOSubobjectsIntoExport() by removing the IsDefaultSubobject() check, which eliminates the need for the extra iteration step that was added a while ago to include any editinline subobject references.
- Modified a few other paths to use GetObjectsWithOuter() and/or FReferenceFinder in place of FFindInstancedReferenceSubobjectHelper.
- Deprecated FFindInstancedReferenceSubobjectHelper/FInstancedPropertyPath/FInstancedSubObjRef to be removed in a future release.

diff --git a/Engine/Source/Editor/Kismet/Private/BlueprintEditorModule.cpp b/Engine/Source/Editor/Kismet/Private/BlueprintEditorModule.cpp
--- a/Engine/Source/Editor/Kismet/Private/BlueprintEditorModule.cpp
+++ b/Engine/Source/Editor/Kismet/Private/BlueprintEditorModule.cpp
@@ -22,7 +22,6 @@
 #include "HAL/PlatformCrt.h"
 #include "IMessageLogListing.h"
 #include "ISettingsModule.h"
-#include "InstancedReferenceSubobjectHelper.h"
 #include "InstancedStaticMeshSCSEditorCustomization.h"
 #include "Internationalization/Internationalization.h"
 #include "Kismet2/BlueprintEditorUtils.h"
@@ -93,12 +92,12 @@ void FixSubObjectReferencesPostUndoRedo(UObject* InObject)
 {
 	// Post undo/redo, these may have the correct Outer but are not referenced by the CDO's UProperties
 	TArray<UObject*> SubObjects;
-	GetObjectsWithOuter(InObject, SubObjects, false);
+	constexpr bool bIncludeNestedSubObjects = false;
+	GetObjectsWithOuter(InObject, SubObjects, bIncludeNestedSubObjects);
 
 	// Post undo/redo, these may have the in-correct Outer but are incorrectly referenced by the CDO's UProperties
-	TSet<FInstancedSubObjRef> PropertySubObjectReferences;
-	UClass* ObjectClass = InObject->GetClass();
-	FFindInstancedReferenceSubobjectHelper::GetInstancedSubObjects(InObject, PropertySubObjectReferences);
+	TArray<UObject*> PropertySubObjectReferences;
+	FReferenceFinder(PropertySubObjectReferences).FindReferences(InObject);
 
 	TMap<UObject*, UObject*> OldToNewInstanceMap;
 	for (UObject* PropertySubObject : PropertySubObjectReferences)
diff --git a/Engine/Source/Editor/PropertyEditor/Private/PropertyNode.cpp b/Engine/Source/Editor/PropertyEditor/Private/PropertyNode.cpp
--- a/Engine/Source/Editor/PropertyEditor/Private/PropertyNode.cpp
+++ b/Engine/Source/Editor/PropertyEditor/Private/PropertyNode.cpp
@@ -3816,6 +3816,7 @@ void FPropertyNode::DuplicateArrayEntry(FProperty* NodeProperty, FScriptArrayHel
 	}
 	else if (NodeProperty->HasAnyPropertyFlags(CPF_ContainsInstancedReference))
 	{
+PRAGMA_DISABLE_DEPRECATION_WARNINGS
 		// If this is a container with instanced references within it the new entry will reference the old subobjects
 		// Go through and duplicate the subobjects so that each container has unique instances
 		FInstancedPropertyPath NodePropertyPath(NodeProperty);
@@ -3828,6 +3829,7 @@ void FPropertyNode::DuplicateArrayEntry(FProperty* NodeProperty, FScriptArrayHel
 				((FObjectProperty*)Ref.PropertyPath.Head())->SetObjectPropertyValue(PropertyValueAddress, DuplicateObject(Obj, Obj->GetOuter()));
 			}
 		);
+PRAGMA_ENABLE_DEPRECATION_WARNINGS
 	}
 }
 
diff --git a/Engine/Source/Editor/UnrealEd/Private/Kismet2/BlueprintEditorUtils.cpp b/Engine/Source/Editor/UnrealEd/Private/Kismet2/BlueprintEditorUtils.cpp
--- a/Engine/Source/Editor/UnrealEd/Private/Kismet2/BlueprintEditorUtils.cpp
+++ b/Engine/Source/Editor/UnrealEd/Private/Kismet2/BlueprintEditorUtils.cpp
@@ -97,7 +97,6 @@
 #include "Kismet2/StructureEditorUtils.h"
 #include "ScopedTransaction.h"
 #include "ClassViewerFilter.h"
-#include "InstancedReferenceSubobjectHelper.h"
 #include "NodeDependingOnEnumInterface.h"
 
 #include "BlueprintEditorModule.h"
@@ -1413,108 +1412,54 @@ void FBlueprintEditorUtils::RecreateClassMetaData(UBlueprint* Blueprint, UClass*
 
 void FBlueprintEditorUtils::PatchCDOSubobjectsIntoExport(UObject* PreviousCDO, UObject* NewCDO)
 {
-	if (PreviousCDO && NewCDO)
+	struct PatchCDOSubobjectsIntoExport_Impl
 	{
-		struct PatchCDOSubobjectsIntoExport_Impl
+		static void PatchSubObjects(UObject* OldObj, UObject* NewObj)
 		{
-			static void PatchSubObjects(UObject* OldObj, UObject* NewObj, TSet<UObject*>& AlreadyPatched)
-			{
-				TArray<UObject*> OldSubObjects;
-				GetObjectsWithOuter(OldObj, OldSubObjects, /*bIncludeNestedSubObjects =*/false);
-				UPackage* Package= OldObj->GetPackage();
-				checkf(Package, TEXT("Expecting a package for this object"));
+			// Find all subobjects outered to the old object.
+			TArray<UObject*> OldSubObjects;
+			constexpr bool bIncludeNestedObjects = false;
+			GetObjectsWithOuter(OldObj, OldSubObjects, bIncludeNestedObjects);
 
-				// Exit now if we don't have any subobjects to process.
-				if (OldSubObjects.Num() == 0)
-				{
-					return;
-				}
+			// Exit now if we don't have any subobjects to process.
+			if (OldSubObjects.Num() == 0)
+			{
+				return;
+			}
 
-				// Used to keep track of subobjects that are patched through an explicitly instanced reference property.
-				TSet<UObject*> PatchedAsInstancedReferenceSet;
-				PatchedAsInstancedReferenceSet.Reserve(OldSubObjects.Num());
+			// Build a lookup table for subobjects outered to the new object.
+			TMap<FName, UObject*> NewSubObjLookupTable;
+			NewSubObjLookupTable.Reserve(OldSubObjects.Num());
+			ForEachObjectWithOuter(NewObj, [&NewSubObjLookupTable](UObject* NewSubObj)
+			{
+				NewSubObjLookupTable.Add(NewSubObj->GetFName(), NewSubObj);
+			},
+				bIncludeNestedObjects);
 
-				// If the old object's class has explicitly instanced reference properties, the subobject values they contain will be different on the new object
-				// that's replacing it (because they've been re-instanced), so here we patch up the linker's export table to reference the re-instanced ones instead.
-				const UClass* OldObjClass = OldObj->GetClass();
-				if (OldObjClass && OldObjClass->HasAnyClassFlags(CLASS_HasInstancedReference))
+			// For each subobject that's outered to the old object, find the corresponding new subobject that matches its name and patch it into the export table.
+			// This includes any editinline subobjects that are instanced at edit time; the outer object (e.g. CDO) may not have a matching archetype in that case.
+			for (UObject* OldSubObj : OldSubObjects)
+			{
+				if (UObject** NewSubObjPtr = NewSubObjLookupTable.Find(OldSubObj->GetFName()))
 				{
-					// Get the list of subobjects assigned to an explicitly instanced reference property (including containers) for the old object.
-					TSet<FInstancedSubObjRef> OldInstancedSubObjRefs;
-					FFindInstancedReferenceSubobjectHelper::GetInstancedSubObjects(OldObj, OldInstancedSubObjRefs);
-
-					// Resolve new instances through the new object and patch them into the linker's export table.
-					for (const FInstancedSubObjRef& OldInstancedSubObjRef : OldInstancedSubObjRefs)
-					{
-						// For non null subobject, make sure they are in the same package, otherwise it will move them and this is not what we want.
-						if (UObject* OldSubObj = OldInstancedSubObjRef.SubObjInstance; OldSubObj && OldSubObj->IsInPackage(Package))
-						{
-							if (UObject* NewSubObj = OldInstancedSubObjRef.PropertyPath.Resolve(NewObj))
-							{
-								FLinkerLoad::PRIVATE_PatchNewObjectIntoExport(OldSubObj, NewSubObj);
-
-								bool bAlreadyPatched;
-								AlreadyPatched.Add(OldSubObj, &bAlreadyPatched);
-								if (!bAlreadyPatched)
-								{
-									// Recursively find and patch any instances nested within the current subobject.
-									PatchSubObjects(OldSubObj, NewSubObj, AlreadyPatched);
-								}
+					UObject* NewSubObj = *NewSubObjPtr;
+					FLinkerLoad::PRIVATE_PatchNewObjectIntoExport(OldSubObj, NewSubObj);
 
-								// Track the old instanced reference so we don't attempt to patch it again below.
-								PatchedAsInstancedReferenceSet.Add(OldSubObj);
-							}
-						}
-					}
+					// Recursively patch all nested subobjects into the export table as well.
+					PatchSubObjects(OldSubObj, NewSubObj);
 				}
-
-				TMap<FName, UObject*> NewSubObjLookupTable;
-				bool bIsNewSubObjLookupTableInitialized = false;
-
-				// Subobject instances not handled above include those assigned to reference properties for which the object type implicitly defaults
-				// to being instanced (e.g. UActorComponent derivatives) and instances not directly assigned to an instanced object reference property.
-				// Here we iterate over those remaining instances, look for a match in the new object, and patch those along with any nested instances.
-				for (UObject* OldSubObj : OldSubObjects)
+				else
 				{
-					if (!PatchedAsInstancedReferenceSet.Contains(OldSubObj))
-					{
-						if (!bIsNewSubObjLookupTableInitialized)
-						{
-							NewSubObjLookupTable.Reserve(OldSubObjects.Num());
-							ForEachObjectWithOuter(NewObj, [&NewSubObjLookupTable](UObject* NewSubObj)
-							{
-								if (NewSubObj != nullptr)
-								{
-									NewSubObjLookupTable.Add(NewSubObj->GetFName(), NewSubObj);
-								}
-							}, /*bIncludeNestedSubObjects =*/false);
-
-							bIsNewSubObjLookupTableInitialized = true;
-						}
-
-						if (UObject** NewSubObjPtr = NewSubObjLookupTable.Find(OldSubObj->GetFName()))
-						{
-							UObject* NewSubObj = *NewSubObjPtr;
-							if (OldSubObj->IsDefaultSubobject() && NewSubObj->IsDefaultSubobject())
-							{
-								FLinkerLoad::PRIVATE_PatchNewObjectIntoExport(OldSubObj, NewSubObj);
-							}
-
-							PatchSubObjects(OldSubObj, NewSubObj, AlreadyPatched);
-						}
-						// Matches check in FPackageHarvester::ProcessImport. Invalidate these objects
-						// to avoid generation of implicit dependencies
-						else if (OldSubObj->HasAnyFlags(RF_DefaultSubObject | RF_ArchetypeObject))
-						{
-							FLinkerLoad::InvalidateExport(OldSubObj);
-						}
-					}
+					// Invalidate any existing entry in the export table to avoid generation of implicit dependencies.
+					FLinkerLoad::InvalidateExport(OldSubObj);
 				}
 			}
-		};
+		}
+	};
 
-		TSet<UObject*> AlreadyPatched; 
-		PatchCDOSubobjectsIntoExport_Impl::PatchSubObjects(PreviousCDO, NewCDO, AlreadyPatched);
+	if (PreviousCDO && NewCDO)
+	{
+		PatchCDOSubobjectsIntoExport_Impl::PatchSubObjects(PreviousCDO, NewCDO);
 		NewCDO->CheckDefaultSubobjects();
 	}
 }
diff --git a/Engine/Source/Editor/UnrealEd/Private/Kismet2/KismetReinstanceUtilities.cpp b/Engine/Source/Editor/UnrealEd/Private/Kismet2/KismetReinstanceUtilities.cpp
--- a/Engine/Source/Editor/UnrealEd/Private/Kismet2/KismetReinstanceUtilities.cpp
+++ b/Engine/Source/Editor/UnrealEd/Private/Kismet2/KismetReinstanceUtilities.cpp
@@ -42,7 +42,6 @@
 #include "Engine/ActorChannel.h"
 #include "Subsystems/AssetEditorSubsystem.h"
 #include "Engine/ScopedMovementUpdate.h"
-#include "InstancedReferenceSubobjectHelper.h"
 #include "Algo/TopologicalSort.h"
 #include "UObject/OverridableManager.h"
 #include "UObject/PropertyOptional.h"
diff --git a/Engine/Source/Runtime/CoreUObject/Private/InstancedReferenceSubobjectHelper.cpp b/Engine/Source/Runtime/CoreUObject/Private/InstancedReferenceSubobjectHelper.cpp
--- a/Engine/Source/Runtime/CoreUObject/Private/InstancedReferenceSubobjectHelper.cpp
+++ b/Engine/Source/Runtime/CoreUObject/Private/InstancedReferenceSubobjectHelper.cpp
@@ -4,6 +4,8 @@
 #include "UObject/PropertyOptional.h"
 #include "UObject/Package.h"
 #include "UObject/Package.h"
+
+PRAGMA_DISABLE_DEPRECATION_WARNINGS
 	
 UObject* FInstancedPropertyPath::Resolve(const UObject* Container) const
 {
@@ -348,3 +350,5 @@ void FFindInstancedReferenceSubobjectHelper::Duplicate(
 		}
 	}
 }
+
+PRAGMA_ENABLE_DEPRECATION_WARNINGS
\ No newline at end of file
diff --git a/Engine/Source/Runtime/CoreUObject/Private/UObject/OverridableManager.cpp b/Engine/Source/Runtime/CoreUObject/Private/UObject/OverridableManager.cpp
--- a/Engine/Source/Runtime/CoreUObject/Private/UObject/OverridableManager.cpp
+++ b/Engine/Source/Runtime/CoreUObject/Private/UObject/OverridableManager.cpp
@@ -1,7 +1,6 @@
 // Copyright Epic Games, Inc. All Rights Reserved.
 
 #include "UObject/OverridableManager.h"
-#include "InstancedReferenceSubobjectHelper.h"
 #include "UObject/InstanceDataObjectUtils.h"
 #include "UObject/UObjectGlobals.h"
 
@@ -128,16 +127,15 @@ EOverriddenState FOverridableManager::GetOverriddenState(UObject& Object)
 		}
 
 		// Need to check subobjects to 
-		TSet<UObject*> InstancedSubObjects;
-		FFindInstancedReferenceSubobjectHelper::GetInstancedSubObjects(&Object, InstancedSubObjects);
+		TArray<UObject*> InstancedSubObjects;
+		constexpr bool bIncludeNestedObjects = false;
+		GetObjectsWithOuter(&Object, InstancedSubObjects, bIncludeNestedObjects);
 		for (UObject* InstancedSubObject : InstancedSubObjects)
 		{
-			if (InstancedSubObject && InstancedSubObject->IsIn(&Object))
+			checkf(InstancedSubObject, TEXT("Expecting non null SubObjects"));
+			if (GetOverriddenState(*InstancedSubObject) != EOverriddenState::NoOverrides)
 			{
-				if (GetOverriddenState(*InstancedSubObject) != EOverriddenState::NoOverrides)
-				{
-					return EOverriddenState::HasOverrides;
-				}
+				return EOverriddenState::HasOverrides;
 			}
 		}
 	}
@@ -169,8 +167,9 @@ void FOverridableManager::OverrideInstancedSubObject(UObject& Object, UObject& I
 void FOverridableManager::PropagateOverrideToInstancedSubObjects(UObject& Object)
 {
 #if WITH_EDITORONLY_DATA
-	TSet<UObject*> InstancedSubObjects;
-	FFindInstancedReferenceSubobjectHelper::GetInstancedSubObjects(&Object, InstancedSubObjects);
+	TArray<UObject*> InstancedSubObjects;
+	constexpr bool bIncludeNestedObjects = false;
+	GetObjectsWithOuter(&Object, InstancedSubObjects, bIncludeNestedObjects);
 	for (UObject* InstancedSubObject : InstancedSubObjects)
 	{
 		checkf(InstancedSubObject, TEXT("Expecting non null SubObjects"));
@@ -264,8 +263,9 @@ void FOverridableManager::ClearInstancedSubObjectOverrides(UObject& Object, UObj
 void FOverridableManager::PropagateClearOverridesToInstancedSubObjects(UObject& Object)
 {
 #if WITH_EDITORONLY_DATA
-	TSet<UObject*> InstancedSubObjects;
-	FFindInstancedReferenceSubobjectHelper::GetInstancedSubObjects(&Object, InstancedSubObjects);
+	TArray<UObject*> InstancedSubObjects;
+	constexpr bool bIncludeNestedObjects = false;
+	GetObjectsWithOuter(&Object, InstancedSubObjects, bIncludeNestedObjects);
 	for (UObject* InstancedSubObject : InstancedSubObjects)
 	{
 		checkf(InstancedSubObject, TEXT("Expecting non null SubObjects"));
diff --git a/Engine/Source/Runtime/CoreUObject/Private/UObject/PropertyBagRepository.cpp b/Engine/Source/Runtime/CoreUObject/Private/UObject/PropertyBagRepository.cpp
--- a/Engine/Source/Runtime/CoreUObject/Private/UObject/PropertyBagRepository.cpp
+++ b/Engine/Source/Runtime/CoreUObject/Private/UObject/PropertyBagRepository.cpp
@@ -5,7 +5,6 @@
 #if WITH_EDITORONLY_DATA
 
 #include "Containers/Queue.h"
-#include "InstancedReferenceSubobjectHelper.h"
 #include "Misc/StringBuilder.h"
 #include "Serialization/ObjectReader.h"
 #include "Serialization/ObjectWriter.h"
@@ -809,11 +808,12 @@ void FPropertyBagRepository::PostLoadInstanceDataObject(const UObject* Owner)
 			// the owner's PostLoad() may have mutated its instanced subobjects as well (e.g. pointer fixup). to handle
 			// that case, we look for any instanced subobjects that have already had their PostLoad() called, as those will
 			// not have a chance to get their IDO data fixed up to match changes potentially made by its owner's PostLoad().
-			TSet<UObject*> InstancedSubObjects;
-			FFindInstancedReferenceSubobjectHelper::GetInstancedSubObjects(Owner, InstancedSubObjects);
+			TArray<UObject*> InstancedSubObjects;
+			constexpr bool bIncludeNestedObjects = false;
+			GetObjectsWithOuter(Owner, InstancedSubObjects, bIncludeNestedObjects);
 			for (UObject* InstancedSubObject : InstancedSubObjects)
 			{
-				if (!InstancedSubObject->HasAnyFlags(RF_NeedPostLoad) && InstancedSubObject->IsInOuter(Owner))
+				if (!InstancedSubObject->HasAnyFlags(RF_NeedPostLoad))
 				{
 					PostLoadInstanceDataObject(InstancedSubObject);
 				}
diff --git a/Engine/Source/Runtime/CoreUObject/Private/UObject/SavePackage/PackageHarvester.cpp b/Engine/Source/Runtime/CoreUObject/Private/UObject/SavePackage/PackageHarvester.cpp
--- a/Engine/Source/Runtime/CoreUObject/Private/UObject/SavePackage/PackageHarvester.cpp
+++ b/Engine/Source/Runtime/CoreUObject/Private/UObject/SavePackage/PackageHarvester.cpp
@@ -3,7 +3,6 @@
 
 #include "UObject/SavePackage/PackageHarvester.h"
 
-#include "InstancedReferenceSubobjectHelper.h"
 #include "Interfaces/ITargetPlatform.h"
 #include "UObject/OverridableManager.h"
 #include "UObject/SavePackage/SaveContext.h"
@@ -651,8 +650,9 @@ void FPackageHarvester::TryHarvestExportInternal(UObject* InObject)
 	// This is because it needs to make a difference between the pointer that is overridden vs some properties inside the instanced sub object that are overridden.
 	if (FOverridableManager::Get().IsEnabled(*InObject))
 	{
-		TSet<UObject*> InstancedSubObjects;
-		FFindInstancedReferenceSubobjectHelper::GetInstancedSubObjects(InObject, InstancedSubObjects);
+		TArray<UObject*> InstancedSubObjects;
+		constexpr bool bIncludeNestedObjects = false;
+		GetObjectsWithOuter(InObject, InstancedSubObjects, bIncludeNestedObjects);
 		for (UObject* InstancedSubObject : InstancedSubObjects)
 		{
 			checkf(InstancedSubObject, TEXT("Expecting valid subobject"));
diff --git a/Engine/Source/Runtime/CoreUObject/Public/InstancedReferenceSubobjectHelper.h b/Engine/Source/Runtime/CoreUObject/Public/InstancedReferenceSubobjectHelper.h
--- a/Engine/Source/Runtime/CoreUObject/Public/InstancedReferenceSubobjectHelper.h
+++ b/Engine/Source/Runtime/CoreUObject/Public/InstancedReferenceSubobjectHelper.h
@@ -12,7 +12,7 @@
  * property, so that we can easily retrieve instanced sub-objects from a 
  * container object.
  */
-struct FInstancedPropertyPath
+struct UE_DEPRECATED(5.6, "Deprecated; this type is no longer in use.") FInstancedPropertyPath
 {
 private:
 	struct FPropertyLink
@@ -58,12 +58,14 @@ private:
 	TArray<FPropertyLink> PropertyChain;
 };
 
+PRAGMA_DISABLE_DEPRECATION_WARNINGS
+
 /** 
  * Can be used as a raw sub-object pointer, but also contains a 
  * FInstancedPropertyPath to identify the property that this sub-object is 
  * referenced by. Paired together for ease of use (so API users don't have to manage a map).
  */
-struct FInstancedSubObjRef
+struct UE_DEPRECATED(5.6, "Deprecated; this type is no longer in use.") FInstancedSubObjRef
 {
 	FInstancedSubObjRef(UObject* SubObj, const FInstancedPropertyPath& PropertyPathIn)
 		: SubObjInstance(SubObj)
@@ -94,12 +96,13 @@ struct FInstancedSubObjRef
  
 /** 
  * Contains a set of utility functions useful for searching out and identifying
- * instanced sub-objects contained within a specific outer object.
+ * instanced sub-objects contained within a specific outer object. (DEPRECATED)
  */
 class FFindInstancedReferenceSubobjectHelper
 {
 public:
 	template<typename T>
+	UE_DEPRECATED(5.6, "Deprecated; use GetObjectsWithOuter() instead, or use TPropertyValueIterator<FObjectProperty> and check for the CPF_InstancedReference flag if you specifically need this functionality.")
 	static void GetInstancedSubObjects(const UObject* Container, T& OutObjects)
 	{
 		const UClass* ContainerClass = Container->GetClass();
@@ -115,6 +118,7 @@ public:
 	}
 
 	template<typename T>
+	UE_DEPRECATED(5.6, "Deprecated; use GetObjectsWithOuter() instead, or recursively use a TPropertyValueIterator<FObjectProperty> and check for the CPF_InstancedReference flag if you specifically need this functionality.")
 	static void GetInstancedSubObjectsRecursive(const UObject* Container, T& OutObjects)
 	{
 		const UClass* ContainerClass = Container->GetClass();
@@ -136,6 +140,7 @@ public:
 		}
 	}
 
+	UE_DEPRECATED(5.6, "Deprecated; this method was originally implemented for internal use only.")
 	static COREUOBJECT_API void Duplicate(
 		UObject* OldObject, 
 		UObject* NewObject, 
@@ -144,5 +149,8 @@ public:
 		TMap<UObject*, UObject*>* OptionalMappings);
 
 	template<typename T>
+	UE_DEPRECATED(5.6, "Deprecated; use a container helper (e.g. FScriptArrayHelper) instead if you need to iterate over the values of a container property given its value address.")
 	static void ForEachInstancedSubObject(FInstancedPropertyPath& PropertyPath, T ContainerAddress, TFunctionRef<void(const FInstancedSubObjRef& Ref, T PropertyValueAddress)> ObjRefFunc);
 };
+
+PRAGMA_ENABLE_DEPRECATION_WARNINGS
\ No newline at end of file
diff --git a/Engine/Source/Runtime/Engine/Private/UnrealEngine.cpp b/Engine/Source/Runtime/Engine/Private/UnrealEngine.cpp
--- a/Engine/Source/Runtime/Engine/Private/UnrealEngine.cpp
+++ b/Engine/Source/Runtime/Engine/Private/UnrealEngine.cpp
@@ -201,7 +201,6 @@ UnrealEngine.cpp: Implements the UEngine class and helpers.
 #include "ProfilingDebugging/ABTesting.h"
 #include "Performance/EnginePerformanceTargets.h"
 
-#include "InstancedReferenceSubobjectHelper.h"
 #include "Misc/EngineBuildSettings.h"
 #include "Engine/LODActor.h"
 #include "Engine/AssetManager.h"
@@ -17618,6 +17617,124 @@ public:
 	// ~End FArchive Interface
 };
 
+namespace UE::Engine::Private
+{
+	// Internal helper to duplicate or fix up any references to subobjects still owned by the old object that have not
+	// yet been duplicated under the new object. These are typically editinline types that get instanced at edit time.
+	static void DuplicateEditInlineSubObjects(
+		UObject* OldObject,
+		UObject* NewObject,
+		TMap<UObject*, UObject*>& ReferenceReplacementMap,
+		TArray<UObject*>& DuplicatedObjects,
+		TMap<UObject*, UObject*>* OptionalMappings)
+	{
+		if (!OldObject->GetClass()->HasAnyClassFlags(CLASS_HasInstancedReference) ||
+			!NewObject->GetClass()->HasAnyClassFlags(CLASS_HasInstancedReference))
+		{
+			return;
+		}
+
+		// Gather all objects referenced by the old object
+		TArray<UObject*> OldObjectReferences;
+		FReferenceFinder(OldObjectReferences).FindReferences(OldObject);
+		if(OldObjectReferences.IsEmpty())
+		{
+			return;
+		}
+
+		// Gather all objects referenced by the new object
+		TArray<UObject*> NewObjectReferences;
+		NewObjectReferences.Reserve(OldObjectReferences.Num());
+		FReferenceFinder(NewObjectReferences).FindReferences(NewObject);
+		for (UObject* Obj : NewObjectReferences)
+		{
+			// We only care about new references to subobjects still outered to the old object.
+			const bool bNewObjectHasOldOuter = (Obj->GetOuter() == OldObject);
+			if (!bNewObjectHasOldOuter)
+			{
+				continue;
+			}
+
+			// If the old object doesn't reference it anymore then we don't want to keep it either.
+			// In this case the reference on the new object will be cleared during ref replacement.
+			const bool bKeptByOld = OldObjectReferences.Contains(Obj);
+			if (!bKeptByOld)
+			{
+				continue;
+			}
+
+			// Duplicate the instance if it hasn't already been handled. Otherwise, ensure that
+			// we fix up the outer reference on the copy to correctly point back to the new object.
+			const bool bNotHandledYet = !ReferenceReplacementMap.Contains(Obj) &&
+				(!OptionalMappings || !OptionalMappings->Contains(Obj));
+			if (bNotHandledYet)
+			{
+				// This name may have been taken by an instanced subobject, since we're copying from
+				// an object we want to recreate it here with a duplicate of the old object. I'm 
+				// using a rename call here to reserve the name for ourself. These objects come from
+				// FObjectInitializer::InstanceSubobjects()
+				UObject* ExistingObject = StaticFindObjectFast(UObject::StaticClass(), NewObject, Obj->GetFName());
+				if (ExistingObject)
+				{
+					ExistingObject->Rename(
+						nullptr,
+						GetTransientPackage(),
+						REN_DoNotDirty | REN_DontCreateRedirectors | REN_NonTransactional);
+				}
+
+				UObject* NewEditInlineSubobject = StaticDuplicateObject(Obj, NewObject, Obj->GetFName());
+				ReferenceReplacementMap.Add(Obj, NewEditInlineSubobject);
+				if (OptionalMappings)
+				{
+					OptionalMappings->Add(Obj, NewEditInlineSubobject);
+				}
+
+				// NOTE: we cannot patch OldObject's linker table here, since we don't 
+				//       know the relation between the two objects (one could be of a 
+				//       super class, and the other a child)
+
+				// We also need to make sure to fixup any properties here
+				DuplicatedObjects.Add(NewEditInlineSubobject);
+			}
+			else
+			{
+				// make sure the object is outered correctly:
+				check(bNewObjectHasOldOuter);
+				UObject* RealNewSubobject = nullptr;
+				if (OptionalMappings)
+				{
+					UObject* const* Entry = OptionalMappings->Find(Obj);
+					if (Entry)
+					{
+						RealNewSubobject = *Entry;
+					}
+				}
+
+				if (!RealNewSubobject)
+				{
+					RealNewSubobject = ReferenceReplacementMap.FindChecked(Obj);
+				}
+				if (RealNewSubobject->GetOuter() != NewObject)
+				{
+					UObject* ExistingObject = StaticFindObjectFast(UObject::StaticClass(), NewObject, RealNewSubobject->GetFName());
+					if (ExistingObject)
+					{
+						ExistingObject->Rename(
+							nullptr,
+							GetTransientPackage(),
+							REN_DoNotDirty | REN_DontCreateRedirectors | REN_NonTransactional);
+						ReferenceReplacementMap.Add(ExistingObject, RealNewSubobject);
+					}
+					RealNewSubobject->Rename(
+						*(RealNewSubobject->GetName()),
+						NewObject,
+						REN_DoNotDirty | REN_DontCreateRedirectors | REN_NonTransactional);
+				}
+			}
+		}
+	}
+}
+
 PRAGMA_DISABLE_DEPRECATION_WARNINGS
 UEngine::FCopyPropertiesForUnrelatedObjectsParams::FCopyPropertiesForUnrelatedObjectsParams()
 	: bAggressiveDefaultSubobjectReplacement(false)
@@ -17915,16 +18032,18 @@ void UEngine::CopyPropertiesForUnrelatedObjects(UObject* OldObject, UObject* New
 		{
 			if (int32* pOldInstanceIndex = OldInstanceMap.Find(NewInstance->GetPathName(NewObject)))
 			{
-				// Restore modified properties into the new instance
+				// Restore modified properties into the new instance and duplicate any editinline references that were copied from the old subobject instance
 				FInstancedObjectRecord& Record = SavedInstances[*pOldInstanceIndex];
 				FObjectReaderWithReplacement Reader(NewInstance, Record.SavedProperties, true, true, OptionalReferenceReplacementMap, OptionalClassReferenceReplacementMap);
-				FFindInstancedReferenceSubobjectHelper::Duplicate(Record.OldInstance, NewInstance, ReferenceReplacementMap, EditInlineSubobjectsOfComponents, Params.OptionalReplacementMappings);
+				UE::Engine::Private::DuplicateEditInlineSubObjects(Record.OldInstance, NewInstance, ReferenceReplacementMap, EditInlineSubobjectsOfComponents, Params.OptionalReplacementMappings);
 			}
 		}
 		ComponentsOnNewObject.Append(MoveTemp(EditInlineSubobjectsOfComponents));
 	}
 
-	FFindInstancedReferenceSubobjectHelper::Duplicate(OldObject, NewObject, ReferenceReplacementMap, ComponentsOnNewObject, Params.OptionalReplacementMappings);
+	// Duplicate any editinline subobjects that were not otherwise instanced by the new object at construction time. That is, if a reference to an object
+	// that's outered to the old object was copied into the new object, it may still need to be reinstanced (duplicated) with the new object as its owner
+	UE::Engine::Private::DuplicateEditInlineSubObjects(OldObject, NewObject, ReferenceReplacementMap, ComponentsOnNewObject, Params.OptionalReplacementMappings);
 
 	// Replace anything with an outer of the old object with NULL, unless it already has a replacement or is marked pending kill
 	// Also, if requested, leave pointers to instances of renewed classes intact, assuming we are in the midst of reinstancing, and those will be reinstanced
