Fixing errors when compiling the editor with UE_WITH_OBJECT_HANDLE_TRACKING=0

diff --git a/Engine/Source/Editor/UnrealEd/Classes/CookOnTheSide/CookOnTheFlyServer.h b/Engine/Source/Editor/UnrealEd/Classes/CookOnTheSide/CookOnTheFlyServer.h
--- a/Engine/Source/Editor/UnrealEd/Classes/CookOnTheSide/CookOnTheFlyServer.h
+++ b/Engine/Source/Editor/UnrealEd/Classes/CookOnTheSide/CookOnTheFlyServer.h
@@ -1489,7 +1489,9 @@ private:
 	TUniquePtr<UE::Cook::FCachedDependencies> CachedDependencies;
 
 	/** Registration handle for TObjectPtr's OnRead delegate. */
+#if UE_WITH_OBJECT_HANDLE_TRACKING
 	UE::CoreUObject::FObjectHandleTrackingCallbackId ObjectHandleReadHandle;
+#endif // UE_WITH_OBJECT_HANDLE_TRACKING
 	/** Data tracking the package currently having calls made (Load/Save/Other) from COTFS. Used for diagnostics. */
 	FActivePackageData ActivePackageData;
 
diff --git a/Engine/Source/Editor/UnrealEd/Private/CookOnTheFlyServer.cpp b/Engine/Source/Editor/UnrealEd/Private/CookOnTheFlyServer.cpp
--- a/Engine/Source/Editor/UnrealEd/Private/CookOnTheFlyServer.cpp
+++ b/Engine/Source/Editor/UnrealEd/Private/CookOnTheFlyServer.cpp
@@ -3665,7 +3665,12 @@ UE::Cook::EPollStatus UCookOnTheFlyServer::PreSaveGeneratorPackage(UE::Cook::FPa
 
 	{
 		FScopedActivePackage ScopedActivePackage(*this, GenerationHelper.GetOwner().GetPackageName(),
-			PackageAccessTrackingOps::NAME_CookerBuildObject);
+#if UE_WITH_OBJECT_HANDLE_TRACKING
+			PackageAccessTrackingOps::NAME_CookerBuildObject
+#else
+			FName()
+#endif
+		);
 		if (!Splitter->PreSaveGeneratorPackage(Package, SplitDataObject, GeneratedPackagesForPresave, KeepReferencedPackages))
 		{
 			UE_LOG(LogCook, Error, TEXT("PackageSplitter returned false from PreSaveGeneratorPackage. Splitter=%s"),
@@ -3816,7 +3821,12 @@ UE::Cook::EPollStatus UCookOnTheFlyServer::TryPopulateGeneratedPackage(UE::Cook:
 	TArray<UPackage*> KeepReferencedPackages;
 	{
 		FScopedActivePackage ScopedActivePackage(*this, GenerationHelper.GetOwner().GetPackageName(),
-			PackageAccessTrackingOps::NAME_CookerBuildObject);
+#if UE_WITH_OBJECT_HANDLE_TRACKING
+			PackageAccessTrackingOps::NAME_CookerBuildObject
+#else
+			FName()
+#endif
+		);
 		if (!Splitter->PreSaveGeneratedPackage(OwnerPackage, OwnerObject, PopulateData, KeepReferencedPackages))
 		{
 			UE_LOG(LogCook, Error, TEXT("PackageSplitter returned false from PreSaveGeneratedPackage. Splitter=%s, Generated=%s."),
@@ -12851,11 +12861,13 @@ void UCookOnTheFlyServer::RegisterCookByTheBookDelegates()
 	{
 		FCoreUObjectDelegates::PackageCreatedForLoad.AddUObject(this, &UCookOnTheFlyServer::MaybeMarkPackageAsAlreadyLoaded);
 	}
+#if UE_WITH_OBJECT_HANDLE_TRACKING
 	if (bHiddenDependenciesDebug)
 	{
 		ObjectHandleReadHandle = UE::CoreUObject::AddObjectHandleReadCallback(
 			[this](const TArrayView<const UObject*const>& ReadObjects) { UCookOnTheFlyServer::OnObjectHandleReadDebug(ReadObjects); });
 	}
+#endif // UE_WITH_OBJECT_HANDLE_TRACKING
 }
 
 void UCookOnTheFlyServer::UnregisterCookByTheBookDelegates()
@@ -12864,11 +12876,13 @@ void UCookOnTheFlyServer::UnregisterCookByTheBookDelegates()
 	{
 		FCoreUObjectDelegates::PackageCreatedForLoad.RemoveAll(this);
 	}
+#if UE_WITH_OBJECT_HANDLE_TRACKING
 	if (ObjectHandleReadHandle.IsValid())
 	{
 		UE::CoreUObject::RemoveObjectHandleReadCallback(ObjectHandleReadHandle);
 		ObjectHandleReadHandle = UE::CoreUObject::FObjectHandleTrackingCallbackId();
 	}
+#endif // UE_WITH_OBJECT_HANDLE_TRACKING
 }
 
 TArray<FName> UCookOnTheFlyServer::GetNeverCookPackageNames(TArrayView<const FString> ExtraNeverCookDirectories) const
@@ -13460,7 +13474,13 @@ void UCookOnTheFlyServer::RouteBeginCacheForCookedPlatformData(UE::Cook::FPackag
 	TOptional<FScopedActivePackage> ScopedActivePackage;
 	if (!ActivePackageData.bActive)
 	{
-		ScopedActivePackage.Emplace(*this, PackageName, PackageAccessTrackingOps::NAME_CookerBuildObject);
+		ScopedActivePackage.Emplace(*this, PackageName,
+#if UE_WITH_OBJECT_HANDLE_TRACKING
+			PackageAccessTrackingOps::NAME_CookerBuildObject
+#else
+			FName()
+#endif
+		);
 	}
 	Obj->BeginCacheForCookedPlatformData(TargetPlatform);
 	*ExistingEvent = ECachedCookedPlatformDataEvent::BeginCacheForCookedPlatformDataCalled;
@@ -13488,7 +13508,13 @@ bool UCookOnTheFlyServer::RouteIsCachedCookedPlatformDataLoaded(UE::Cook::FPacka
 	TOptional<FScopedActivePackage> ScopedActivePackage;
 	if (!ActivePackageData.bActive)
 	{
-		ScopedActivePackage.Emplace(*this, PackageName, PackageAccessTrackingOps::NAME_CookerBuildObject);
+		ScopedActivePackage.Emplace(*this, PackageName,
+#if UE_WITH_OBJECT_HANDLE_TRACKING
+			PackageAccessTrackingOps::NAME_CookerBuildObject
+#else
+			FName()
+#endif
+			);
 	}
 
 	bool bResult = Obj->IsCachedCookedPlatformDataLoaded(TargetPlatform);
diff --git a/Engine/Source/Editor/UnrealEd/Private/Cooker/CookGenerationHelper.cpp b/Engine/Source/Editor/UnrealEd/Private/Cooker/CookGenerationHelper.cpp
--- a/Engine/Source/Editor/UnrealEd/Private/Cooker/CookGenerationHelper.cpp
+++ b/Engine/Source/Editor/UnrealEd/Private/Cooker/CookGenerationHelper.cpp
@@ -749,7 +749,12 @@ bool FGenerationHelper::TryGenerateList()
 	TArray<ICookPackageSplitter::FGeneratedPackage> GeneratorDatas;
 	{
 		UCookOnTheFlyServer::FScopedActivePackage ScopedActivePackage(COTFS, OwnerPackageName,
-			PackageAccessTrackingOps::NAME_CookerBuildObject);
+#if UE_WITH_OBJECT_HANDLE_TRACKING
+			PackageAccessTrackingOps::NAME_CookerBuildObject
+#else
+			FName()
+#endif
+			);
 		GeneratorDatas = GetCookPackageSplitterInstance()->GetGenerateList(LocalOwnerPackage, OwnerObject);
 	}
 
@@ -910,8 +915,12 @@ bool FGenerationHelper::TryCallPopulateGeneratorPackage(
 		return false;
 	}
 	UCookOnTheFlyServer::FScopedActivePackage ScopedActivePackage(COTFS, OwnerPackageName,
-		PackageAccessTrackingOps::NAME_CookerBuildObject);
-
+#if UE_WITH_OBJECT_HANDLE_TRACKING
+		PackageAccessTrackingOps::NAME_CookerBuildObject
+#else
+		FName()
+#endif
+	);
 	TArray<UPackage*> KeepReferencedPackages;
 	TArray<UObject*> ObjectsToMove;
 	bool bPopulateSucceeded = CookPackageSplitterInstance->PopulateGeneratorPackage(LocalOwnerPackage,
@@ -974,8 +983,12 @@ bool FGenerationHelper::TryCallPopulateGeneratedPackage(UE::Cook::FCookGeneratio
 		Info.GeneratedRootPath, Package, Info.IsCreateAsMap() };
 
 	UCookOnTheFlyServer::FScopedActivePackage ScopedActivePackage(COTFS, OwnerPackageName,
-		PackageAccessTrackingOps::NAME_CookerBuildObject);
-
+#if UE_WITH_OBJECT_HANDLE_TRACKING
+		PackageAccessTrackingOps::NAME_CookerBuildObject
+#else
+		FName()
+#endif
+	);
 	TArray<UPackage*> KeepReferencedPackages;
 	bool bPopulateSucceeded = CookPackageSplitterInstance->PopulateGeneratedPackage(Package,
 		LocalSplitDataObject, SplitterInfo, OutObjectsToMove, KeepReferencedPackages);
@@ -1343,7 +1356,12 @@ void FGenerationHelper::ResetSaveState(FCookGenerationInfo& Info, UPackage* Pack
 		{
 			UCookOnTheFlyServer& COTFS = GetOwner().GetPackageDatas().GetCookOnTheFlyServer();
 			UCookOnTheFlyServer::FScopedActivePackage ScopedActivePackage(COTFS, GetOwner().GetPackageName(),
-				PackageAccessTrackingOps::NAME_CookerBuildObject);
+#if UE_WITH_OBJECT_HANDLE_TRACKING
+				PackageAccessTrackingOps::NAME_CookerBuildObject
+#else
+				FName()
+#endif
+			);
 			if (Info.IsGenerator())
 			{
 				GetCookPackageSplitterInstance()->PostSaveGeneratorPackage(Package, SplitObject);
diff --git a/Engine/Source/Editor/UnrealEd/Private/Cooker/CookPackageData.cpp b/Engine/Source/Editor/UnrealEd/Private/Cooker/CookPackageData.cpp
--- a/Engine/Source/Editor/UnrealEd/Private/Cooker/CookPackageData.cpp
+++ b/Engine/Source/Editor/UnrealEd/Private/Cooker/CookPackageData.cpp
@@ -2650,7 +2650,12 @@ void FPackageDatas::PollPendingCookedPlatformDatas(bool bForce, double& LastCook
 			if (bEntering)
 			{
 				CookOnTheFlyServer.SetActivePackage(Package->GetFName(),
-					PackageAccessTrackingOps::NAME_CookerBuildObject);
+#if UE_WITH_OBJECT_HANDLE_TRACKING
+					PackageAccessTrackingOps::NAME_CookerBuildObject
+#else
+					FName()
+#endif
+				);
 			}
 			else
 			{
diff --git a/Engine/Source/Editor/UnrealEd/Private/Cooker/PackageBuildDependencyTracker.cpp b/Engine/Source/Editor/UnrealEd/Private/Cooker/PackageBuildDependencyTracker.cpp
--- a/Engine/Source/Editor/UnrealEd/Private/Cooker/PackageBuildDependencyTracker.cpp
+++ b/Engine/Source/Editor/UnrealEd/Private/Cooker/PackageBuildDependencyTracker.cpp
@@ -1,6 +1,7 @@
 // Copyright Epic Games, Inc. All Rights Reserved.
 
 #include "PackageBuildDependencyTracker.h"
+#include "Misc/ConfigAccessTracking.h"
 
 #if UE_WITH_PACKAGE_ACCESS_TRACKING
 #include "Async/UniqueLock.h"
@@ -166,7 +167,7 @@ void DumpBuildDependencyTrackerStats()
 #if UE_WITH_PACKAGE_ACCESS_TRACKING
 	FPackageBuildDependencyTracker::Get().DumpStats();
 #endif
-#if UE_WITH_CONFIG_TRACKING
+#if UE_WITH_CONFIG_TRACKING && UE_WITH_OBJECT_HANDLE_TRACKING
 	UE::ConfigAccessTracking::FCookConfigAccessTracker::Get().DumpStats();
 #endif
 }
