[Curve Editor] Fixed Crashes in array remove/clear and undo/redo, when the curve data is inline alloacted in a container, or wrapped by InstancedStruct

diff --git a/Engine/Source/Editor/DetailCustomizations/Private/CurveColorCustomization.cpp b/Engine/Source/Editor/DetailCustomizations/Private/CurveColorCustomization.cpp
--- a/Engine/Source/Editor/DetailCustomizations/Private/CurveColorCustomization.cpp
+++ b/Engine/Source/Editor/DetailCustomizations/Private/CurveColorCustomization.cpp
@@ -150,6 +150,7 @@ void FCurveColorCustomization::CustomizeHeader(TSharedRef<IPropertyHandle> InStr
 		else
 		{
 			CurveWidget->SetCurveOwner(this);
+			CurveWidget->SetPropertyUtils(StructCustomizationUtils.GetPropertyUtilities());
 		}
 	}
 	else
diff --git a/Engine/Source/Editor/DetailCustomizations/Private/CurveStructCustomization.cpp b/Engine/Source/Editor/DetailCustomizations/Private/CurveStructCustomization.cpp
--- a/Engine/Source/Editor/DetailCustomizations/Private/CurveStructCustomization.cpp
+++ b/Engine/Source/Editor/DetailCustomizations/Private/CurveStructCustomization.cpp
@@ -152,6 +152,7 @@ void FCurveStructCustomization::CustomizeHeader( TSharedRef<IPropertyHandle> InS
 		else
 		{
 			CurveWidget->SetCurveOwner(this, InStructPropertyHandle->IsEditable());
+			CurveWidget->SetPropertyUtils(StructCustomizationUtils.GetPropertyUtilities());
 		}
 	}
 	else
diff --git a/Engine/Source/Editor/DetailCustomizations/Private/CurveVectorCustomization.cpp b/Engine/Source/Editor/DetailCustomizations/Private/CurveVectorCustomization.cpp
--- a/Engine/Source/Editor/DetailCustomizations/Private/CurveVectorCustomization.cpp
+++ b/Engine/Source/Editor/DetailCustomizations/Private/CurveVectorCustomization.cpp
@@ -150,6 +150,7 @@ void FCurveVectorCustomization::CustomizeHeader(TSharedRef<IPropertyHandle> InSt
 		else if (RuntimeCurve)
 		{
 			CurveWidget->SetCurveOwner(this, InStructPropertyHandle->IsEditable());
+			CurveWidget->SetPropertyUtils(StructCustomizationUtils.GetPropertyUtilities());
 		}
 	}
 	else
diff --git a/Engine/Source/Editor/UnrealEd/Private/SCurveEditor.cpp b/Engine/Source/Editor/UnrealEd/Private/SCurveEditor.cpp
--- a/Engine/Source/Editor/UnrealEd/Private/SCurveEditor.cpp
+++ b/Engine/Source/Editor/UnrealEd/Private/SCurveEditor.cpp
@@ -36,6 +36,7 @@
 #include "Widgets/Input/STextEntryPopup.h"
 #include "Widgets/Notifications/SNotificationList.h"
 #include "Framework/Notifications/NotificationManager.h"
+#include "IPropertyUtilities.h"
 
 #define LOCTEXT_NAMESPACE "SCurveEditor"
  
@@ -84,6 +85,9 @@ void SCurveEditor::Construct(const FArguments& InArgs)
 
 	bIsUsingSlider = false;
 	bAllowAutoFrame = true;
+	bRequireFocusToZoom = false;
+
+	bIsPendingRebuilt = false;
 
 	// if editor size is set, use it, otherwise, use default value
 	if (DesiredSize.Get().IsZero())
@@ -467,8 +471,6 @@ void SCurveEditor::Construct(const FArguments& InArgs)
 	}
 
 	FCoreUObjectDelegates::OnObjectPropertyChanged.AddSP(this, &SCurveEditor::OnObjectPropertyChanged);
-
-	bRequireFocusToZoom = false;
 }
 
 FText SCurveEditor::GetIsCurveVisibleToolTip(TSharedPtr<FCurveViewModel> CurveViewModel) const
@@ -3013,6 +3015,11 @@ void SCurveEditor::SetOutputMinMax(float NewMin, float NewMax)
 	}
 }
 
+void SCurveEditor::EmptyAllCurveViewModels()
+{
+	CurveViewModels.Empty();
+}
+
 void SCurveEditor::ClearSelectedCurveViewModels()
 {
 	for(auto CurveViewModel : CurveViewModels)
@@ -3807,16 +3814,37 @@ void SCurveEditor::RedoAction()
 
 void SCurveEditor::OnObjectPropertyChanged(UObject* Object, FPropertyChangedEvent& PropertyChangedEvent)
 {
-	if ( CurveOwner && CurveOwner->GetOwners().Contains(Object) )
+	if (CurveOwner && CurveOwner->GetOwners().Contains(Object))
 	{
-		if (PropertyChangedEvent.ChangeType == EPropertyChangeType::ArrayClear)
+		// CurveEditor will hold dangling references to CurveData and access them in some cases:
+		// - Curve Data is inline allocated, or wrapped by InstancedStruct, in a container. And we do Array Remove/Clear op, or Add that causes resize
+		// - Curve Data is wrapped by InstancedStruct in a Container, and we do undo/redo(will trigger emptying and refilling the container, causing address changed)
+		if (!bIsPendingRebuilt
+			&& (GIsTransacting
+			|| PropertyChangedEvent.ChangeType == EPropertyChangeType::ArrayAdd
+			|| PropertyChangedEvent.ChangeType == EPropertyChangeType::ArrayRemove
+			|| PropertyChangedEvent.ChangeType == EPropertyChangeType::ArrayClear))
 		{
+			bIsPendingRebuilt = true;
+
+			//To clean up dangling references
 			EmptyAllSelection();
+			EmptyAllCurveViewModels();
+			RegisterActiveTimer(0.f, FWidgetActiveTimerDelegate::CreateSPLambda(
+				this,
+				[this](double, float)
+				{
+					if(const TSharedPtr<IPropertyUtilities> PropertyUtilities = PropertyUtilitiesWeak.Pin())
+					{
+						//Rebuild the widget so it will reference to the new correct address of Curves, and show things correctly
+						PropertyUtilities->ForceRefresh();
+					}
+					return EActiveTimerReturnType::Stop;
+				}
+				));
 		}
-		else
-		{
-			ValidateSelection();
-		}
+
+		ValidateSelection();
 	}
 }
 
diff --git a/Engine/Source/Editor/UnrealEd/Public/SCurveEditor.h b/Engine/Source/Editor/UnrealEd/Public/SCurveEditor.h
--- a/Engine/Source/Editor/UnrealEd/Public/SCurveEditor.h
+++ b/Engine/Source/Editor/UnrealEd/Public/SCurveEditor.h
@@ -5,6 +5,7 @@
 #include "CoreMinimal.h"
 #include "Misc/Attribute.h"
 #include "Templates/SubclassOf.h"
+#include "Templates/SharedPointer.h"
 #include "Layout/Geometry.h"
 #include "Input/Reply.h"
 #include "Layout/Visibility.h"
@@ -21,6 +22,7 @@
 class FPaintArgs;
 class FSlateWindowElementList;
 class FUICommandList;
+class IPropertyUtilities;
 class IMenu;
 class SBox;
 class SErrorText;
@@ -240,6 +242,11 @@ public:
 
 	/** Set the curve that is being edited by this track widget. Also provide an option to enable/disable editing */
 	UNREALED_API void SetCurveOwner(FCurveOwnerInterface* InCurveOwner, bool bCanEdit = true);
+	
+	UNREALED_API void SetPropertyUtils(const TSharedPtr<IPropertyUtilities>& InPropertyUtils)
+	{
+		PropertyUtilitiesWeak = InPropertyUtils;
+	}
 
 	/** Set new zoom to fit **/
 	UNREALED_API void SetZoomToFit(bool bNewZoomToFitVertical, bool bNewZoomToFitHorizontal);
@@ -651,6 +658,9 @@ protected:
 	/** Access the user-supplied settings object */
 	UCurveEditorSettings* GetSettings() const { return Settings; }
 
+	/** Empty all the curve view models */
+	void EmptyAllCurveViewModels();
+
 	/** Clear the selected curve view models */
 	UNREALED_API void ClearSelectedCurveViewModels();
 
@@ -849,4 +859,8 @@ protected:
 
 	/** Generic Popup Entry */
 	TWeakPtr<IMenu> EntryPopupMenu;
+
+private:
+	bool bIsPendingRebuilt;
+	TWeakPtr<IPropertyUtilities> PropertyUtilitiesWeak;
 };
