Fixed FBlueprintEditorUtils::ReplaceAllActorRefrences which was always dirtying the level package even if the level script BP wasn't concerned by the replaced actor.

diff --git a/Engine/Source/Editor/UnrealEd/Private/Kismet2/BlueprintEditorUtils.cpp b/Engine/Source/Editor/UnrealEd/Private/Kismet2/BlueprintEditorUtils.cpp
--- a/Engine/Source/Editor/UnrealEd/Private/Kismet2/BlueprintEditorUtils.cpp
+++ b/Engine/Source/Editor/UnrealEd/Private/Kismet2/BlueprintEditorUtils.cpp
@@ -8110,8 +8110,16 @@ bool FBlueprintEditorUtils::FindReferencesToActorFromLevelScript(ULevelScriptBlu
 
 void FBlueprintEditorUtils::ReplaceAllActorRefrences(ULevelScriptBlueprint* InLevelScriptBlueprint, AActor* InOldActor, AActor* InNewActor)
 {
-	InLevelScriptBlueprint->Modify();
-	FBlueprintEditorUtils::MarkBlueprintAsModified(InLevelScriptBlueprint);
+	bool bBlueprintModified = false;
+	auto BlueprintModify = [&bBlueprintModified, InLevelScriptBlueprint]()
+	{
+		if (!bBlueprintModified)
+		{
+			InLevelScriptBlueprint->Modify();
+			FBlueprintEditorUtils::MarkBlueprintAsModified(InLevelScriptBlueprint);
+			bBlueprintModified = true;
+		}
+	};
 
 	// Literal nodes are the common "get" type nodes and need to be updated with the new object reference
 	TArray< UK2Node_Literal* > LiteralNodes;
@@ -8121,6 +8129,7 @@ void FBlueprintEditorUtils::ReplaceAllActorRefrences(ULevelScriptBlueprint* InLe
 	{
 		if(LiteralNode->GetObjectRef() == InOldActor)
 		{
+			BlueprintModify();
 			LiteralNode->Modify();
 			LiteralNode->SetObjectRef(InNewActor);
 			LiteralNode->ReconstructNode();
@@ -8135,6 +8144,7 @@ void FBlueprintEditorUtils::ReplaceAllActorRefrences(ULevelScriptBlueprint* InLe
 	{
 		if(ActorEventNode->GetReferencedLevelActor() == InOldActor)
 		{
+			BlueprintModify();
 			ActorEventNode->Modify();
 			ActorEventNode->EventOwner = InNewActor;
 			ActorEventNode->ReconstructNode();
