Parallelize loading the cached localization data from the package headers

The behavior of this can be controlled by the "Localization.GatherTextFromAssetsCommandlet.ParallelizeProcessAndRemoveCachedPackages" and "Localization.GatherTextFromAssetsCommandlet.ProcessAndRemoveCachedPackagesMaxThreads", but in local testing with the default settings it yielded an 8x improvement compared with loading the cached localization data on a single thread.

diff --git a/Engine/Source/Editor/UnrealEd/Classes/Commandlets/GatherTextFromAssetsCommandlet.h b/Engine/Source/Editor/UnrealEd/Classes/Commandlets/GatherTextFromAssetsCommandlet.h
--- a/Engine/Source/Editor/UnrealEd/Classes/Commandlets/GatherTextFromAssetsCommandlet.h
+++ b/Engine/Source/Editor/UnrealEd/Classes/Commandlets/GatherTextFromAssetsCommandlet.h
@@ -92,6 +92,9 @@ private:
 		/** The set of external actors to process for a world partition map package */
 		TSet<FGuid> ExternalActors;
 
+		/** The localization ID of this package, if any */
+		FString PackageLocalizationId;
+
 		/** Localization cache state of this package */
 		EPackageLocCacheState PackageLocCacheState;
 
diff --git a/Engine/Source/Editor/UnrealEd/Private/Commandlets/GatherTextFromAssetsCommandlet.cpp b/Engine/Source/Editor/UnrealEd/Private/Commandlets/GatherTextFromAssetsCommandlet.cpp
--- a/Engine/Source/Editor/UnrealEd/Private/Commandlets/GatherTextFromAssetsCommandlet.cpp
+++ b/Engine/Source/Editor/UnrealEd/Private/Commandlets/GatherTextFromAssetsCommandlet.cpp
@@ -251,7 +251,13 @@ namespace UE::Private::GatherTextFromAssetsCommandlet
 	}
 
 	static bool bParallelizeIncludeExcludePathFiltering = true;
-	static FAutoConsoleVariableRef CVarParallelizeIncludeExcludePathFiltering(TEXT("Localization.GatherTextFromAssetsCommandlet.ParallelizeIncludeExcludePathFiltering"), bParallelizeIncludeExcludePathFiltering, TEXT("True to parallelize the include exclude path filtering. False to force the include exclu de process to be single threaded for easier debugging."));
+	static FAutoConsoleVariableRef CVarParallelizeIncludeExcludePathFiltering(TEXT("Localization.GatherTextFromAssetsCommandlet.ParallelizeIncludeExcludePathFiltering"), bParallelizeIncludeExcludePathFiltering, TEXT("True to parallelize the include exclude path filtering. False to force it to be single threaded for easier debugging."));
+
+	static bool bParallelizeProcessAndRemoveCachedPackages = true;
+	static FAutoConsoleVariableRef CVarParallelizeProcessAndRemoveCachedPackages(TEXT("Localization.GatherTextFromAssetsCommandlet.ParallelizeProcessAndRemoveCachedPackages"), bParallelizeProcessAndRemoveCachedPackages, TEXT("True to parallelize the 'process and remove cached packages' step. False to force it to be single threaded for easier debugging."));
+
+	static int32 ProcessAndRemoveCachedPackagesMaxThreads = -1;
+	static FAutoConsoleVariableRef CVarProcessAndRemoveCachedPackagesMaxThreads(TEXT("Localization.GatherTextFromAssetsCommandlet.ProcessAndRemoveCachedPackagesMaxThreads"), ProcessAndRemoveCachedPackagesMaxThreads, TEXT("Max number of threads to use if parallelizing the 'process and remove cached packages' step, or <= 0 to use as many threads as possible."));
 }
 
 #define LOC_DEFINE_REGION
@@ -1058,42 +1064,86 @@ void UGatherTextFromAssetsCommandlet::PopulatePackagesPendingGather(TSet<FName>
 void UGatherTextFromAssetsCommandlet::ProcessAndRemoveCachedPackages(TMap<FName, TSet<FGuid>>& OutExternalActorsWithStaleOrMissingCaches)
 {
 	UE_SCOPED_TIMER(TEXT("UGatherTextFromAssetsCommandlet::ProcessAndRemoveCachedPackages"), LogGatherTextFromAssetsCommandlet, Display);
-	int32 NumPackagesProcessed = 0;
-	int32 PackageCount = PackagesPendingGather.Num();
-	TMap<FString, FName> AssignedPackageLocalizationIds;
-	PackagesPendingGather.RemoveAll([&](FPackagePendingGather& PackagePendingGather) -> bool
+
+	const int32 InitialNumPackagesToGather = PackagesPendingGather.Num();
+
+	// Load any cached localization data in parallel to maximize file throughput
+	// We do not update any shared state during this pass, and only update the data within FPackagePendingGather
+	{
+		const EParallelForFlags LoadingLoopFlags = UE::Private::GatherTextFromAssetsCommandlet::bParallelizeProcessAndRemoveCachedPackages ? EParallelForFlags::None : EParallelForFlags::ForceSingleThread;
+		const int32 LoadingLoopNumElements = InitialNumPackagesToGather;
+		const int32 LoadingLoopMinBatchSize = UE::Private::GatherTextFromAssetsCommandlet::ProcessAndRemoveCachedPackagesMaxThreads <= 0
+			? 1 // Batch size of 1 uses as many threads as ParallelFor allows
+			: (LoadingLoopNumElements / FMath::Max(UE::Private::GatherTextFromAssetsCommandlet::ProcessAndRemoveCachedPackagesMaxThreads - 1, 1)) + 1; // -1 from MaxThreads as ParallelFor will include the game thread internally
+		const int32 LoadingLoopNumThreads = ParallelForImpl::GetNumberOfThreadTasks(LoadingLoopNumElements, LoadingLoopMinBatchSize, LoadingLoopFlags);
+
+		const double LoadStartTime = FPlatformTime::Seconds();
+		UE_LOG(LogGatherTextFromAssetsCommandlet, Display, TEXT("Loading the cached localization data for %d package(s) using %d thread(s). This may take a while..."), InitialNumPackagesToGather, LoadingLoopNumThreads);
+
+		ParallelFor(TEXT("UGatherTextFromAssetsCommandlet::ProcessAndRemoveCachedPackages"), LoadingLoopNumElements, LoadingLoopMinBatchSize, [this](int32 Index)
 		{
+			FPackagePendingGather& PackagePendingGather = PackagesPendingGather[Index];
+
 			const FNameBuilder PackageNameStr(PackagePendingGather.PackageName);
-			const int32 CurrentPackageNum = ++NumPackagesProcessed;
-			const float PercentageComplete = static_cast<float>(CurrentPackageNum) / static_cast<float>(PackageCount) * 100.0f;
+			const bool bIsExternalActorPackage = PackageNameStr.ToView().Contains(FPackagePath::GetExternalActorsFolderName());
 
 			TUniquePtr<FArchive> FileReader(IFileManager::Get().CreateFileReader(*PackagePendingGather.PackageFilename));
 			if (!FileReader)
 			{
-				return false;
+				return;
 			}
 
-			const bool bIsExternalActorPackage = PackageNameStr.ToView().Contains(FPackagePath::GetExternalActorsFolderName());
-
 			// Read package file summary from the file.
 			FPackageFileSummary PackageFileSummary;
 			*FileReader << PackageFileSummary;
 
+			PackagePendingGather.PackageLocalizationId = PackageFileSummary.LocalizationId;
+			PackagePendingGather.PackageLocCacheState = CalculatePackageLocCacheState(PackageFileSummary, PackagePendingGather.PackageName, bIsExternalActorPackage);
+
+			// Read the cached localization data
+			if (PackagePendingGather.PackageLocCacheState == EPackageLocCacheState::Cached && PackageFileSummary.GatherableTextDataOffset > 0)
+			{
+				FileReader->Seek(PackageFileSummary.GatherableTextDataOffset);
+
+				PackagePendingGather.GatherableTextDataArray.SetNum(PackageFileSummary.GatherableTextDataCount);
+				for (int32 GatherableTextDataIndex = 0; GatherableTextDataIndex < PackageFileSummary.GatherableTextDataCount; ++GatherableTextDataIndex)
+				{
+					*FileReader << PackagePendingGather.GatherableTextDataArray[GatherableTextDataIndex];
+				}
+			}
+		}, LoadingLoopFlags);
+
+		UE_LOG(LogGatherTextFromAssetsCommandlet, Display, TEXT("Loading the cached localization data for %d package(s) took %.2f seconds."), InitialNumPackagesToGather, FPlatformTime::Seconds() - LoadStartTime);
+	}
+
+	// Now that everything has been loaded (or not), run through the cached data, ingest it, update the metrics, and update PackagesPendingGather
+	{
+		const double IngestStartTime = FPlatformTime::Seconds();
+		UE_LOG(LogGatherTextFromAssetsCommandlet, Display, TEXT("Ingesting the cached localization data for %d package(s). This may take a while..."), InitialNumPackagesToGather);
+
+		int32 NumPackagesProcessed = 0;
+		TMap<FString, FName> AssignedPackageLocalizationIds;
+		PackagesPendingGather.RemoveAll([this, &OutExternalActorsWithStaleOrMissingCaches, &AssignedPackageLocalizationIds, &NumPackagesProcessed, InitialNumPackagesToGather](const FPackagePendingGather& PackagePendingGather) -> bool
+		{
+			const FNameBuilder PackageNameStr(PackagePendingGather.PackageName);
+			const int32 CurrentPackageNum = ++NumPackagesProcessed;
+			const float PercentageComplete = static_cast<float>(CurrentPackageNum) / static_cast<float>(InitialNumPackagesToGather) * 100.0f;
+
+			const bool bIsExternalActorPackage = PackageNameStr.ToView().Contains(FPackagePath::GetExternalActorsFolderName());
+
 			// Track the package localization ID of this package (if known) and detect duplicates
-			if (!PackageFileSummary.LocalizationId.IsEmpty())
+			if (!PackagePendingGather.PackageLocalizationId.IsEmpty())
 			{
-				if (const FName* ExistingLongPackageName = AssignedPackageLocalizationIds.Find(PackageFileSummary.LocalizationId))
+				if (const FName* ExistingLongPackageName = AssignedPackageLocalizationIds.Find(PackagePendingGather.PackageLocalizationId))
 				{
-					UE_LOG(LogGatherTextFromAssetsCommandlet, Warning, TEXT("Package '%s' and '%s' have the same localization ID (%s). Please reset one of these (Asset Localization -> Reset Localization ID) to avoid conflicts."), *PackageNameStr, *ExistingLongPackageName->ToString(), *PackageFileSummary.LocalizationId);
+					UE_LOG(LogGatherTextFromAssetsCommandlet, Warning, TEXT("Package '%s' and '%s' have the same localization ID (%s). Please reset one of these (Asset Localization -> Reset Localization ID) to avoid conflicts."), *PackageNameStr, *ExistingLongPackageName->ToString(), *PackagePendingGather.PackageLocalizationId);
 				}
 				else
 				{
-					AssignedPackageLocalizationIds.Add(PackageFileSummary.LocalizationId, PackagePendingGather.PackageName);
+					AssignedPackageLocalizationIds.Add(PackagePendingGather.PackageLocalizationId, PackagePendingGather.PackageName);
 				}
 			}
 
-			PackagePendingGather.PackageLocCacheState = CalculatePackageLocCacheState(PackageFileSummary, PackagePendingGather.PackageName, bIsExternalActorPackage);
-
 			if (PackagePendingGather.PackageLocCacheState != EPackageLocCacheState::Cached)
 			{
 				// External actors actors must be gathered via their owner world rather than via a raw LoadPackage call
@@ -1118,20 +1168,12 @@ void UGatherTextFromAssetsCommandlet::ProcessAndRemoveCachedPackages(TMap<FName,
 			}
 
 			// Process packages that don't require loading to process.
-			if (PackageFileSummary.GatherableTextDataOffset > 0)
-				{
+			if (PackagePendingGather.GatherableTextDataArray.Num() > 0)
+			{
 				UE_LOG(LogGatherTextFromAssetsCommandlet, Display, TEXT("[%6.2f%%] Gathering package: '%s'..."), PercentageComplete, *PackageNameStr);
 
 				UE::Private::GatherTextFromAssetsCommandlet::AssetGatherCacheMetrics.CountCachedAsset();
 
-				FileReader->Seek(PackageFileSummary.GatherableTextDataOffset);
-
-				PackagePendingGather.GatherableTextDataArray.SetNum(PackageFileSummary.GatherableTextDataCount);
-				for (int32 GatherableTextDataIndex = 0; GatherableTextDataIndex < PackageFileSummary.GatherableTextDataCount; ++GatherableTextDataIndex)
-				{
-					*FileReader << PackagePendingGather.GatherableTextDataArray[GatherableTextDataIndex];
-				}
-
 				ProcessGatherableTextDataArray(PackagePendingGather.GatherableTextDataArray);
 			}
 
@@ -1145,6 +1187,9 @@ void UGatherTextFromAssetsCommandlet::ProcessAndRemoveCachedPackages(TMap<FName,
 
 			return true;
 		});
+
+		UE_LOG(LogGatherTextFromAssetsCommandlet, Display, TEXT("Ingesting the cached localization data for %d package(s) took %.2f seconds."), InitialNumPackagesToGather, FPlatformTime::Seconds() - IngestStartTime);
+	}
 }
 
 void UGatherTextFromAssetsCommandlet::MergeInExternalActorsWithStaleOrMissingCaches(TMap<FName, TSet<FGuid>>& StaleExternalActors)
