Adjust placement of toolbar play buttons for PIE and SIE

* The PIE and SIE resume and pause buttons now always appear to the left in the toolbar.
* Possess (SIE -> PIE) and Eject (PIE -> SIE) now share button positions for easy toggling.
* Revert some changes that were previously made for 5.5 but that we don't need after this change. Diff against  files before CL 36660744 to better see what has been added and removed for 5.5.

diff --git a/Engine/Source/Editor/EditorStyle/Private/StarshipStyle.cpp b/Engine/Source/Editor/EditorStyle/Private/StarshipStyle.cpp
--- a/Engine/Source/Editor/EditorStyle/Private/StarshipStyle.cpp
+++ b/Engine/Source/Editor/EditorStyle/Private/StarshipStyle.cpp
@@ -5753,6 +5753,7 @@ void FStarshipEditorStyle::FStyle::SetupPersonaStyle()
 		Set("PlayWorld.Simulate", new IMAGE_BRUSH_SVG("Starship/MainToolbar/simulate", Icon20x20));
 		Set( "PlayWorld.RepeatLastPlay", new IMAGE_BRUSH_SVG("Starship/Common/play", Icon20x20));
 		Set( "PlayWorld.PlayInViewport", new IMAGE_BRUSH_SVG("Starship/Common/play", Icon20x20));
+		Set("PlayWorld.PlaySimulate", new IMAGE_BRUSH_SVG("Starship/MainToolbar/PlaySimulate_20", Icon20x20));
 
 		Set( "PlayWorld.PlayInEditorFloating", new IMAGE_BRUSH_SVG("Starship/MainToolbar/PlayNewEditorWindow", Icon20x20));
 		Set( "PlayWorld.PlayInVR", new IMAGE_BRUSH_SVG("Starship/MainToolbar/PlayVRPreview", Icon20x20));
@@ -5760,19 +5761,11 @@ void FStarshipEditorStyle::FStyle::SetupPersonaStyle()
 		Set( "PlayWorld.PlayInVulkanPreview", new IMAGE_BRUSH_SVG("Starship/MainToolbar/PlayMobilePreview", Icon20x20));
 		Set( "PlayWorld.PlayInNewProcess", new IMAGE_BRUSH_SVG("Starship/MainToolbar/PlayStandaloneGame", Icon20x20));
 
-		Set("PlayWorld.PlayPossess", new IMAGE_BRUSH_SVG("Starship/MainToolbar/PlayPossess_20", Icon20x20));
-
-		Set("PlayWorld.SimulatePauseEject_PlaySimulate", new IMAGE_BRUSH_SVG("Starship/MainToolbar/PlaySimulate_20", Icon20x20));
-		Set("PlayWorld.SimulatePauseEject_Eject", new IMAGE_BRUSH_SVG("Starship/MainToolbar/eject", Icon40x40));
-		Set("PlayWorld.SimulatePauseEject_Pause", new IMAGE_BRUSH_SVG("Starship/MainToolbar/pause", Icon40x40));
-
 		Set( "PlayWorld.RepeatLastLaunch", new IMAGE_BRUSH_SVG( "Starship/MainToolbar/launch", Icon20x20 ) );
 
 		Set( "PlayWorld.PlayInCameraLocation", new IMAGE_BRUSH_SVG( "Starship/AssetIcons/CameraActor_16", Icon20x20 ) );
 		Set( "PlayWorld.PlayInDefaultPlayerStart", new IMAGE_BRUSH_SVG( "Starship/AssetIcons/PlayerStart_16", Icon20x20 ) );
 
-		Set("PlayWorld.ResumePIESession", new IMAGE_BRUSH_SVG("Starship/Common/play", Icon40x40));
-		Set("PlayWorld.ResumePIESession.Small", new IMAGE_BRUSH_SVG("Starship/Common/play", Icon20x20));
 		Set("PlayWorld.ResumePlaySession", new IMAGE_BRUSH_SVG("Starship/MainToolbar/simulate", Icon40x40));
 		Set("PlayWorld.ResumePlaySession.Small", new IMAGE_BRUSH_SVG("Starship/MainToolbar/simulate", Icon20x20));
 		Set( "PlayWorld.PausePlaySession", new IMAGE_BRUSH_SVG("Starship/MainToolbar/pause", Icon40x40));
@@ -5787,6 +5780,7 @@ void FStarshipEditorStyle::FStyle::SetupPersonaStyle()
 		Set("PlayWorld.LateJoinSession.Small", new IMAGE_BRUSH_SVG("Starship/MainToolbar/simulate", Icon20x20));
 
 		Set( "PlayWorld.PossessPlayer", new IMAGE_BRUSH_SVG("Starship/AssetIcons/PlayerController_16", Icon20x20));
+		Set("PlayWorld.PlayPossess", new IMAGE_BRUSH_SVG("Starship/MainToolbar/PlayPossess_20", Icon20x20));
 		Set( "PlayWorld.EjectFromPlayer", new IMAGE_BRUSH_SVG("Starship/MainToolbar/eject", Icon40x40));
 
 		Set( "PlayWorld.ShowCurrentStatement", new IMAGE_BRUSH_SVG( "Starship/MainToolbar/DebugFindNode", Icon40x40 ) );
diff --git a/Engine/Source/Editor/UnrealEd/Private/Kismet2/DebuggerCommands.cpp b/Engine/Source/Editor/UnrealEd/Private/Kismet2/DebuggerCommands.cpp
--- a/Engine/Source/Editor/UnrealEd/Private/Kismet2/DebuggerCommands.cpp
+++ b/Engine/Source/Editor/UnrealEd/Private/Kismet2/DebuggerCommands.cpp
@@ -190,21 +190,12 @@ public:
 	// Mouse control
 	static void GetMouseControlExecute();
 
-	static void PossessPlayer_Clicked();
-	static bool CanPossessPlayer();
-
 	static void PossessEjectPlayer_Clicked();
 	static bool CanPossessEjectPlayer();
 	static FText GetPossessEjectLabel();
 	static FText GetPossessEjectTooltip();
 	static FSlateIcon GetPossessEjectImage();
 
-	static void SimulatePauseEject_Clicked();
-	static bool CanSimulatePauseEject();
-	static FText GetSimulatePauseEjectLabel();
-	static FText GetSimulatePauseEjectTooltip();
-	static FSlateIcon GetSimulatePauseEjectIcon();
-
 	static bool CanLateJoin();
 	static bool CanShowLateJoinButton();
 
@@ -356,6 +347,7 @@ void FPlayWorldCommands::RegisterCommands()
 	// PIE
 	UI_COMMAND(RepeatLastPlay, "Play", "Launches a game preview session in the same mode as the last game preview session launched from the Game Preview Modes dropdown next to the Play button on the level editor toolbar", EUserInterfaceActionType::Button, FInputChord(EKeys::P, EModifierKey::Alt))
 	UI_COMMAND(PlayInViewport, "Selected Viewport", "Play this level in the active level editor viewport", EUserInterfaceActionType::Check, FInputChord());
+	UI_COMMAND(PlaySimulate, "Start simulation", "Simulate this level in the active level editor viewport", EUserInterfaceActionType::Button, FInputChord());
 	UI_COMMAND(PlayInEditorFloating, "New Editor Window (PIE)", "Play this level in a new window", EUserInterfaceActionType::Check, FInputChord());
 	UI_COMMAND(PlayInVR, "VR Preview", "Play this level in VR", EUserInterfaceActionType::Check, FInputChord());
 	UI_COMMAND(PlayInMobilePreview, "Standalone Game Mobile", "Play this level as a mobile device preview in Mobile mode (runs in its own process)", EUserInterfaceActionType::Check, FInputChord());
@@ -367,17 +359,14 @@ void FPlayWorldCommands::RegisterCommands()
 	UI_COMMAND(PlayInSettings, "Advanced Settings...", "Open the settings for the 'Play In' feature", EUserInterfaceActionType::Button, FInputChord());
 
 	// SIE & PIE controls
-	UI_COMMAND(StopPlaySession, "Stop", "Stop simulation", EUserInterfaceActionType::Button, FInputChord(EKeys::Escape));
-	UI_COMMAND(ResumePlaySession, "Resume", "Resume simulation", EUserInterfaceActionType::Button, FInputChord());
-	UI_COMMAND(PausePlaySession, "Pause", "Pause simulation", EUserInterfaceActionType::Button, FInputChord());
-	UI_COMMAND(ResumePIESession, "Resume", "Resume play-in-editor simulation", EUserInterfaceActionType::Button, FInputChord());
+	UI_COMMAND(StopPlaySession, "Stop", "Stop", EUserInterfaceActionType::Button, FInputChord(EKeys::Escape));
+	UI_COMMAND(ResumePlaySession, "Resume", "Resume", EUserInterfaceActionType::Button, FInputChord());
+	UI_COMMAND(PausePlaySession, "Pause", "Pause", EUserInterfaceActionType::Button, FInputChord());
 	UI_COMMAND(GetMouseControl, "Mouse Control", "Get mouse cursor while in PIE", EUserInterfaceActionType::Button, FInputChord(EModifierKey::Shift, EKeys::F1));
 	UI_COMMAND(LateJoinSession, "Add Client", "Add another client", EUserInterfaceActionType::Button, FInputChord());
 	UI_COMMAND(SingleFrameAdvance, "Skip", "Advances a single frame", EUserInterfaceActionType::Button, FInputChord());
 	UI_COMMAND(TogglePlayPauseOfPlaySession, "Toggle Play/Pause", "Resume playing if paused, or pause if playing", EUserInterfaceActionType::Button, FInputChord(EKeys::Pause));
-	UI_COMMAND(PlayPossess, "Possess Player", "Possesses the player character", EUserInterfaceActionType::Button, FInputChord());
 	UI_COMMAND(PossessEjectPlayer, "Possess or Eject Player", "Possesses or ejects the player from the camera", EUserInterfaceActionType::Button, FInputChord(EKeys::F8));
-	UI_COMMAND(SimulatePauseEject, "Simulate, Pause Simulation or Eject Player", "Starts or pauses simulation or eject the player from the camera", EUserInterfaceActionType::Button, FInputChord());
 	UI_COMMAND(ShowCurrentStatement, "Locate", "Locate the currently active node", EUserInterfaceActionType::Button, FInputChord());
 	UI_COMMAND(AbortExecution, "Abort", "Abort the execution", EUserInterfaceActionType::Button, FInputChord());
 	UI_COMMAND(ContinueExecution, "Continue", "Continue the execution", EUserInterfaceActionType::Button, FInputChord(EModifierKey::Alt | EModifierKey::Shift, EKeys::F10));
@@ -458,6 +447,14 @@ void FPlayWorldCommands::BindGlobalPlayWorldCommands()
 		FIsActionButtonVisible::CreateStatic(&FInternalPlayWorldCommandCallbacks::CanShowNonPlayWorldOnlyActions)
 	);
 
+	ActionList.MapAction(
+		Commands.PlaySimulate,
+		FExecuteAction::CreateStatic(&FInternalPlayWorldCommandCallbacks::Simulate_Clicked),
+		FCanExecuteAction::CreateStatic(&FInternalPlayWorldCommandCallbacks::Simulate_CanExecute),
+		FIsActionChecked::CreateStatic(&FInternalPlayWorldCommandCallbacks::PlayInModeIsChecked, PlayMode_Simulate),
+		FIsActionButtonVisible::CreateStatic(&FInternalPlayWorldCommandCallbacks::CanShowNonPlayWorldOnlyActions)
+	);
+
 	ActionList.MapAction(Commands.PlayInEditorFloating,
 		FExecuteAction::CreateStatic(&FInternalPlayWorldCommandCallbacks::PlayInEditorFloating_Clicked),
 		FCanExecuteAction::CreateStatic(&FInternalPlayWorldCommandCallbacks::PlayInEditorFloating_CanExecute),
@@ -543,14 +540,6 @@ void FPlayWorldCommands::BindGlobalPlayWorldCommands()
 		FIsActionButtonVisible::CreateStatic(&FPlayWorldCommandCallbacks::HasPlayWorldAndRunning)
 	);
 
-	ActionList.MapAction(
-		Commands.ResumePIESession,
-		FExecuteAction::CreateStatic(&FPlayWorldCommandCallbacks::ResumePlaySession_Clicked),
-		FCanExecuteAction::CreateStatic(&FPlayWorldCommandCallbacks::IsInPIE_AndPaused),
-		FIsActionChecked(),
-		FIsActionButtonVisible::CreateStatic(&FPlayWorldCommandCallbacks::IsInPIE_AndPaused)
-	);
-
 	ActionList.MapAction(Commands.SingleFrameAdvance,
 		FExecuteAction::CreateStatic(&FInternalPlayWorldCommandCallbacks::SingleFrameAdvance_Clicked),
 		FCanExecuteAction::CreateStatic(&FPlayWorldCommandCallbacks::HasPlayWorldAndPaused),
@@ -573,30 +562,13 @@ void FPlayWorldCommands::BindGlobalPlayWorldCommands()
 		FIsActionButtonVisible::CreateStatic(&FInternalPlayWorldCommandCallbacks::HasPlayWorld)
 	);
 
-	// Possess the player (SIE->PIE).
-	ActionList.MapAction(
-		Commands.PlayPossess,
-		FExecuteAction::CreateStatic(&FInternalPlayWorldCommandCallbacks::PossessPlayer_Clicked),
-		FCanExecuteAction::CreateStatic(&FInternalPlayWorldCommandCallbacks::CanPossessPlayer),
-		FIsActionChecked(),
-		FIsActionButtonVisible::CreateStatic(&FInternalPlayWorldCommandCallbacks::CanPossessPlayer)
-	);
-
 	// Toggle PIE/SIE, Eject (PIE->SIE), and Possess (SIE->PIE)
-	ActionList.MapAction(Commands.PossessEjectPlayer,
+	ActionList.MapAction(
+		Commands.PossessEjectPlayer,
 		FExecuteAction::CreateStatic(&FInternalPlayWorldCommandCallbacks::PossessEjectPlayer_Clicked),
 		FCanExecuteAction::CreateStatic(&FInternalPlayWorldCommandCallbacks::CanPossessEjectPlayer),
 		FIsActionChecked(),
-		FIsActionButtonVisible::CreateStatic(&DebuggerCommands::AllowPlayWorldFeature)
-	);
-
-	// Starts SIE if we're not in PIE or SIE, toggles pause/resume if we're in SIE, or ejects (PIE -> SIE) if we're in PIE.
-	ActionList.MapAction(
-		Commands.SimulatePauseEject,
-		FExecuteAction::CreateStatic(&FInternalPlayWorldCommandCallbacks::SimulatePauseEject_Clicked),
-		FCanExecuteAction::CreateStatic(&FInternalPlayWorldCommandCallbacks::CanSimulatePauseEject),
-		FIsActionChecked(),
-		FIsActionButtonVisible::CreateStatic(&DebuggerCommands::AllowPlayWorldFeature)
+		FIsActionButtonVisible::CreateStatic(&FInternalPlayWorldCommandCallbacks::CanPossessEjectPlayer)
 	);
 
 	// Breakpoint-only commands
@@ -679,15 +651,27 @@ void FPlayWorldCommands::BuildToolbar(FToolMenuSection& InSection, bool bInclude
 	// Play
 	InSection.AddEntry(PlayMenuEntry);
 
-	// Resume PIE, pause PIE/SIE, or possess player (SIE -> PIE). (Only one will be visible, and only in PIE/SIE.)
+	// First slot (only one is visible at a time.)
 	{
 		{
 			FToolMenuEntry ResumeEntry = FToolMenuEntry::InitToolBarButton(
-				FPlayWorldCommands::Get().ResumePIESession,
+				FPlayWorldCommands::Get().ResumePlaySession,
+				TAttribute<FText>(),
 				TAttribute<FText>(),
-				LOCTEXT("ResumePIE", "Resume play-in-editor session"),
-				FSlateIcon(FAppStyle::GetAppStyleSetName(), "PlayWorld.ResumePIESession"),
-				FName(TEXT("ResumePIESession"))
+				TAttribute<FSlateIcon>::CreateLambda(
+					[]() -> FSlateIcon
+					{
+						if (FPlayWorldCommandCallbacks::IsInPIE())
+						{
+							return FSlateIcon(FAppStyle::GetAppStyleSetName(), "PlayWorld.PlayInViewport");
+						}
+						else
+						{
+							return FSlateIcon(FAppStyle::GetAppStyleSetName(), "PlayWorld.PlaySimulate");
+						}
+					}
+				),
+				FName(TEXT("ResumePlaySession"))
 			);
 			ResumeEntry.StyleNameOverride = FName("Toolbar.BackplateLeft");
 			InSection.AddEntry(ResumeEntry);
@@ -702,21 +686,8 @@ void FPlayWorldCommands::BuildToolbar(FToolMenuSection& InSection, bool bInclude
 				FName(TEXT("PausePlaySession"))
 			);
 			PauseEntry.StyleNameOverride = FName("Toolbar.BackplateLeft");
-
 			InSection.AddEntry(PauseEntry);
 		}
-
-		{
-			FToolMenuEntry PossessEntry = FToolMenuEntry::InitToolBarButton(
-				FPlayWorldCommands::Get().PlayPossess,
-				TAttribute<FText>(),
-				FPlayWorldCommands::Get().PlayPossess->GetDescription(),
-				FPlayWorldCommands::Get().PlayPossess->GetIcon(),
-				FName(TEXT("PlayPossess"))
-			);
-			PossessEntry.StyleNameOverride = FName("Toolbar.BackplateLeft");
-			InSection.AddEntry(PossessEntry);
-		}
 	}
 
 	FToolMenuEntry SingleFrameAdvanceEntry = FToolMenuEntry::InitToolBarButton(FPlayWorldCommands::Get().SingleFrameAdvance, TAttribute<FText>(), TAttribute<FText>(), TAttribute<FSlateIcon>(), FName(TEXT("SingleFrameAdvance")));
@@ -724,7 +695,13 @@ void FPlayWorldCommands::BuildToolbar(FToolMenuSection& InSection, bool bInclude
 	InSection.AddEntry(SingleFrameAdvanceEntry);
 
 	// Stop
-	FToolMenuEntry StopEntry = FToolMenuEntry::InitToolBarButton(FPlayWorldCommands::Get().StopPlaySession, TAttribute<FText>(), TAttribute<FText>(), TAttribute<FSlateIcon>(), FName(TEXT("StopPlaySession")));
+	FToolMenuEntry StopEntry = FToolMenuEntry::InitToolBarButton(
+		FPlayWorldCommands::Get().StopPlaySession,
+		TAttribute<FText>(),
+		TAttribute<FText>(),
+		TAttribute<FSlateIcon>(),
+		FName(TEXT("StopPlaySession"))
+	);
 	StopEntry.StyleNameOverride = FName("Toolbar.BackplateCenterStop");
 
 	InSection.AddEntry(StopEntry);
@@ -734,24 +711,37 @@ void FPlayWorldCommands::BuildToolbar(FToolMenuSection& InSection, bool bInclude
 	LateJoinEntry.StyleNameOverride = FName("Toolbar.BackplateCenter");
 	InSection.AddEntry(LateJoinEntry);
 
-	// Start SIE / pause SIE/PIE / eject (PIE -> SIE).
+	// Fourth slot (only one is visible at a time.)
 	{
-		FToolMenuEntry SimulatePauseEjectMenuEntry = FToolMenuEntry::InitToolBarButton(
-			FPlayWorldCommands::Get().SimulatePauseEject,
-			TAttribute<FText>::Create(
-				TAttribute<FText>::FGetter::CreateStatic(&FInternalPlayWorldCommandCallbacks::GetSimulatePauseEjectLabel)
-			),
-			TAttribute<FText>::Create(
-				TAttribute<FText>::FGetter::CreateStatic(&FInternalPlayWorldCommandCallbacks::GetSimulatePauseEjectTooltip)
-			),
-			TAttribute<FSlateIcon>::Create(TAttribute<FSlateIcon>::FGetter::CreateStatic(
-				&FInternalPlayWorldCommandCallbacks::GetSimulatePauseEjectIcon
-			)),
-			FName(TEXT("SimulatePauseEject"))
-		);
-		SimulatePauseEjectMenuEntry.StyleNameOverride = FName("Toolbar.BackplateCenter");
+		{
+			FToolMenuEntry PlaySimulateEntry = FToolMenuEntry::InitToolBarButton(
+				FPlayWorldCommands::Get().PlaySimulate,
+				TAttribute<FText>(),
+				TAttribute<FText>(),
+				TAttribute<FSlateIcon>(),
+				FName(TEXT("PlaySimulate"))
+			);
+			PlaySimulateEntry.StyleNameOverride = FName("Toolbar.BackplateCenter");
+			InSection.AddEntry(PlaySimulateEntry);
+		}
 
-		InSection.AddEntry(SimulatePauseEjectMenuEntry);
+		{
+			FToolMenuEntry PossessEjectEntry = FToolMenuEntry::InitToolBarButton(
+				FPlayWorldCommands::Get().PossessEjectPlayer,
+				TAttribute<FText>::Create(
+					TAttribute<FText>::FGetter::CreateStatic(&FInternalPlayWorldCommandCallbacks::GetPossessEjectLabel)
+				),
+				TAttribute<FText>::Create(
+					TAttribute<FText>::FGetter::CreateStatic(&FInternalPlayWorldCommandCallbacks::GetPossessEjectTooltip)
+				),
+				TAttribute<FSlateIcon>::Create(
+					TAttribute<FSlateIcon>::FGetter::CreateStatic(&FInternalPlayWorldCommandCallbacks::GetPossessEjectImage)
+				),
+				FName(TEXT("PossessEjectPlayer"))
+			);
+			PossessEjectEntry.StyleNameOverride = FName("Toolbar.BackplateCenter");
+			InSection.AddEntry(PossessEjectEntry);
+		}
 	}
 
 	FUIAction SpecialPIEOptionsMenuAction;
@@ -1176,11 +1166,6 @@ bool FPlayWorldCommandCallbacks::IsInPIE_AndRunning()
 	return IsInPIE() && ((GEditor->PlayWorld == NULL) || !(GEditor->PlayWorld->bDebugPauseExecution));
 }
 
-bool FPlayWorldCommandCallbacks::IsInPIE_AndPaused()
-{
-	return GUnrealEd && IsInPIE() && HasPlayWorld() && GUnrealEd->PlayWorld->bDebugPauseExecution;
-}
-
 bool FPlayWorldCommandCallbacks::HasPlayWorld()
 {
 	return GEditor->PlayWorld != NULL;
@@ -1195,7 +1180,7 @@ bool FPlayWorldCommandCallbacks::HasPlayWorldAndPaused()
 
 bool FPlayWorldCommandCallbacks::HasPlayWorldAndRunning()
 {
-	return IsInPIE() && HasPlayWorld() && !GUnrealEd->PlayWorld->bDebugPauseExecution;
+	return HasPlayWorld() && !GUnrealEd->PlayWorld->bDebugPauseExecution;
 }
 
 
@@ -1240,7 +1225,7 @@ FSlateIcon FInternalPlayWorldCommandCallbacks::GetPossessEjectImage()
 {
 	if (IsInSIE())
 	{
-		return FSlateIcon(FAppStyle::GetAppStyleSetName(), "PlayWorld.PossessPlayer");
+		return FSlateIcon(FAppStyle::GetAppStyleSetName(), "PlayWorld.PlayPossess");
 	}
 	else
 	{
@@ -1248,75 +1233,6 @@ FSlateIcon FInternalPlayWorldCommandCallbacks::GetPossessEjectImage()
 	}
 }
 
-FText FInternalPlayWorldCommandCallbacks::GetSimulatePauseEjectLabel()
-{
-	if (IsInPIE())
-	{
-		return LOCTEXT("EjectLabel", "Eject");
-	}
-	else if (IsInSIE())
-	{
-		if (FPlayWorldCommandCallbacks::HasPlayWorldAndPaused())
-		{
-			return LOCTEXT("SimulatePauseEjectLabel_Resume", "Resume");
-		}
-		else
-		{
-			return LOCTEXT("SimulatePauseEjectLabel_Pause", "Pause");
-		}
-	}
-	else
-	{
-		return FPlayWorldCommands::Get().Simulate->GetLabel();
-	}
-}
-
-FText FInternalPlayWorldCommandCallbacks::GetSimulatePauseEjectTooltip()
-{
-	if (IsInPIE())
-	{
-		return LOCTEXT("EjectToolTip", "Detaches from the player controller, allowing regular editor controls");
-	}
-	else if (IsInSIE())
-	{
-		if (FPlayWorldCommandCallbacks::HasPlayWorldAndPaused())
-		{
-			return LOCTEXT("SimulatePauseEjectTooltip_Resume", "Resumes the simulation");
-		}
-		else
-		{
-			return LOCTEXT("SimulatePauseEjectTooltip_Pause", "Pauses the simulation");
-		}
-	}
-	else
-	{
-		return FPlayWorldCommands::Get().Simulate->GetDescription();
-	}
-}
-
-FSlateIcon FInternalPlayWorldCommandCallbacks::GetSimulatePauseEjectIcon()
-{
-	if (IsInPIE())
-	{
-		return FSlateIcon(FAppStyle::GetAppStyleSetName(), "PlayWorld.SimulatePauseEject_Eject");
-	}
-	else if (IsInSIE())
-	{
-		if (FPlayWorldCommandCallbacks::HasPlayWorldAndPaused())
-		{
-			return FSlateIcon(FAppStyle::GetAppStyleSetName(), "PlayWorld.SimulatePauseEject_PlaySimulate");
-		}
-		else
-		{
-			return FSlateIcon(FAppStyle::GetAppStyleSetName(), "PlayWorld.SimulatePauseEject_Pause");
-		}
-	}
-	else
-	{
-		return FSlateIcon(FAppStyle::GetAppStyleSetName(), "PlayWorld.SimulatePauseEject_PlaySimulate");
-	}
-}
-
 bool FInternalPlayWorldCommandCallbacks::CanLateJoin()
 {
 	if (FInternalPlayWorldCommandCallbacks::IsStoppedAtBreakpoint_InEngineMode())
@@ -2134,23 +2050,6 @@ bool FInternalPlayWorldCommandCallbacks::IsStoppedAtBreakpoint_InEngineMode()
 	return GIntraFrameDebuggingGameThread && FKismetDebugUtilities::GetCurrentDebuggingWorld() == nullptr;
 }
 
-void FInternalPlayWorldCommandCallbacks::PossessPlayer_Clicked()
-{
-	GEditor->RequestToggleBetweenPIEandSIE();
-}
-
-bool FInternalPlayWorldCommandCallbacks::CanPossessPlayer()
-{
-	if (IsInSIE() && !IsStoppedAtBreakpoint())
-	{
-		for (auto It = GUnrealEd->SlatePlayInEditorMap.CreateIterator(); It; ++It)
-		{
-			return It.Value().DestinationSlateViewport.IsValid();
-		}
-	}
-	return false;
-}
-
 void FInternalPlayWorldCommandCallbacks::PossessEjectPlayer_Clicked()
 {
 	GEditor->RequestToggleBetweenPIEandSIE();
@@ -2169,45 +2068,6 @@ bool FInternalPlayWorldCommandCallbacks::CanPossessEjectPlayer()
 	return false;
 }
 
-void FInternalPlayWorldCommandCallbacks::SimulatePauseEject_Clicked()
-{
-	if (IsInPIE())
-	{
-		GEditor->RequestToggleBetweenPIEandSIE();
-	}
-	else if (IsInSIE())
-	{
-		if (FPlayWorldCommandCallbacks::HasPlayWorldAndPaused())
-		{
-			FPlayWorldCommandCallbacks::ResumePlaySession_Clicked();
-		}
-		else
-		{
-			FPlayWorldCommandCallbacks::PausePlaySession_Clicked();
-		}
-	}
-	else
-	{
-		FPlayWorldCommands::GlobalPlayWorldActions->ExecuteAction(FPlayWorldCommands::Get().Simulate.ToSharedRef());
-	}
-}
-
-bool FInternalPlayWorldCommandCallbacks::CanSimulatePauseEject()
-{
-	if (IsInPIE())
-	{
-		return FInternalPlayWorldCommandCallbacks::CanPossessEjectPlayer();
-	}
-	else if (IsInSIE())
-	{
-		return HasPlayWorld();
-	}
-	else
-	{
-		return FPlayWorldCommands::GlobalPlayWorldActions->CanExecuteAction(FPlayWorldCommands::Get().Simulate.ToSharedRef());
-	}
-}
-
 void FInternalPlayWorldCommandCallbacks::AddMessageLog(const FText& Text, const FText& Detail, const FString& TutorialLink, const FString& DocumentationLink)
 {
 	TSharedRef<FTokenizedMessage> Message = FTokenizedMessage::Create(EMessageSeverity::Error);
diff --git a/Engine/Source/Editor/UnrealEd/Public/Kismet2/DebuggerCommands.h b/Engine/Source/Editor/UnrealEd/Public/Kismet2/DebuggerCommands.h
--- a/Engine/Source/Editor/UnrealEd/Public/Kismet2/DebuggerCommands.h
+++ b/Engine/Source/Editor/UnrealEd/Public/Kismet2/DebuggerCommands.h
@@ -83,6 +83,7 @@ public:
 	/** Play in editor (PIE) */
 	TSharedPtr<FUICommandInfo> RepeatLastPlay;
 	TSharedPtr<FUICommandInfo> PlayInViewport;
+	TSharedPtr<FUICommandInfo> PlaySimulate;
 	TSharedPtr<FUICommandInfo> PlayInEditorFloating;
 	TSharedPtr<FUICommandInfo> PlayInVR;
 	TSharedPtr<FUICommandInfo> PlayInMobilePreview;
@@ -98,16 +99,13 @@ public:
 	/** SIE & PIE controls */
 	TSharedPtr<FUICommandInfo> ResumePlaySession;
 	TSharedPtr<FUICommandInfo> PausePlaySession;
-	TSharedPtr<FUICommandInfo> ResumePIESession;
 	TSharedPtr<FUICommandInfo> SingleFrameAdvance;
 	TSharedPtr<FUICommandInfo> TogglePlayPauseOfPlaySession;
 	TSharedPtr<FUICommandInfo> StopPlaySession;
 	TSharedPtr<FUICommandInfo> LateJoinSession;
-	TSharedPtr<FUICommandInfo> PlayPossess;
 	TSharedPtr<FUICommandInfo> PossessEjectPlayer;
 	TSharedPtr<FUICommandInfo> ShowCurrentStatement;
 	TSharedPtr<FUICommandInfo> GetMouseControl;
-	TSharedPtr<FUICommandInfo> SimulatePauseEject;
 
 	/** BP Debugging controls */
 	TSharedPtr<FUICommandInfo> AbortExecution;
@@ -157,7 +155,6 @@ public:
 
 	static UNREALED_API bool IsInSIE_AndRunning();
 	static UNREALED_API bool IsInPIE_AndRunning();
-	static UNREALED_API bool IsInPIE_AndPaused();
 
 	static UNREALED_API bool HasPlayWorld();
 	static UNREALED_API bool HasPlayWorldAndPaused();
