[PCG] Fix bug where pasting an editor graph node that fails to import object properties will break the parser and crash the editor.

diff --git a/Engine/Source/Editor/UnrealEd/Private/EditorObject.cpp b/Engine/Source/Editor/UnrealEd/Private/EditorObject.cpp
--- a/Engine/Source/Editor/UnrealEd/Private/EditorObject.cpp
+++ b/Engine/Source/Editor/UnrealEd/Private/EditorObject.cpp
@@ -517,6 +517,13 @@ static const TCHAR* ImportProperties(
 				// since we're redefining an object in the same text block, only need to import properties again
 				SourceText = ImportObjectProperties( (uint8*)BaseTemplate, SourceText, TemplateClass, SubobjectRoot, BaseTemplate,
 													Warn, Depth + 1, ContextSupplier ? ContextSupplier->CurrentLine : 0, &InstanceGraph, ObjectRemapper );
+
+				// If the object import fails, early out and stop parsing.
+				if (SourceText == nullptr)
+				{
+					Warn->Logf(ELogVerbosity::Error, TEXT("BEGIN OBJECT: Could not import properties from sub-object %s."), *SubobjectRoot->GetName());
+					return nullptr;
+				}
 			}
 			else 
 			{
@@ -731,6 +738,13 @@ static const TCHAR* ImportProperties(
 					&InstanceGraph,
 					ObjectRemapper
 					);
+
+				// If the object import fails, early out and stop parsing.
+				if (SourceText == nullptr)
+				{
+					Warn->Logf(ELogVerbosity::Error, TEXT("BEGIN OBJECT: Could not import properties from sub-object %s."), *SubobjectRoot->GetName());
+					return nullptr;
+				}
 			}
 		}
 		else if( FParse::Command(&Str,TEXT("CustomProperties")))
@@ -1740,4 +1754,3 @@ const TCHAR* ImportObjectsPropertiesStep(FMultiStepsImportObjectParams& InParams
 }
 
 }	// End namespace EditorUtilities
-
diff --git a/Engine/Source/Editor/UnrealEd/Private/Factories/EditorFactories.cpp b/Engine/Source/Editor/UnrealEd/Private/Factories/EditorFactories.cpp
--- a/Engine/Source/Editor/UnrealEd/Private/Factories/EditorFactories.cpp
+++ b/Engine/Source/Editor/UnrealEd/Private/Factories/EditorFactories.cpp
@@ -5444,8 +5444,10 @@ void FCustomizableTextObjectFactory::ProcessBuffer(UObject* InParent, EObjectFla
 		const FString& PropText = PropMap.FindChecked(CreatedObject);
 
 		// Import the properties and give the derived factory a shot at it
-		ImportObjectProperties((uint8*)CreatedObject, *PropText, CreatedObject->GetClass(), CreatedObject, CreatedObject, WarningContext, 0, 0, &InstanceGraph);
-		ProcessConstructedObject(CreatedObject);
+		if (ImportObjectProperties((uint8*)CreatedObject, *PropText, CreatedObject->GetClass(), CreatedObject, CreatedObject, WarningContext, 0, 0, &InstanceGraph))
+		{
+			ProcessConstructedObject(CreatedObject);
+		}
 	}
 	PostProcessConstructedObjects();
 }
@@ -8059,4 +8061,3 @@ UTexture2D* UUDIMTextureFunctionLibrary::MakeUDIMVirtualTextureFromTexture2Ds(FS
 }
 
 #undef LOCTEXT_NAMESPACE
-
