Hide deprecated properties from data table editor

[RN] rnx

diff --git a/Engine/Source/Editor/UnrealEd/Private/DataTableEditorUtils.cpp b/Engine/Source/Editor/UnrealEd/Private/DataTableEditorUtils.cpp
--- a/Engine/Source/Editor/UnrealEd/Private/DataTableEditorUtils.cpp
+++ b/Engine/Source/Editor/UnrealEd/Private/DataTableEditorUtils.cpp
@@ -778,7 +778,7 @@ void FDataTableEditorUtils::CacheDataForEditing(const UScriptStruct* RowStruct,
 	{
 		const FProperty* Prop = *It;
 		check(Prop);
-		if (!Prop->HasMetaData(FName(TEXT("HideFromDataTableEditorColumn"))))
+		if (!Prop->HasAnyPropertyFlags(CPF_Deprecated) && !Prop->HasMetaData(FName(TEXT("HideFromDataTableEditorColumn"))))
 		{
 			StructProps.Add(Prop);
 		}
