IncrementalCook: Fix spurious detections of infinite loop in the handling of the graph of transitive build dependencies; they were caused by incorrect assumption of maximum loop count that was broken by repeated unnecessary revisits of vertices in the graph.
#rnx

diff --git a/Engine/Source/Editor/UnrealEd/Private/Cooker/CookRequestCluster.cpp b/Engine/Source/Editor/UnrealEd/Private/Cooker/CookRequestCluster.cpp
--- a/Engine/Source/Editor/UnrealEd/Private/Cooker/CookRequestCluster.cpp
+++ b/Engine/Source/Editor/UnrealEd/Private/Cooker/CookRequestCluster.cpp
@@ -819,7 +819,12 @@ void FRequestCluster::FGraphSearch::TickExploration(bool& bOutDone)
 {
 	bool bHadActivity = false;
 
-	int32 RunawayLoopCount = 0;
+	int64 RunawayLoopCount = 0;
+	// Add a counter to detect bugs that lead to an infinite loop. Calculation of upper bound:
+	// Each time through the loop we either process or delay one platform of one vertex. Once delayed, a vertex is not added to
+	// the queue again until all of its UnreadyDependencies for that platform are cleared, therefore we only delay it once per platform.
+	// An upper bound is therefore 2 * NumPlatforms * NumVertices.
+	int64 RunawayLoopUpperBound = 2 * Vertices.Num() * (FetchPlatforms.Num() - 1);
 	for (;;)
 	{
 		TOptional<FVertexData*> FrontVertex = AsyncQueueResults.Dequeue();
@@ -844,13 +849,15 @@ void FRequestCluster::FGraphSearch::TickExploration(bool& bOutDone)
 		ExploreEdgesContext.Explore(*Vertex);
 		bHadActivity = true;
 
-		if (RunawayLoopCount++ > 2 * Vertices.Num())
+		if (RunawayLoopCount++ > RunawayLoopUpperBound)
 		{
 			UE_LOG(LogCook, Fatal, TEXT("Infinite loop detected in FRequestCluster::TickExploration's AsyncQueueResults."));
 		}
 	}
 
 	RunawayLoopCount = 0;
+	// Calculation of upper bound: we visit each vertex at most once per platform.
+	RunawayLoopUpperBound = Vertices.Num() * (FetchPlatforms.Num() - 1);
 	while (!VisitVertexQueue.IsEmpty())
 	{
 		bHadActivity = true;
@@ -867,7 +874,7 @@ void FRequestCluster::FGraphSearch::TickExploration(bool& bOutDone)
 		Snapshot.Append(VisitVertexQueue);
 		VisitVertexQueue = MoveTemp(Snapshot);
 
-		if (RunawayLoopCount++ > 2 * Vertices.Num())
+		if (RunawayLoopCount++ > RunawayLoopUpperBound)
 		{
 			UE_LOG(LogCook, Fatal, TEXT("Infinite loop detected in FRequestCluster::TickExploration's VisitVertexQueue."));
 		}
@@ -958,7 +965,7 @@ void FRequestCluster::FGraphSearch::ResolveTransitiveBuildDependencyCycle()
 
 	// Every package in the PendingTransitiveBuildDependencyVertices set is one that is not invalidated by any of its
 	// direct dependencies, but it has transitive build dependencies that might be invalidated.
-	// If we have run out of direct dependency work to do, then all there are no transitive build dependencies on any
+	// If we have run out of direct dependency work to do, then there are no transitive build dependencies on any
 	// vertex not in the set.
 	// No direct dependency invalidations and no transitive build dependency invalidations, by our interpretation of a
 	// cycle above, mean that the package is not invalidated.
@@ -1319,6 +1326,8 @@ bool FRequestCluster::FGraphSearch::FExploreEdgesContext::TryCalculateIterativel
 		return true;
 	}
 
+	Vertex->UnreadyDependencies.Reset();
+	Vertex->bWaitingOnUnreadyDependencies = false;
 	bool bAllPlatformsAreReady = true;
 	for (int32 PlatformIndex : PlatformsToProcess)
 	{
@@ -1456,6 +1465,7 @@ bool FRequestCluster::FGraphSearch::FExploreEdgesContext::TryCalculateIterativel
 					}
 					// It's okay to add duplicates to IterativelyModifiedListeners; we remove them when broadcasting
 					TransitiveBuildVertex->IterativelyModifiedListeners.Add(this->Vertex);
+					Vertex->UnreadyDependencies.Add(TransitiveBuildVertex);
 				}
 
 				bAllPlatformsAreReady = false;
@@ -1469,16 +1479,42 @@ bool FRequestCluster::FGraphSearch::FExploreEdgesContext::TryCalculateIterativel
 	if (!bAllPlatformsAreReady)
 	{
 		GraphSearch.PendingTransitiveBuildDependencyVertices.Add(Vertex);
+		Vertex->bWaitingOnUnreadyDependencies = true;
 		return false;
 	}
 
+	Vertex->UnreadyDependencies.Empty(); // Already reset above, call empty to free the memory
 	if (!Vertex->IterativelyModifiedListeners.IsEmpty())
 	{
+		bool bIterativelyModified = false;
+		for (int32 PlatformIndex = 0; PlatformIndex < LocalNumFetchPlatforms; ++PlatformIndex)
+		{
+			if (PlatformIndex == PlatformAgnosticPlatformIndex)
+			{
+				continue;
+			}
+			FQueryPlatformData& QueryPlatformData = Vertex->PlatformData[PlatformIndex];
+			if (QueryPlatformData.bIterativelyUnmodified.IsSet() && !*QueryPlatformData.bIterativelyUnmodified)
+			{
+				bIterativelyModified = true;
+				break;
+			}
+		}
 		Algo::Sort(Vertex->IterativelyModifiedListeners);
 		Vertex->IterativelyModifiedListeners.SetNum(Algo::Unique(Vertex->IterativelyModifiedListeners));
 		for (FVertexData* ListenerVertex : Vertex->IterativelyModifiedListeners)
 		{
-			GraphSearch.KickVertex(ListenerVertex);
+			if (!ListenerVertex->bWaitingOnUnreadyDependencies)
+			{
+				continue;
+			}
+
+			ListenerVertex->UnreadyDependencies.Remove(Vertex);
+			if (bIterativelyModified || ListenerVertex->UnreadyDependencies.IsEmpty())
+			{
+				ListenerVertex->bWaitingOnUnreadyDependencies = false;
+				GraphSearch.KickVertex(ListenerVertex);
+			}
 		}
 		Vertex->IterativelyModifiedListeners.Empty();
 	}
diff --git a/Engine/Source/Editor/UnrealEd/Private/Cooker/CookRequestCluster.h b/Engine/Source/Editor/UnrealEd/Private/Cooker/CookRequestCluster.h
--- a/Engine/Source/Editor/UnrealEd/Private/Cooker/CookRequestCluster.h
+++ b/Engine/Source/Editor/UnrealEd/Private/Cooker/CookRequestCluster.h
@@ -144,9 +144,11 @@ private:
 		/* Async thread is not allowed to access PackageData, so store its name.The name is immutable for vertex lifetime. */
 		FName PackageName;
 		TArray<FVertexData*> IterativelyModifiedListeners;
+		TSet<FVertexData*> UnreadyDependencies;
 		UE::Cook::FPackageData* PackageData = nullptr;
 		bool bAnyCookable = true;
 		bool bPulledIntoCluster = false;
+		bool bWaitingOnUnreadyDependencies = false;
 		/** Settings and Results for each of the GraphSearch's FetchPlatforms. Element n corresponds to FetchPlatform n. */
 		TUniquePtr<FQueryPlatformData[]> PlatformData;
 	};
