IncrementalCook: BuildResultProjections and TObjectPtr.: Add mechanism for opting out of the default transitive build dependency on the target package when a TObjectPtr is dereferenced. The opt out mechanism includes a way to instead add a dependency on a cook result projection declared by the package; currently we only support a generic cook result projection that adds a dependency on the package and the object's class.
Change FCookDependency::NativeClass to support finding the closest native ancestor rather than giving an error if the class is non-native; this is needed for the PackageAndClass BuildResultProject dependency.
#rnx

diff --git a/Engine/Source/Editor/UnrealEd/Private/Cooker/PackageBuildDependencyTracker.cpp b/Engine/Source/Editor/UnrealEd/Private/Cooker/PackageBuildDependencyTracker.cpp
--- a/Engine/Source/Editor/UnrealEd/Private/Cooker/PackageBuildDependencyTracker.cpp
+++ b/Engine/Source/Editor/UnrealEd/Private/Cooker/PackageBuildDependencyTracker.cpp
@@ -4,8 +4,12 @@
 #include "Misc/ConfigAccessTracking.h"
 
 #if UE_WITH_PACKAGE_ACCESS_TRACKING
+#include "AssetRegistry/IAssetRegistry.h"
+#include "AssetRegistry/AssetData.h"
 #include "Async/UniqueLock.h"
 #include "Cooker/CookConfigAccessTracker.h"
+#include "Cooker/CookDependency.h"
+#include "CookOnTheSide/CookLog.h"
 #include "HAL/Platform.h"
 #include "Logging/LogMacros.h"
 #include "Misc/PackageAccessTrackingOps.h"
@@ -44,10 +48,10 @@ void FPackageBuildDependencyTracker::DumpStats() const
 	UE::TUniqueLock RecordsScopeLock(RecordsLock);
 	uint64 ReferencingPackageCount = 0;
 	uint64 ReferenceCount = 0;
-	for (const TPair<FName, TSet<FBuildDependencyAccessData>>& PackageAccessRecord : Records)
+	for (const TPair<FName, TMap<FBuildDependencyAccessData, FResultProjectionList>>& PackageAccessRecord : Records)
 	{
 		++ReferencingPackageCount;
-		for (const FBuildDependencyAccessData& AccessedData : PackageAccessRecord.Value)
+		for (const TPair<FBuildDependencyAccessData, FResultProjectionList>& AccessedData : PackageAccessRecord.Value)
 		{
 			++ReferenceCount;
 		}
@@ -58,26 +62,26 @@ void FPackageBuildDependencyTracker::DumpStats() const
 	if (bDetailedDump)
 	{
 		UE_LOG(LogPackageBuildDependencyTracker, Display, TEXT("========================================================================="));
-		for (const TPair<FName, TSet<FBuildDependencyAccessData>>& PackageAccessRecord : Records)
+		for (const TPair<FName, TMap<FBuildDependencyAccessData, FResultProjectionList>>& PackageAccessRecord : Records)
 		{
 			UE_LOG(LogPackageBuildDependencyTracker, Display, TEXT("%s:"), *PackageAccessRecord.Key.ToString());
-			for (const FBuildDependencyAccessData& AccessedData : PackageAccessRecord.Value)
+			for (const TPair<FBuildDependencyAccessData, FResultProjectionList>& AccessedData : PackageAccessRecord.Value)
 			{
-				UE_LOG(LogPackageBuildDependencyTracker, Display, TEXT("    %s"), *AccessedData.ReferencedPackage.ToString());
+				UE_LOG(LogPackageBuildDependencyTracker, Display, TEXT("    %s"), *AccessedData.Key.ReferencedPackage.ToString());
 			}
 		}
 	}
 }
 
-TArray<FBuildDependencyAccessData> FPackageBuildDependencyTracker::GetAccessDatas(FName ReferencerPackage) const
+TArray<TPair<FBuildDependencyAccessData, FResultProjectionList>> FPackageBuildDependencyTracker::GetAccessDatas(FName ReferencerPackage) const
 {
 	UE::TUniqueLock RecordsScopeLock(Singleton.RecordsLock);
-	const TSet<FBuildDependencyAccessData>* ReferencerSet = Records.Find(ReferencerPackage);
-	if (!ReferencerSet)
+	const TMap<FBuildDependencyAccessData, FResultProjectionList>* ReferencerMap = Records.Find(ReferencerPackage);
+	if (!ReferencerMap)
 	{
-		return TArray<FBuildDependencyAccessData>();
+		return TArray<TPair<FBuildDependencyAccessData, FResultProjectionList>>();
 	}
-	return ReferencerSet->Array();
+	return ReferencerMap->Array();
 }
 
 FPackageBuildDependencyTracker::FPackageBuildDependencyTracker()
@@ -99,6 +103,8 @@ static bool ShouldSkipDependency(const UObject* Object)
 		(Object->GetClass() == UClass::StaticClass());
 }
 
+static FName NAME_EngineTransient(TEXT("/Engine/Transient"));
+
 void FPackageBuildDependencyTracker::StaticOnObjectHandleRead(const TArrayView<const UObject* const>& Objects)
 {
 	int Count = Objects.Num();
@@ -113,7 +119,12 @@ void FPackageBuildDependencyTracker::StaticOnObjectHandleRead(const TArrayView<c
 		return;
 	}
 
-	if (AccumulatedScopeData->BuildOpName.IsNone())
+	FName Referencer = AccumulatedScopeData->PackageName;
+	FName CookResultProjection = AccumulatedScopeData->CookResultProjection;
+	if (AccumulatedScopeData->BuildOpName.IsNone()
+		| CookResultProjection == UE::Cook::ResultProjection::None
+		| Referencer.IsNone()
+		| Referencer == NAME_EngineTransient)
 	{
 		return;
 	}
@@ -125,10 +136,10 @@ void FPackageBuildDependencyTracker::StaticOnObjectHandleRead(const TArrayView<c
 			continue;
 		}
 
-		FName Referencer = AccumulatedScopeData->PackageName;
 		UPackage* ReferencedPackage = ReadObject->GetOutermost();
 		FName Referenced = ReferencedPackage->GetFName();
-		if ((Referencer == Referenced) || ReferencedPackage->HasAnyPackageFlags(PKG_CompiledIn))
+		if ((Referencer == Referenced) || ReferencedPackage->HasAnyPackageFlags(PKG_CompiledIn)
+			|| Referenced == NAME_EngineTransient)
 		{
 			continue;
 		}
@@ -143,24 +154,95 @@ void FPackageBuildDependencyTracker::StaticOnObjectHandleRead(const TArrayView<c
 
 		FBuildDependencyAccessData AccessData{ Referenced, AccumulatedScopeData->TargetPlatform };
 		UE::TUniqueLock RecordsScopeLock(Singleton.RecordsLock);
+
+		bool bNeedsAdd = true;
 		if (Referencer == Singleton.LastReferencer)
 		{
-			if (AccessData != Singleton.LastAccessData)
+			if (AccessData == Singleton.LastAccessData)
 			{
+				if (CookResultProjection == UE::Cook::ResultProjection::All
+					&& Singleton.LastCookResultProjection == UE::Cook::ResultProjection::All)
+				{
+					bNeedsAdd = false;
+				}
+				else
+				{
+					Singleton.LastCookResultProjection = CookResultProjection;
+				}
+			}
+			else
+			{
+				Singleton.LastCookResultProjection = CookResultProjection;
 				Singleton.LastAccessData = AccessData;
-				Singleton.LastReferencerSet->Add(AccessData);
 			}
 		}
 		else
 		{
+			Singleton.LastCookResultProjection = CookResultProjection;
 			Singleton.LastAccessData = AccessData;
 			Singleton.LastReferencer = Referencer;
-			Singleton.LastReferencerSet = &Singleton.Records.FindOrAdd(Referencer);
-			Singleton.LastReferencerSet->Add(AccessData);
+			Singleton.LastReferencerMap = &Singleton.Records.FindOrAdd(Referencer);
+		}
+		if (bNeedsAdd)
+		{
+			if (CookResultProjection == UE::Cook::ResultProjection::All)
+			{
+				Singleton.LastReferencerMap->FindOrAdd(AccessData).AddProjectionAll();
+			}
+			else
+			{
+				Singleton.LastReferencerMap->FindOrAdd(AccessData).AddProjection(CookResultProjection, ReadObject->GetClass()->GetClassPathName());
+			}
+
 		}
 	}
 }
-#endif // UE_WITH_OBJECT_HANDLE_TRACKING
+
+void FResultProjectionList::AddProjectionAll(bool* bOutExists)
+{
+	if (bHasAll)
+	{
+		if (bOutExists)
+		{
+			*bOutExists = true;
+		}
+		return;
+	}
+	bHasAll = true;
+	Classes.Empty();
+	ResultProjections.Empty();
+	if (bOutExists)
+	{
+		*bOutExists = false;
+	}
+}
+
+void FResultProjectionList::AddProjection(FName CookResultProjection, FTopLevelAssetPath ClassPath, bool* bOutExists)
+{
+	if (bHasAll)
+	{
+		if (bOutExists)
+		{
+			*bOutExists = true;
+		}
+		return;
+	}
+
+	if (CookResultProjection == UE::Cook::ResultProjection::PackageAndClass)
+	{
+		Classes.Add(ClassPath, bOutExists);
+		ResultProjections.Add(CookResultProjection);
+	}
+	else
+	{
+		// ResultProjection::None and ResultProjection::All should have been handled by caller and this function not called
+		check(CookResultProjection != UE::Cook::ResultProjection::None);
+		check(CookResultProjection != UE::Cook::ResultProjection::All);
+		ResultProjections.Add(CookResultProjection, bOutExists);
+	}
+}
+
+#endif // UE_WITH_PACKAGE_ACCESS_TRACKING
 
 void DumpBuildDependencyTrackerStats()
 {
diff --git a/Engine/Source/Editor/UnrealEd/Private/Cooker/PackageBuildDependencyTracker.h b/Engine/Source/Editor/UnrealEd/Private/Cooker/PackageBuildDependencyTracker.h
--- a/Engine/Source/Editor/UnrealEd/Private/Cooker/PackageBuildDependencyTracker.h
+++ b/Engine/Source/Editor/UnrealEd/Private/Cooker/PackageBuildDependencyTracker.h
@@ -30,6 +30,16 @@ struct FBuildDependencyAccessData
 	}
 };
 
+struct FResultProjectionList
+{
+	TSet<FTopLevelAssetPath> Classes;
+	TSet<FName> ResultProjections;
+	bool bHasAll = false;
+
+	void AddProjectionAll(bool* bOutExists = nullptr);
+	void AddProjection(FName CookResultProjection, FTopLevelAssetPath ClassPath, bool* bOutExists = nullptr);
+};
+
 class FPackageBuildDependencyTracker: public FNoncopyable
 {
 public:
@@ -38,7 +48,7 @@ public:
 	void Disable();
 	bool IsEnabled() const;
 	void DumpStats() const;
-	TArray<FBuildDependencyAccessData> GetAccessDatas(FName ReferencerPackage) const;
+	TArray<TPair<FBuildDependencyAccessData, FResultProjectionList>> GetAccessDatas(FName ReferencerPackage) const;
 
 private:
 	FPackageBuildDependencyTracker();
@@ -52,10 +62,11 @@ private:
 	// cook time for shader invalidation on a large project.
 	mutable UE::FMutex RecordsLock;
 
-	TMap<FName, TSet<FBuildDependencyAccessData>> Records;
+	TMap<FName, TMap<FBuildDependencyAccessData, FResultProjectionList>> Records;
 	FName LastReferencer = NAME_None;
 	FBuildDependencyAccessData LastAccessData{ NAME_None, nullptr };
-	TSet<FBuildDependencyAccessData>* LastReferencerSet = nullptr;
+	FName LastCookResultProjection = NAME_None;
+	TMap<FBuildDependencyAccessData, FResultProjectionList>* LastReferencerMap = nullptr;
 	UE::CoreUObject::FObjectHandleTrackingCallbackId ObjectHandleReadHandle;
 	bool bEnabled = false;
 	static FPackageBuildDependencyTracker Singleton;
diff --git a/Engine/Source/Editor/UnrealEd/Private/TargetDomain/TargetDomainUtils.cpp b/Engine/Source/Editor/UnrealEd/Private/TargetDomain/TargetDomainUtils.cpp
--- a/Engine/Source/Editor/UnrealEd/Private/TargetDomain/TargetDomainUtils.cpp
+++ b/Engine/Source/Editor/UnrealEd/Private/TargetDomain/TargetDomainUtils.cpp
@@ -300,7 +300,8 @@ EPackageMountPoint GetPackageMountPoint(FName PackageName, FName TransientPackag
 	}
 	TStringBuilder<256> StringBuffer;
 	PackageName.ToString(StringBuffer);
-	if (FPackageName::IsMemoryPackage(StringBuffer))
+	if (FPackageName::IsMemoryPackage(StringBuffer) ||
+		FPackageName::IsTempPackage(StringBuffer))
 	{
 		return EPackageMountPoint::Transient;
 	}
@@ -339,7 +340,13 @@ FCookDependencies FCookDependencies::Collect(UPackage* Package, const ITargetPla
 
 	FCookDependencies Result;
 	Result.PackageName = Package->GetFName();
-	TSet<FName> BuildDependenciesSet;
+	enum class EBuildDependency
+	{
+		PackageOnly,
+		Transitive,
+	};
+
+	TMap<FName, EBuildDependency> PackageBuildDependencies;
 
 	TArray<FName> AssetDependencies;
 	if (!GeneratedResult)
@@ -359,22 +366,65 @@ FCookDependencies FCookDependencies::Collect(UPackage* Package, const ITargetPla
 
 	if (Tracker.IsEnabled())
 	{
-		TArray<FBuildDependencyAccessData> AccessDatas = Tracker.GetAccessDatas(Result.PackageName);
+		TArray<TPair<FBuildDependencyAccessData, FResultProjectionList>> AccessDatas
+			= Tracker.GetAccessDatas(Result.PackageName);
 
-		BuildDependenciesSet.Reserve(AccessDatas.Num());
-		for (FBuildDependencyAccessData& AccessData : AccessDatas)
+		PackageBuildDependencies.Reserve(AccessDatas.Num());
+		for (TPair<FBuildDependencyAccessData, FResultProjectionList>& Pair : AccessDatas)
 		{
+			FBuildDependencyAccessData& AccessData = Pair.Key;
 			if (AccessData.TargetPlatform == TargetPlatform || AccessData.TargetPlatform == nullptr)
 			{
-				BuildDependenciesSet.Add(AccessData.ReferencedPackage);
+				FResultProjectionList& ProjectionList = Pair.Value;
+				constexpr bool bAutoTransitiveDependenciesEnabled = false;
+				if (!bAutoTransitiveDependenciesEnabled)
+				{
+					// We have not yet enabled marking the auto-added dependencies from TObjectPtr resolve
+					// as transitive, because it causes a performance regression and we are still working on fixing the
+					// regression.
+					PackageBuildDependencies.FindOrAdd(AccessData.ReferencedPackage, EBuildDependency::PackageOnly);
+				}
+				else if (ProjectionList.bHasAll)
+				{
+					// TObjectPtr UE::BuildProjection::All dependencies are added as transitive build dependencies.
+					// We have to do this to be conservative, since we do not know which bytes from the target are
+					// dependended upon and which of the target's build dependencies influence those bytes.
+					PackageBuildDependencies.FindOrAdd(AccessData.ReferencedPackage) = EBuildDependency::Transitive;
+				}
+				else
+				{
+					for (FTopLevelAssetPath ClassPath : ProjectionList.Classes)
+					{
+						Result.ClassDependencies.Add(ClassPath.ToString());
+					}
+
+					for (FName ResultProjection : ProjectionList.ResultProjections)
+					{
+						if (ResultProjection == UE::Cook::ResultProjection::PackageAndClass)
+						{
+							PackageBuildDependencies.FindOrAdd(AccessData.ReferencedPackage, EBuildDependency::PackageOnly);
+						}
+						else
+						{
+							UE_LOG(LogCook, Error,
+								TEXT("When saving %s, found ResultProjection %s, which is a system-specific ResultProjection, and this is not yet implemented. ")
+								TEXT("Find the call to UE_COOK_RESULTPROJECTION_SCOPED passing in this name and remove it."),
+								*Result.PackageName.ToString(), *ResultProjection.ToString());
+						}
+					}
+				}
 			}
 		}
 	}
 	else
 #endif
 	{
-		// Defensively treat all asset dependencies as both build and runtime dependencies
-		BuildDependenciesSet.Append(AssetDependencies);
+		// When PackageAccessTracking is disabled, defensively treat all asset dependencies as transitive build
+		// dependencies.
+		for (FName AssetDependency : AssetDependencies)
+		{
+			PackageBuildDependencies.FindOrAdd(AssetDependency) = EBuildDependency::Transitive;
+		}
 	}
 
 #if UE_WITH_CONFIG_TRACKING
@@ -441,11 +491,12 @@ FCookDependencies FCookDependencies::Collect(UPackage* Package, const ITargetPla
 					{
 						RuntimeDependencies.Add(DependencyPackageName);
 					}
-					Result.TransitiveBuildDependencies.Add(MoveTemp(CookDependency));
+					PackageBuildDependencies.FindOrAdd(DependencyPackageName) = EBuildDependency::Transitive;
 					break;
 				}
 				case UE::Cook::ECookDependency::Package:
-					BuildDependenciesSet.Add(CookDependency.GetPackageName());
+					// Add it as a PackageOnly build dependency (lowest level) unless a higher level already exists.
+					PackageBuildDependencies.FindOrAdd(CookDependency.GetPackageName(), EBuildDependency::PackageOnly);
 					break;
 				case UE::Cook::ECookDependency::Config:
 					Result.ConfigDependencies.Add(FString(CookDependency.GetConfigPath()));
@@ -499,22 +550,37 @@ FCookDependencies FCookDependencies::Collect(UPackage* Package, const ITargetPla
 		}
 	}
 
+	FName TransientPackageName = GetTransientPackage()->GetFName();
 	Algo::Sort(Result.CookDependencies);
 	Result.CookDependencies.SetNum(Algo::Unique(Result.CookDependencies));
+	for (TPair<FName, EBuildDependency>& Pair : PackageBuildDependencies)
+	{
+		if (Pair.Key == Result.PackageName)
+		{
+			// Skip dependencies to self, for performance
+			continue;
+		}
+		if (GetPackageMountPoint(Pair.Key, TransientPackageName) != EPackageMountPoint::Content)
+		{
+			// Skip script and temp packages
+			continue;
+		}
+		if (Pair.Value == EBuildDependency::Transitive)
+		{
+			Result.TransitiveBuildDependencies.Add(UE::Cook::FCookDependency::Package(Pair.Key));
+		}
+		else
+		{
+			Result.BuildPackageDependencies.Add(Pair.Key);
+		}
+	}
 	Algo::Sort(Result.TransitiveBuildDependencies);
-	Result.TransitiveBuildDependencies.SetNum(Algo::Unique(Result.TransitiveBuildDependencies));
+	Result.BuildPackageDependencies.Sort(FNameLexicalLess());
 	Algo::Sort(Result.ConfigDependencies);
 	Result.ConfigDependencies.SetNum(Algo::Unique(Result.ConfigDependencies));
 	Algo::Sort(Result.ClassDependencies);
 	Result.ClassDependencies.SetNum(Algo::Unique(Result.ClassDependencies));
 
-	FName TransientPackageName = GetTransientPackage()->GetFName();
-	Result.BuildPackageDependencies = BuildDependenciesSet.Array();
-	Result.BuildPackageDependencies.RemoveAllSwap([TransientPackageName](FName InPackageName)
-		{
-			return GetPackageMountPoint(InPackageName, TransientPackageName) != EPackageMountPoint::Content;
-		}, EAllowShrinking::Yes);
-	Result.BuildPackageDependencies.Sort(FNameLexicalLess());
 
 	for (TArray<FName>::TIterator Iter(RuntimeDependencies); Iter; ++Iter)
 	{
diff --git a/Engine/Source/Runtime/CoreUObject/Private/Cooker/CookDependency.cpp b/Engine/Source/Runtime/CoreUObject/Private/Cooker/CookDependency.cpp
--- a/Engine/Source/Runtime/CoreUObject/Private/Cooker/CookDependency.cpp
+++ b/Engine/Source/Runtime/CoreUObject/Private/Cooker/CookDependency.cpp
@@ -42,6 +42,15 @@ const TMap<FName, FCookDependencyFunction>& GetCookDependencyFunctions();
 namespace UE::Cook
 {
 
+namespace ResultProjection
+{
+
+FName All(TEXT("UE_Cook_ResultProjection_All"));
+FName PackageAndClass(TEXT("UE_Cook_ResultProjection_PackageAndClass"));
+FName None(TEXT("UE_Cook_ResultProjection_None"));
+
+}
+
 FCookDependency FCookDependency::File(FStringView InFileName)
 {
 	FCookDependency Result(ECookDependency::File);
@@ -152,15 +161,17 @@ FCookDependency FCookDependency::NativeClass(const UClass* InClass)
 {
 	if (InClass)
 	{
-		if (!InClass->IsNative())
+		while (!InClass->IsNative())
 		{
-			UE_LOG(LogCore, Error, TEXT("Invalid FCookDependency::NativeClass(%s). The class is not native. ")
-				TEXT("The dependency will be ignored."),
-				InClass ? *InClass->GetPathName() : TEXT("<null>"));
-			InClass = nullptr;
+			InClass = InClass->GetSuperClass();
+			check(InClass); // Every class other than UObject has a super class, and UObject is native.
 		}
+		return NativeClass(InClass->GetPathName());
+	}
+	else
+	{
+		return NativeClass(FStringView());
 	}
-	return NativeClass(InClass ? InClass->GetPathName() : FStringView());
 }
 
 FCookDependency FCookDependency::NativeClass(FStringView ClassPath)
diff --git a/Engine/Source/Runtime/CoreUObject/Private/Misc/PackageAccessTracking.cpp b/Engine/Source/Runtime/CoreUObject/Private/Misc/PackageAccessTracking.cpp
--- a/Engine/Source/Runtime/CoreUObject/Private/Misc/PackageAccessTracking.cpp
+++ b/Engine/Source/Runtime/CoreUObject/Private/Misc/PackageAccessTracking.cpp
@@ -1,6 +1,8 @@
 // Copyright Epic Games, Inc. All Rights Reserved.
 
 #include "Misc/PackageAccessTracking.h"
+
+#include "Cooker/CookDependency.h"
 #include "Misc/PackageAccessTrackingOps.h"
 
 #if UE_WITH_PACKAGE_ACCESS_TRACKING
@@ -23,11 +25,12 @@ bool IsBuildOpName(FName OpName)
 		;
 }
 
-FTrackedData::FTrackedData(FName InPackageName, FName InOpName, const ITargetPlatform* InTargetPlatform,
-	const UObject* InObject)
+FTrackedData::FTrackedData(FName InPackageName, FName InOpName, FName InCookResultProjection,
+	const ITargetPlatform* InTargetPlatform, const UObject* InObject)
 	: PackageName(InPackageName)
 	, OpName(InOpName)
 	, BuildOpName(IsBuildOpName(InOpName) ? InOpName : NAME_None)
+	, CookResultProjection(InCookResultProjection)
 	, Object(InObject)
 	, TargetPlatform(InTargetPlatform)
 {
@@ -50,6 +53,9 @@ FTrackedData::FTrackedData(FTrackedData& DirectData, FTrackedData* OuterAccumula
 		// The accumulated OpName is the most recent non-empty OpName
 		OpName = !DirectData.OpName.IsNone() ? DirectData.OpName : OuterAccumulatedData->OpName;
 
+		CookResultProjection = !DirectData.CookResultProjection.IsNone() ?
+			DirectData.CookResultProjection : OuterAccumulatedData->CookResultProjection;
+
 		// The accumulated BuildOpName is by default the most recent non-empty BuildOpName
 		// ResetContext ops set it back to none
 		if (OpName == PackageAccessTrackingOps::NAME_ResetContext)
@@ -72,32 +78,62 @@ FTrackedData::FTrackedData(FTrackedData& DirectData, FTrackedData* OuterAccumula
 		PackageName = DirectData.PackageName;
 		OpName = DirectData.OpName;
 		BuildOpName = DirectData.BuildOpName;
+		CookResultProjection = DirectData.CookResultProjection;
 		Object = DirectData.Object;
 		TargetPlatform = DirectData.TargetPlatform;
 	}
 }
 
 FPackageAccessRefScope::FPackageAccessRefScope(FName InPackageName, FName InOpName,
-	const ITargetPlatform* InTargetPlatform, const UObject* InObject)
-	: DirectData(InPackageName, InOpName, InTargetPlatform, InObject)
+	FName InCookResultProjection, const ITargetPlatform* InTargetPlatform, const UObject* InObject)
+	: DirectData(InPackageName, InOpName, InCookResultProjection, InTargetPlatform, InObject)
 	, AccumulatedData(DirectData, GetCurrentThreadAccumulatedData())
 	, Outer(CurrentThreadScope)
 {
 	CurrentThreadScope = this;
 }
 
+FPackageAccessRefScope::FPackageAccessRefScope(FName InPackageName, FName InOpName)
+	: FPackageAccessRefScope(InPackageName, InOpName
+#if WITH_EDITOR
+	, UE::Cook::ResultProjection::All
+#else
+	, NAME_None
+#endif
+	, nullptr, nullptr)
+{
+}
+
 FPackageAccessRefScope::FPackageAccessRefScope(const UObject* InObject, FName InOpName)
-	: FPackageAccessRefScope(InObject->GetPackage()->GetFName(), InOpName, nullptr, InObject)
+	: FPackageAccessRefScope(InObject->GetPackage()->GetFName(), InOpName
+#if WITH_EDITOR
+	, UE::Cook::ResultProjection::All
+#else
+	, NAME_None
+#endif
+	, nullptr, InObject)
 {
 }
 
 FPackageAccessRefScope::FPackageAccessRefScope(FName InOpName)
-	: FPackageAccessRefScope(NAME_None, InOpName, nullptr)
+	: FPackageAccessRefScope(NAME_None, InOpName
+#if WITH_EDITOR
+	, UE::Cook::ResultProjection::All
+#else
+	, NAME_None
+#endif
+	, nullptr, nullptr)
 {
 }
 
 FPackageAccessRefScope::FPackageAccessRefScope(const ITargetPlatform* InTargetPlatform)
-	: FPackageAccessRefScope(NAME_None, NAME_None, InTargetPlatform)
+	: FPackageAccessRefScope(NAME_None, NAME_None, NAME_None, InTargetPlatform, nullptr)
+{
+}
+
+FPackageAccessRefScope::FPackageAccessRefScope(FPackageAccessRefScope::ECookResultProjectionType,
+	FName InCookResultProjection)
+	: FPackageAccessRefScope(NAME_None, NAME_None, InCookResultProjection, nullptr, nullptr)
 {
 }
 
diff --git a/Engine/Source/Runtime/CoreUObject/Public/Cooker/CookDependency.h b/Engine/Source/Runtime/CoreUObject/Public/Cooker/CookDependency.h
--- a/Engine/Source/Runtime/CoreUObject/Public/Cooker/CookDependency.h
+++ b/Engine/Source/Runtime/CoreUObject/Public/Cooker/CookDependency.h
@@ -6,6 +6,7 @@
 #include "Containers/StringView.h"
 #include "Containers/UnrealString.h"
 #include "HAL/PreprocessorHelpers.h"
+#include "Misc/PackageAccessTracking.h"
 #include "Serialization/CompactBinary.h"
 #include "Templates/Function.h"
 #include "Templates/UniquePtr.h"
@@ -19,12 +20,27 @@ namespace UE::ConfigAccessTracking { struct FConfigAccessData; }
 struct FARFilter;
 #endif
 
+#if WITH_EDITOR && UE_WITH_PACKAGE_ACCESS_TRACKING
+/**
+ * Specify that within the current scope, dereferenced TObjectPtrs cause a dependency on the specified
+ * ResultProjection of the dereferenced object and package, instead of the outer scope's projection or the default
+ * projection of UE::CooK::ResultProjection::All. ProjectionName must be either one of the generic terms provided
+ * in UE::Cook::ResultProjection or it must be a name declared as a resultprojection by one of the objects in
+ * the target package.
+ */
+#define UE_COOK_RESULTPROJECTION_SCOPED(ProjectionName) \
+	PackageAccessTracking_Private::FPackageAccessRefScope ANONYMOUS_VARIABLE(PackageAccessTracker_)( \
+		PackageAccessTracking_Private::FPackageAccessRefScope::CookResultProjectionType, ProjectionName)
+#else
+#define UE_COOK_RESULTPROJECTION_SCOPED(ProjectionName)
+#endif
+
 #if WITH_EDITOR
 
 namespace UE::Cook
 {
 
-/** Context passed into FCookDependencyFunction to provide calling flags and receive their hash output. */
+/** Context passed into UE_COOK_DEPENDENCY_FUNCTION functions to provide calling flags and receive hash output. */
 struct FCookDependencyContext
 {
 public:
@@ -99,6 +115,34 @@ enum class ECookDependency : uint8
 	Count,
 };
 
+namespace ResultProjection
+{
+/**
+ * FName passed into FCookDependency::ResultProjection or UE_COOK_RESULTPROJECTION_SCOPED to indicate
+ * the entire target package can contribute to the source package's cook. This is the default projection
+ * used TObjectPtr resolves when no more limited scope has been declared. It causes a transitive dependency;
+ * all cook dependencies declared by the target package are also used as dependencies of the source package.
+ */
+extern COREUOBJECT_API FName All;
+/*
+ * FName passed into FCookDependency::ResultProjection or UE_COOK_RESULTPROJECTION_SCOPED to indicate the bytes
+ * within the target package, and the native class used by objects referenced from it, can contribute the the
+ * source package's cook. This is more limited than the transitive dependency indicated by BuildProjectionAll
+ * because it does not include e.g. other packages that the target package depends on for its cook.
+ * This ResultProjection is applied automatically for all hard imports saved into the source cooked package;
+ * the SavePackage code applies it because generic LinkerLoad code can cause those imports to be included
+ * or excluded from the package based on values of e.g. NeedsLoadForClient.
+ */
+extern COREUOBJECT_API FName PackageAndClass;
+/*
+ * FName passed into UE_COOK_RESULTPROJECTION_SCOPED to indicate that dereferences of TObjectPtr within the scope
+ * should not be automatically added as dependencies of any kind. The calling code for those scopes either
+ * knows that the resolves are spurious or is declaring the proper ResultProjection dependency manually.
+ */
+extern COREUOBJECT_API FName None;
+
+}	
+
 /**
  * TargetDomain dependencies that can be reported from the class instances in a package. These dependencies are
  * stored in the cook oplog and are evaluated during incremental cook. If any of them changes, the package is
@@ -345,6 +389,7 @@ struct FCookDependencyFunctionRegistration
 
 #else // WITH_EDITOR
 
+#define UE_COOK_RESULTPROJECTION_SCOPED(ProjectionName)
 #define UE_COOK_DEPENDENCY_FUNCTION(Name, Function)
 #define UE_COOK_DEPENDENCY_FUNCTION_CALL(Name) NAME_None
 
diff --git a/Engine/Source/Runtime/CoreUObject/Public/Misc/PackageAccessTracking.h b/Engine/Source/Runtime/CoreUObject/Public/Misc/PackageAccessTracking.h
--- a/Engine/Source/Runtime/CoreUObject/Public/Misc/PackageAccessTracking.h
+++ b/Engine/Source/Runtime/CoreUObject/Public/Misc/PackageAccessTracking.h
@@ -42,14 +42,15 @@ namespace PackageAccessTracking_Private
 	struct FTrackedData
 	{
 		/** Standard constructor; sets variables from DirectData passed to a FPackageAccessRefScope. */
-		COREUOBJECT_API FTrackedData(FName PackageName, FName OpName, const ITargetPlatform* InTargetPlatform,
-			const UObject* InObject);
+		COREUOBJECT_API FTrackedData(FName PackageNameIn, FName InOpName, FName InCookResultProjection,
+			const ITargetPlatform* InTargetPlatform, const UObject* InObject);
 		/** Accumulating constructor; sets new AccumulatedData by combining DirectData with Outer's AccumulatedData. */
 		COREUOBJECT_API FTrackedData(FTrackedData& DirectData, FTrackedData* OuterAccumulatedData);
 
 		FName PackageName;
 		FName OpName;
 		FName BuildOpName;
+		FName CookResultProjection;
 		const UObject* Object = nullptr;
 		const ITargetPlatform* TargetPlatform = nullptr;
 	};
@@ -57,11 +58,15 @@ namespace PackageAccessTracking_Private
 	class FPackageAccessRefScope
 	{
 	public:
-		COREUOBJECT_API FPackageAccessRefScope(FName InPackageName, FName InOpName,
-			const ITargetPlatform* InTargetPlatform=nullptr, const UObject* InObject=nullptr);
+		COREUOBJECT_API FPackageAccessRefScope(FName InPackageName, FName InOpName);
 		COREUOBJECT_API FPackageAccessRefScope(const UObject* InObject, FName InOpName);
 		COREUOBJECT_API FPackageAccessRefScope(FName InOpName);
 		COREUOBJECT_API FPackageAccessRefScope(const ITargetPlatform* InTargetPlatform);
+		enum ECookResultProjectionType
+		{
+			CookResultProjectionType,
+		};
+		COREUOBJECT_API FPackageAccessRefScope(ECookResultProjectionType, FName InCookResultProjection);
 
 		COREUOBJECT_API ~FPackageAccessRefScope();
 
@@ -75,6 +80,9 @@ namespace PackageAccessTracking_Private
 		static COREUOBJECT_API FPackageAccessRefScope* GetCurrentThreadScope();
 		static COREUOBJECT_API FTrackedData* GetCurrentThreadAccumulatedData();
 	private:
+		COREUOBJECT_API FPackageAccessRefScope(FName InPackageName, FName InOpName, FName InCookResultProjection,
+			const ITargetPlatform* InTargetPlatform, const UObject* InObject);
+
 		FTrackedData DirectData;
 		FTrackedData AccumulatedData;
 		FPackageAccessRefScope* Outer = nullptr;
