Cooker: Suppress static analysis warning V792: The 'IsNone' function located to the right of the operator '|' will be called regardless of the value of the left operand. Perhaps, it is better to use '||'.
The non-conditional evaluation is intentional, to avoid a branch.
#rnx

diff --git a/Engine/Source/Editor/UnrealEd/Private/Cooker/PackageBuildDependencyTracker.cpp b/Engine/Source/Editor/UnrealEd/Private/Cooker/PackageBuildDependencyTracker.cpp
--- a/Engine/Source/Editor/UnrealEd/Private/Cooker/PackageBuildDependencyTracker.cpp
+++ b/Engine/Source/Editor/UnrealEd/Private/Cooker/PackageBuildDependencyTracker.cpp
@@ -122,9 +122,9 @@ void FPackageBuildDependencyTracker::StaticOnObjectHandleRead(const TArrayView<c
 	FName Referencer = AccumulatedScopeData->PackageName;
 	FName CookResultProjection = AccumulatedScopeData->CookResultProjection;
 	if (AccumulatedScopeData->BuildOpName.IsNone()
-		| CookResultProjection == UE::Cook::ResultProjection::None
-		| Referencer.IsNone()
-		| Referencer == NAME_EngineTransient)
+		| CookResultProjection == UE::Cook::ResultProjection::None // -V792
+		| Referencer.IsNone() // -V792
+		| Referencer == NAME_EngineTransient) // -V792
 	{
 		return;
 	}
