Cooker: Implement budgeted time-based triggering of SoftGC instead of relying on the memory threshold approach, because the memory threshold approach is stymied by our lack of knowledge of how much memory is buffered inside the memory allocator.
[RN] Minor, Cooking

diff --git a/Engine/Config/BaseEditor.ini b/Engine/Config/BaseEditor.ini
--- a/Engine/Config/BaseEditor.ini
+++ b/Engine/Config/BaseEditor.ini
@@ -370,11 +370,18 @@ MaxConcurrentShaderJobs=8192
 MaxDiffsToLog=5
 IgnoreHeaderDiffs=True
 
-; SoftGC triggers a partial garbage collect when available physical memory on the machine reaches configured milestones
-; e.g. 5/10, 4/10, 3/10, 2/10, 1/10.
 bUseSoftGC=true
-SoftGCStartNumerator=5
+; When SoftGCStartNumerator is > 0, a SoftGC garbage collect is triggered when available physical memory
+; on the machine reaches configured milestones e.g. 5/10, 4/10, 3/10, 2/10, 1/10.
+; This mechanism has been replaced by SoftGCTimeFractionBudget and will be removed in UE5.7.
+SoftGCStartNumerator=0
 SoftGCDenominator=10
+; When SoftGCTimeFractionBudget is > 0, a SoftGC garbage collect is triggered as frequently as possible
+; while not going over budget, potentially as often as after every package save. SoftGCTimeFractionBudget
+; specifies what fraction of the time spent loading and saving packages is allowed to be spent running SoftGC.
+; When SoftGC is insufficient and we run out of memory, a full GC will be executed instead and will not be considered
+; part of the SoftGC's time budget.
+SoftGCTimeFractionBudget = 0.05
 
 ; If true, iterative cooks will be allowed even if referenced ini values have changed.
 IterativeIgnoreIni=false
diff --git a/Engine/Source/Editor/UnrealEd/Classes/CookOnTheSide/CookOnTheFlyServer.h b/Engine/Source/Editor/UnrealEd/Classes/CookOnTheSide/CookOnTheFlyServer.h
--- a/Engine/Source/Editor/UnrealEd/Classes/CookOnTheSide/CookOnTheFlyServer.h
+++ b/Engine/Source/Editor/UnrealEd/Classes/CookOnTheSide/CookOnTheFlyServer.h
@@ -187,6 +187,7 @@ namespace UE::Cook
 	class FRequestCluster;
 	class FRequestQueue;
 	class FSaveCookedPackageContext;
+	class FSoftGCHistory;
 	class FWorkerRequestsLocal;
 	class FWorkerRequestsRemote;
 	class ICookOnTheFlyRequestManager;
@@ -365,6 +366,8 @@ private:
 	int64 SoftGCNextAvailablePhysicalTarget = -1;
 	int32 SoftGCStartNumerator = 5;
 	int32 SoftGCDenominator = 10;
+	float SoftGCTimeFractionBudget = 0.05;
+	TUniquePtr<UE::Cook::FSoftGCHistory> SoftGCHistory;
 	bool bUseSoftGC = false;
 	bool bWarnedExceededMaxMemoryWithinGCCooldown = false;
 	bool bGarbageCollectTypeSoft = false;
@@ -483,6 +486,7 @@ private:
 	};
 	void SetIdleStatus(UE::Cook::FTickStackData& StackData, EIdleStatus InStatus);
 	void UpdateDisplay(UE::Cook::FTickStackData& StackData, bool bForceDisplay);
+	FString GetCookSettingsForMemoryLogText() const;
 	enum class ECookAction
 	{
 		Done,			// The cook is complete; no requests remain in any non-idle state
@@ -839,7 +843,8 @@ public:
 		int32 NumObjectsBeforeGC, const FPlatformMemoryStats& MemStatsBeforeGC,
 		const FGenericMemoryStats& AllocatorStatsBeforeGC,
 		int32 NumObjectsAfterGC, const FPlatformMemoryStats& MemStatsAfterGC,
-		const FGenericMemoryStats& AllocatorStatsAfterGC);
+		const FGenericMemoryStats& AllocatorStatsAfterGC,
+		float GCDurationSeconds);
 	UNREALED_API bool NeedsDiagnosticSecondGC() const;
 
 	/**
@@ -1656,6 +1661,7 @@ private:
 	friend UE::Cook::FPlatformManager;
 	friend UE::Cook::FRequestCluster;
 	friend UE::Cook::FSaveCookedPackageContext;
+	friend UE::Cook::FSoftGCHistory;
 	friend UE::Cook::FWorkerRequestsLocal;
 	friend UE::Cook::FWorkerRequestsRemote;
 };
diff --git a/Engine/Source/Editor/UnrealEd/Private/Commandlets/CookCommandlet.cpp b/Engine/Source/Editor/UnrealEd/Private/Commandlets/CookCommandlet.cpp
--- a/Engine/Source/Editor/UnrealEd/Private/Commandlets/CookCommandlet.cpp
+++ b/Engine/Source/Editor/UnrealEd/Private/Commandlets/CookCommandlet.cpp
@@ -653,6 +653,7 @@ void UCookCommandlet::ConditionalCollectGarbage(uint32 TickResults, UCookOnTheFl
 		return;
 	}
 
+	double GCStartTime = FPlatformTime::Seconds();
 	FString GCReason;
 	FString GCType = bPartialGC ? TEXT(" partial gc") : TEXT("");
 	if ((TickResults & UCookOnTheFlyServer::COSR_RequiresGC_PackageCount) != 0)
@@ -749,8 +750,9 @@ void UCookCommandlet::ConditionalCollectGarbage(uint32 TickResults, UCookOnTheFl
 	FLowLevelMemTracker::Get().UpdateStatsPerFrame();
 #endif
 
+	float GCDurationSeconds = FPlatformTime::Seconds() - GCStartTime;
 	bool bWasDueToOOM = (TickResults & UCookOnTheFlyServer::COSR_RequiresGC_OOM) != 0;
 	COTFS.EvaluateGarbageCollectionResults(bWasDueToOOM, bPartialGC, TickResults,
 		NumObjectsBeforeGC, MemStatsBeforeGC, AllocatorStatsBeforeGC,
-		NumObjectsAfterGC, MemStatsAfterGC, AllocatorStatsAfterGC);
+		NumObjectsAfterGC, MemStatsAfterGC, AllocatorStatsAfterGC, GCDurationSeconds);
 }
diff --git a/Engine/Source/Editor/UnrealEd/Private/CookOnTheFlyServer.cpp b/Engine/Source/Editor/UnrealEd/Private/CookOnTheFlyServer.cpp
--- a/Engine/Source/Editor/UnrealEd/Private/CookOnTheFlyServer.cpp
+++ b/Engine/Source/Editor/UnrealEd/Private/CookOnTheFlyServer.cpp
@@ -44,6 +44,7 @@
 #include "Cooker/NetworkFileCookOnTheFlyRequestManager.h"
 #include "Cooker/OnDemandShaderCompilation.h"
 #include "Cooker/PackageTracker.h"
+#include "Cooker/SoftGC.h"
 #include "Cooker/WorkerRequestsLocal.h"
 #include "Cooker/WorkerRequestsRemote.h"
 #include "CookerSettings.h"
@@ -1884,6 +1885,26 @@ void UCookOnTheFlyServer::UpdateDisplay(UE::Cook::FTickStackData& StackData, boo
 	}
 }
 
+FString UCookOnTheFlyServer::GetCookSettingsForMemoryLogText() const
+{
+	return FString::Printf(
+		TEXT("\n\tMemoryMaxUsedVirtual %dMiB")
+		TEXT("\n\tMemoryMaxUsedPhysical %dMiB")
+		TEXT("\n\tMemoryMinFreeVirtual %dMiB")
+		TEXT("\n\tMemoryMinFreePhysical %dMiB")
+		TEXT("\n\tMemoryTriggerGCAtPressureLevel %s")
+		TEXT("\n\tSoftGCMemoryUseTrigger %s%s")
+		TEXT("\n\tSoftGCTimeBudgetTrigger %s%s"),
+		MemoryMaxUsedVirtual / 1024 / 1024, MemoryMaxUsedPhysical / 1024 / 1024,
+		MemoryMinFreeVirtual / 1024 / 1024, MemoryMinFreePhysical / 1024 / 1024,
+		*LexToString(MemoryTriggerGCAtPressureLevel),
+		(bUseSoftGC && SoftGCStartNumerator > 0) ? TEXT("enabled") : TEXT("disabled"),
+		(bUseSoftGC && SoftGCStartNumerator > 0) ? *FString::Printf(TEXT(" (%d/%d)"), SoftGCStartNumerator, SoftGCDenominator) : TEXT(""),
+		(bUseSoftGC && SoftGCTimeFractionBudget > 0) ? TEXT("enabled") : TEXT("disabled"),
+		(bUseSoftGC && SoftGCTimeFractionBudget > 0) ? *FString::Printf(TEXT(" (%.3g)"), SoftGCTimeFractionBudget) : TEXT("")
+	);
+}
+
 namespace UE::Cook::Pollable
 {
 constexpr double TimePeriodNever = MAX_flt / 2;
@@ -5144,31 +5165,46 @@ bool UCookOnTheFlyServer::PumpHasExceededMaxMemory(uint32& OutResultFlags)
 	}
 
 	// If a normal GC was not triggered, check the SoftGC trigger conditions
-	double CurrentTime = 0.;
+	double CurrentTime = FPlatformTime::Seconds();
 	bool bIsSoftGC = false;
-	if (!bTriggerGC && bUseSoftGC && IsDirectorCookByTheBook() && !IsCookingInEditor())
+	if (bUseSoftGC && IsDirectorCookByTheBook() && !IsCookingInEditor())
 	{
-		if (SoftGCNextAvailablePhysicalTarget == -1) // Uninitialized
+		if (!bTriggerGC && SoftGCStartNumerator > 0)
 		{
-			int32 StartNumerator = FMath::Max(SoftGCStartNumerator, 1);
-			int32 Denominator = FMath::Max(SoftGCDenominator, 1);
-			// e.g. Start the target at 5/10, and decrease it by 1/10 each time the target is reached
-			SoftGCNextAvailablePhysicalTarget = (static_cast<int64>(MemStats.TotalPhysical)*StartNumerator)
-				/Denominator;
-		}
+			if (SoftGCNextAvailablePhysicalTarget == -1) // Uninitialized
+			{
+				int32 StartNumerator = FMath::Max(SoftGCStartNumerator, 1);
+				int32 Denominator = FMath::Max(SoftGCDenominator, 1);
+				// e.g. Start the target at 5/10, and decrease it by 1/10 each time the target is reached
+				SoftGCNextAvailablePhysicalTarget = (static_cast<int64>(MemStats.TotalPhysical) * StartNumerator)
+					/ Denominator;
+			}
 
-		if (SoftGCNextAvailablePhysicalTarget < -1)
-		{
-			// No further targets, no further SoftGC
+			if (SoftGCNextAvailablePhysicalTarget < -1)
+			{
+				// No further targets, no further SoftGC
+			}
+			else if (static_cast<int64>(MemStats.AvailablePhysical) <= SoftGCNextAvailablePhysicalTarget)
+			{
+				constexpr float SoftGCInstigateCooldown = 5 * 60.f;
+				CurrentTime = FPlatformTime::Seconds();
+				if (LastSoftGCTime + SoftGCInstigateCooldown <= CurrentTime)
+				{
+					TriggerMessages.Appendf(TEXT("\n  CookSettings.SoftGCMemoryTrigger: Available physical memory %dMiB is less than the current target for SoftGC %dMiB."),
+						static_cast<uint32>(MemStats.AvailablePhysical / 1024 / 1024), static_cast<uint32>(SoftGCNextAvailablePhysicalTarget / 1024 / 1024));
+					bTriggerGC = true;
+					bIsSoftGC = true;
+				}
+			}
 		}
-		else if (static_cast<int64>(MemStats.AvailablePhysical) <= SoftGCNextAvailablePhysicalTarget)
+
+		if (!bTriggerGC && SoftGCTimeFractionBudget > 0.)
 		{
-			constexpr float SoftGCInstigateCooldown = 5 * 60.f;
-			CurrentTime = FPlatformTime::Seconds();
-			if (LastSoftGCTime + SoftGCInstigateCooldown <= CurrentTime)
+			FString TriggerDiagnostics;
+			if (SoftGCHistory->IsTriggeringWithinBudget(*this, CurrentTime, &TriggerDiagnostics))
 			{
-				TriggerMessages.Appendf(TEXT("\n  CookSettings.bUseSoftGC: Available physical memory %dMiB is less than the current target for SoftGC %dMiB."),
-					static_cast<uint32>(MemStats.AvailablePhysical / 1024 / 1024), static_cast<uint32>(SoftGCNextAvailablePhysicalTarget / 1024 / 1024));
+				TriggerMessages.Appendf(TEXT("\n  CookSettings.SoftGCTimeTrigger: Periodic triggering of SoftGC: %s."),
+					*TriggerDiagnostics);
 				bTriggerGC = true;
 				bIsSoftGC = true;
 			}
@@ -5180,11 +5216,6 @@ bool UCookOnTheFlyServer::PumpHasExceededMaxMemory(uint32& OutResultFlags)
 		return false;
 	}
 
-	if (CurrentTime == 0.)
-	{
-		CurrentTime = FPlatformTime::Seconds();
-	}
-
 	// Don't allow a second GC (soft or normal) within the GC cooldown period after a full GC
 	constexpr float GCCooldown = 60.f;
 	if (LastFullGCTime + GCCooldown > CurrentTime)
@@ -5726,7 +5757,8 @@ void UCookOnTheFlyServer::EvaluateGarbageCollectionResults(bool bWasDueToOOM, bo
 	int32 NumObjectsBeforeGC, const FPlatformMemoryStats& MemStatsBeforeGC,
 	const FGenericMemoryStats& AllocatorStatsBeforeGC,
 	int32 NumObjectsAfterGC, const FPlatformMemoryStats& MemStatsAfterGC,
-	const FGenericMemoryStats& AllocatorStatsAfterGC)
+	const FGenericMemoryStats& AllocatorStatsAfterGC,
+	float GCDurationSeconds)
 {
 	using namespace UE::Cook;
 
@@ -5741,28 +5773,36 @@ void UCookOnTheFlyServer::EvaluateGarbageCollectionResults(bool bWasDueToOOM, bo
 	if (bWasSoftGC)
 	{
 		LastSoftGCTime = LastGCTime;
-		int32 StartNumerator = FMath::Max(SoftGCStartNumerator, 1);
-		int32 Denominator = FMath::Max(SoftGCDenominator, 1);
-		// Calculate the new SoftGCNextAvailablePhysicalTarget. Use the floor of NewAvailableMemory/Denominator,
-		// unless we are already 50% of the way through that level, in which case use the next value below that
-		int64 PhysicalMemoryQuantum = static_cast<int64>(MemStatsAfterGC.TotalPhysical) / Denominator;
-		int32 NextTarget =
-			static_cast<int64>(MemStatsAfterGC.AvailablePhysical - PhysicalMemoryQuantum/2) / PhysicalMemoryQuantum;
-		NextTarget = FMath::Min(NextTarget, StartNumerator);
-		if (NextTarget <= 0)
-		{
-			SoftGCNextAvailablePhysicalTarget = -2; // disabled, no further targets
-		}
-		else
+		if (SoftGCStartNumerator)
 		{
-			SoftGCNextAvailablePhysicalTarget = static_cast<int64>((MemStatsAfterGC.TotalPhysical) * NextTarget)
-				/ Denominator;
+			int32 StartNumerator = FMath::Max(SoftGCStartNumerator, 1);
+			int32 Denominator = FMath::Max(SoftGCDenominator, 1);
+			// Calculate the new SoftGCNextAvailablePhysicalTarget. Use the floor of NewAvailableMemory/Denominator,
+			// unless we are already 50% of the way through that level, in which case use the next value below that
+			int64 PhysicalMemoryQuantum = static_cast<int64>(MemStatsAfterGC.TotalPhysical) / Denominator;
+			int32 NextTarget =
+				static_cast<int64>(MemStatsAfterGC.AvailablePhysical - PhysicalMemoryQuantum / 2) / PhysicalMemoryQuantum;
+			NextTarget = FMath::Min(NextTarget, StartNumerator);
+			if (NextTarget <= 0)
+			{
+				SoftGCNextAvailablePhysicalTarget = -2; // disabled, no further targets
+			}
+			else
+			{
+				SoftGCNextAvailablePhysicalTarget = static_cast<int64>((MemStatsAfterGC.TotalPhysical) * NextTarget)
+					/ Denominator;
+			}
 		}
 	}
 	else
 	{
+		LastSoftGCTime = LastGCTime;
 		LastFullGCTime = LastGCTime;
 	}
+	if (SoftGCHistory)
+	{
+		SoftGCHistory->AddDurationMeasurement(GCDurationSeconds);
+	}
 
 	if (IsCookingInEditor())
 	{
@@ -5799,6 +5839,7 @@ void UCookOnTheFlyServer::EvaluateGarbageCollectionResults(bool bWasDueToOOM, bo
 	{
 		UE_LOG(LogCook, Display, TEXT("GarbageCollection Results:\n")
 			TEXT("\tType: %s\n")
+			TEXT("\tDuration: %.3fs\n")
 			TEXT("\tNumObjects:\n")
 			TEXT("\t\tCapacity:         %10d\n")
 			TEXT("\t\tBefore GC:        %10d\n")
@@ -5807,8 +5848,9 @@ void UCookOnTheFlyServer::EvaluateGarbageCollectionResults(bool bWasDueToOOM, bo
 			TEXT("\tVirtual Memory:\n")
 			TEXT("\t\tBefore GC:        %10" INT64_FMT " MB\n")
 			TEXT("\t\tAfter GC:         %10" INT64_FMT " MB\n")
-			TEXT("\t\tFreed by GC:      %10" INT64_FMT " MB\n"),
+			TEXT("\t\tFreed by GC:      %10" INT64_FMT " MB"),
 			(bWasSoftGC ? TEXT("Soft") : (bWasPartialGC ? TEXT("Partial") : TEXT("Full"))),
+			GCDurationSeconds,
 			NumObjectsCapacity, (int64)NumObjectsBeforeGC, (int64)NumObjectsAfterGC, NumObjectsFreed,
 			VirtualMemBeforeGC / BytesPerMeg, VirtualMemAfterGC / BytesPerMeg, VirtualMemFreed / BytesPerMeg
 		);
@@ -5835,8 +5877,10 @@ void UCookOnTheFlyServer::EvaluateGarbageCollectionResults(bool bWasDueToOOM, bo
 			UE_LOG(LogCook, Display, TEXT("GarbageCollection Results:"));
 		}
 		UE_LOG(LogCook, Display, TEXT("\tMemoryAnalysis: General:\n")
-			TEXT("\t\tType: %s"),
-			(bWasSoftGC ? TEXT("Soft") : (bWasPartialGC ? TEXT("Partial") : TEXT("Full"))));
+			TEXT("\t\tType: %s\n")
+			TEXT("\tDuration: %.3fs"),
+			(bWasSoftGC ? TEXT("Soft") : (bWasPartialGC ? TEXT("Partial") : TEXT("Full"))),
+			GCDurationSeconds);
 		UE_LOG(LogCook, Display, TEXT("\tMemoryAnalysis: NumObjects:\n")
 			TEXT("\t\tCapacity:         %10" INT64_FMT "\n")
 			TEXT("\t\tProcess Min:      %10" INT64_FMT "\n")
@@ -7232,6 +7276,7 @@ void FInitializeConfigSettings::LoadLocal(const FString& InOutputDirectoryOverri
 	GConfig->GetBool(TEXT("CookSettings"), TEXT("bUseSoftGC"), bUseSoftGC, GEditorIni);
 	GConfig->GetInt(TEXT("CookSettings"), TEXT("SoftGCStartNumerator"), SoftGCStartNumerator, GEditorIni);
 	GConfig->GetInt(TEXT("CookSettings"), TEXT("SoftGCDenominator"), SoftGCDenominator, GEditorIni);
+	GConfig->GetFloat(TEXT("CookSettings"), TEXT("SoftGCTimeFractionBudget"), SoftGCTimeFractionBudget, GEditorIni);
 
 	MemoryExpectedFreedToSpreadRatio = 0.10f;
 	GConfig->GetFloat(TEXT("CookSettings"), TEXT("MemoryExpectedFreedToSpreadRatio"),
@@ -7246,18 +7291,6 @@ void FInitializeConfigSettings::LoadLocal(const FString& InOutputDirectoryOverri
 	
 	GConfig->GetArray(TEXT("CookSettings"), TEXT("CookOnTheFlyConfigSettingDenyList"), ConfigSettingDenyList, GEditorIni);
 
-	UE_LOG(LogCook, Display, TEXT("CookSettings for Memory:")
-		TEXT("\n\tMemoryMaxUsedVirtual %dMiB")
-		TEXT("\n\tMemoryMaxUsedPhysical %dMiB")
-		TEXT("\n\tMemoryMinFreeVirtual %dMiB")
-		TEXT("\n\tMemoryMinFreePhysical %dMiB")
-		TEXT("\n\tMemoryTriggerGCAtPressureLevel %s")
-		TEXT("\n\tUseSoftGC %s%s"),
-		MemoryMaxUsedVirtual / 1024 / 1024, MemoryMaxUsedPhysical / 1024 / 1024, MemoryMinFreeVirtual / 1024 / 1024, MemoryMinFreePhysical / 1024 / 1024,
-		*LexToString(MemoryTriggerGCAtPressureLevel),
-		bUseSoftGC ? TEXT("true") : TEXT("false"),
-		bUseSoftGC ? *FString::Printf(TEXT(" (%d/%d)"), SoftGCStartNumerator, SoftGCDenominator) : TEXT(""));
-
 	const FConfigSection* CacheSettings = GConfig->GetSection(TEXT("CookPlatformDataCacheSettings"), false, GEditorIni);
 	if (CacheSettings)
 	{
@@ -7346,9 +7379,16 @@ void UCookOnTheFlyServer::SetInitializeConfigSettings(UE::Cook::FInitializeConfi
 		}
 	}
 
+	if (SoftGCTimeFractionBudget > 0)
+	{
+		SoftGCHistory.Reset(new UE::Cook::FSoftGCHistory());
+	}
+
 	// The rest of this function parses config settings that are reparsed on every CookDirector and CookWorker rather than
 	// being replicated from CookDirector to CookWorker
 
+	UE_LOG(LogCook, Display, TEXT("CookSettings for Memory:%s"), *GetCookSettingsForMemoryLogText());
+
 	// Debugging hidden dependencies
 	bOnlyEditorOnlyDebug = FParse::Param(FCommandLine::Get(), TEXT("OnlyEditorOnlyDebug"));
 	bSkipOnlyEditorOnly = false;
diff --git a/Engine/Source/Editor/UnrealEd/Private/Cooker/CookDirector.cpp b/Engine/Source/Editor/UnrealEd/Private/Cooker/CookDirector.cpp
--- a/Engine/Source/Editor/UnrealEd/Private/Cooker/CookDirector.cpp
+++ b/Engine/Source/Editor/UnrealEd/Private/Cooker/CookDirector.cpp
@@ -1277,20 +1277,8 @@ void FCookDirector::ActivateMachineResourceReduction()
 	COTFS.MemoryMinFreePhysical = 0;
 	COTFS.MemoryTriggerGCAtPressureLevel = FGenericPlatformMemoryStats::EMemoryPressureStatus::Critical;
 
-	UE_LOG(LogCook, Display, TEXT("CookMultiprocess changed CookSettings for Memory:")
-		TEXT("\n\tMemoryMaxUsedVirtual %dMiB")
-		TEXT("\n\tMemoryMaxUsedPhysical %dMiB")
-		TEXT("\n\tMemoryMinFreeVirtual %dMiB")
-		TEXT("\n\tMemoryMinFreePhysical %dMiB")
-		TEXT("\n\tMemoryTriggerGCAtPressureLevel %s")
-		TEXT("\n\tUseSoftGC %s%s"),
-		COTFS.MemoryMaxUsedVirtual / 1024 / 1024, COTFS.MemoryMaxUsedPhysical / 1024 / 1024,
-		COTFS.MemoryMinFreeVirtual / 1024 / 1024, COTFS.MemoryMinFreePhysical / 1024 / 1024,
-		*LexToString(COTFS.MemoryTriggerGCAtPressureLevel),
-		COTFS.bUseSoftGC ? TEXT("true") : TEXT("false"),
-		COTFS.bUseSoftGC ? *FString::Printf(TEXT(" (%d/%d)"),
-			COTFS.SoftGCStartNumerator, COTFS.SoftGCDenominator) : TEXT("")
-	);
+	UE_LOG(LogCook, Display, TEXT("CookMultiprocess changed CookSettings for Memory:%s"),
+		*COTFS.GetCookSettingsForMemoryLogText());
 
 	// Set CoreLimit for updating workerthreads in this process and passing to the commandline for workers
 	int32 NumProcesses = RequestedCookWorkerCount + 1;
diff --git a/Engine/Source/Editor/UnrealEd/Private/Cooker/CookTypes.cpp b/Engine/Source/Editor/UnrealEd/Private/Cooker/CookTypes.cpp
--- a/Engine/Source/Editor/UnrealEd/Private/Cooker/CookTypes.cpp
+++ b/Engine/Source/Editor/UnrealEd/Private/Cooker/CookTypes.cpp
@@ -739,6 +739,7 @@ FCbWriter& operator<<(FCbWriter& Writer, const UE::Cook::FInitializeConfigSettin
 	Writer << "bUseSoftGC" << Value.bUseSoftGC;
 	Writer << "SoftGCStartNumerator" << Value.SoftGCStartNumerator;
 	Writer << "SoftGCDenominator" << Value.SoftGCDenominator;
+	Writer << "SoftGCTimeFractionBudget" << Value.SoftGCTimeFractionBudget;
 	Writer << "MinFreeUObjectIndicesBeforeGC" << Value.MinFreeUObjectIndicesBeforeGC;
 	Writer << "MaxNumPackagesBeforePartialGC" << Value.MaxNumPackagesBeforePartialGC;
 	Writer << "ConfigSettingDenyList" << Value.ConfigSettingDenyList;
@@ -775,6 +776,7 @@ bool LoadFromCompactBinary(FCbFieldView Field, UE::Cook::FInitializeConfigSettin
 	bOk = LoadFromCompactBinary(Field["bUseSoftGC"], OutValue.bUseSoftGC) & bOk;
 	bOk = LoadFromCompactBinary(Field["SoftGCStartNumerator"], OutValue.SoftGCStartNumerator) & bOk;
 	bOk = LoadFromCompactBinary(Field["SoftGCDenominator"], OutValue.SoftGCDenominator) & bOk;
+	bOk = LoadFromCompactBinary(Field["SoftGCTimeFractionBudget"], OutValue.SoftGCTimeFractionBudget) & bOk;
 	bOk = LoadFromCompactBinary(Field["MinFreeUObjectIndicesBeforeGC"], OutValue.MinFreeUObjectIndicesBeforeGC) & bOk;
 	bOk = LoadFromCompactBinary(Field["MaxNumPackagesBeforePartialGC"], OutValue.MaxNumPackagesBeforePartialGC) & bOk;
 	bOk = LoadFromCompactBinary(Field["ConfigSettingDenyList"], OutValue.ConfigSettingDenyList) & bOk;
@@ -804,6 +806,7 @@ void FInitializeConfigSettings::MoveOrCopy(SourceType&& Source, TargetType&& Tar
 	Target.bUseSoftGC = Source.bUseSoftGC;
 	Target.SoftGCStartNumerator = Source.SoftGCStartNumerator;
 	Target.SoftGCDenominator = Source.SoftGCDenominator;
+	Target.SoftGCTimeFractionBudget = Source.SoftGCTimeFractionBudget;
 	Target.MinFreeUObjectIndicesBeforeGC = Source.MinFreeUObjectIndicesBeforeGC;
 	Target.MaxNumPackagesBeforePartialGC = Source.MaxNumPackagesBeforePartialGC;
 	Target.ConfigSettingDenyList = MoveTempIfPossible(Source.ConfigSettingDenyList);
diff --git a/Engine/Source/Editor/UnrealEd/Private/Cooker/CookTypes.h b/Engine/Source/Editor/UnrealEd/Private/Cooker/CookTypes.h
--- a/Engine/Source/Editor/UnrealEd/Private/Cooker/CookTypes.h
+++ b/Engine/Source/Editor/UnrealEd/Private/Cooker/CookTypes.h
@@ -431,15 +431,16 @@ namespace UE::Cook
 		uint32 PackagesPerGC = 0;
 		float MemoryExpectedFreedToSpreadRatio = 0.f;
 		double IdleTimeToGC = 0.;
-		uint64 MemoryMaxUsedVirtual;
-		uint64 MemoryMaxUsedPhysical;
-		uint64 MemoryMinFreeVirtual;
-		uint64 MemoryMinFreePhysical;
+		uint64 MemoryMaxUsedVirtual = 0;
+		uint64 MemoryMaxUsedPhysical = 0;
+		uint64 MemoryMinFreeVirtual = 0;
+		uint64 MemoryMinFreePhysical = 0;
 		FGenericPlatformMemoryStats::EMemoryPressureStatus MemoryTriggerGCAtPressureLevel;
-		int32 MinFreeUObjectIndicesBeforeGC;
-		int32 MaxNumPackagesBeforePartialGC;
-		int32 SoftGCStartNumerator;
-		int32 SoftGCDenominator;
+		int32 MinFreeUObjectIndicesBeforeGC = 0;
+		int32 MaxNumPackagesBeforePartialGC = 0;
+		int32 SoftGCStartNumerator = 0;
+		int32 SoftGCDenominator = 1;
+		float SoftGCTimeFractionBudget = 0.f;
 		TArray<FString> ConfigSettingDenyList;
 		/** max number of objects of a specific type which are allowed to async cache at once */
 		TMap<FName, int32> MaxAsyncCacheForType;
diff --git a/Engine/Source/Editor/UnrealEd/Private/Cooker/SoftGC.cpp b/Engine/Source/Editor/UnrealEd/Private/Cooker/SoftGC.cpp
new file mode 100644
--- /dev/null
+++ b/Engine/Source/Editor/UnrealEd/Private/Cooker/SoftGC.cpp
@@ -0,0 +1,91 @@
+// Copyright Epic Games, Inc. All Rights Reserved.
+
+#include "Cooker/SoftGC.h"
+
+#include "CookOnTheSide/CookOnTheFlyServer.h"
+
+namespace UE::Cook
+{
+
+void FSoftGCHistory::AddDurationMeasurement(float DurationSeconds)
+{
+	int32 HistoryLength = DurationHistory.Num();
+	while (HistoryLength >= MaxHistoryLength)
+	{
+		if (HistoryLength <= 1)
+		{
+			if (HistoryLength > 0)
+			{
+				DurationHistory.PopFront();
+			}
+			AverageDurationSeconds = 0.0;
+			HistoryLength = 0;
+		}
+		else
+		{
+			AverageDurationSeconds = (AverageDurationSeconds * HistoryLength - DurationHistory.PopFrontValue()) / (HistoryLength - 1);
+			--HistoryLength;
+		}
+	}
+	if (HistoryLength < MaxHistoryLength)
+	{
+		if (HistoryLength == 0)
+		{
+			AverageDurationSeconds = DurationSeconds;
+		}
+		else
+		{
+			AverageDurationSeconds = (AverageDurationSeconds * HistoryLength + DurationSeconds) / (HistoryLength + 1);
+		}
+		DurationHistory.Add(DurationSeconds);
+	}
+}
+
+bool FSoftGCHistory::IsTriggeringWithinBudget(UCookOnTheFlyServer& COTFS, double CurrentTimeSeconds, FString* OutDiagnostics)
+{
+	if (OutDiagnostics)
+	{
+		OutDiagnostics->Reset();
+	}
+	if (COTFS.CookedPackageCountSinceLastGC == 0)
+	{
+		return false;
+	}
+	if (DurationHistory.Num() == 0)
+	{
+		if (OutDiagnostics)
+		{
+			*OutDiagnostics = TEXT("No duration data");
+		}
+		return true;
+	}
+	// TimeBudget/(Time + TimeBudget) == BudgetFraction
+	// TimeBudget == (BudgetFraction/(1 - BudgetFraction))*Time
+	if (COTFS.SoftGCTimeFractionBudget > .99f)
+	{
+		if (OutDiagnostics)
+		{
+			*OutDiagnostics = FString::Printf(
+				TEXT("SoftGCTimeFractionBudget == %.3f, above threshold to always trigger"),
+				COTFS.SoftGCTimeFractionBudget);
+		}
+		return true;
+	}
+	float TimeSinceLastGCSeconds = CurrentTimeSeconds - COTFS.LastSoftGCTime;
+	float CurrentTimeBudget = TimeSinceLastGCSeconds *
+		COTFS.SoftGCTimeFractionBudget / (1 - COTFS.SoftGCTimeFractionBudget);
+	if (CurrentTimeBudget >= AverageDurationSeconds)
+	{
+		if (OutDiagnostics)
+		{
+			*OutDiagnostics = FString::Printf(
+				TEXT("SoftGCTimeFractionBudget == %.3f. TimeSinceLastGCSeconds == %.3f. CurrentTimeBudget == %.3f. ExpectedDuration == %.3f"),
+				COTFS.SoftGCTimeFractionBudget, TimeSinceLastGCSeconds, CurrentTimeBudget, AverageDurationSeconds);
+		}
+		return true;
+	}
+	return false;
+}
+
+}
+
diff --git a/Engine/Source/Editor/UnrealEd/Private/Cooker/SoftGC.h b/Engine/Source/Editor/UnrealEd/Private/Cooker/SoftGC.h
new file mode 100644
--- /dev/null
+++ b/Engine/Source/Editor/UnrealEd/Private/Cooker/SoftGC.h
@@ -0,0 +1,26 @@
+// Copyright Epic Games, Inc. All Rights Reserved.
+
+#pragma once
+
+#include "Containers/RingBuffer.h"
+#include "Containers/UnrealString.h"
+
+class UCookOnTheFlyServer;
+
+namespace UE::Cook
+{
+
+class FSoftGCHistory
+{
+public:
+	void AddDurationMeasurement(float DurationSeconds);
+	bool IsTriggeringWithinBudget(UCookOnTheFlyServer& COTFS, double CurrentTimeSeconds, FString* OutDiagnostics);
+
+private:
+	TRingBuffer<float> DurationHistory;
+	int MaxHistoryLength = 5;
+	float AverageDurationSeconds = 0.f;
+};
+
+}
+
