[Viewport Toolbar] Add argument to optionally show/hide Far View Plane widget within ExtendCamera method

Unshelved from pending changelist '37178925':

diff --git a/Engine/Source/Editor/UnrealEd/Private/ViewportToolbar/UnrealEdViewportToolbar.cpp b/Engine/Source/Editor/UnrealEd/Private/ViewportToolbar/UnrealEdViewportToolbar.cpp
--- a/Engine/Source/Editor/UnrealEd/Private/ViewportToolbar/UnrealEdViewportToolbar.cpp
+++ b/Engine/Source/Editor/UnrealEd/Private/ViewportToolbar/UnrealEdViewportToolbar.cpp
@@ -2578,14 +2578,14 @@ void PopulateCameraMenu(UToolMenu* InMenu)
 	OrthographicCameraSection.AddMenuEntry(FEditorViewportCommands::Get().Back);
 }
 
-void ExtendCameraSubmenu(FName InCameraOptionsSubmenuName)
+void ExtendCameraSubmenu(FName InCameraOptionsSubmenuName, bool bInShowFarViewPlaneEntry /* = true */)
 {
 	UToolMenu* const Submenu = UToolMenus::Get()->ExtendMenu(InCameraOptionsSubmenuName);
 
 	Submenu->AddDynamicSection(
 		"EditorCameraExtensionDynamicSection",
 		FNewToolMenuDelegate::CreateLambda(
-			[](UToolMenu* InDynamicMenu)
+			[bInShowFarViewPlaneEntry](UToolMenu* InDynamicMenu)
 			{
 				UUnrealEdViewportToolbarContext* const EditorViewportContext =
 					InDynamicMenu->FindContext<UUnrealEdViewportToolbarContext>();
@@ -2611,11 +2611,14 @@ void ExtendCameraSubmenu(FName InCameraOptionsSubmenuName)
 					LOCTEXT("CameraSubmenu_FieldOfViewLabel", "Field of View")
 				));
 
-				PostOrthoSection.AddEntry(FToolMenuEntry::InitWidget(
-					"CameraFarViewPlane",
-					UE::UnrealEd::CreateFarViewPlaneMenuWidget(EditorViewport.ToSharedRef()),
-					LOCTEXT("CameraSubmenu_FarViewPlaneLabel", "Far View Plane")
-				));
+				if (bInShowFarViewPlaneEntry)
+				{
+					PostOrthoSection.AddEntry(FToolMenuEntry::InitWidget(
+						"CameraFarViewPlane",
+						UE::UnrealEd::CreateFarViewPlaneMenuWidget(EditorViewport.ToSharedRef()),
+						LOCTEXT("CameraSubmenu_FarViewPlaneLabel", "Far View Plane")
+					));
+				}
 			}
 		)
 	);
diff --git a/Engine/Source/Editor/UnrealEd/Public/ViewportToolbar/UnrealEdViewportToolbar.h b/Engine/Source/Editor/UnrealEd/Public/ViewportToolbar/UnrealEdViewportToolbar.h
--- a/Engine/Source/Editor/UnrealEd/Public/ViewportToolbar/UnrealEdViewportToolbar.h
+++ b/Engine/Source/Editor/UnrealEd/Public/ViewportToolbar/UnrealEdViewportToolbar.h
@@ -148,7 +148,7 @@ UNREALED_API void PopulateCameraMenu(UToolMenu* InMenu);
 /**
  * Adds Field of View and Far View Plane entries to the specified Camera Submenu
  */
-UNREALED_API void ExtendCameraSubmenu(FName InCameraOptionsSubmenuName);
+UNREALED_API void ExtendCameraSubmenu(FName InCameraOptionsSubmenuName, bool bInShowFarViewPlaneEntry = true);
 
 UNREALED_API void GenerateViewportTypeMenu(UToolMenu* InMenu);
 
