CookGenerationHelper: Fix spurious report of missing splitobject when a generator package is at the start of the save state, without yet initializing its SplitDataObject pointer, and garbage collection is triggered.
#rnx

diff --git a/Engine/Source/Editor/UnrealEd/Private/Cooker/CookGenerationHelper.cpp b/Engine/Source/Editor/UnrealEd/Private/Cooker/CookGenerationHelper.cpp
--- a/Engine/Source/Editor/UnrealEd/Private/Cooker/CookGenerationHelper.cpp
+++ b/Engine/Source/Editor/UnrealEd/Private/Cooker/CookGenerationHelper.cpp
@@ -653,12 +653,9 @@ const FCookGenerationInfo* FGenerationHelper::FindInfo(const FPackageData& Packa
 	return const_cast<FGenerationHelper*>(this)->FindInfo(PackageData);
 }
 
-UObject* FGenerationHelper::FindOrLoadSplitDataObject()
+UObject* FGenerationHelper::GetWeakSplitDataObject()
 {
-	if (!IsValid())
-	{
-		return nullptr;
-	}
+	ConditionalInitialize();
 	UObject* Result = SplitDataObject.Get();
 	if (Result)
 	{
@@ -675,23 +672,32 @@ UObject* FGenerationHelper::FindOrLoadSplitDataObject()
 	}
 
 	Result = FindObject<UObject>(nullptr, *ObjectPath);
-	if (!Result)
+	if (Result)
 	{
-		FPackageData& OwnerPackageData = GetOwner();
-		FName OwnerPackageName = OwnerPackageData.GetPackageName();
-		UCookOnTheFlyServer& COTFS = OwnerPackageData.GetPackageDatas().GetCookOnTheFlyServer();
-		UPackage* LocalOwnerPackage;
-		COTFS.LoadPackageForCooking(OwnerPackageData, LocalOwnerPackage);
+		SplitDataObject = Result;
+	}
+	return Result;
+}
 
-		Result = FindObject<UObject>(nullptr, *ObjectPath);
-		if (!Result)
-		{
-			return nullptr;
-		}
+UObject* FGenerationHelper::FindOrLoadSplitDataObject()
+{
+	if (!IsValid())
+	{
+		return nullptr;
+	}
+	UObject* Result = GetWeakSplitDataObject();
+	if (Result)
+	{
+		return Result;
 	}
 
-	SplitDataObject = Result;
-	return Result;
+	FPackageData& OwnerPackageData = GetOwner();
+	FName OwnerPackageName = OwnerPackageData.GetPackageName();
+	UCookOnTheFlyServer& COTFS = OwnerPackageData.GetPackageDatas().GetCookOnTheFlyServer();
+	UPackage* LocalOwnerPackage;
+	COTFS.LoadPackageForCooking(OwnerPackageData, LocalOwnerPackage);
+
+	return GetWeakSplitDataObject();
 }
 
 UPackage* FGenerationHelper::GetOwnerPackage()
diff --git a/Engine/Source/Editor/UnrealEd/Private/Cooker/CookGenerationHelper.h b/Engine/Source/Editor/UnrealEd/Private/Cooker/CookGenerationHelper.h
--- a/Engine/Source/Editor/UnrealEd/Private/Cooker/CookGenerationHelper.h
+++ b/Engine/Source/Editor/UnrealEd/Private/Cooker/CookGenerationHelper.h
@@ -221,8 +221,8 @@ public:
 	bool IsRequiresGeneratorPackageDestructBeforeResplit() const;
 	/** Return the Splitter's value for virtual bool DoesGeneratedRequireGenerator(). */
 	ICookPackageSplitter::EGeneratedRequiresGenerator DoesGeneratedRequireGenerator() const;
-	/** Return the cached pointer to the SplitDataObject. Returns null if no longer in memory or marked as garbage. */
-	UObject* GetWeakSplitDataObject() const;
+	/** Return the pointer to the SplitDataObject. Returns null if not in memory or marked as garbage. */
+	UObject* GetWeakSplitDataObject();
 	/**
 	 * Return our cached pointer to the SplitDataObject if set. If not set, load the package and find it. Can still
 	 * return null if !IsValid or if not found even after loading the package.
@@ -669,12 +669,6 @@ inline ICookPackageSplitter::EGeneratedRequiresGenerator FGenerationHelper::Does
 	return DoesGeneratedRequireGeneratorValue;
 }
 
-inline UObject* FGenerationHelper::GetWeakSplitDataObject() const
-{
-	ConditionalInitialize();
-	return SplitDataObject.Get();
-}
-
 inline TConstArrayView<FWeakObjectPtr> FGenerationHelper::GetOwnerObjectsToMove() const
 {
 	return OwnerObjectsToMove;
