[ContentBrowser] Temporary CIS fix

diff --git a/Engine/Source/Editor/UnrealEd/Private/DragAndDrop/AssetDragDropOp.cpp b/Engine/Source/Editor/UnrealEd/Private/DragAndDrop/AssetDragDropOp.cpp
--- a/Engine/Source/Editor/UnrealEd/Private/DragAndDrop/AssetDragDropOp.cpp
+++ b/Engine/Source/Editor/UnrealEd/Private/DragAndDrop/AssetDragDropOp.cpp
@@ -11,11 +11,6 @@
 #include "ThumbnailRendering/ThumbnailManager.h"
 #include "Engine/Level.h"
 
-namespace UE::Editor::ContentBrowser
-{
-	CONTENTBROWSER_API extern bool GEnableContentBrowserNewStyle;
-}
-
 TSharedRef<FAssetDragDropOp> FAssetDragDropOp::New(const FAssetData& InAssetData, TScriptInterface<IAssetFactoryInterface> AssetFactory)
 {
 	TArray<FAssetData> AssetDataArray;
@@ -90,7 +85,7 @@ TSharedPtr<SWidget> FAssetDragDropOp::GetDefaultDecorator() const
 	}
 	else if (AssetThumbnail.IsValid())
 	{
-		if (UE::Editor::ContentBrowser::GEnableContentBrowserNewStyle)
+		if (0)
 		{
 			FAssetThumbnailConfig AssetThumbnailConfig;
 			AssetThumbnailConfig.bShowAssetChip = true;
