Fix assert when doing legacy fbx scene re-import

#rnx

diff --git a/Engine/Source/Editor/UnrealEd/Private/Fbx/SSceneImportStaticMeshListView.cpp b/Engine/Source/Editor/UnrealEd/Private/Fbx/SSceneImportStaticMeshListView.cpp
--- a/Engine/Source/Editor/UnrealEd/Private/Fbx/SSceneImportStaticMeshListView.cpp
+++ b/Engine/Source/Editor/UnrealEd/Private/Fbx/SSceneImportStaticMeshListView.cpp
@@ -328,16 +328,18 @@ void SFbxSceneStaticMeshListView::OnFinishedChangingProperties(const FPropertyCh
 
 void SFbxSceneStaticMeshListView::OnChangedOverrideOptions(TSharedPtr<FString> ItemSelected, ESelectInfo::Type SelectInfo)
 {
-	check(ItemSelected.IsValid());
-	if (ItemSelected.Get()->Compare(UFbxSceneImportFactory::DefaultOptionName) == 0)
+	if (ItemSelected.IsValid())
 	{
-		CurrentMeshImportOptions = GlobalImportSettings;
-	}
-	else if(OverrideNameOptionsMap->Contains(*(ItemSelected.Get())))
-	{
-		CurrentMeshImportOptions = *OverrideNameOptionsMap->Find(*(ItemSelected.Get()));
+		if (ItemSelected.Get()->Compare(UFbxSceneImportFactory::DefaultOptionName) == 0)
+		{
+			CurrentMeshImportOptions = GlobalImportSettings;
+		}
+		else if (OverrideNameOptionsMap->Contains(*(ItemSelected.Get())))
+		{
+			CurrentMeshImportOptions = *OverrideNameOptionsMap->Find(*(ItemSelected.Get()));
+		}
+		SFbxSceneOptionWindow::CopyFbxOptionsToStaticMeshOptions(CurrentMeshImportOptions, SceneImportOptionsStaticMeshDisplay);
 	}
-	SFbxSceneOptionWindow::CopyFbxOptionsToStaticMeshOptions(CurrentMeshImportOptions, SceneImportOptionsStaticMeshDisplay);
 }
 
 #undef LOCTEXT_NAMESPACE
diff --git a/Engine/Source/Editor/UnrealEd/Private/Fbx/SSceneReimportSkeletalMeshListView.cpp b/Engine/Source/Editor/UnrealEd/Private/Fbx/SSceneReimportSkeletalMeshListView.cpp
--- a/Engine/Source/Editor/UnrealEd/Private/Fbx/SSceneReimportSkeletalMeshListView.cpp
+++ b/Engine/Source/Editor/UnrealEd/Private/Fbx/SSceneReimportSkeletalMeshListView.cpp
@@ -420,16 +420,18 @@ TSharedRef< ITableRow > SFbxSceneSkeletalMeshReimportListView::OnGenerateRowFbxS
 
 void SFbxSceneSkeletalMeshReimportListView::OnChangedOverrideOptions(TSharedPtr<FString> ItemSelected, ESelectInfo::Type SelectInfo)
 {
-	check(ItemSelected.IsValid());
-	if (ItemSelected.Get()->Compare(UFbxSceneImportFactory::DefaultOptionName) == 0)
+	if (ItemSelected.IsValid())
 	{
-		CurrentMeshImportOptions = GlobalImportSettings;
-	}
-	else if (OverrideNameOptionsMap->Contains(*(ItemSelected)))
-	{
-		CurrentMeshImportOptions = *OverrideNameOptionsMap->Find(*(ItemSelected));
+		if (ItemSelected.Get()->Compare(UFbxSceneImportFactory::DefaultOptionName) == 0)
+		{
+			CurrentMeshImportOptions = GlobalImportSettings;
+		}
+		else if (OverrideNameOptionsMap->Contains(*(ItemSelected)))
+		{
+			CurrentMeshImportOptions = *OverrideNameOptionsMap->Find(*(ItemSelected));
+		}
+		SFbxSceneOptionWindow::CopyFbxOptionsToSkeletalMeshOptions(CurrentMeshImportOptions, SceneImportOptionsSkeletalMeshDisplay);
 	}
-	SFbxSceneOptionWindow::CopyFbxOptionsToSkeletalMeshOptions(CurrentMeshImportOptions, SceneImportOptionsSkeletalMeshDisplay);
 }
 
 void SFbxSceneSkeletalMeshReimportListView::OnFinishedChangingProperties(const FPropertyChangedEvent& PropertyChangedEvent)
diff --git a/Engine/Source/Editor/UnrealEd/Private/Fbx/SSceneReimportStaticMeshListView.cpp b/Engine/Source/Editor/UnrealEd/Private/Fbx/SSceneReimportStaticMeshListView.cpp
--- a/Engine/Source/Editor/UnrealEd/Private/Fbx/SSceneReimportStaticMeshListView.cpp
+++ b/Engine/Source/Editor/UnrealEd/Private/Fbx/SSceneReimportStaticMeshListView.cpp
@@ -420,16 +420,18 @@ TSharedRef< ITableRow > SFbxSceneStaticMeshReimportListView::OnGenerateRowFbxSce
 
 void SFbxSceneStaticMeshReimportListView::OnChangedOverrideOptions(TSharedPtr<FString> ItemSelected, ESelectInfo::Type SelectInfo)
 {
-	check(ItemSelected.IsValid());
-	if (ItemSelected.Get()->Compare(UFbxSceneImportFactory::DefaultOptionName) == 0)
+	if (ItemSelected.IsValid())
 	{
-		CurrentMeshImportOptions = GlobalImportSettings;
-	}
-	else if (OverrideNameOptionsMap->Contains(*(ItemSelected)))
-	{
-		CurrentMeshImportOptions = *OverrideNameOptionsMap->Find(*(ItemSelected));
+		if (ItemSelected.Get()->Compare(UFbxSceneImportFactory::DefaultOptionName) == 0)
+		{
+			CurrentMeshImportOptions = GlobalImportSettings;
+		}
+		else if (OverrideNameOptionsMap->Contains(*(ItemSelected)))
+		{
+			CurrentMeshImportOptions = *OverrideNameOptionsMap->Find(*(ItemSelected));
+		}
+		SFbxSceneOptionWindow::CopyFbxOptionsToStaticMeshOptions(CurrentMeshImportOptions, SceneImportOptionsStaticMeshDisplay);
 	}
-	SFbxSceneOptionWindow::CopyFbxOptionsToStaticMeshOptions(CurrentMeshImportOptions, SceneImportOptionsStaticMeshDisplay);
 }
 
 void SFbxSceneStaticMeshReimportListView::OnFinishedChangingProperties(const FPropertyChangedEvent& PropertyChangedEvent)
diff --git a/Engine/Source/Editor/UnrealEd/Private/Fbx/SSceneSkeletalMeshListView.cpp b/Engine/Source/Editor/UnrealEd/Private/Fbx/SSceneSkeletalMeshListView.cpp
--- a/Engine/Source/Editor/UnrealEd/Private/Fbx/SSceneSkeletalMeshListView.cpp
+++ b/Engine/Source/Editor/UnrealEd/Private/Fbx/SSceneSkeletalMeshListView.cpp
@@ -327,16 +327,18 @@ void SFbxSceneSkeletalMeshListView::OnFinishedChangingProperties(const FProperty
 
 void SFbxSceneSkeletalMeshListView::OnChangedOverrideOptions(TSharedPtr<FString> ItemSelected, ESelectInfo::Type SelectInfo)
 {
-	check(ItemSelected.IsValid());
-	if (ItemSelected.Get()->Compare(UFbxSceneImportFactory::DefaultOptionName) == 0)
+	if (ItemSelected.IsValid())
 	{
-		CurrentMeshImportOptions = GlobalImportSettings;
-	}
-	else if (OverrideNameOptionsMap->Contains(*(ItemSelected.Get())))
-	{
-		CurrentMeshImportOptions = *OverrideNameOptionsMap->Find(*(ItemSelected.Get()));
+		if (ItemSelected.Get()->Compare(UFbxSceneImportFactory::DefaultOptionName) == 0)
+		{
+			CurrentMeshImportOptions = GlobalImportSettings;
+		}
+		else if (OverrideNameOptionsMap->Contains(*(ItemSelected.Get())))
+		{
+			CurrentMeshImportOptions = *OverrideNameOptionsMap->Find(*(ItemSelected.Get()));
+		}
+		SFbxSceneOptionWindow::CopyFbxOptionsToSkeletalMeshOptions(CurrentMeshImportOptions, SceneImportOptionsSkeletalMeshDisplay);
 	}
-	SFbxSceneOptionWindow::CopyFbxOptionsToSkeletalMeshOptions(CurrentMeshImportOptions, SceneImportOptionsSkeletalMeshDisplay);
 }
 
 #undef LOCTEXT_NAMESPACE
