[BUGFIX] FBX Export Level Instance
- Correctly prints that this feature is currently unsupported instead of there is nothing to export.

#rnx

diff --git a/Engine/Source/Editor/UnrealEd/Private/EditorExporters.cpp b/Engine/Source/Editor/UnrealEd/Private/EditorExporters.cpp
--- a/Engine/Source/Editor/UnrealEd/Private/EditorExporters.cpp
+++ b/Engine/Source/Editor/UnrealEd/Private/EditorExporters.cpp
@@ -72,6 +72,7 @@
 #include "LandscapeInfo.h"
 #include "LandscapeComponent.h"
 #include "LandscapeDataAccess.h"
+#include "LevelInstance/LevelInstanceActor.h"
 #include "UnrealExporter.h"
 #include "InstancedFoliage.h"
 #include "Engine/Selection.h"
@@ -90,6 +91,7 @@
 
 #include "Misc/App.h"
 #include "Widgets/Notifications/SNotificationList.h"
+#include "Framework/Docking/TabManager.h"
 #include "Framework/Notifications/NotificationManager.h"
 
 DEFINE_LOG_CATEGORY_STATIC(LogEditorExporters, Log, All);
@@ -1732,10 +1734,16 @@ ULevelExporterFBX::ULevelExporterFBX(const FObjectInitializer& ObjectInitializer
 	FormatDescription.Add(TEXT("FBX File"));
 }
 
-bool IsSomethingToExport(AActor* Actor)
+bool IsSomethingToExport(AActor* Actor, TArray<FString>& OutReasons)
 {
 	check(Actor);
 
+	if (ALevelInstance* LevelInstance = Cast<ALevelInstance>(Actor))
+	{
+		OutReasons.Add(FString::Printf(TEXT("%s: Exporting Level Instances to FBX is not supported."), *Actor->GetActorLabel()));
+		return false;
+	}
+
 	TInlineComponentArray<USceneComponent*> ComponentsToExport;
 	for (UActorComponent* ActorComp : Actor->GetComponents())
 	{
@@ -1755,7 +1763,7 @@ bool IsSomethingToExport(AActor* Actor)
 			(SkelMeshComp && SkelMeshComp->GetSkeletalMeshAsset()) ||
 			Component->IsA(UCameraComponent::StaticClass()) ||
 			Component->IsA(ULightComponent::StaticClass()) ||
-			(ChildActorComp && ChildActorComp->GetChildActor() && IsSomethingToExport(ChildActorComp->GetChildActor()))
+			(ChildActorComp && ChildActorComp->GetChildActor() && IsSomethingToExport(ChildActorComp->GetChildActor(), OutReasons))
 			)
 		{
 			return true;
@@ -1764,7 +1772,7 @@ bool IsSomethingToExport(AActor* Actor)
 	return false;
 }
 
-bool IsSomethingToExport(ULevel* InLevel, bool bSelectedOnly)
+bool IsSomethingToExport(ULevel* InLevel, bool bSelectedOnly, TArray<FString>& OutReasons)
 {
 	TArray<AActor*> ActorToExport;
 	int32 ActorCount = InLevel->Actors.Num();
@@ -1773,7 +1781,7 @@ bool IsSomethingToExport(ULevel* InLevel, bool bSelectedOnly)
 		AActor* Actor = InLevel->Actors[ActorIndex];
 		if (Actor != nullptr && (!bSelectedOnly || (bSelectedOnly && Actor->IsSelected())))
 		{
-			if (IsSomethingToExport(Actor))
+			if (IsSomethingToExport(Actor, OutReasons))
 			{
 				return true;
 			}
@@ -1795,14 +1803,14 @@ bool IsSomethingToExport(ULevel* InLevel, bool bSelectedOnly)
 	return false;
 }
 
-bool IsSomethingToExport(UWorld* World, bool bSelectedOnly)
+bool IsSomethingToExport(UWorld* World, bool bSelectedOnly, TArray<FString>& OutReasons)
 {
 	if (bSelectedOnly && World->GetModel()->Nodes.Num())
 	{
 		return true;
 	}
 	ULevel* Level = World->PersistentLevel;
-	if (IsSomethingToExport(Level, bSelectedOnly))
+	if (IsSomethingToExport(Level, bSelectedOnly, OutReasons))
 	{
 		return true;
 	}
@@ -1812,7 +1820,7 @@ bool IsSomethingToExport(UWorld* World, bool bSelectedOnly)
 		ULevel* CurLevel = World->GetLevel(CurLevelIndex);
 		if (CurLevel != NULL && CurLevel != Level)
 		{
-			if (IsSomethingToExport(CurLevel, bSelectedOnly))
+			if (IsSomethingToExport(CurLevel, bSelectedOnly, OutReasons))
 			{
 				return true;
 			}
@@ -1821,20 +1829,48 @@ bool IsSomethingToExport(UWorld* World, bool bSelectedOnly)
 	return false;
 }
 
+void OpenOutputLogTab()
+{
+	FGlobalTabmanager::Get()->TryInvokeTab(FName("OutputLog"));
+}
+
 bool ULevelExporterFBX::ExportBinary( UObject* Object, const TCHAR* Type, FArchive& Ar, FFeedbackContext* Warn, int32 FileIndex, uint32 PortFlags )
 {
 	UWorld* World = CastChecked<UWorld>(Object);
+	TArray<FString> GeneratedWarnings;
 	//Check if there is something to export
-	if (!IsSomethingToExport(World, bSelectedOnly))
+	if (!IsSomethingToExport(World, bSelectedOnly, GeneratedWarnings))
 	{
-		UE_LOG(LogEditorExporters, Warning, TEXT("There is nothing to export to a fbx file."));
+		if (GeneratedWarnings.Num() > 0)
+		{
+			for (const FString& WarningMessage : GeneratedWarnings)
+			{
+				UE_LOG(LogEditorExporters, Warning, TEXT("%s"), *WarningMessage);
+			}
+		}
+		else
+		{
+			UE_LOG(LogEditorExporters, Warning, TEXT("There is nothing to export to a fbx file."));
+		}
+
 		if (!GIsAutomationTesting && !FApp::IsUnattended())
 		{
 			FNotificationInfo* NotificationInfo = new FNotificationInfo(FText::GetEmpty());
-			NotificationInfo->Text = bSelectedOnly
-				? FText(NSLOCTEXT("UnrealEd", "ExportingLevelToFBX_Selection", "The selection has nothing that can be exported to fbx!"))
-				: FText(NSLOCTEXT("UnrealEd", "ExportingLevelToFBX_World", "The world has nothing that can be exported to fbx!"));
-			NotificationInfo->ExpireDuration = 5.0f;
+
+			if (GeneratedWarnings.Num() > 0)
+			{
+				NotificationInfo->Text = FText(NSLOCTEXT("UnrealEd", "ExportingLevelToFBX_GeneratedWarnings", "Some warnings were generated while exporting to FBX!"));
+				NotificationInfo->Hyperlink = FSimpleDelegate::CreateStatic(OpenOutputLogTab);
+				NotificationInfo->HyperlinkText = NSLOCTEXT("UnrealEd", "ShowOutputLogButton", "Show Output Log");
+				NotificationInfo->ExpireDuration = 8.0f;
+			}
+			else
+			{
+				NotificationInfo->Text = bSelectedOnly
+					? FText(NSLOCTEXT("UnrealEd", "ExportingLevelToFBX_Selection", "The selection has nothing that can be exported to FBX!"))
+					: FText(NSLOCTEXT("UnrealEd", "ExportingLevelToFBX_World", "The world has nothing that can be exported to FBX!"));
+				NotificationInfo->ExpireDuration = 5.0f;
+			}
 			//QueueNotification will add the pointer to a list and the FSlateNotificationManager::Tick will delete the pointer when adding the notification to the active list.
 			FSlateNotificationManager::Get().QueueNotification(NotificationInfo);
 		}
