Source Control
- Added a notification if we encounter a MarkForAdd failure.
- Shown for 30 seconds unless dimissed.
- Has a link to open the changelists window.

diff --git a/Engine/Source/Editor/UnrealEd/Classes/Editor/EditorEngine.h b/Engine/Source/Editor/UnrealEd/Classes/Editor/EditorEngine.h
--- a/Engine/Source/Editor/UnrealEd/Classes/Editor/EditorEngine.h
+++ b/Engine/Source/Editor/UnrealEd/Classes/Editor/EditorEngine.h
@@ -3049,6 +3049,10 @@ private:
 	/** List of files we are deferring adding to source control */
 	TArray<FString> DeferredFilesToAddToSourceControl;
 
+	/** Ongoing failure notification */
+	TWeakPtr<class SNotificationItem> SourceControlFailureNotificationPtr;
+	void OnDismissSourceControlFailureNotification();
+
 protected:
 
 	/**
diff --git a/Engine/Source/Editor/UnrealEd/Private/EditorEngine.cpp b/Engine/Source/Editor/UnrealEd/Private/EditorEngine.cpp
--- a/Engine/Source/Editor/UnrealEd/Private/EditorEngine.cpp
+++ b/Engine/Source/Editor/UnrealEd/Private/EditorEngine.cpp
@@ -138,6 +138,7 @@
 #include "ContentBrowserModule.h"
 #include "ISourceCodeAccessor.h"
 #include "ISourceCodeAccessModule.h"
+#include "ISourceControlWindowsModule.h"
 
 #include "Settings/EditorSettings.h"
 
@@ -4170,9 +4171,43 @@ void UEditorEngine::RunDeferredMarkForAddFiles(bool)
 		ISourceControlProvider& SourceControlProvider = ISourceControlModule::Get().GetProvider();
 		if(SourceControlProvider.IsAvailable())
 		{
-			ECommandResult::Type Result =  SourceControlProvider.Execute(ISourceControlOperation::Create<FMarkForAdd>(), SourceControlHelpers::PackageFilenames(DeferredFilesToAddToSourceControl));
+			ECommandResult::Type Result = SourceControlProvider.Execute(ISourceControlOperation::Create<FMarkForAdd>(), SourceControlHelpers::PackageFilenames(DeferredFilesToAddToSourceControl));
 
 			bAddToUncontrolled = (Result == ECommandResult::Failed);
+
+			if (Result == ECommandResult::Failed && !SourceControlFailureNotificationPtr.Pin().IsValid())
+			{
+				FTextBuilder InfoTextBuilder;
+				InfoTextBuilder.AppendLine(
+					FText::Format(LOCTEXT("Error_MarkForAddFailed", "{0} Mark for Add failed."),
+					FText::FromName(SourceControlProvider.GetName())));
+
+				if (FUncontrolledChangelistsModule::Get().IsEnabled())
+				{
+					InfoTextBuilder.AppendLine(); // Empty line
+					InfoTextBuilder.AppendLine(LOCTEXT("Error_MarkForAddFailed_Uncontrolled", "Adding files to the Uncontrolled changelist instead."));
+				}
+
+				FNotificationInfo Info(InfoTextBuilder.ToText());
+
+				Info.bFireAndForget = true;
+				Info.bUseSuccessFailIcons = true;
+				Info.Hyperlink = FSimpleDelegate::CreateLambda([]() { ISourceControlWindowsModule::Get().ShowChangelistsTab(); });
+				Info.HyperlinkText = LOCTEXT("ViewChangesHyperlink", "View Changes");
+
+				Info.ButtonDetails.Emplace(
+					LOCTEXT("Dismiss", "Dismiss"),
+					LOCTEXT("DismissTooltip", "Dismiss this notification"),
+					FSimpleDelegate::CreateUObject(this, &UEditorEngine::OnDismissSourceControlFailureNotification),
+					SNotificationItem::CS_Fail);
+
+				if (TSharedPtr<SNotificationItem> Notification = FSlateNotificationManager::Get().AddNotification(Info))
+				{
+					Notification->SetExpireDuration(30.f);
+					Notification->SetCompletionState(SNotificationItem::CS_Fail);
+					SourceControlFailureNotificationPtr = Notification;
+				}
+			}
 		}
 	}
 
@@ -4186,6 +4221,16 @@ void UEditorEngine::RunDeferredMarkForAddFiles(bool)
 	DeferredFilesToAddToSourceControl.Empty();
 }
 
+void UEditorEngine::OnDismissSourceControlFailureNotification()
+{
+	if (TSharedPtr<SNotificationItem> Notification = SourceControlFailureNotificationPtr.Pin())
+	{
+		Notification->SetExpireDuration(0.f);
+		Notification->ExpireAndFadeout();
+		SourceControlFailureNotificationPtr.Reset();
+	}
+}
+
 bool UEditorEngine::InitializePhysicsSceneForSaveIfNecessary(UWorld* World, bool &bOutForceInitialized)
 {
 	TRACE_CPUPROFILER_EVENT_SCOPE(UEditorEngine::InitializePhysicsSceneForSaveIfNecessary);
diff --git a/Engine/Source/Editor/UnrealEd/UnrealEd.Build.cs b/Engine/Source/Editor/UnrealEd/UnrealEd.Build.cs
--- a/Engine/Source/Editor/UnrealEd/UnrealEd.Build.cs
+++ b/Engine/Source/Editor/UnrealEd/UnrealEd.Build.cs
@@ -40,7 +40,8 @@ public class UnrealEd : ModuleRules
 				"Blutility",
 				"SlateReflector",
 				"PackagesDialog",
-				"GraphEditor"
+				"GraphEditor",
+				"SourceControlWindows",
 			}
 		);
 
