Give the new viewport toolbar's realtime toggle a disabled state

* Set the realtime toggle to disabled when the current viewport's realtime state can't be toggled because of overrides.
* Update the tooltip of the realtime toggle to show the override message when an override is set.

diff --git a/Engine/Source/Editor/UnrealEd/Private/ViewportToolbar/UnrealEdViewportToolbar.cpp b/Engine/Source/Editor/UnrealEd/Private/ViewportToolbar/UnrealEdViewportToolbar.cpp
--- a/Engine/Source/Editor/UnrealEd/Private/ViewportToolbar/UnrealEdViewportToolbar.cpp
+++ b/Engine/Source/Editor/UnrealEd/Private/ViewportToolbar/UnrealEdViewportToolbar.cpp
@@ -3200,16 +3200,26 @@ FToolMenuEntry CreateToggleRealtimeEntry()
 							if (TSharedPtr<SEditorViewport> EditorViewport = EditorViewportWeak.Pin())
 							{
 								EditorViewport->OnToggleRealtime();
-								// Calling UToolMenu::RefreshAllWidgets here is cheating. We do it because the menu
-								// entry's TAttribute<FSlateIcon> is only called once when the menu is opened (because
-								// FBaseMenuBuilder::AddMenuEntry takes an FSlateIcon and not a TAttribute<FSlateIcon>).
-								// So when we refresh all widgets here, we force the open menu to close and hide the
-								// fact that the icon wouldn't have updated if the menu stayed open.
+
+								// HACK: We unfortunately have to call UToolMenu::RefreshAllWidgets here to be able
+								// to dynamically set StyleNameOverride below and have that take effect. Without refreshing
+								// widgets, the style name is not applied because the UI has already been built.
 								UToolMenus::Get()->RefreshAllWidgets();
 							}
 						}
 					);
 
+					RealtimeToggleAction.CanExecuteAction = FToolMenuCanExecuteAction::CreateLambda(
+						[EditorViewportWeak](const FToolMenuContext& Context) -> bool
+						{
+							if (TSharedPtr<SEditorViewport> EditorViewport = EditorViewportWeak.Pin())
+							{
+								return !EditorViewport->GetViewportClient()->IsRealtimeOverrideSet();
+							}
+							return true;
+						}
+					);
+
 					RealtimeToggleAction.GetActionCheckState = FToolMenuGetActionCheckState::CreateLambda(
 						[EditorViewportWeak](const FToolMenuContext& Context) -> ECheckBoxState
 						{
@@ -3225,10 +3235,6 @@ FToolMenuEntry CreateToggleRealtimeEntry()
 
 				TAttribute<FText> Tooltip;
 				{
-					const FText NonRealtimeTooltip = LOCTEXT(
-						"ToggleRealtimeTooltip_WarnRealtimeOff",
-						"This viewport is not updating in realtime.  Click to turn on realtime mode."
-					);
 					const FText RealtimeTooltip =
 						LOCTEXT("ToggleRealtimeTooltip", "Toggle realtime rendering of the viewport");
 
@@ -3237,15 +3243,31 @@ FToolMenuEntry CreateToggleRealtimeEntry()
 					if (EditorViewportContext)
 					{
 						Tooltip = TAttribute<FText>::CreateLambda(
-							[EditorViewportWeak, NonRealtimeTooltip, RealtimeTooltip]() -> FText
+							[EditorViewportWeak, RealtimeTooltip]() -> FText
 							{
-								bool bDisplayTopLevel = false;
 								if (const TSharedPtr<SEditorViewport> EditorViewport = EditorViewportWeak.Pin())
 								{
-									bDisplayTopLevel = !EditorViewport->IsRealtime();
+									if (!EditorViewport->IsRealtime())
+									{
+										return LOCTEXT(
+											"ToggleRealtimeTooltip_WarnRealtimeOff", "Warning: This viewport is not updating in realtime. Click to turn on realtime mode."
+										);
+									}
+
+									if (EditorViewport->GetViewportClient()->IsRealtimeOverrideSet())
+									{
+										const FText Format = LOCTEXT(
+											"ToggleRealtimeTooltip_RealtimeOverrides",
+											"Realtime rendering cannot be toggled because an override has been set: {0}"
+										);
+										const FText Message =
+											EditorViewport->GetViewportClient()->GetRealtimeOverrideMessage();
+
+										return FText::Format(Format, Message);
+									}
 								}
 
-								return bDisplayTopLevel ? NonRealtimeTooltip : RealtimeTooltip;
+								return RealtimeTooltip;
 							}
 						);
 					}
