[Viewport Toolbar] Add back proper visibility check for Scalability menu

diff --git a/Engine/Source/Editor/UnrealEd/Private/SCommonEditorViewportToolbarBase.cpp b/Engine/Source/Editor/UnrealEd/Private/SCommonEditorViewportToolbarBase.cpp
--- a/Engine/Source/Editor/UnrealEd/Private/SCommonEditorViewportToolbarBase.cpp
+++ b/Engine/Source/Editor/UnrealEd/Private/SCommonEditorViewportToolbarBase.cpp
@@ -626,7 +626,8 @@ FText SCommonEditorViewportToolbarBase::GetScalabilityWarningLabel() const
 
 EVisibility SCommonEditorViewportToolbarBase::GetScalabilityWarningVisibility() const
 {
-	return UE::UnrealEd::IsScalabilityWarningVisible() ? EVisibility::Visible : EVisibility::Collapsed;
+	return UE::UnrealEd::IsScalabilityWarningVisible() && GetShowScalabilityMenu() ? EVisibility::Visible
+																				   : EVisibility::Collapsed;
 }
 
 TSharedRef<SWidget> SCommonEditorViewportToolbarBase::GetScalabilityWarningMenuContent() const
