Favor using uncontrolled changelists (if possible) when using the BackupAndOverwrite auto-save method

This can be controlled by the "Editor.AutoSave.UseUncontrolledChangelists" CVar

diff --git a/Engine/Source/Editor/UnrealEd/Private/FileHelpers.cpp b/Engine/Source/Editor/UnrealEd/Private/FileHelpers.cpp
--- a/Engine/Source/Editor/UnrealEd/Private/FileHelpers.cpp
+++ b/Engine/Source/Editor/UnrealEd/Private/FileHelpers.cpp
@@ -5487,7 +5487,10 @@ static bool InternalCheckoutAndSavePackages(const TArray<UPackage*>& PackagesToS
 			const bool bErrorIfAlreadyCheckedOut = false;
 			const bool bConfirmPackageBranchCheckOutStatus = false;
 
-			FEditorFileUtils::CheckoutPackages(PackagesToSave, &PackagesCheckedOut, bErrorIfAlreadyCheckedOut, bConfirmPackageBranchCheckOutStatus);
+			if (ISourceControlModule::Get().IsEnabled())
+			{
+				FEditorFileUtils::CheckoutPackages(PackagesToSave, &PackagesCheckedOut, bErrorIfAlreadyCheckedOut, bConfirmPackageBranchCheckOutStatus);
+			}
 
 			// Cannot mark files for add until after packages saved
 			TArray<UPackage*> PackagesToMarkForAdd;
@@ -5505,7 +5508,7 @@ static bool InternalCheckoutAndSavePackages(const TArray<UPackage*>& PackagesToS
 
 			// Mark files for add now that packages have saved
 			PackagesToMarkForAdd.RemoveAll([&FailedPackages](UPackage* Package) { return FailedPackages.Contains(Package); });
-			if (PackagesToMarkForAdd.Num() > 0)
+			if (PackagesToMarkForAdd.Num() > 0 && ISourceControlModule::Get().IsEnabled())
 			{
 				FEditorFileUtils::CheckoutPackages(PackagesToMarkForAdd, nullptr, bErrorIfAlreadyCheckedOut, bConfirmPackageBranchCheckOutStatus);
 			}
diff --git a/Engine/Source/Editor/UnrealEd/Private/PackageAutoSaver.cpp b/Engine/Source/Editor/UnrealEd/Private/PackageAutoSaver.cpp
--- a/Engine/Source/Editor/UnrealEd/Private/PackageAutoSaver.cpp
+++ b/Engine/Source/Editor/UnrealEd/Private/PackageAutoSaver.cpp
@@ -5,6 +5,7 @@
 #include "UObject/ObjectSaveContext.h"
 #include "UObject/Package.h"
 #include "HAL/FileManager.h"
+#include "HAL/PlatformFileManager.h"
 #include "Misc/Paths.h"
 #include "Misc/ScopedSlowTask.h"
 #include "Misc/App.h"
@@ -22,6 +23,7 @@
 #include "UnrealEdMisc.h"
 #include "FileHelpers.h"
 #include "InterchangeManager.h"
+#include "UncontrolledChangelistsModule.h"
 #include "UnrealEdGlobals.h"
 #include "PackageRestore.h"
 #include "Framework/Notifications/NotificationManager.h"
@@ -29,12 +31,19 @@
 #include "AutoSaveUtils.h"
 #include "ShaderCompiler.h"
 #include "EditorLevelUtils.h"
+#include "ISourceControlModule.h"
 #include "IVREditorModule.h"
 #include "LevelEditorViewport.h"
 #include "AssetCompilingManager.h"
 #include "AssetDefinition.h"
 #include "AssetDefinitionRegistry.h"
 
+namespace PackageAutoSaverOptions
+{
+	static bool bUseUncontrolledChangelists = true;
+	static FAutoConsoleVariableRef CVarUseUncontrolledChangelists(TEXT("Editor.AutoSave.UseUncontrolledChangelists"), bUseUncontrolledChangelists, TEXT("True to use uncontrolled changelists (if possible) when using the BackupAndOverwrite auto-save method."));
+}
+
 namespace PackageAutoSaverJson
 {
 	typedef TCHAR CharType;
@@ -252,7 +261,7 @@ void FPackageAutoSaver::AttemptAutoSave()
 
 			if (LoadingSavingSettings->AutoSaveMethod == EAutoSaveMethod::BackupAndRestore)
 			{
-				bIsAutoSaving = true;
+				TGuardValue<bool> IsAutoSavingScope(bIsAutoSaving, true);
 
 				SlowTask.EnterProgressFrame(50);
 
@@ -287,7 +296,6 @@ void FPackageAutoSaver::AttemptAutoSave()
 						{
 							if (UPackage* Pkg = PackageToBackup.Get())
 							{
-								PackageFilename.Reset();
 								if (FPackageName::DoesPackageExist(PackageToBackup->GetPathName(), &PackageFilename))
 								{
 									const FString PackageAutoSaveFilename = FEditorFileUtils::GetAutoSaveFilename(Pkg, AutoSaveDir, NewAutoSaveIndex, FPaths::GetExtension(PackageFilename, /*bIncludeDot*/true));
@@ -307,12 +315,18 @@ void FPackageAutoSaver::AttemptAutoSave()
 					}
 				}
 
+				// If possible, auto-save will use an uncontrolled changelist rather than actually checkout the files
+				// TODO: Have a dedicated auto-save uncontrolled changelist rather than use the default
+				FUncontrolledChangelistsModule& UncontrolledChangelistsModule = FUncontrolledChangelistsModule::Get();
+				const bool bMakeFilesWritable = PackageAutoSaverOptions::bUseUncontrolledChangelists && UncontrolledChangelistsModule.IsEnabled() && ISourceControlModule::Get().GetProvider().UsesLocalReadOnlyState();
+
 				// Build the complete list of packages to save
 				TArray<UPackage*> PackagesToSave;
 				{
-					auto AppendPackagesToSave = [&PackagesToSave](const TSet<TWeakObjectPtr<UPackage>, TWeakObjectPtrSetKeyFuncs<TWeakObjectPtr<UPackage>>>& PotentialPackagesToSave)
+					auto AppendPackagesToSave = [&PackagesToSave, &UncontrolledChangelistsModule, bMakeFilesWritable](const TSet<TWeakObjectPtr<UPackage>, TWeakObjectPtrSetKeyFuncs<TWeakObjectPtr<UPackage>>>& PotentialPackagesToSave)
 					{
 						FNameBuilder PkgName;
+						FString PackageFilename;
 						for (const TWeakObjectPtr<UPackage>& PotentialPackageToSave : PotentialPackagesToSave)
 						{
 							if (UPackage* Pkg = PotentialPackageToSave.Get())
@@ -324,6 +338,18 @@ void FPackageAutoSaver::AttemptAutoSave()
 								if (FPackageName::IsValidLongPackageName(PkgName.ToView(), /*bIncludeReadOnlyRoots*/false))
 								{
 									PackagesToSave.Add(Pkg);
+
+									// If using an uncontrolled changelist to track writable files, then make the file writable now so that we can overwrite it during save
+									if (bMakeFilesWritable)
+									{
+										if (FPackageName::DoesPackageExist(Pkg->GetPathName(), &PackageFilename) && FPlatformFileManager::Get().GetPlatformFile().IsReadOnly(*PackageFilename))
+										{
+											if (FPlatformFileManager::Get().GetPlatformFile().SetReadOnly(*PackageFilename, false))
+											{
+												UncontrolledChangelistsModule.OnMakeWritable(PackageFilename);
+											}
+										}
+									}
 								}
 							}
 						}
@@ -342,7 +368,12 @@ void FPackageAutoSaver::AttemptAutoSave()
 				if (PackagesToSave.Num() > 0)
 				{
 					// Note: The in-place save does a regular save of dirty packages, so it doesn't set the bIsAutoSaving flag since it functions like a user-initiated save
-					const bool bSuccess = UEditorLoadingAndSavingUtils::SavePackages(PackagesToSave, false);
+					bool bSuccess = false;
+					{
+						// Disable source control around UEditorLoadingAndSavingUtils::SavePackages when making files writable to avoid it still checking them them out
+						FScopedDisableSourceControl DisableSourceControl(bMakeFilesWritable);
+						bSuccess = UEditorLoadingAndSavingUtils::SavePackages(PackagesToSave, false);
+					}
 
 					// Note: We don't update DirtyMapsForAutoSave/DirtyContentForAutoSave manually post-save, as any packages that were actually saved will have been removed from the list via the save callback
 					if (bSuccess)
@@ -384,8 +415,6 @@ void FPackageAutoSaver::AttemptAutoSave()
 
 			ResetAutoSaveTimer();
 			bDelayingDueToFailedSave = false;
-
-			bIsAutoSaving = false;
 		}
 		else
 		{
