Disallow addition of interfaces with anim functions to non animbps

[RN] Fixed crash when adding an interface with functions that have "AnimBlueprintFunction" metadata to a non-animation blueprint

diff --git a/Engine/Source/Editor/UnrealEd/Private/Kismet2/BlueprintEditorUtils.cpp b/Engine/Source/Editor/UnrealEd/Private/Kismet2/BlueprintEditorUtils.cpp
--- a/Engine/Source/Editor/UnrealEd/Private/Kismet2/BlueprintEditorUtils.cpp
+++ b/Engine/Source/Editor/UnrealEd/Private/Kismet2/BlueprintEditorUtils.cpp
@@ -6395,7 +6395,20 @@ bool FBlueprintEditorUtils::ImplementNewInterface(UBlueprint* Blueprint, FTopLev
 	for( TFieldIterator<UFunction> FunctionIter(InterfaceClass, EFieldIteratorFlags::IncludeSuper); FunctionIter; ++FunctionIter )
 	{
 		UFunction* Function = *FunctionIter;
-		const bool bIsAnimFunction = Function->HasMetaData(FBlueprintMetadata::MD_AnimBlueprintFunction) && Blueprint->IsA<UAnimBlueprint>();
+		const bool bIsAnimFunction = Function->HasMetaData(FBlueprintMetadata::MD_AnimBlueprintFunction);
+		if(bIsAnimFunction && !Blueprint->IsA<UAnimBlueprint>())
+		{
+			bAllFunctionsAdded = false;
+			ShowNotification(
+				FText::Format(
+					LOCTEXT("AddedAnimFunctionToNonAnimBlueprint", "Blueprint '{0}' is not an animation blueprint and therefore cannot implement interface '{1}'"),
+					FText::FromString(Blueprint->GetName()),
+					FText::FromString(InterfaceClassPathName.ToString())
+				),
+				EMessageSeverity::Error
+			);
+			break;
+		}
 		if( (UEdGraphSchema_K2::CanKismetOverrideFunction(Function) && !UEdGraphSchema_K2::FunctionCanBePlacedAsEvent(Function)) || 
 			bIsAnimFunction)
 		{
