Build combined thumbnail cache in memory
#rnx

diff --git a/Engine/Source/Editor/UnrealEd/Private/ThumbnailExternalCache.cpp b/Engine/Source/Editor/UnrealEd/Private/ThumbnailExternalCache.cpp
--- a/Engine/Source/Editor/UnrealEd/Private/ThumbnailExternalCache.cpp
+++ b/Engine/Source/Editor/UnrealEd/Private/ThumbnailExternalCache.cpp
@@ -69,41 +69,12 @@ namespace ThumbnailExternalCache
 	}
 }
 
-struct FThumbnailDeduplicateKey
-{
-	FThumbnailDeduplicateKey() {}
-	FThumbnailDeduplicateKey(uint64 InHash, int32 InNumBytes) : Hash(InHash), NumBytes(InNumBytes) {}
-
-	uint64 Hash = 0;
-	int32 NumBytes = 0;
-
-	bool operator ==( const FThumbnailDeduplicateKey& Other ) const
-	{
-		return Hash == Other.Hash && NumBytes == Other.NumBytes;
-	}
-};
-
-inline uint32 GetTypeHash(const FThumbnailDeduplicateKey& InValue)
-{
-	return GetTypeHash(InValue.Hash);
-}
-
 struct FPackageThumbnailRecord
 {
 	FName Name;
 	int64 Offset = 0;
 };
 
-class FSaveThumbnailCacheTask
-{
-public:
-	FObjectThumbnail ObjectThumbnail;
-	FName Name;
-	uint64 CompressedBytesHash = 0;
-
-	void Compress(const FThumbnailExternalCacheSettings& InSettings);
-};
-
 class FSaveThumbnailCache
 {
 public:
@@ -111,6 +82,7 @@ public:
 	~FSaveThumbnailCache();
 
 	void Save(FArchive& Ar, const TArrayView<FAssetData> InAssetDatas, const FThumbnailExternalCacheSettings& InSettings);
+	void Save(FArchive& Ar, FCombinedThumbnailCacheToSave& CombinedCache, bool bSort);
 };
 
 FSaveThumbnailCache::FSaveThumbnailCache()
@@ -121,33 +93,119 @@ FSaveThumbnailCache::~FSaveThumbnailCache()
 {
 }
 
-void FSaveThumbnailCache::Save(FArchive& Ar, const TArrayView<FAssetData> InAssetDatas, const FThumbnailExternalCacheSettings& InSettings)
+void FThumbnailExternalCache::LoadCompressAndAppend(const TArrayView<FAssetData> InAssetDatas, FCombinedThumbnailCacheToSave& CombinedCache)
 {
-	// Reduce peak memory to support larger asset counts by loading then saving in batches
-	int32 TaskBatchSize = 100000;
+	const double StartTime = FPlatformTime::Seconds();
+
+	struct FAssetToProcess
+	{
+		const FAssetData* AssetData = nullptr;
+		TSharedPtr<FSaveThumbnailCacheTask> Task;
+	};
+
+	TArray<FAssetToProcess> AssetsToProcess;
+	AssetsToProcess.Reserve(InAssetDatas.Num());
+
+	CombinedCache.Tasks.Reserve(CombinedCache.Tasks.Num() + InAssetDatas.Num());
+
+	for (const FAssetData& AssetData : InAssetDatas)
+	{
+		FNameBuilder ObjectFullNameBuilder;
+		AssetData.GetFullName(ObjectFullNameBuilder);
+		FName ObjectFullName(ObjectFullNameBuilder);
+
+		if (!CombinedCache.Tasks.Contains(ObjectFullName))
+		{
+			FAssetToProcess& AssetToProcess = AssetsToProcess.AddDefaulted_GetRef();
+			AssetToProcess.AssetData = &AssetData;			 
+			AssetToProcess.Task = MakeShared<FSaveThumbnailCacheTask>();
+			AssetToProcess.Task->Name = ObjectFullName;
 
+			CombinedCache.Tasks.Add(ObjectFullName, AssetToProcess.Task);
+		}
+	}
+
+	// Load and compress
+	ParallelFor(AssetsToProcess.Num(), [&AssetsToProcess, &CombinedCache](int32 Index)
+	{
+		FAssetToProcess& AssetToProcess = AssetsToProcess[Index];
+		if (ThumbnailTools::LoadThumbnailFromPackage(*AssetToProcess.AssetData, AssetToProcess.Task->ObjectThumbnail) && !AssetToProcess.Task->ObjectThumbnail.IsEmpty())
+		{
+			AssetToProcess.Task->Compress(CombinedCache.Settings);
+		}
+		else
+		{
+			AssetToProcess.Task.Reset();
+		}
+	});
+
+	// Deduplicate to free up memory
+	// Assumes compression method is deterministic (refactor if not deterministic to hash decompressed thumbnails)
+	CombinedCache.DeduplicateMap.Reserve(CombinedCache.DeduplicateMap.Num() + AssetsToProcess.Num());
+	for (FAssetToProcess& AssetToProcess : AssetsToProcess)
+	{
+		if (AssetToProcess.Task.IsValid() && !AssetToProcess.Task->ObjectThumbnail.IsEmpty())
+		{
+			FSaveThumbnailCacheDeduplicateKey DeduplicateKey(AssetToProcess.Task->CompressedBytesHash, AssetToProcess.Task->ObjectThumbnail.GetCompressedDataSize());
+			if (TSharedPtr<FSaveThumbnailCacheTask>* ExistingData = CombinedCache.DeduplicateMap.Find(DeduplicateKey))
+			{
+				// Point to the other object's thumbnail
+				CombinedCache.Tasks[AssetToProcess.Task->Name] = *ExistingData;
+			}
+			else
+			{
+				CombinedCache.DeduplicateMap.Add(DeduplicateKey, AssetToProcess.Task);
+			}
+		}
+	}
+
+	CombinedCache.AccumlatedLoadTime += FPlatformTime::Seconds() - StartTime;
+}
+
+void FSaveThumbnailCache::Save(FArchive& Ar, FCombinedThumbnailCacheToSave& CombinedCache, bool bSort)
+{
 	const double TimeStart = FPlatformTime::Seconds();
 
-	const int32 NumAssetDatas = InAssetDatas.Num();
+	// Sorting is often done to reduce size of patches
+	if (bSort)
+	{
+		CombinedCache.Tasks.KeyStableSort(FNameLexicalLess());
+	}
+
+	// Only write information about assets that contain thumbnails
+	TArray<TPair<FName, TSharedPtr<FSaveThumbnailCacheTask>>> AssetsToWrite;
+	AssetsToWrite.Reserve(CombinedCache.Tasks.Num());
+	for (TPair<FName, TSharedPtr<FSaveThumbnailCacheTask>>& It : CombinedCache.Tasks)
+	{
+		if (It.Value.IsValid() && !It.Value->ObjectThumbnail.IsEmpty())
+		{
+			AssetsToWrite.Add(It);
+		}
+	}
 
-	UE_LOG(LogThumbnailExternalCache, Log, TEXT("Saving thumbnails for %d assets to %s"), NumAssetDatas, *Ar.GetArchiveName());
+	const int32 NumAssetDatas = AssetsToWrite.Num();
+
+	UE_LOG(LogThumbnailExternalCache, Log, TEXT("Saving thumbnails for %d/%d assets (%d unique thumbnails) to %s"), NumAssetDatas, CombinedCache.Tasks.Num(), CombinedCache.DeduplicateMap.Num(), *Ar.GetArchiveName());
 
 	FText StatusText = LOCTEXT("SaveStatus", "Saving Thumbnails: {0}");
-	FScopedSlowTask SlowTask( (float)NumAssetDatas / (float)TaskBatchSize, FText::Format(StatusText, FText::AsNumber(NumAssetDatas)));
+	FScopedSlowTask SlowTask((float)NumAssetDatas, FText::Format(StatusText, FText::AsNumber(NumAssetDatas)));
 	SlowTask.MakeDialog(/*bShowCancelButton*/ false);
 
+	const double SaveTimeStart = FPlatformTime::Seconds();
+
 	TArray<FPackageThumbnailRecord> PackageThumbnailRecords;
 	PackageThumbnailRecords.Reset();
 	PackageThumbnailRecords.Reserve(NumAssetDatas);
 
-	TMap<FThumbnailDeduplicateKey, int64> DeduplicateMap;	
+	TMap<FSaveThumbnailCacheDeduplicateKey, int64> DeduplicateMap;
 	DeduplicateMap.Reset();
-	DeduplicateMap.Reserve(NumAssetDatas);
+	DeduplicateMap.Reserve(CombinedCache.DeduplicateMap.Num());
 	int32 NumDuplicates = 0;
 	int64 DuplicateBytesSaved = 0;
 
 	int64 TotalCompressedBytes = 0;
 
+	// Write Header
 	{
 		FThumbnailExternalCache::FThumbnailExternalCacheHeader Header;
 		Header.HeaderId = ThumbnailExternalCache::ExpectedHeaderId;
@@ -158,80 +216,34 @@ void FSaveThumbnailCache::Save(FArchive& Ar, const TArrayView<FAssetData> InAsse
 	}
 	const int64 ThumbnailTableOffsetPos = Ar.Tell() - sizeof(int64);
 
-	double LoadTime = 0.0;
-	double SaveTime = 0.0;
-
-	for (int32 StartIndex = 0; StartIndex < InAssetDatas.Num(); StartIndex += TaskBatchSize)
+	// Write compressed image data
+	for (TPair<FName, TSharedPtr<FSaveThumbnailCacheTask>>& It : AssetsToWrite)
 	{
-		const int32 EndIndex = FMath::Min<int32>(StartIndex + TaskBatchSize, InAssetDatas.Num());
-		const TArrayView<FAssetData> AssetsForSegment(&InAssetDatas[StartIndex], EndIndex - StartIndex);
-
-		TArray<FSaveThumbnailCacheTask> Tasks;
-		Tasks.AddDefaulted(AssetsForSegment.Num());
-
-		const double LoadTimeStart = FPlatformTime::Seconds();
-
-		// Load then recompress if needed. Thread for improved load performance.
-		ParallelFor(AssetsForSegment.Num(), [AssetsForSegment, &Tasks, &InSettings](int32 Index)
+		// Add table of contents entry
+		FPackageThumbnailRecord& PackageThumbnailRecord = PackageThumbnailRecords.AddDefaulted_GetRef();
+		PackageThumbnailRecord.Name = It.Key;
+
+		// Image data
+		FSaveThumbnailCacheTask& Task = *It.Value;
+		FSaveThumbnailCacheDeduplicateKey DeduplicateKey(Task.CompressedBytesHash, Task.ObjectThumbnail.GetCompressedDataSize());
+		if (const int64* ExistingOffset = DeduplicateMap.Find(DeduplicateKey))
 		{
-			FSaveThumbnailCacheTask& Task = Tasks[Index];
-			if (ThumbnailTools::LoadThumbnailFromPackage(AssetsForSegment[Index], Task.ObjectThumbnail) && !Task.ObjectThumbnail.IsEmpty())
-			{
-				{
-					FNameBuilder ObjectFullNameBuilder;
-					AssetsForSegment[Index].GetFullName(ObjectFullNameBuilder);
-					Task.Name = FName(ObjectFullNameBuilder);
-				}
-
-				Task.Compress(InSettings);
-			}
-		});
-
-		const double SaveTimeStart = FPlatformTime::Seconds();
-
-		// Save compressed image data
-		for (FSaveThumbnailCacheTask& Task : Tasks)
-		{
-			// Skip empty
-			if (Task.ObjectThumbnail.IsEmpty())
-			{
-				continue;
-			}
-
-			// Add table of contents entry
-			FPackageThumbnailRecord& PackageThumbnailRecord = PackageThumbnailRecords.AddDefaulted_GetRef();
-			PackageThumbnailRecord.Name = Task.Name;
-
-			FThumbnailDeduplicateKey DeduplicateKey(Task.CompressedBytesHash, Task.ObjectThumbnail.GetCompressedDataSize());
-			if (const int64* ExistingOffset = DeduplicateMap.Find(DeduplicateKey))
-			{
-				// Reference existing compressed image data
-				PackageThumbnailRecord.Offset = *ExistingOffset;
-				DuplicateBytesSaved += DeduplicateKey.NumBytes;
-				++NumDuplicates;
-			}
-			else
-			{
-				// Save compressed image data
-				PackageThumbnailRecord.Offset = Ar.Tell();
-				Task.ObjectThumbnail.Serialize(Ar);
-				DeduplicateMap.Add(DeduplicateKey, PackageThumbnailRecord.Offset);
-				TotalCompressedBytes += DeduplicateKey.NumBytes;
-			}
-
-			// Free memory
-			Task.ObjectThumbnail.AccessCompressedImageData().Empty();
+			// Reference existing compressed image data
+			PackageThumbnailRecord.Offset = *ExistingOffset;
+			DuplicateBytesSaved += DeduplicateKey.NumBytes;
+			++NumDuplicates;
 		}
-
-		SaveTime += FPlatformTime::Seconds() - SaveTimeStart;
-		LoadTime += SaveTimeStart - LoadTimeStart;
-
-		SlowTask.EnterProgressFrame((float)AssetsForSegment.Num() / (float)TaskBatchSize, FText::Format(StatusText, FText::AsNumber(NumAssetDatas - EndIndex)));
-		if (SlowTask.ShouldCancel())
+		else
 		{
-			PackageThumbnailRecords.Reset();
-			break;
+			// Save compressed image data
+			PackageThumbnailRecord.Offset = Ar.Tell();
+			Task.ObjectThumbnail.Serialize(Ar);
+			DeduplicateMap.Add(DeduplicateKey, PackageThumbnailRecord.Offset);
+			TotalCompressedBytes += DeduplicateKey.NumBytes;
 		}
+
+		// Free memory
+		Task.ObjectThumbnail.AccessCompressedImageData().Empty();
 	}
 
 	// Save table of contents
@@ -256,11 +268,22 @@ void FSaveThumbnailCache::Save(FArchive& Ar, const TArrayView<FAssetData> InAsse
 	Ar.Seek(ThumbnailTableOffsetPos);
 	Ar << NewThumbnailTableOffset;
 
-	UE_LOG(LogThumbnailExternalCache, Log, TEXT("Load Time: %f secs, Save Time: %f secs, Total Time: %f secs"), LoadTime, SaveTime, FPlatformTime::Seconds() - TimeStart);
+	const double SaveTime = FPlatformTime::Seconds() - SaveTimeStart;
+
+	UE_LOG(LogThumbnailExternalCache, Log, TEXT("Load Time: %f secs, Save Time: %f secs, Total Time: %f secs"), CombinedCache.AccumlatedLoadTime, SaveTime, (FPlatformTime::Seconds() - TimeStart) + CombinedCache.AccumlatedLoadTime);
 	UE_LOG(LogThumbnailExternalCache, Log, TEXT("Thumbnails: %d, %f MB"), PackageThumbnailRecords.Num(), (TotalCompressedBytes / (1024.0 * 1024.0)));
 	UE_LOG(LogThumbnailExternalCache, Log, TEXT("Duplicates: %d, %f MB"), NumDuplicates, (DuplicateBytesSaved / (1024.0 * 1024.0)));
 }
 
+void FSaveThumbnailCache::Save(FArchive& Ar, const TArrayView<FAssetData> InAssetDatas, const FThumbnailExternalCacheSettings& InSettings)
+{
+	FCombinedThumbnailCacheToSave CombinedCache;
+	CombinedCache.Settings = InSettings;
+
+	FThumbnailExternalCache::Get().LoadCompressAndAppend(InAssetDatas, CombinedCache);
+	Save(Ar, CombinedCache, true);
+}
+
 void FSaveThumbnailCacheTask::Compress(const FThumbnailExternalCacheSettings& InSettings)
 {
 	ThumbnailExternalCache::ResizeThumbnailIfNeeded(ObjectThumbnail, InSettings.MaxImageSize);
@@ -426,6 +449,22 @@ void FThumbnailExternalCache::SortAssetDatas(TArray<FAssetData>& AssetDatas)
 	Algo::SortBy(AssetDatas, [](const FAssetData& Data) { return Data.PackageName; }, FNameLexicalLess());
 }
 
+bool FThumbnailExternalCache::SaveExternalCache(const FString& InFilename, FCombinedThumbnailCacheToSave& InCache, const bool bSort)
+{
+	bIsSavingCache = true;
+	if (TUniquePtr<FArchive> FileWriter = TUniquePtr<FArchive>(IFileManager::Get().CreateFileWriter(*InFilename)))
+	{
+		FSaveThumbnailCache SaveJob;
+		SaveJob.Save(*FileWriter, InCache, bSort);
+		bIsSavingCache = false;
+		return true;
+	}
+
+	bIsSavingCache = false;
+
+	return false;
+}
+
 bool FThumbnailExternalCache::SaveExternalCache(const FString& InFilename, const TArrayView<FAssetData> InAssetDatas, const FThumbnailExternalCacheSettings& InSettings)
 {
 	bIsSavingCache = true;
diff --git a/Engine/Source/Editor/UnrealEd/Public/ThumbnailExternalCache.h b/Engine/Source/Editor/UnrealEd/Public/ThumbnailExternalCache.h
--- a/Engine/Source/Editor/UnrealEd/Public/ThumbnailExternalCache.h
+++ b/Engine/Source/Editor/UnrealEd/Public/ThumbnailExternalCache.h
@@ -18,6 +18,43 @@ struct FThumbnailExternalCacheSettings
 	int32 MaxImageSize = INT32_MAX;
 };
 
+class FSaveThumbnailCacheTask
+{
+public:
+	FObjectThumbnail ObjectThumbnail;
+	FName Name;
+	uint64 CompressedBytesHash = 0;
+
+	void Compress(const FThumbnailExternalCacheSettings& InSettings);
+};
+
+struct FSaveThumbnailCacheDeduplicateKey
+{
+	FSaveThumbnailCacheDeduplicateKey() {}
+	FSaveThumbnailCacheDeduplicateKey(uint64 InHash, int32 InNumBytes) : Hash(InHash), NumBytes(InNumBytes) {}
+
+	uint64 Hash = 0;
+	int32 NumBytes = 0;
+
+	bool operator ==(const FSaveThumbnailCacheDeduplicateKey& Other) const
+	{
+		return Hash == Other.Hash && NumBytes == Other.NumBytes;
+	}
+};
+
+inline uint32 GetTypeHash(const FSaveThumbnailCacheDeduplicateKey& InValue)
+{
+	return GetTypeHash(InValue.Hash);
+}
+
+struct FCombinedThumbnailCacheToSave
+{
+	TMap<FName, TSharedPtr<FSaveThumbnailCacheTask>> Tasks;
+	FThumbnailExternalCacheSettings Settings;
+	TMap<FSaveThumbnailCacheDeduplicateKey, TSharedPtr<FSaveThumbnailCacheTask>> DeduplicateMap;
+	double AccumlatedLoadTime = 0.0;
+};
+
 class FThumbnailExternalCache
 {
 public:
@@ -37,6 +74,9 @@ public:
 	/** Save thumbnails for the given assets to an external file. Deterministic if assets were sorted. */
 	UNREALED_API bool SaveExternalCache(const FString& InFilename, const TArrayView<FAssetData> InAssetDatas, const FThumbnailExternalCacheSettings& InSettings);
 
+	/** Save thumbnails for the given assets to an external file. Deterministic if bSort=true. */
+	UNREALED_API bool SaveExternalCache(const FString& InFilename, FCombinedThumbnailCacheToSave& InCache, const bool bSort);
+
 	/** Sort asset data list if determinism needed */
 	UNREALED_API static void SortAssetDatas(TArray<FAssetData>& InOutAssetDatas);
 
@@ -81,6 +121,8 @@ public:
 		TMap<FName, FThumbnailEntry> NameToEntry;
 	};
 
+	UNREALED_API void LoadCompressAndAppend(const TArrayView<FAssetData> InAssetDatas, FCombinedThumbnailCacheToSave& CombinedCache);
+
 private:
 
 	void SaveExternalCache(FArchive& Ar, const TArrayView<FAssetData> InAssetDatas, const FThumbnailExternalCacheSettings& InSettings);
