Added an uncontrolled changelist for auto-save, and move newly uncontrolled files to that changelist after performing an auto-save

diff --git a/Engine/Source/Editor/UnrealEd/Private/PackageAutoSaver.cpp b/Engine/Source/Editor/UnrealEd/Private/PackageAutoSaver.cpp
--- a/Engine/Source/Editor/UnrealEd/Private/PackageAutoSaver.cpp
+++ b/Engine/Source/Editor/UnrealEd/Private/PackageAutoSaver.cpp
@@ -316,7 +316,6 @@ void FPackageAutoSaver::AttemptAutoSave()
 				}
 
 				// If possible, auto-save will use an uncontrolled changelist rather than actually checkout the files
-				// TODO: Have a dedicated auto-save uncontrolled changelist rather than use the default
 				FUncontrolledChangelistsModule& UncontrolledChangelistsModule = FUncontrolledChangelistsModule::Get();
 				const bool bMakeFilesWritable = PackageAutoSaverOptions::bUseUncontrolledChangelists && UncontrolledChangelistsModule.IsEnabled() && ISourceControlModule::Get().GetProvider().UsesLocalReadOnlyState();
 
@@ -370,9 +369,38 @@ void FPackageAutoSaver::AttemptAutoSave()
 					// Note: The in-place save does a regular save of dirty packages, so it doesn't set the bIsAutoSaving flag since it functions like a user-initiated save
 					bool bSuccess = false;
 					{
+						FUncontrolledChangelistStatePtr DefaultUncontrolledChangelistState = UncontrolledChangelistsModule.GetDefaultChangelistState();
+
+						TSet<FString> PreSaveFilesInDefaultUncontrolledChangelist;
+						if (bMakeFilesWritable && DefaultUncontrolledChangelistState)
+						{
+							PreSaveFilesInDefaultUncontrolledChangelist.Append(DefaultUncontrolledChangelistState->GetFilenames());
+						}
+
 						// Disable source control around UEditorLoadingAndSavingUtils::SavePackages when making files writable to avoid it still checking them them out
 						FScopedDisableSourceControl DisableSourceControl(bMakeFilesWritable);
 						bSuccess = UEditorLoadingAndSavingUtils::SavePackages(PackagesToSave, false);
+
+						// Move any files that the auto-save added to the default uncontrolled changelist to the auto-save specific uncontrolled changelist
+						// Note: We don't check bSuccess here, as a subset of files might still have been saved
+						if (bMakeFilesWritable && DefaultUncontrolledChangelistState)
+						{
+							TSet<FString> PostSaveFilesInDefaultUncontrolledChangelist = TSet<FString>(DefaultUncontrolledChangelistState->GetFilenames());
+							TSet<FString> NewFilesInDefaultUncontrolledChangelist = PostSaveFilesInDefaultUncontrolledChangelist.Difference(PreSaveFilesInDefaultUncontrolledChangelist);
+
+							if (NewFilesInDefaultUncontrolledChangelist.Num() > 0)
+							{
+								FUncontrolledChangelist AutoSaveUncontrolledChangelist(FGuid(0x0D227D90, 0x2AC34632, 0xB1A35533, 0xF3CC134E));
+								if (UncontrolledChangelistsModule.CreateUncontrolledChangelist(NSLOCTEXT("AutoSave", "UncontrolledChangelistTitle", "Auto Save Uncontrolled Changelist"), AutoSaveUncontrolledChangelist))
+								{
+									TArray<FSourceControlStateRef> FileStates;
+									if (ISourceControlModule::Get().GetProvider().GetState(NewFilesInDefaultUncontrolledChangelist.Array(), FileStates, EStateCacheUsage::Use) == ECommandResult::Succeeded)
+									{
+										UncontrolledChangelistsModule.MoveFilesToUncontrolledChangelist(TArray<FSourceControlStateRef>(), FileStates, AutoSaveUncontrolledChangelist);
+									}
+								}
+							}
+						}
 					}
 
 					// Note: We don't update DirtyMapsForAutoSave/DirtyContentForAutoSave manually post-save, as any packages that were actually saved will have been removed from the list via the save callback
