IncrementalCook: ICookPackageSplitter: Prevent spurious warning of missing splitter objects for GenerationHelpers that were found invalid. GenerationHelpers found invalid can occur if the splitter created generated packages in a previous cook but in the current cook has changed and now returns false for ShouldSplit. The PreGarbageCollect and PostGarbageCollect calls can occur on Invalid CookPackageSplitters and need to not expect splitter objects to be loaded in that case.
#rnx

diff --git a/Engine/Source/Editor/UnrealEd/Private/Cooker/CookGenerationHelper.cpp b/Engine/Source/Editor/UnrealEd/Private/Cooker/CookGenerationHelper.cpp
--- a/Engine/Source/Editor/UnrealEd/Private/Cooker/CookGenerationHelper.cpp
+++ b/Engine/Source/Editor/UnrealEd/Private/Cooker/CookGenerationHelper.cpp
@@ -1537,7 +1537,7 @@ void FGenerationHelper::PreGarbageCollect(const TRefCountPtr<FGenerationHelper>&
 	{
 		PreGarbageCollectGCLifetimeData();
 	}
-	if (!IsInitialized())
+	if (!IsInitialized() || !IsValid())
 	{
 		return;
 	}
@@ -1762,12 +1762,12 @@ void FGenerationHelper::MarkPackageIterativelySkipped(FPackageData& PackageData)
 void FGenerationHelper::PostGarbageCollect(const TRefCountPtr<FGenerationHelper>& RefcountHeldByCaller, FCookGCDiagnosticContext& Context)
 {
 	PostGarbageCollectGCLifetimeData(Context);
-	if (!IsInitialized())
+	bCurrentGCHasKeptGeneratorPackage = false;
+	bCurrentGCHasKeptGeneratorKeepPackages = false;
+	if (!IsInitialized() || !IsValid())
 	{
 		return;
 	}
-	bCurrentGCHasKeptGeneratorPackage = false;
-	bCurrentGCHasKeptGeneratorKeepPackages = false;
 
 	FPackageData& Owner = GetOwner();
 	if (Owner.IsInStateProperty(EPackageStateProperty::Saving))
@@ -1891,7 +1891,7 @@ void FGenerationHelper::VerifyGeneratorPackageGarbageCollected(FCookGCDiagnostic
 
 void FGenerationHelper::UpdateSaveAfterGarbageCollect(const FPackageData& PackageData, bool& bInOutDemote)
 {
-	if (!IsInitialized())
+	if (!IsInitialized() || !IsValid())
 	{
 		return;
 	}
