[ContentBrowser] Graphic Tweaks for the new content browser

diff --git a/Engine/Plugins/Editor/EngineAssetDefinitions/Source/Private/AssetDefinition_ForceFeedbackEffect.cpp b/Engine/Plugins/Editor/EngineAssetDefinitions/Source/Private/AssetDefinition_ForceFeedbackEffect.cpp
--- a/Engine/Plugins/Editor/EngineAssetDefinitions/Source/Private/AssetDefinition_ForceFeedbackEffect.cpp
+++ b/Engine/Plugins/Editor/EngineAssetDefinitions/Source/Private/AssetDefinition_ForceFeedbackEffect.cpp
@@ -319,15 +319,9 @@ bool UAssetDefinition_ForceFeedbackEffect::GetThumbnailActionOverlay(const FAsse
 		return LOCTEXT("Thumbnail_PlayForceFeedbackToolTip", "Play selected force feedback effect");
 	};
 
-	OutActionOverlayInfo.ActionButtonWidget = SNew(SButton)
-		.ButtonStyle(FAppStyle::Get(), "HoverHintOnly")
-		.ContentPadding(0.0f)
+	OutActionOverlayInfo.ActionButtonArgs = SButton::FArguments()
 		.ToolTipText_Lambda(OnToolTipTextLambda)
-		.OnClicked_Lambda(OnClickedLambda)
-		[
-			SNew(SImage)
-			.Image_Lambda(OnGetDisplayBrushLambda)
-		];
+		.OnClicked_Lambda(OnClickedLambda);
 
 	return true;
 }
diff --git a/Engine/Plugins/Runtime/Metasound/Source/MetasoundEditor/Private/MetasoundAssetDefinitions.cpp b/Engine/Plugins/Runtime/Metasound/Source/MetasoundEditor/Private/MetasoundAssetDefinitions.cpp
--- a/Engine/Plugins/Runtime/Metasound/Source/MetasoundEditor/Private/MetasoundAssetDefinitions.cpp
+++ b/Engine/Plugins/Runtime/Metasound/Source/MetasoundEditor/Private/MetasoundAssetDefinitions.cpp
@@ -348,15 +348,9 @@ bool UAssetDefinition_MetaSoundSource::GetThumbnailActionOverlay(const FAssetDat
 		return FReply::Handled();
 	};
 
-	OutActionOverlayInfo.ActionButtonWidget = SNew(SButton)
-		.ButtonStyle(FAppStyle::Get(), "HoverHintOnly")
-		.ContentPadding(0.0f)
+	OutActionOverlayInfo.ActionButtonArgs = SButton::FArguments()
 		.ToolTipText_Lambda(OnToolTipTextLambda)
-		.OnClicked_Lambda(OnClickedLambda)
-		[
-			SNew(SImage)
-			.Image_Lambda(OnGetDisplayBrushLambda)
-		];
+		.OnClicked_Lambda(OnClickedLambda);
 
 	return true;
 }
diff --git a/Engine/Source/Editor/AssetDefinition/Public/AssetDefinition.h b/Engine/Source/Editor/AssetDefinition/Public/AssetDefinition.h
--- a/Engine/Source/Editor/AssetDefinition/Public/AssetDefinition.h
+++ b/Engine/Source/Editor/AssetDefinition/Public/AssetDefinition.h
@@ -9,6 +9,7 @@
 #include "Misc/ScopedSlowTask.h"
 #include "Toolkits/IToolkit.h"
 #include "Misc/AssetFilterData.h"
+#include "Widgets/Input/SButton.h"
 
 #include "AssetDefinition.generated.h"
 
@@ -335,8 +336,12 @@ public:
 	/** Image widget displayed in the top left border */
 	TSharedPtr<SWidget> ActionImageWidget;
 
+	UE_DEPRECATED(5.6, "Please use the ActionButtonArgs instead")
 	/** Button widget displayed in the center when hovering the thumbnail */
 	TSharedPtr<SWidget> ActionButtonWidget;
+
+	/** Button Args used to construct the Action Button */
+	SButton::FArguments ActionButtonArgs;
 };
 
 struct FAssetButtonActionExtension
diff --git a/Engine/Source/Editor/AudioEditor/Private/AssetTypeActions/AssetDefinition_SoundBase.cpp b/Engine/Source/Editor/AudioEditor/Private/AssetTypeActions/AssetDefinition_SoundBase.cpp
--- a/Engine/Source/Editor/AudioEditor/Private/AssetTypeActions/AssetDefinition_SoundBase.cpp
+++ b/Engine/Source/Editor/AudioEditor/Private/AssetTypeActions/AssetDefinition_SoundBase.cpp
@@ -238,15 +238,9 @@ bool UAssetDefinition_SoundBase::GetThumbnailActionOverlay(const FAssetData& InA
 		return FReply::Handled();
 	};
 
-	OutActionOverlayInfo.ActionButtonWidget = SNew(SButton)
-		.ButtonStyle(FAppStyle::Get(), "HoverHintOnly")
-		.ContentPadding(0.0f)
+	OutActionOverlayInfo.ActionButtonArgs = SButton::FArguments()
 		.ToolTipText_Lambda(OnToolTipTextLambda)
-		.OnClicked_Lambda(OnClickedLambda)
-		[
-			SNew(SImage)
-			.Image_Lambda(OnGetDisplayBrushLambda)
-		];
+		.OnClicked_Lambda(OnClickedLambda);
 
 	return true;
 }
diff --git a/Engine/Source/Editor/ContentBrowser/ContentBrowser.Build.cs b/Engine/Source/Editor/ContentBrowser/ContentBrowser.Build.cs
--- a/Engine/Source/Editor/ContentBrowser/ContentBrowser.Build.cs
+++ b/Engine/Source/Editor/ContentBrowser/ContentBrowser.Build.cs
@@ -20,6 +20,7 @@ public class ContentBrowser : ModuleRules
 				"AppFramework",
 				"Core",
 				"CoreUObject",
+				"Documentation",
 				"DeveloperSettings",
 				"ApplicationCore",
                 "InputCore",
diff --git a/Engine/Source/Editor/ContentBrowser/Private/AssetViewWidgets.cpp b/Engine/Source/Editor/ContentBrowser/Private/AssetViewWidgets.cpp
--- a/Engine/Source/Editor/ContentBrowser/Private/AssetViewWidgets.cpp
+++ b/Engine/Source/Editor/ContentBrowser/Private/AssetViewWidgets.cpp
@@ -54,6 +54,7 @@
 #include "ProfilingDebugging/CpuProfilerTrace.h"
 #include "Rendering/SlateRenderer.h"
 #include "SAssetTagItem.h"
+#include "SRenameWindow.h"
 #include "SThumbnailEditModeTools.h"
 #include "SlateOptMacros.h"
 #include "SlotBase.h"
@@ -1144,6 +1145,11 @@ TSharedRef<SWidget> SAssetViewItem::CreateToolTipWidget() const
 
 			if (GEnableContentBrowserNewStyle)
 			{
+				// Default values
+				FName FolderBrushName = TEXT("ContentBrowser.ListViewFolderIcon");
+				FName FolderShadowBrushName = TEXT("ContentBrowser.FolderItem.DropShadow");
+				ContentBrowserUtils::TryGetFolderBrushAndShadowName(AssetItem->GetItem(), FolderBrushName, FolderShadowBrushName);
+
 				return SNew(SVerticalBox)
 
 					+ SVerticalBox::Slot()
@@ -1153,6 +1159,7 @@ TSharedRef<SWidget> SAssetViewItem::CreateToolTipWidget() const
 
 						+ SVerticalBox::Slot()
 						.AutoHeight()
+						.Padding(0.f, 0.f, 0.f, 6.f)
 						[
 							SNew(SHorizontalBox)
 
@@ -1165,6 +1172,27 @@ TSharedRef<SWidget> SAssetViewItem::CreateToolTipWidget() const
 								.ColorAndOpacity(FStyleColors::White)
 								.Font(FAppStyle::GetFontStyle("ContentBrowser.Tooltip.EntryFont"))
 							]
+						]
+
+						+ SVerticalBox::Slot()
+						.AutoHeight()
+						[
+							SNew(SHorizontalBox)
+
+							+ SHorizontalBox::Slot()
+							.AutoWidth()
+							.VAlign(VAlign_Center)
+							.Padding(0.f, 0.f, 4.f, 0.f)
+							[
+								SNew(SBox)
+								.WidthOverride(16.f)
+								.HeightOverride(16.f)
+								[
+									SNew(SImage)
+									.ColorAndOpacity(GetAssetColor())
+									.Image(FAppStyle::GetBrush(FolderBrushName))
+								]
+							]
 
 							+ SHorizontalBox::Slot()
 							.AutoWidth()
@@ -1172,8 +1200,7 @@ TSharedRef<SWidget> SAssetViewItem::CreateToolTipWidget() const
 							[
 								SNew(STextBlock)
 								.Font(FAppStyle::GetFontStyle("ContentBrowser.Tooltip.EntryFont"))
-								.ColorAndOpacity(FStyleColors::White)
-								.Text(LOCTEXT("FolderNameBracketedLabel", " (Folder)"))
+								.Text(LOCTEXT("FolderNameBracketedLabel", "Folder"))
 							]
 						]
 					]
@@ -2059,6 +2086,8 @@ void SAssetTileItem::Construct( const FArguments& InArgs )
 			ThumbnailConfig.bAllowAssetStatusThumbnailOverlay = true;
 			ThumbnailConfig.bShowAssetChip = true;
 			ThumbnailConfig.AssetChipBorderImageOverride = TAttribute<const FSlateBrush*>::CreateSP(this, &SAssetTileItem::GetAssetAreaOverlayBackgroundImage);
+			ThumbnailConfig.bCanDisplayEditModePrimitiveTools = CanDisplayPrimitiveTools();
+			ThumbnailConfig.IsEditModeVisible = MakeAttributeSP(this, &SAssetTileItem::GetThumbnailEditModeUIVisibility);
 		}
 
 		if(!AssetItem->GetItem().IsSupported())
@@ -2316,7 +2345,14 @@ void SAssetTileItem::Construct( const FArguments& InArgs )
 
 	if(AssetItem.IsValid())
 	{
-		AssetItem->OnRenameRequested().BindSP(InlineRenameWidget.Get(), &SInlineEditableTextBlock::EnterEditingMode);
+		if (GEnableContentBrowserNewStyle)
+		{
+			AssetItem->OnRenameRequested().BindSP(this, &SAssetTileItem::OnRequestingRename);
+		}
+		else
+		{
+			AssetItem->OnRenameRequested().BindSP(InlineRenameWidget.Get(), &SInlineEditableTextBlock::EnterEditingMode);
+		}
 		AssetItem->OnRenameCanceled().BindSP(InlineRenameWidget.Get(), &SInlineEditableTextBlock::ExitEditingMode);
 	}
 }
@@ -2382,6 +2418,34 @@ const FSlateBrush* SAssetTileItem::GetBorderImage() const
 	return FStyleDefaults::GetNoBrush();
 }
 
+void SAssetTileItem::OnRequestingRename()
+{
+	if (InlineRenameWidget.IsValid())
+	{
+		if (CurrentThumbnailSize.Get() != EThumbnailSize::Tiny)
+		{
+			InlineRenameWidget->EnterEditingMode();
+			return;
+		}
+
+		const FGeometry& WidgetGeometry = GetCachedGeometry();
+		FVector2f Position = WidgetGeometry.GetAbsolutePosition();
+
+		// Lower the Y Position of the Widget by Half the Height of the TileItem minus half of the RenameWindow to keep it centered.
+		constexpr float RenameWindowHalfSize = 15.f;
+		Position.Y += (WidgetGeometry.GetAbsoluteSize().Y * 0.5f) - RenameWindowHalfSize;
+
+		// DPI scale and AppScale is handled by the SRenameWindow
+		FSlateApplication::Get().AddWindow(
+			SNew(SRenameWindow)
+			.ScreenPosition(Position)
+			.OnBeginTextEdit(this, &SAssetTileItem::HandleBeginNameChange)
+			.OnVerifyTextChanged(this, &SAssetTileItem::HandleVerifyNameChanged)
+			.InitialText(InlineRenameWidget->GetText())
+			.OnTextCommitted(this, &SAssetTileItem::HandleNameCommitted));
+	}
+}
+
 float SAssetTileItem::GetExtraStateIconWidth() const
 {
 	return GetStateIconImageSize().Get();
diff --git a/Engine/Source/Editor/ContentBrowser/Private/Menus/AssetViewMenus.cpp b/Engine/Source/Editor/ContentBrowser/Private/Menus/AssetViewMenus.cpp
--- a/Engine/Source/Editor/ContentBrowser/Private/Menus/AssetViewMenus.cpp
+++ b/Engine/Source/Editor/ContentBrowser/Private/Menus/AssetViewMenus.cpp
@@ -5,16 +5,60 @@
 #include "Misc/DelayedAutoRegister.h"
 #include "SAssetView.h"
 #include "SContentBrowser.h"
+#include "SDocumentationToolTip.h"
 #include "SFilterList.h"
 #include "ToolMenu.h"
 #include "ToolMenus.h"
 #include "UObject/NameTypes.h"
+#include "Widgets/Text/SRichTextBlock.h"
 
 #define LOCTEXT_NAMESPACE "ContentBrowser"
 
 namespace UE::Editor::ContentBrowser
 {
 	extern CONTENTBROWSER_API bool GEnableContentBrowserNewStyle;
+
+	namespace Private
+	{
+		TSharedRef<SWidget> GetExtendedToolTipMouseWheel()
+		{
+			constexpr float ExtendedToolTipPadding = 8.f;
+			return SNew(SBox)
+					.Padding(ExtendedToolTipPadding)
+					[
+						SNew(SRichTextBlock)
+						.DecoratorStyleSet(&FAppStyle::Get())
+						.TextStyle(&FAppStyle::GetWidgetStyle<FTextBlockStyle>("RichTextBlock.DarkText"))
+						.Text(LOCTEXT("CustomSizeToolTip",
+							"Hold <RichTextBlock.BoldDarkText>Ctrl</> and use the <RichTextBlock.BoldDarkText>Scroll Wheel</> to scale the thumbnails."
+							"\nHold <RichTextBlock.BoldDarkText>Ctrl+Shift</> and use the <RichTextBlock.BoldDarkText>Scroll Wheel</> to cycle through predefined thumbnail sizes."))
+					];
+		}
+
+		TSharedRef<IToolTip> GetThumbnailSizeToolTip()
+		{
+			return SNew(SToolTip)
+					.BorderImage(FAppStyle::GetBrush("ToolTip.BrightBackground"))
+					.TextMargin(FMargin(1.f, 0.f))
+					.Content()
+					[
+						SNew(SDocumentationToolTip)
+						.OverrideExtendedToolTipContent(GetExtendedToolTipMouseWheel())
+						[
+							SNew(STextBlock)
+							.ColorAndOpacity(FLinearColor::Black)
+							.Text(LOCTEXT("ThumbnailSizeToolTip", "Adjust the size of thumbnails."))
+						]
+					];
+		}
+
+		TSharedRef<SWidget> GetThumbnailSizeWidget()
+		{
+			return SNew(STextBlock)
+				.Text(LOCTEXT("ThumbnailSize", "Thumbnail Size"))
+				.ToolTip(GetThumbnailSizeToolTip());
+		}
+	}
 }
 
 using namespace UE::Editor::ContentBrowser;
@@ -194,12 +238,24 @@ FDelayedAutoRegisterHelper SAssetView::AssetViewOptionsMenuRegistration(
 							}
 						};
 
-						Section.AddEntry(FToolMenuEntry::InitSubMenu(
-							"ThumbnailSize",
-							LOCTEXT("ThumbnailSize", "Thumbnail Size"),
-							LOCTEXT("ThumbnailSizeToolTip", "Adjust the size of thumbnails."),
-							FNewToolMenuDelegate::CreateLambda(CreateThumbnailSizeSubMenu)
-						));
+						if (!GEnableContentBrowserNewStyle)
+						{
+							Section.AddEntry(FToolMenuEntry::InitSubMenu(
+								"ThumbnailSize",
+								LOCTEXT("ThumbnailSize", "Thumbnail Size"),
+								LOCTEXT("ThumbnailSizeToolTip", "Adjust the size of thumbnails."),
+								FNewToolMenuDelegate::CreateLambda(CreateThumbnailSizeSubMenu)
+							));
+						}
+						else
+						{
+							Section.AddEntry(FToolMenuEntry::InitSubMenu(
+								"ThumbnailSize",
+								FToolUIActionChoice(),
+								Private::GetThumbnailSizeWidget(),
+								FNewToolMenuDelegate::CreateLambda(CreateThumbnailSizeSubMenu)
+							));
+						}
 
 						if (!GEnableContentBrowserNewStyle)
 						{
diff --git a/Engine/Source/Editor/ContentBrowser/Private/SAssetView.cpp b/Engine/Source/Editor/ContentBrowser/Private/SAssetView.cpp
--- a/Engine/Source/Editor/ContentBrowser/Private/SAssetView.cpp
+++ b/Engine/Source/Editor/ContentBrowser/Private/SAssetView.cpp
@@ -2489,25 +2489,38 @@ FReply SAssetView::OnKeyDown( const FGeometry& MyGeometry, const FKeyEvent& InKe
 FReply SAssetView::OnMouseWheel( const FGeometry& MyGeometry, const FPointerEvent& MouseEvent )
 {
 	// Make sure to not change the thumbnail scaling when we're in Columns view since thumbnail scaling isn't applicable there.
-	if( MouseEvent.IsControlDown() && IsThumbnailScalingAllowed() )
+	if (MouseEvent.IsControlDown() && IsThumbnailScalingAllowed())
 	{
 		if (GEnableContentBrowserNewStyle)
 		{
-			const float NewDelta = MouseEvent.GetWheelDelta() * 0.4f;
-			if ((ZoomScale == 1.f && NewDelta > 0) || (ZoomScale == 0.f && NewDelta < 0))
+			if (MouseEvent.IsShiftDown())
 			{
-				int32 Step = (int32)FMath::Sign(NewDelta);
-				EThumbnailSize OldSize = ThumbnailSize;
-				ThumbnailSize = (EThumbnailSize)FMath::Clamp<int32>(((int32)ThumbnailSize + Step), 0, (int32)EThumbnailSize::MAX - 1);
-				if (OldSize != ThumbnailSize)
+				const int32 Delta = MouseEvent.GetWheelDelta() > 0 ? 1 : -1;
+				EThumbnailSize DesiredThumbnailSize = (EThumbnailSize)FMath::Clamp<int32>((int32)ThumbnailSize + Delta, 0, (int32)EThumbnailSize::MAX - 1);
+
+				if (DesiredThumbnailSize != ThumbnailSize)
 				{
-					OnThumbnailSizeChanged(ThumbnailSize);
-					ZoomScale = NewDelta > 0.f ? 0.f : 1.f;
+					OnThumbnailSizeChanged(DesiredThumbnailSize);
 				}
 			}
 			else
 			{
-				ZoomScale = FMath::Clamp(ZoomScale + NewDelta, 0.f, 1.f);
+				const float NewDelta = MouseEvent.GetWheelDelta() * 0.4f;
+				if ((ZoomScale == 1.f && NewDelta > 0) || (ZoomScale == 0.f && NewDelta < 0))
+				{
+					int32 Step = (int32)FMath::Sign(NewDelta);
+					EThumbnailSize OldSize = ThumbnailSize;
+					ThumbnailSize = (EThumbnailSize)FMath::Clamp<int32>(((int32)ThumbnailSize + Step), 0, (int32)EThumbnailSize::MAX - 1);
+					if (OldSize != ThumbnailSize)
+					{
+						OnThumbnailSizeChanged(ThumbnailSize);
+						ZoomScale = NewDelta > 0.f ? 0.f : 1.f;
+					}
+				}
+				else
+				{
+					ZoomScale = FMath::Clamp(ZoomScale + NewDelta, 0.f, 1.f);
+				}
 			}
 		}
 		else
diff --git a/Engine/Source/Editor/ContentBrowser/Private/SRenameWindow.cpp b/Engine/Source/Editor/ContentBrowser/Private/SRenameWindow.cpp
new file mode 100644
--- /dev/null
+++ b/Engine/Source/Editor/ContentBrowser/Private/SRenameWindow.cpp
@@ -0,0 +1,140 @@
+﻿// Copyright Epic Games, Inc. All Rights Reserved.
+
+#include "SRenameWindow.h"
+#include "Framework/Application/SlateApplication.h"
+#include "Interfaces/IMainFrameModule.h"
+#include "Widgets/Input/SEditableTextBox.h"
+
+#define LOCTEXT_NAMESPACE "SRenameWindow"
+
+/**
+ * EditableTextBox used by the SRenameWindow
+ */
+class SRenameEditableTextBox : public SEditableTextBox
+{
+	SLATE_BEGIN_ARGS(SRenameEditableTextBox)
+	{}
+
+	SLATE_EVENT(FOnTextCommitted, OnTextCommitted)
+	SLATE_EVENT(FOnVerifyTextChanged, OnVerifyTextChanged)
+	SLATE_ARGUMENT(FText, InitialText)
+
+	SLATE_END_ARGS()
+
+	void Construct(const FArguments& InArgs, TSharedRef<SWindow> InOwnerWindow)
+	{
+		OwnerWindow = InOwnerWindow;
+		OnRenameCommittedDelegate = InArgs._OnTextCommitted;
+
+		SEditableTextBox::Construct(SEditableTextBox::FArguments()
+			.OnVerifyTextChanged(InArgs._OnVerifyTextChanged)
+			.Text(InArgs._InitialText)
+			.OnTextCommitted(this, &SRenameEditableTextBox::OnRenameTextCommitted));
+	}
+
+	void EndRename()
+	{
+		if (!bOnTextCommittedCalled)
+		{
+			bOnTextCommittedCalled = true;
+			OnRenameCommittedDelegate.ExecuteIfBound(GetText(), ETextCommit::OnCleared);
+		}
+	}
+
+	virtual FReply OnKeyDown(const FGeometry& MyGeometry, const FKeyEvent& InKeyEvent) override
+	{
+		if (InKeyEvent.GetKey() == EKeys::Escape)
+		{
+			if (const TSharedPtr<SWindow>& Window = OwnerWindow.Pin())
+			{
+				EndRename();
+				return Window->OnKeyDown(MyGeometry, InKeyEvent);
+			}
+		}
+		return SEditableTextBox::OnKeyDown(MyGeometry, InKeyEvent);
+	}
+
+private:
+	void OnRenameTextCommitted(const FText& InText, ETextCommit::Type InCommitType)
+	{
+		if (!bOnTextCommittedCalled)
+		{
+			bOnTextCommittedCalled = true;
+			OnRenameCommittedDelegate.ExecuteIfBound(InText, InCommitType);
+		}
+	}
+
+private:
+	/** The OwnerWindow of this Widget */
+	TWeakPtr<SWindow> OwnerWindow;
+
+	/** Callback to call during the text OnCommitted callback */
+	FOnTextCommitted OnRenameCommittedDelegate;
+
+	/** Will be set to true once the OnCommitted is called once, there are cases which is called more time, but we want to call the passed callback just once */
+	bool bOnTextCommittedCalled = false;
+};
+
+void SRenameWindow::Construct(const FArguments& InArgs)
+{
+	float DPIScale = 1.f;
+	{
+		IMainFrameModule& MainFrameModule = FModuleManager::LoadModuleChecked<IMainFrameModule>(TEXT("MainFrame"));
+		const TSharedPtr<SWindow>& MainFrameParentWindow = MainFrameModule.GetParentWindow();
+		if (MainFrameParentWindow.IsValid())
+		{
+			DPIScale = MainFrameParentWindow->GetDPIScaleFactor();
+		}
+	}
+
+	const float AppScale = FSlateApplication::Get().GetApplicationScale();
+	constexpr float RenameWindowMinWidth = 100.f;
+	constexpr float RenameWindowFixedHeight = 30.f;
+
+	SWindow::Construct(SWindow::FArguments()
+		.SizingRule(ESizingRule::Autosized)
+		.AutoCenter(EAutoCenter::None)
+		.ScreenPosition(InArgs._ScreenPosition * (1.f/DPIScale))
+		.Type(EWindowType::Normal)
+		.LayoutBorder(0)
+		.ActivationPolicy(EWindowActivationPolicy::FirstShown)
+		.bDragAnywhere(false)
+		.CreateTitleBar(false)
+		.FocusWhenFirstShown(true)
+		.MinWidth(RenameWindowMinWidth * AppScale * DPIScale)
+		.MinHeight(RenameWindowFixedHeight * AppScale * DPIScale)
+		.MaxHeight(RenameWindowFixedHeight * AppScale * DPIScale)
+		[
+			SAssignNew(RenameEditableTextBox, SRenameEditableTextBox, StaticCastSharedRef<SWindow>(AsShared()))
+			.OnVerifyTextChanged(InArgs._OnVerifyTextChanged)
+			.InitialText(InArgs._InitialText)
+			.OnTextCommitted(InArgs._OnTextCommitted)
+		]);
+
+	// Once you open the SRenameWindow call the OnBeginText already since that is the moment where the Rename begin
+	InArgs._OnBeginTextEdit.ExecuteIfBound(InArgs._InitialText);
+
+	SetWidgetToFocusOnActivate(RenameEditableTextBox);
+	OnWindowDeactivated.BindSP(this, &SRenameWindow::DeactivateWindow);
+}
+
+FReply SRenameWindow::OnKeyDown(const FGeometry& MyGeometry, const FKeyEvent& InKeyEvent)
+{
+	if (InKeyEvent.GetKey() == EKeys::Escape)
+	{
+		DeactivateWindow();
+		return FReply::Handled();
+	}
+	return SWindow::OnKeyDown(MyGeometry, InKeyEvent);
+}
+
+void SRenameWindow::DeactivateWindow()
+{
+	if (RenameEditableTextBox.IsValid())
+	{
+		RenameEditableTextBox->EndRename();
+	}
+	RequestDestroyWindow();
+}
+
+#undef LOCTEXT_NAMESPACE
diff --git a/Engine/Source/Editor/ContentBrowser/Private/SRenameWindow.h b/Engine/Source/Editor/ContentBrowser/Private/SRenameWindow.h
new file mode 100644
--- /dev/null
+++ b/Engine/Source/Editor/ContentBrowser/Private/SRenameWindow.h
@@ -0,0 +1,40 @@
+﻿// Copyright Epic Games, Inc. All Rights Reserved.
+
+#pragma once
+
+#include "Widgets/SWindow.h"
+#include "Widgets/Text/SInlineEditableTextBlock.h"
+
+class SRenameEditableTextBox;
+
+/**
+ * floating Rename Window
+ */
+class SRenameWindow : public SWindow
+{
+	SLATE_BEGIN_ARGS(SRenameWindow)
+		: _ScreenPosition(FVector2D(0.f, 0.f))
+	{}
+
+	SLATE_ARGUMENT(UE::Slate::FDeprecateVector2DParameter, ScreenPosition)
+
+	SLATE_EVENT(FOnTextCommitted, OnTextCommitted)
+
+	SLATE_EVENT(FOnVerifyTextChanged, OnVerifyTextChanged)
+
+	SLATE_EVENT(FOnBeginTextEdit, OnBeginTextEdit)
+
+	SLATE_ARGUMENT(FText, InitialText)
+
+	SLATE_END_ARGS()
+
+	void Construct(const FArguments& InArgs);
+
+	virtual FReply OnKeyDown(const FGeometry& MyGeometry, const FKeyEvent& InKeyEvent) override;
+
+private:
+	void DeactivateWindow();
+
+private:
+	TSharedPtr<SRenameEditableTextBox> RenameEditableTextBox;
+};
diff --git a/Engine/Source/Editor/ContentBrowser/Public/AssetViewWidgets.h b/Engine/Source/Editor/ContentBrowser/Public/AssetViewWidgets.h
--- a/Engine/Source/Editor/ContentBrowser/Public/AssetViewWidgets.h
+++ b/Engine/Source/Editor/ContentBrowser/Public/AssetViewWidgets.h
@@ -706,6 +706,9 @@ public:
 	static float GetSmallFontHeight() { return SmallFontHeight; }
 
 protected:
+	/** Called when requesting renaming */
+	virtual void OnRequestingRename();
+
 	/** SAssetViewItem interface */
 	virtual float GetNameTextWrapWidth() const override { return LastGeometry.GetLocalSize().X - 15.f; }
 
diff --git a/Engine/Source/Editor/EditorStyle/Private/StarshipStyle.cpp b/Engine/Source/Editor/EditorStyle/Private/StarshipStyle.cpp
--- a/Engine/Source/Editor/EditorStyle/Private/StarshipStyle.cpp
+++ b/Engine/Source/Editor/EditorStyle/Private/StarshipStyle.cpp
@@ -274,6 +274,13 @@ void FStarshipEditorStyle::FStyle::SetupGeneralStyles()
 	{
 		Set( "RichTextBlock.TextHighlight", FTextBlockStyle(NormalText)
 			.SetColorAndOpacity( FLinearColor( 1.0f, 1.0f, 1.0f ) ) );
+		Set("RichTextBlock.DarkText", FTextBlockStyle(NormalText)
+			.SetColorAndOpacity(FLinearColor(0.0f, 0.0f, 0.0f)));
+		Set("RichTextBlock.ForegroundText", FTextBlockStyle(NormalText)
+			.SetColorAndOpacity(FStyleColors::Secondary));
+		Set( "RichTextBlock.BoldDarkText", FTextBlockStyle(NormalText)
+			.SetFont(DEFAULT_FONT("Bold", FStarshipCoreStyle::RegularTextSize))
+			.SetColorAndOpacity(FLinearColor(0.0f, 0.0f, 0.0f)));
 		Set( "RichTextBlock.Bold", FTextBlockStyle(NormalText)
 			.SetFont( DEFAULT_FONT("Bold", FStarshipCoreStyle::RegularTextSize )) );
 		Set( "RichTextBlock.BoldHighlight", FTextBlockStyle(NormalText)
@@ -1240,9 +1247,9 @@ void FStarshipEditorStyle::FStyle::SetupGeneralStyles()
 		Set( "AssetThumbnail.HintShadowOffset", FVector2D(1,1) );
 		Set( "AssetThumbnail.HintShadowColorAndOpacity", FLinearColor(0, 0, 0, 0.5) );
 		Set( "AssetThumbnail.HintBackground", new BOX_BRUSH( "Common/TableViewHeader", FMargin(8.0f/32.0f) ) );
-		Set( "AssetThumbnail.AssetTileViewWhiteChipBorder", new IMAGE_BRUSH_SVG("Starship/ContentBrowser/CornerChipOuter_18", FVector2D(18, 18)));
-		Set( "AssetThumbnail.AssetTileViewWhiteChip", new IMAGE_BRUSH_SVG("Starship/ContentBrowser/CornerChipInner_18", FVector2D(18, 18)));
-		Set( "AssetThumbnail.AssetBorder", new FSlateRoundedBoxBrush(FStyleColors::Transparent, FVector4(4.0f, 2.0f, 4.0f, 4.0f), FStyleColors::Secondary, 2.f));
+		Set( "AssetThumbnail.AssetTileViewWhiteChipBorder", new IMAGE_BRUSH_SVG("Starship/ContentBrowser/CornerChipOuter2_20", FVector2D(20, 20)));
+		Set( "AssetThumbnail.AssetTileViewWhiteChip", new IMAGE_BRUSH_SVG("Starship/ContentBrowser/CornerChipInner2_20", FVector2D(20, 20)));
+		Set( "AssetThumbnail.AssetBorder", new FSlateRoundedBoxBrush(FStyleColors::Transparent, FVector4(4.0f, 4.0f, 4.0f, 4.0f), FStyleColors::Secondary, 2.f));
 		FLinearColor OverlayColorAssetStatusOverlay = FStyleColors::Panel.GetSpecifiedColor();
 		OverlayColorAssetStatusOverlay.A = 0.75f;
 		Set("AssetThumbnail.AssetThumbnailStatusBar", new FSlateRoundedBoxBrush(OverlayColorAssetStatusOverlay, 2.f));
@@ -1250,7 +1257,22 @@ void FStarshipEditorStyle::FStyle::SetupGeneralStyles()
 		OverlayColorAssetThumbnailOverlay.A = 0.75f;
 		Set("AssetThumbnail.AssetThumbnailBar", new FSlateRoundedBoxBrush(OverlayColorAssetThumbnailOverlay, 4.f));
 		Set("AssetThumbnail.ToolTip.CommandBorder", new FSlateRoundedBoxBrush(FStyleColors::Transparent, 4.f, FStyleColors::White, 1.f));
+		Set("AssetThumbnail.ToolTip.DarkCommandBorder", new FSlateRoundedBoxBrush(FStyleColors::Transparent, 4.f, FStyleColors::Black, 1.f));
+		Set("AssetThumbnail.ToolTip.ForegroundCommandBorder", new FSlateRoundedBoxBrush(FStyleColors::Transparent, 4.f, FStyleColors::Foreground, 1.f));
 		Set("AssetThumbnail.Tooltip.Border", new FSlateRoundedBoxBrush(FStyleColors::Secondary, 0.f, COLOR("#484848FF"), 2.f));
+
+		const FButtonStyle EditModePrimitives = FButtonStyle()
+			.SetNormal(FSlateRoundedBoxBrush(FStyleColors::Secondary, 4.f, COLOR("#121212FF"), 1.f))
+			.SetHovered(FSlateRoundedBoxBrush(FStyleColors::Hover, 4.f, COLOR("#121212FF"), 1.f))
+			.SetPressed(FSlateRoundedBoxBrush(FStyleColors::Recessed, 4.f, COLOR("#121212FF"), 1.f));
+		Set("AssetThumbnail.EditMode.Primitive", EditModePrimitives);
+
+		const FButtonStyle ActionButton = FButtonStyle()
+			.SetNormal(IMAGE_BRUSH_SVG("Starship/AssetActions/PlayButtonBackground", FVector2D(32, 32), FStyleColors::Secondary))
+			.SetHovered(IMAGE_BRUSH_SVG("Starship/AssetActions/PlayButtonBackground", FVector2D(32, 32), FStyleColors::Hover))
+			.SetPressed(IMAGE_BRUSH_SVG("Starship/AssetActions/PlayButtonBackground", FVector2D(32, 32), FStyleColors::Recessed));
+		Set("AssetThumbnail.Action.Button", ActionButton);
+
 	}
 
 	// Open any asset dialog
@@ -6679,10 +6701,10 @@ void FStarshipEditorStyle::FStyle::SetupContentBrowserStyle()
 		Set("ContentBrowser.AssetTileItem.AssetBorderSelectedHoverBackground", new FSlateRoundedBoxBrush(FStyleColors::Transparent, FVector4(4.0f, 2.0f, 4.0f, 4.0f), FStyleColors::PrimaryHover, 2.f));
 
 		// Tile Item Border
-		Set("ContentBrowser.AssetTileItem.AssetContent", new FSlateRoundedBoxBrush(FStyleColors::Secondary, FVector4(4.0f, 2.0f, 4.0f, 4.0f)));
-		Set("ContentBrowser.AssetTileItem.AssetContentHoverBackground", new FSlateRoundedBoxBrush(FStyleColors::Hover, FVector4(4.0f, 2.0f, 4.0f, 4.0f)));
-		Set("ContentBrowser.AssetTileItem.AssetContentSelectedBackground", new FSlateRoundedBoxBrush(FStyleColors::Primary, FVector4(4.0f, 2.0f, 4.0f, 4.0f)));
-		Set("ContentBrowser.AssetTileItem.AssetContentSelectedHoverBackground", new FSlateRoundedBoxBrush(FStyleColors::PrimaryHover, FVector4(4.0f, 2.0f, 4.0f, 4.0f)));
+		Set("ContentBrowser.AssetTileItem.AssetContent", new FSlateRoundedBoxBrush(FStyleColors::Secondary, FVector4(4.0f, 4.0f, 4.0f, 4.0f)));
+		Set("ContentBrowser.AssetTileItem.AssetContentHoverBackground", new FSlateRoundedBoxBrush(FStyleColors::Hover, FVector4(4.0f, 4.0f, 4.0f, 4.0f)));
+		Set("ContentBrowser.AssetTileItem.AssetContentSelectedBackground", new FSlateRoundedBoxBrush(FStyleColors::Primary, FVector4(4.0f, 4.0f, 4.0f, 4.0f)));
+		Set("ContentBrowser.AssetTileItem.AssetContentSelectedHoverBackground", new FSlateRoundedBoxBrush(FStyleColors::PrimaryHover, FVector4(4.0f, 4.0f, 4.0f, 4.0f)));
 
 		{
 			FLinearColor OverlayColor = FStyleColors::Panel.GetSpecifiedColor();
@@ -6716,7 +6738,7 @@ void FStarshipEditorStyle::FStyle::SetupContentBrowserStyle()
 		Set( "ContentBrowser.DirectoryUp", new IMAGE_BRUSH("Icons/icon_folder_up_16x", Icon16x16) );
 		Set( "ContentBrowser.PathPickerButton", new IMAGE_BRUSH("Icons/ellipsis_12x", Icon12x12, FLinearColor::Black) );
 
-		Set( "ContentBrowser.ContentDirty", new IMAGE_BRUSH( "ContentBrowser/ContentDirty", Icon16x16) );
+		Set( "ContentBrowser.ContentDirty", new IMAGE_BRUSH_SVG("Starship/Common/DirtyBadge", Icon16x16) );
 		Set( "ContentBrowser.AssetDragDropTooltipBackground", new BOX_BRUSH( "Old/Menu_Background", FMargin(8.0f/64.0f) ) );
 		Set( "ContentBrowser.CollectionTreeDragDropBorder", new BOX_BRUSH( "Old/Window/ViewportDebugBorder", 0.8f ) );
 		Set( "ContentBrowser.PopupMessageIcon", new IMAGE_BRUSH( "Icons/alert", Icon32x32) );
diff --git a/Engine/Source/Editor/UnrealEd/Private/AssetThumbnail.cpp b/Engine/Source/Editor/UnrealEd/Private/AssetThumbnail.cpp
--- a/Engine/Source/Editor/UnrealEd/Private/AssetThumbnail.cpp
+++ b/Engine/Source/Editor/UnrealEd/Private/AssetThumbnail.cpp
@@ -46,6 +46,7 @@
 #include "ClassIconFinder.h"
 #include "IAssetSystemInfoProvider.h"
 #include "IVREditorModule.h"
+#include "SAssetThumbnailEditModeTools.h"
 #include "SDocumentationToolTip.h"
 #include "Fonts/FontMeasure.h"
 #include "Framework/Application/SlateApplication.h"
@@ -116,7 +117,9 @@ public:
 		SLATE_ARGUMENT(bool, AllowAssetStatusThumbnailOverlay)
 		SLATE_ATTRIBUTE(TSharedPtr<IToolTip>, TooltipOverride)
 		SLATE_ARGUMENT(bool, ShowAssetChip)
-		SLATE_ARGUMENT(TAttribute<const FSlateBrush*>, AssetChipBorderImageOverride)
+		SLATE_ATTRIBUTE(const FSlateBrush*, AssetChipBorderImageOverride)
+		SLATE_ARGUMENT(bool, CanDisplayEditModePrimitiveTools)
+		SLATE_ATTRIBUTE(EVisibility, IsEditModeVisible)
 		SLATE_ARGUMENT(EThumbnailColorStripOrientation, ColorStripOrientation)
 
 	SLATE_END_ARGS()
@@ -141,6 +144,7 @@ public:
 		AssetSystemInfoProvider = InArgs._AssetSystemInfoProvider;
 		bShowAssetChip = InArgs._ShowAssetChip;
 		AssetChipBorderImageOverride = InArgs._AssetChipBorderImageOverride;
+		EditModeVisibility = InArgs._IsEditModeVisible;
 		AssetThumbnail->OnAssetDataChanged().AddSP(this, &SAssetThumbnail::OnAssetDataChanged);
 
 		const FAssetData& AssetData = AssetThumbnail->GetAssetData();
@@ -357,6 +361,20 @@ public:
 			];
 		}
 
+		if (bEnableContentBrowserNewStyle)
+		{
+			// AssetEditMode, do not create it if there is no config argument to show it
+			if (EditModeVisibility.IsSet())
+			{
+				OverlayWidget->AddSlot()
+				[
+					SNew(SAssetThumbnailEditModeTools, AssetThumbnail)
+					.SmallView(InArgs._CanDisplayEditModePrimitiveTools)
+					.Visibility(EditModeVisibility)
+				];
+			}
+		}
+
 		const UAssetDefinition* AssetDefinition = UAssetDefinitionRegistry::Get()->GetAssetDefinitionForClass(Class);
 
 		if (bEnableContentBrowserNewStyle)
@@ -426,8 +444,9 @@ public:
 				constexpr int32 OverlayZOrder = 1;
 				if (OutThumbnailInfo.ActionImageWidget.IsValid())
 				{
-					constexpr float PaddingFromTopLeftBorder = 2.f;
+					constexpr float PaddingFromTopLeftBorder = 4.f;
 					constexpr float TopLeftImageSize = 20.f;
+					constexpr float BorderPaddingToActionImageWidget = 4.f;
 
 					OverlayWidget->AddSlot()
 					.ZOrder(OverlayZOrder)
@@ -435,13 +454,14 @@ public:
 					.HAlign(HAlign_Left)
 					.Padding(PaddingFromTopLeftBorder, PaddingFromTopLeftBorder, 0.f, 0.f)
 					[
-						SNew(SBorder)
-						.BorderImage(FAppStyle::GetBrush(Style, ".AssetThumbnailBar"))
-						.Visibility_Lambda([this] () { return IsHovered() ? EVisibility::Collapsed : EVisibility::Visible; })
+						SNew(SBox)
+						.WidthOverride(TopLeftImageSize)
+						.HeightOverride(TopLeftImageSize)
 						[
-							SNew(SBox)
-							.WidthOverride(TopLeftImageSize)
-							.HeightOverride(TopLeftImageSize)
+							SNew(SBorder)
+							.BorderImage(FAppStyle::GetBrush(Style, ".AssetThumbnailBar"))
+							.Padding(BorderPaddingToActionImageWidget)
+							.Visibility_Lambda([this] () { return IsHovered() || (EditModeVisibility.IsSet() && EditModeVisibility.Get().IsVisible()) ? EVisibility::Collapsed : EVisibility::Visible; })
 							[
 								OutThumbnailInfo.ActionImageWidget.ToSharedRef()
 							]
@@ -449,6 +469,7 @@ public:
 					];
 				}
 
+				PRAGMA_DISABLE_DEPRECATION_WARNINGS
 				if (OutThumbnailInfo.ActionButtonWidget.IsValid())
 				{
 					constexpr float CenterImageSize = 32.f;
@@ -459,8 +480,9 @@ public:
 					.HAlign(HAlign_Center)
 					[
 						SNew(SBorder)
-						.BorderImage(FAppStyle::GetBrush(Style, ".AssetThumbnailBar"))
-						.Visibility_Lambda([this] () { return IsHovered() ? EVisibility::Visible : EVisibility::Collapsed; })
+						.Padding(0.f)
+						.BorderImage(FStyleDefaults::GetNoBrush())
+						.Visibility(this, &SAssetThumbnail::GetActionButtonVisibility)
 						[
 							SNew(SBox)
 							.WidthOverride(CenterImageSize)
@@ -471,6 +493,38 @@ public:
 						]
 					];
 				}
+				PRAGMA_ENABLE_DEPRECATION_WARNINGS
+				else
+				{
+					// Set the default style and padding for the Button
+					OutThumbnailInfo.ActionButtonArgs
+						.ButtonStyle(&FAppStyle::GetWidgetStyle<FButtonStyle>(Style, ".Action.Button"))
+						.ContentPadding(8.f);
+
+					TAttribute<EVisibility> ActionVisibility = TAttribute<EVisibility>::CreateSP(this, &SAssetThumbnail::GetActionButtonVisibility);
+					TSharedRef<SButton> ActionButton = SNew(SButton);
+					ActionButton->Construct(OutThumbnailInfo.ActionButtonArgs);
+					ActionButton->SetVisibility(ActionVisibility);
+					ActionButton->SetToolTipText(OutThumbnailInfo.ActionButtonArgs._ToolTipText);
+					ActionButton->SetToolTip(OutThumbnailInfo.ActionButtonArgs._ToolTip);
+					ActionButton->SetCursor(EMouseCursor::Default);
+					ActionButton->SetContent(OutThumbnailInfo.ActionImageWidget.ToSharedRef());
+
+					constexpr float CenterImageSize = 32.f;
+
+					OverlayWidget->AddSlot()
+                    .ZOrder(OverlayZOrder)
+                    .VAlign(VAlign_Center)
+                    .HAlign(HAlign_Center)
+                    [
+                    	SNew(SBox)
+						.WidthOverride(CenterImageSize)
+						.HeightOverride(CenterImageSize)
+						[
+							ActionButton
+						]
+                    ];
+				}
 			}
 			else
 			{
@@ -783,7 +837,11 @@ private:
 			.BorderImage(FAppStyle::GetBrush(Style, ".ToolTip.CommandBorder"))
 			[
 				SNew(STextBlock)
+#if PLATFORM_MAC
+				.Text(NSLOCTEXT("AssetThumbnail", "CommandOptionLabel", " Command + Option "))
+#else
 				.Text(NSLOCTEXT("AssetThumbnail", "CtrlAltLabel", " Ctrl + Alt "))
+#endif
 			];
 		const TSharedRef<FSlateFontMeasure> FontMeasure = FSlateApplication::Get().GetRenderer()->GetFontMeasureService();
 		const int16 Baseline = FontMeasure->GetBaseline(FStyleDefaults::GetFontInfo());
@@ -838,12 +896,20 @@ private:
 		{
 			if (AssetStatus.IsVisible.IsSet() && AssetStatus.IsVisible.Get().IsVisible())
 			{
-				return EVisibility::Visible;
+				if (!EditModeVisibility.IsSet() || !EditModeVisibility.Get().IsVisible())
+				{
+					return EVisibility::Visible;
+				}
 			}
 		}
 		return EVisibility::Collapsed;
 	}
 
+	EVisibility GetActionButtonVisibility() const
+	{
+		return IsHovered() && (!EditModeVisibility.IsSet() || !EditModeVisibility.Get().IsVisible()) ? EVisibility::Visible : EVisibility::Collapsed;
+	}
+
 	TSharedRef<SWidget> CreateStatusWidget(int32 StatusIndex, const FAssetDisplayInfo& InStatusInfo) const
 	{
 		return SNew(SBox)
@@ -1322,7 +1388,17 @@ private:
 			RenderedThumbnailWidget->SetVisibility( ViewportVisibility );
 			if( ClassIconWidget.IsValid() )
 			{
-				ClassIconWidget->SetVisibility( ViewportVisibility );
+				static const IConsoleVariable* CVarEnableContentBrowserNewStyle = IConsoleManager::Get().FindConsoleVariable(TEXT("ContentBrowser.EnableNewStyle"));
+				const bool bEnableContentBrowserNewStyle = CVarEnableContentBrowserNewStyle && CVarEnableContentBrowserNewStyle->GetBool();
+
+				if (bEnableContentBrowserNewStyle)
+				{
+					ClassIconWidget->SetVisibility(TAttribute<EVisibility>::CreateLambda([this, ViewportVisibility] () { return !EditModeVisibility.IsSet() || !EditModeVisibility.Get().IsVisible() ? ViewportVisibility : EVisibility::Collapsed; }));
+				}
+				else
+				{
+					ClassIconWidget->SetVisibility( ViewportVisibility );
+				}
 			}
 		}
 	}
@@ -1368,7 +1444,7 @@ private:
 	TSharedPtr<IAssetSystemInfoProvider> AssetSystemInfoProvider;
 	bool bShowAssetChip;
 	TAttribute<const FSlateBrush*> AssetChipBorderImageOverride;
-
+	TAttribute<EVisibility> EditModeVisibility;
 
 	bool bAllowHintText;
 	bool bAllowRealTimeOnHovered;
@@ -1545,6 +1621,8 @@ TSharedRef<SWidget> FAssetThumbnail::MakeThumbnailWidget( const FAssetThumbnailC
 		.AllowAssetStatusThumbnailOverlay(InConfig.bAllowAssetStatusThumbnailOverlay)
 		.ShowAssetChip(InConfig.bShowAssetChip)
 		.AssetChipBorderImageOverride(InConfig.AssetChipBorderImageOverride)
+		.CanDisplayEditModePrimitiveTools(InConfig.bCanDisplayEditModePrimitiveTools)
+		.IsEditModeVisible(InConfig.IsEditModeVisible)
 		.ColorStripOrientation(InConfig.ColorStripOrientation);
 }
 
diff --git a/Engine/Source/Editor/UnrealEd/Private/SAssetThumbnailEditModeTools.cpp b/Engine/Source/Editor/UnrealEd/Private/SAssetThumbnailEditModeTools.cpp
new file mode 100644
--- /dev/null
+++ b/Engine/Source/Editor/UnrealEd/Private/SAssetThumbnailEditModeTools.cpp
@@ -0,0 +1,344 @@
+﻿// Copyright Epic Games, Inc. All Rights Reserved.
+
+
+#include "SAssetThumbnailEditModeTools.h"
+
+#include "AssetRegistry/AssetData.h"
+#include "AssetThumbnail.h"
+#include "AssetToolsModule.h"
+#include "Containers/EnumAsByte.h"
+#include "Delegates/Delegate.h"
+#include "Editor/UnrealEdEngine.h"
+#include "GenericPlatform/ICursor.h"
+#include "IAssetTools.h"
+#include "IAssetTypeActions.h"
+#include "Input/Events.h"
+#include "InputCoreTypes.h"
+#include "Internationalization/Internationalization.h"
+#include "Layout/Children.h"
+#include "Layout/Margin.h"
+#include "Math/UnrealMathSSE.h"
+#include "Math/Vector2D.h"
+#include "Misc/Attribute.h"
+#include "Modules/ModuleManager.h"
+#include "SlotBase.h"
+#include "Styling/AppStyle.h"
+#include "Styling/SlateColor.h"
+#include "Templates/Casts.h"
+#include "ThumbnailRendering/SceneThumbnailInfo.h"
+#include "ThumbnailRendering/SceneThumbnailInfoWithPrimitive.h"
+#include "Types/SlateEnums.h"
+#include "UObject/NameTypes.h"
+#include "UObject/Object.h"
+#include "UObject/ObjectPtr.h"
+#include "UObject/SoftObjectPath.h"
+#include "UnrealEdGlobals.h"
+#include "Widgets/Layout/SBox.h"
+#include "Widgets/Images/SImage.h"
+#include "Widgets/Input/SButton.h"
+#include "Widgets/SBoxPanel.h"
+
+#define LOCTEXT_NAMESPACE "SAssetThumbnailEditModeTools"
+
+void SAssetThumbnailEditModeTools::Construct(const FArguments& InArgs, const TSharedPtr<FAssetThumbnail>& InAssetThumbnail)
+{
+	AssetThumbnail = InAssetThumbnail;
+	bModifiedThumbnailWhileDragging = false;
+	DragStartLocation = FIntPoint(ForceInitToZero);
+	bInSmallView = InArgs._SmallView;
+
+	constexpr float EditModeButtonSize = 20.f;
+	constexpr float EditModeButtonPadding = 4.f;
+	constexpr float EditModeButtonContentPadding = 2.f;
+
+	ChildSlot
+	[
+		SNew(SHorizontalBox)
+
+		// Primitive tools
+		+SHorizontalBox::Slot()
+		.VAlign(VAlign_Bottom)
+		.HAlign(HAlign_Left)
+		.Padding(EditModeButtonPadding, 0.f, 0.f, EditModeButtonPadding)
+		[
+			SNew(SBox)
+			.HeightOverride(EditModeButtonSize)
+			.WidthOverride(EditModeButtonSize)
+			[
+				SNew(SButton)
+				.Visibility(this, &SAssetThumbnailEditModeTools::GetPrimitiveToolsVisibility)
+				.ButtonStyle(FAppStyle::Get(), "AssetThumbnail.EditMode.Primitive")
+				.OnClicked(this, &SAssetThumbnailEditModeTools::ChangePrimitive)
+				.ToolTipText(LOCTEXT("CyclePrimitiveThumbnailShapes", "Cycle through primitive shape for this thumbnail"))
+				.ContentPadding(EditModeButtonContentPadding)
+				.Content()
+				[
+					SNew(SImage)
+					.Image(this, &SAssetThumbnailEditModeTools::GetCurrentPrimitiveBrush)
+				]
+			]
+		]
+
+		+SHorizontalBox::Slot()
+		.HAlign(HAlign_Right)
+		.VAlign(VAlign_Bottom)
+		.Padding(0.f, 0.f, EditModeButtonPadding, EditModeButtonPadding)
+		[
+			SNew(SBox)
+			.HeightOverride(EditModeButtonSize)
+			.WidthOverride(EditModeButtonSize)
+			[
+				SNew(SButton)
+				.Visibility(this, &SAssetThumbnailEditModeTools::GetPrimitiveToolsResetToDefaultVisibility)
+				.ButtonStyle(FAppStyle::Get(), "AssetThumbnail.EditMode.Primitive")
+				.OnClicked(this, &SAssetThumbnailEditModeTools::ResetToDefault)
+				.ToolTipText(LOCTEXT("ResetThumbnailToDefault", "Resets thumbnail to the default"))
+				.ContentPadding(EditModeButtonContentPadding)
+				.Content()
+				[
+					SNew(SImage)
+					.Image(FAppStyle::GetBrush("PropertyWindow.DiffersFromDefault"))
+				]
+			]
+		]
+	];
+}
+
+EVisibility SAssetThumbnailEditModeTools::GetPrimitiveToolsVisibility() const
+{
+	const bool bIsVisible = bInSmallView && (GetSceneThumbnailInfo() != nullptr);
+	return bIsVisible ? EVisibility::Visible : EVisibility::Collapsed;
+}
+
+EVisibility SAssetThumbnailEditModeTools::GetPrimitiveToolsResetToDefaultVisibility() const
+{
+	USceneThumbnailInfo* ThumbnailInfo = GetSceneThumbnailInfo();
+	
+	EVisibility ResetToDefaultVisibility = EVisibility::Collapsed;
+	if (ThumbnailInfo)
+	{
+		ResetToDefaultVisibility = ThumbnailInfo && ThumbnailInfo->DiffersFromDefault() ? EVisibility::Visible : EVisibility::Collapsed;
+	}
+
+	return ResetToDefaultVisibility;
+}
+
+const FSlateBrush* SAssetThumbnailEditModeTools::GetCurrentPrimitiveBrush() const
+{
+	if (USceneThumbnailInfoWithPrimitive* ThumbnailInfo = GetSceneThumbnailInfoWithPrimitive())
+	{
+		// Note this is for the icon only.  we are assuming the thumbnail renderer does the right thing when rendering
+		const EThumbnailPrimType PrimType = ThumbnailInfo->bUserModifiedShape ? ThumbnailInfo->PrimitiveType.GetValue() : (EThumbnailPrimType)ThumbnailInfo->DefaultPrimitiveType.Get(EThumbnailPrimType::TPT_Sphere);
+		switch (PrimType)
+		{
+			case TPT_None: return FAppStyle::GetBrush("ContentBrowser.PrimitiveCustom");
+			case TPT_Sphere: return FAppStyle::GetBrush("ContentBrowser.PrimitiveSphere");
+			case TPT_Cube: return FAppStyle::GetBrush("ContentBrowser.PrimitiveCube");
+			case TPT_Cylinder: return FAppStyle::GetBrush("ContentBrowser.PrimitiveCylinder");
+			case TPT_Plane:
+			default:
+			// Fall through and return a plane
+			break;
+		}
+	}
+
+	return FAppStyle::GetBrush("ContentBrowser.PrimitivePlane");
+}
+
+FReply SAssetThumbnailEditModeTools::ChangePrimitive()
+{
+	USceneThumbnailInfoWithPrimitive* ThumbnailInfo = GetSceneThumbnailInfoWithPrimitive();
+	if (ThumbnailInfo)
+	{
+		uint8 PrimitiveIdx = ThumbnailInfo->PrimitiveType.GetIntValue() + 1;
+		if (PrimitiveIdx >= TPT_MAX)
+		{
+			if (ThumbnailInfo->PreviewMesh.IsValid())
+			{
+				PrimitiveIdx = TPT_None;
+			}
+			else
+			{
+				PrimitiveIdx = TPT_None + 1;
+			}
+		}
+
+		ThumbnailInfo->PrimitiveType = TEnumAsByte<EThumbnailPrimType>(PrimitiveIdx);
+		ThumbnailInfo->bUserModifiedShape = true;
+
+		if (AssetThumbnail.IsValid())
+		{
+			AssetThumbnail.Pin()->RefreshThumbnail();
+		}
+
+		ThumbnailInfo->MarkPackageDirty();
+	}
+
+	return FReply::Handled();
+}
+
+FReply SAssetThumbnailEditModeTools::ResetToDefault()
+{
+	if (USceneThumbnailInfo* ThumbnailInfo = GetSceneThumbnailInfo())
+	{
+		ThumbnailInfo->ResetToDefault();
+
+		if (AssetThumbnail.IsValid())
+		{
+			AssetThumbnail.Pin()->RefreshThumbnail();
+		}
+
+		ThumbnailInfo->MarkPackageDirty();
+
+	}
+
+	return FReply::Handled();
+}
+
+FReply SAssetThumbnailEditModeTools::OnMouseButtonDown(const FGeometry& MyGeometry, const FPointerEvent& MouseEvent)
+{
+	if (AssetThumbnail.IsValid() &&
+		(MouseEvent.GetEffectingButton() == EKeys::LeftMouseButton || MouseEvent.GetEffectingButton() == EKeys::RightMouseButton))
+	{
+		// Load the asset, unless it is in an unloaded map package or already loaded
+		const FAssetData& AssetData = AssetThumbnail.Pin()->GetAssetData();
+		
+		// Getting the asset loads it, if it isn't already.
+		UObject* Asset = AssetData.GetAsset();
+
+		USceneThumbnailInfo* ThumbnailInfo = GetSceneThumbnailInfo();
+		if ( ThumbnailInfo )
+		{
+			FThumbnailRenderingInfo* RenderInfo = GUnrealEd->GetThumbnailManager()->GetRenderingInfo(Asset);
+			if (RenderInfo != nullptr && RenderInfo->Renderer != nullptr)
+			{
+				bModifiedThumbnailWhileDragging = false;
+				DragStartLocation = FIntPoint(FMath::TruncToInt32(MouseEvent.GetScreenSpacePosition().X), FMath::TruncToInt32(MouseEvent.GetScreenSpacePosition().Y));
+
+				return FReply::Handled().CaptureMouse(AsShared()).UseHighPrecisionMouseMovement(AsShared()).PreventThrottling();
+			}
+		}
+
+		// This thumbnail does not have a scene thumbnail info but thumbnail editing is enabled. Just consume the input.
+		return FReply::Handled();
+	}
+		
+	return FReply::Unhandled();
+}
+
+FReply SAssetThumbnailEditModeTools::OnMouseButtonUp(const FGeometry& MyGeometry, const FPointerEvent& MouseEvent)
+{
+	if (HasMouseCapture())
+	{
+		if (bModifiedThumbnailWhileDragging)
+		{
+			USceneThumbnailInfo* ThumbnailInfo = GetSceneThumbnailInfo();
+			if (ThumbnailInfo)
+			{
+				ThumbnailInfo->MarkPackageDirty();
+			}
+
+			bModifiedThumbnailWhileDragging = false;
+		}
+
+		return FReply::Handled().ReleaseMouseCapture().SetMousePos(DragStartLocation);
+	}
+
+	return FReply::Unhandled();
+}
+
+FReply SAssetThumbnailEditModeTools::OnMouseMove(const FGeometry& MyGeometry, const FPointerEvent& MouseEvent)
+{
+	if (HasMouseCapture())
+	{
+		if (!MouseEvent.GetCursorDelta().IsZero())
+		{
+			if (USceneThumbnailInfo* ThumbnailInfo = GetSceneThumbnailInfo())
+			{
+				const bool bLeftMouse = MouseEvent.IsMouseButtonDown(EKeys::LeftMouseButton);
+				const bool bRightMouse = MouseEvent.IsMouseButtonDown(EKeys::RightMouseButton);
+
+				if (bLeftMouse)
+				{
+					ThumbnailInfo->OrbitYaw += -MouseEvent.GetCursorDelta().X;
+					ThumbnailInfo->OrbitPitch += -MouseEvent.GetCursorDelta().Y;
+
+					// Normalize the values
+					if (ThumbnailInfo->OrbitYaw > 180)
+					{
+						ThumbnailInfo->OrbitYaw -= 360;
+					}
+					else if (ThumbnailInfo->OrbitYaw < -180)
+					{
+						ThumbnailInfo->OrbitYaw += 360;
+					}
+					
+					if (ThumbnailInfo->OrbitPitch > 90)
+					{
+						ThumbnailInfo->OrbitPitch = 90;
+					}
+					else if (ThumbnailInfo->OrbitPitch < -90)
+					{
+						ThumbnailInfo->OrbitPitch = -90;
+					}
+				}
+				else if (bRightMouse)
+				{
+					// Since zoom is a modifier of on the camera distance from the bounding sphere of the object, it is normalized in the thumbnail preview scene.
+					ThumbnailInfo->OrbitZoom += MouseEvent.GetCursorDelta().Y;
+				}
+
+				// Dirty the package when the mouse is released
+				bModifiedThumbnailWhileDragging = true;
+			}
+		}
+
+		// Refresh the thumbnail. Do this even if the mouse did not move in case the thumbnail varies with time.
+		if (AssetThumbnail.IsValid())
+		{
+			AssetThumbnail.Pin()->RefreshThumbnail();
+		}
+
+		return FReply::Handled().PreventThrottling();
+	}
+
+	return FReply::Unhandled();
+}
+
+FCursorReply SAssetThumbnailEditModeTools::OnCursorQuery(const FGeometry& MyGeometry, const FPointerEvent& CursorEvent) const
+{
+	return HasMouseCapture() ? 
+		FCursorReply::Cursor(EMouseCursor::None) :
+		FCursorReply::Cursor(EMouseCursor::Default);
+}
+
+USceneThumbnailInfo* SAssetThumbnailEditModeTools::GetSceneThumbnailInfo() const
+{
+	USceneThumbnailInfo* SceneThumbnailInfo = SceneThumbnailInfoPtr.Get();
+	
+	if (!SceneThumbnailInfo)
+	{
+		if (AssetThumbnail.IsValid())
+		{
+			if (UObject* Asset = AssetThumbnail.Pin()->GetAsset())
+			{
+				static const FName AssetToolsName("AssetTools");
+				FAssetToolsModule& AssetToolsModule = FModuleManager::LoadModuleChecked<FAssetToolsModule>(AssetToolsName);
+				TWeakPtr<IAssetTypeActions> AssetTypeActions = AssetToolsModule.Get().GetAssetTypeActionsForClass(Asset->GetClass());
+				if (AssetTypeActions.IsValid())
+				{
+					SceneThumbnailInfo = Cast<USceneThumbnailInfo>(AssetTypeActions.Pin()->GetThumbnailInfo(Asset));
+				}
+			}
+		}
+	}
+
+	return SceneThumbnailInfo;
+}
+
+USceneThumbnailInfoWithPrimitive* SAssetThumbnailEditModeTools::GetSceneThumbnailInfoWithPrimitive() const
+{
+	return Cast<USceneThumbnailInfoWithPrimitive>(GetSceneThumbnailInfo());
+}
+
+#undef LOCTEXT_NAMESPACE
diff --git a/Engine/Source/Editor/UnrealEd/Private/SAssetThumbnailEditModeTools.h b/Engine/Source/Editor/UnrealEd/Private/SAssetThumbnailEditModeTools.h
new file mode 100644
--- /dev/null
+++ b/Engine/Source/Editor/UnrealEd/Private/SAssetThumbnailEditModeTools.h
@@ -0,0 +1,70 @@
+﻿// Copyright Epic Games, Inc. All Rights Reserved.
+
+#pragma once
+
+#include "Templates/SharedPointer.h"
+#include "UObject/WeakObjectPtr.h"
+#include "UObject/WeakObjectPtrTemplates.h"
+#include "Widgets/SCompoundWidget.h"
+
+class FAssetThumbnail;
+class USceneThumbnailInfo;
+class USceneThumbnailInfoWithPrimitive;
+struct FSlateBrush;
+
+// This is the equivalent of the ContentBrowser one, with some Graphic changes.
+// The ContentBrowser one will be later removed when the new style will be enabled by default.
+
+/**
+ * AssetThumbnail EditMode widget, used for the Edit mode of the thumbnails (Ex: changing Shape/Angle)
+ */
+class SAssetThumbnailEditModeTools : public SCompoundWidget
+{
+public:
+	SLATE_BEGIN_ARGS(SAssetThumbnailEditModeTools)
+		: _SmallView(false)
+	{}
+		
+	SLATE_ARGUMENT(bool, SmallView)
+
+	SLATE_END_ARGS()
+
+	/** Constructs this widget with InArgs */
+	void Construct(const FArguments& InArgs, const TSharedPtr<FAssetThumbnail>& InAssetThumbnail);
+
+	// SWidget Interface
+	virtual FReply OnMouseButtonDown(const FGeometry& MyGeometry, const FPointerEvent& MouseEvent) override;
+	virtual FReply OnMouseButtonUp(const FGeometry& MyGeometry, const FPointerEvent& MouseEvent) override;
+	virtual FReply OnMouseMove(const FGeometry& MyGeometry, const FPointerEvent& MouseEvent) override;
+	virtual FCursorReply OnCursorQuery(const FGeometry& MyGeometry, const FPointerEvent& CursorEvent) const override;
+
+protected:
+	/** Gets the visibility for the primitives toolbar */
+	EVisibility GetPrimitiveToolsVisibility() const;
+
+	/** Gets the visibility for the the thumbnail reset to default button */
+	EVisibility GetPrimitiveToolsResetToDefaultVisibility() const;
+
+	/** Gets the brush used to display the currently used primitive */
+	const FSlateBrush* GetCurrentPrimitiveBrush() const;
+
+	/** Sets the primitive type for the supplied thumbnail, if possible */
+	FReply ChangePrimitive();
+
+	/** Resets the primitive to default */
+	FReply ResetToDefault();
+
+	/** Helper accessors for ThumbnailInfo objects */
+	USceneThumbnailInfo* GetSceneThumbnailInfo() const;
+	USceneThumbnailInfoWithPrimitive* GetSceneThumbnailInfoWithPrimitive() const;
+
+protected:
+	bool bInSmallView = false;
+	bool bModifiedThumbnailWhileDragging = false;
+	FIntPoint DragStartLocation = FIntPoint(0, 0);
+	TWeakPtr<FAssetThumbnail> AssetThumbnail;
+
+private:
+	// Never access this directly.
+	mutable TWeakObjectPtr<USceneThumbnailInfo> SceneThumbnailInfoPtr;
+};
diff --git a/Engine/Source/Editor/UnrealEd/Public/AssetThumbnail.h b/Engine/Source/Editor/UnrealEd/Public/AssetThumbnail.h
--- a/Engine/Source/Editor/UnrealEd/Public/AssetThumbnail.h
+++ b/Engine/Source/Editor/UnrealEd/Public/AssetThumbnail.h
@@ -68,8 +68,10 @@ struct FAssetThumbnailConfig
 		, Padding(0)
 		, AssetSystemInfoProvider(nullptr)
 		, bAllowAssetStatusThumbnailOverlay( false )
-		, bShowAssetChip(false)
+		, bShowAssetChip(true)
 		, AssetChipBorderImageOverride()
+		, bCanDisplayEditModePrimitiveTools(false)
+		, IsEditModeVisible()
 	{
 	}
 
@@ -91,6 +93,8 @@ struct FAssetThumbnailConfig
 	bool bAllowAssetStatusThumbnailOverlay;
 	bool bShowAssetChip;
 	TAttribute<const FSlateBrush*> AssetChipBorderImageOverride;
+	bool bCanDisplayEditModePrimitiveTools;
+	TAttribute<EVisibility> IsEditModeVisible;
 };
 
 
