Honor the check references on delete flag when the user hasn't selected an actor.

[RN]

diff --git a/Engine/Source/Editor/UnrealEd/Private/UnrealEdSrv.cpp b/Engine/Source/Editor/UnrealEd/Private/UnrealEdSrv.cpp
--- a/Engine/Source/Editor/UnrealEd/Private/UnrealEdSrv.cpp
+++ b/Engine/Source/Editor/UnrealEd/Private/UnrealEdSrv.cpp
@@ -1545,7 +1545,8 @@ bool UUnrealEdEngine::Exec_Edit( UWorld* InWorld, const TCHAR* Str, FOutputDevic
 		{
 			const FScopedTransaction Transaction(NSLOCTEXT("UnrealEd", "Cut", "Cut"));
 			FEditorDelegates::OnEditCutActorsBegin.Broadcast();
-			CopySelectedActorsToClipboard(InWorld, true);
+			const bool bCheckRef = GetDefault<ULevelEditorMiscSettings>()->bCheckReferencesOnDelete;
+			CopySelectedActorsToClipboard(InWorld, true, false /*bIsMove*/, bCheckRef);
 			FEditorDelegates::OnEditCutActorsEnd.Broadcast();
 		}
 	}
