Free cook functions, for use in editor

[RN] Added a CookAsset function to CookFunctionLibrary, which allows for an asset to be cooked from an editor script

diff --git a/Engine/Source/Editor/UnrealEd/Private/Cooker/CookFunctionLibrary.cpp b/Engine/Source/Editor/UnrealEd/Private/Cooker/CookFunctionLibrary.cpp
new file mode 100644
--- /dev/null
+++ b/Engine/Source/Editor/UnrealEd/Private/Cooker/CookFunctionLibrary.cpp
@@ -0,0 +1,237 @@
+// Copyright Epic Games, Inc. All Rights Reserved.
+
+#include "Cooker/CookFunctionLibrary.h"
+
+#include "Commandlets/Commandlet.h"
+#include "Cooker/AsyncIODelete.h"
+#include "Cooker/CookDependency.h"
+#include "Cooker/CookSandbox.h"
+#include "Cooker/LooseCookedPackageWriter.h"
+#include "CookOnTheSide/CookLog.h"
+#include "Editor/EditorEngine.h"
+#include "HAL/FileManager.h"
+#include "Interfaces/ITargetPlatform.h"
+#include "Interfaces/ITargetPlatformManagerModule.h"
+#include "LooseFilesCookArtifactReader.h"
+#include "Serialization/CompactBinaryPackage.h"
+#include "Serialization/CompactBinaryWriter.h"
+#include "Subsystems/UnrealEditorSubsystem.h"
+#include "TargetDomain/TargetDomainUtils.h"
+#include "UObject/ArchiveCookContext.h"
+#include "UObject/FieldPath.h"
+#include "UObject/SavePackage.h"
+#include "UObject/UnrealTypePrivate.h"
+
+extern UNREALED_API UEditorEngine* GEditor;
+
+namespace UE::Private
+{
+	void SaveDepsToFile(const UE::TargetDomain::FCookDependencies& CookDependencies, const FString& Filename)
+	{
+		FCbWriter Writer;
+		Writer.BeginObject();
+		Writer << "CookTestSnapshot" << CookDependencies;
+		Writer.EndObject();
+		TUniquePtr<FArchive> FileArchive(IFileManager::Get().CreateDebugFileWriter(*Filename));
+		Writer.Save(*FileArchive);
+	}
+
+	bool LoadDepsFromFile(UE::TargetDomain::FCookDependencies& CookDependencies, const FString& Filename)
+	{
+		TArray<uint8> Data;
+		{
+			TUniquePtr<FArchive> FileArchive(IFileManager::Get().CreateFileReader(*Filename));
+			Data.SetNum(FileArchive->TotalSize());
+			FileArchive->Serialize(Data.GetData(), Data.Num());
+		}
+
+		FSharedBuffer SharedBuffer = FSharedBuffer::MakeView(MakeMemoryView(Data)); // strange boilerplate, i don't understand it
+		FCbObject CbObject(SharedBuffer);
+		FCbField TestSnapshot = CbObject["CookTestSnapshot"];
+		return LoadFromCompactBinary(TestSnapshot.AsObject(), CookDependencies);
+	}
+
+	FString GetDepsFilename(const UObject* Object, const FString& DestinationSubfolder)
+	{
+		return FPaths::ConvertRelativePathToFull(
+			FPaths::ProjectSavedDir() +
+			FString(TEXT("Temp/")) + DestinationSubfolder + Object->GetPackage()->GetName() + TEXT(".cookdeps"));
+	}
+}
+
+void UCookFunctionLibrary::CookAsset(UObject* Object, const FString& ForPlatform, const FString& DestinationSubfolder, const FString& CookCommandlineArgs)
+{
+	using namespace UE::TargetDomain;
+	if (Object == nullptr)
+	{
+		UE_LOG(LogCook, Warning, TEXT("CookAsset expected an object to cook, but received nullptr"));
+		return;
+	}
+
+	ITargetPlatformManagerModule* PMM = GetTargetPlatformManager();
+	if (PMM == nullptr)
+	{
+		return;
+	}
+
+	UPackage* Package = Cast<UPackage>(Object);
+	if (!Package)
+	{
+		Package = Object->GetPackage();
+	}
+
+	if(Package == GetTransientPackage())
+	{
+		UE_LOG(LogCook, Warning, 
+			TEXT("CookAsset cannot cook the transient package: %s"), *Object->GetPathName());
+		return;
+	}
+
+	const ITargetPlatform* TargetPlatform = PMM->FindTargetPlatform(*ForPlatform);
+	if (!TargetPlatform)
+	{
+		TargetPlatform = PMM->GetRunningTargetPlatform();
+		if (!TargetPlatform)
+		{
+			UE_LOG(LogCook, Warning, 
+				TEXT("Could not find any platform to cook for when requested to cook %s"), 
+				*ForPlatform);
+			return;
+		}
+		else
+		{
+			UE_LOG(LogCook, Warning, 
+				TEXT("Could not find requested platform %s, fell back to %s!"), 
+				*ForPlatform, *TargetPlatform->IniPlatformName());
+		}
+	}
+
+	bool bUnversioned = false;
+	if(!CookCommandlineArgs.IsEmpty())
+	{
+		TArray<FString> Tokens;
+		TArray<FString> Switches;
+		UCommandlet::ParseCommandLine(*CookCommandlineArgs, Tokens, Switches);
+		if(Tokens.Num() > 0)
+		{
+			UE_LOG(LogCook, Warning, 
+				TEXT("CookAsset does not expect tokens - they have been discarded: %s"), 
+				*FString::Join(Tokens, TEXT(" ")));
+		}
+		if(Switches.Num() > 0)
+		{
+			bUnversioned = (Switches.RemoveSingle(TEXT("UNVERSIONED")) == 1); 
+			if(Switches.Num() > 0)
+			{
+				UE_LOG(LogCook, Warning, 
+					TEXT("CookAsset found switches it does not yet support - they have been discarded: %s"), 
+					*FString::Join(Switches, TEXT(" ")));
+			}
+		}
+	}
+
+	const FString OutputName = FPaths::ConvertRelativePathToFull(
+		FPaths::ProjectSavedDir() +
+		FString(TEXT("Temp/")) + DestinationSubfolder + Package->GetName() + TEXT(".uasset"));
+
+	FArchiveCookContext ArchiveCookContext(Package,
+		UE::Cook::ECookType::ByTheBook,
+		UE::Cook::ECookingDLC::No, // used only by UMaterialInterface::Serialize, we could expose via command line arguments
+		TargetPlatform);
+	FArchiveCookData CookData(*TargetPlatform, ArchiveCookContext);
+
+	// The AsyncIODelete object optimizes the cleaning of the saved/cook directory. Happily not relevant
+	// for callers of this function, but an important optimization for the cook commandlet.
+	TUniquePtr<FAsyncIODelete> AsyncIODelete = MakeUnique<FAsyncIODelete>(); 
+	ICookedPackageWriter::FBeginCacheCallback BeginCacheCallback(
+		[](ICookedPackageWriter::FBeginCacheForCookedPlatformDataInfo& Info) -> EPackageWriterResult
+		{
+			return EPackageWriterResult::Success; // saving will fail if we don't say good things happened
+		});
+
+	TArray<TSharedRef<IPlugin> > PluginsToRemap; // none?
+	TUniquePtr<UE::Cook::FCookSandbox> SandboxFileObj = MakeUnique<UE::Cook::FCookSandbox>(OutputName, PluginsToRemap);
+	ICookedPackageWriter::FRegisterDeterminismHelperCallback RegisterDeterminismHelperCallback;
+	FLooseCookedPackageWriter* LooseWriter = new FLooseCookedPackageWriter(
+		OutputName, 
+		OutputName, 
+		TargetPlatform, 
+		*AsyncIODelete, 
+		*SandboxFileObj,
+		MoveTemp(BeginCacheCallback),
+		MoveTemp(RegisterDeterminismHelperCallback),
+		MakeShared<FLooseFilesCookArtifactReader>());
+	// ~FSavePackageContext will delete the loose writer, hence naked new above:
+	FSavePackageContext SaveContext(TargetPlatform, LooseWriter); 
+	FSavePackageArgs SaveArgs;
+	SaveArgs.TopLevelFlags = RF_Public;
+	SaveArgs.bForceByteSwapping = (!TargetPlatform->IsLittleEndian()) ^ (!PLATFORM_LITTLE_ENDIAN);;
+	SaveArgs.bWarnOfLongFilename = false;
+	SaveArgs.SaveFlags = SAVE_AllowTimeout|SAVE_KeepGUID;
+	SaveArgs.ArchiveCookData = &CookData;
+	SaveArgs.bSlowTask = true;
+	SaveArgs.SavePackageContext = &SaveContext;
+	if(bUnversioned)
+	{
+		SaveArgs.SaveFlags |= SAVE_Unversioned;
+	}
+
+	// We'll need to support this if we want to support cook diffing or iterative validate, 
+	// but for now it's not needed:
+	// SaveContext.PackageWriter->UpdateSaveArguments(SaveArgs);
+	
+	// The platform determines whether it wants editor only data, but lets restore the package
+	// flag at the end of this function:
+	const bool bWasFilterEditorOnly = Package->HasAllPackagesFlags(PKG_FilterEditorOnly);
+	if(!TargetPlatform->HasEditorOnlyData())
+	{
+		Package->SetPackageFlags(PKG_FilterEditorOnly); 
+	}
+	
+	ICookedPackageWriter::FBeginPackageInfo Info;
+	Info.PackageName = Package->GetFName();
+	Info.LooseFilePath = OutputName;
+	LooseWriter->BeginPackage(Info);
+
+	FSavePackageResultStruct Result = GEditor->Save(Package, Package->FindAssetInPackage(), *OutputName, SaveArgs);
+	if (!Result.IsSuccessful())
+	{
+		UE_LOG(LogCook, Warning, TEXT("Saving failed - asset not cooked!"));
+	}
+	else
+	{
+		ICookedPackageWriter::FCommitPackageInfo CommitInfo;
+		CommitInfo.Status = IPackageWriter::ECommitStatus::Success;
+		CommitInfo.PackageName = Package->GetFName();
+		CommitInfo.WriteOptions = IPackageWriter::EWriteOptions::Write;
+		LooseWriter->CommitPackage(MoveTemp(CommitInfo)); // please clap
+
+		// FGeneratedPackageResultStruct is used to support the creation of generated 
+		// streamingcells from WorldPartition. Supporting streaming cells would require 
+		// several other changes to this function, and is not the priority of this 
+		// routine at this time.
+		FGeneratedPackageResultStruct* Todo = nullptr;
+		TArray<FName> RuntimeDepsTodo; // I don't believe we need valid data in here at this time
+
+		FString ErrorMessage;
+		FCookDependencies Deps = FCookDependencies::Collect(Package, TargetPlatform, &Result, Todo, MoveTemp(RuntimeDepsTodo), &ErrorMessage);
+		if(!ErrorMessage.IsEmpty())
+		{
+			UE_LOG(LogCook, Warning, TEXT("Error collecting cook dependencies: %s"), *ErrorMessage);
+		}
+
+		// persist cook dependencies along side the loose files we wrote above..
+		const FString DepsFilename = UE::Private::GetDepsFilename(Object, DestinationSubfolder + "_" + ForPlatform);
+		UE::Private::SaveDepsToFile(Deps, DepsFilename);
+	}
+
+	if (bWasFilterEditorOnly)
+	{
+		Package->SetPackageFlags(PKG_FilterEditorOnly);
+	}
+	else
+	{
+		Package->ClearPackageFlags(PKG_FilterEditorOnly);
+	}
+}
+
diff --git a/Engine/Source/Editor/UnrealEd/Private/Cooker/CookFunctionLibrary.h b/Engine/Source/Editor/UnrealEd/Private/Cooker/CookFunctionLibrary.h
new file mode 100644
--- /dev/null
+++ b/Engine/Source/Editor/UnrealEd/Private/Cooker/CookFunctionLibrary.h
@@ -0,0 +1,20 @@
+// Copyright Epic Games, Inc. All Rights Reserved.
+
+#include "CoreMinimal.h"
+
+#include "CookFunctionLibrary.generated.h"
+
+UCLASS(transient)
+class UCookFunctionLibrary : public UObject
+{
+	GENERATED_BODY()
+public:
+	/** 
+	 * Writes the cooked version of the provided object's package into the Saved folder, in the subfolder
+	 * defined by DestinationSubfolder. Extra arguments (such as -unversioned) can be provided by
+	 * CookCommandlineArgs. This function is experimental and may not exactly match the behavior of
+	 * the cook commandlet.
+	 */
+	UFUNCTION(BlueprintCallable, Category = "Development")
+	static UNREALED_API void CookAsset(UObject* Object, const FString& ForPlatform, const FString& DestinationSubfolder, const FString& CookCommandlineArgs = TEXT("") );
+};
