Fix missing pragma once

diff --git a/Engine/Source/Editor/UnrealEd/Private/Cooker/CookFunctionLibrary.h b/Engine/Source/Editor/UnrealEd/Private/Cooker/CookFunctionLibrary.h
--- a/Engine/Source/Editor/UnrealEd/Private/Cooker/CookFunctionLibrary.h
+++ b/Engine/Source/Editor/UnrealEd/Private/Cooker/CookFunctionLibrary.h
@@ -1,5 +1,7 @@
 // Copyright Epic Games, Inc. All Rights Reserved.
 
+#pragma once
+
 #include "CoreMinimal.h"
 
 #include "CookFunctionLibrary.generated.h"
