Allow Content Worker to start up engine and wait for a  cook command.
* Cook command file updated to process ini, cvar, dlcname, and cook validation options [RunAssetValidation, ValidationErrorsAreFatal]

#rnx

diff --git a/Engine/Source/Editor/UnrealEd/Private/CookOnTheFlyServer.cpp b/Engine/Source/Editor/UnrealEd/Private/CookOnTheFlyServer.cpp
--- a/Engine/Source/Editor/UnrealEd/Private/CookOnTheFlyServer.cpp
+++ b/Engine/Source/Editor/UnrealEd/Private/CookOnTheFlyServer.cpp
@@ -106,6 +106,7 @@
 #include "Misc/App.h"
 #include "Misc/CommandLine.h"
 #include "Misc/ConfigCacheIni.h"
+#include "Misc/ConfigUtilities.h"
 #include "Misc/CoreDelegates.h"
 #include "Misc/DataValidation.h"
 #include "Misc/FileHelper.h"
@@ -12044,6 +12045,15 @@ const UCookOnTheFlyServer::FCookByTheBookStartupOptions& UCookOnTheFlyServer::Bl
 				{
 					const FStringView PackageToken(TEXT("-Package="));
 
+					auto MarkModified = [&ModifiedStartupOptions, &EffectiveStartupOptions, &CookByTheBookStartupOptions]()
+						{
+							if (!ModifiedStartupOptions.IsSet())
+							{
+								ModifiedStartupOptions.Emplace(CookByTheBookStartupOptions);
+								EffectiveStartupOptions = &ModifiedStartupOptions.GetValue();
+							}
+						};
+
 					if (Token == TEXT("-CookAbort"))
 					{
 						UE_LOG(LogCook, Display, TEXT("Received CookAbort token in cook command file, exiting cook."));
@@ -12055,11 +12065,7 @@ const UCookOnTheFlyServer::FCookByTheBookStartupOptions& UCookOnTheFlyServer::Bl
 					}
 					else if (Token.StartsWith(PackageToken))
 					{
-						if (!ModifiedStartupOptions.IsSet())
-						{
-							ModifiedStartupOptions.Emplace(CookByTheBookStartupOptions);
-							EffectiveStartupOptions = &ModifiedStartupOptions.GetValue();
-						}
+						MarkModified();
 
 						UE::String::ParseTokens(Token.RightChop(PackageToken.Len()), TEXT('+'),
 							[&ModifiedStartupOptions](FStringView PackageToken)
@@ -12069,6 +12075,88 @@ const UCookOnTheFlyServer::FCookByTheBookStartupOptions& UCookOnTheFlyServer::Bl
 								UE_LOG(LogCook, Display, TEXT("Received Package token in cook command file, adding package '%.*s' to cook workload."), PackageToken.Len(), PackageToken.GetData());
 							}, UE::String::EParseTokensOptions::SkipEmpty);
 					}
+					else if (Token.StartsWith(TEXTVIEW("-DLCName=")))
+					{
+						MarkModified();
+
+						int32 TokenIndex = 0;
+						Token.FindChar(TEXT('='), TokenIndex);
+						ModifiedStartupOptions->DLCName = Token.RightChop(TokenIndex+1);
+						UE_LOG(LogCook, Display, TEXT("Updated DLCName=%s"), *ModifiedStartupOptions->DLCName);
+					}
+					else if (Token.StartsWith(TEXTVIEW("-RunAssetValidation")))
+					{
+						MarkModified();
+
+						ModifiedStartupOptions->CookOptions |= ECookByTheBookOptions::RunAssetValidation;
+						UE_LOG(LogCook, Display, TEXT("RunAssetValidation enabled"));
+					}
+					else if (Token.StartsWith(TEXTVIEW("-ValidationErrorsAreFatal")))
+					{
+						MarkModified();
+
+						ModifiedStartupOptions->CookOptions |= ECookByTheBookOptions::ValidationErrorsAreFatal;
+						UE_LOG(LogCook, Display, TEXT("ValidationErrorsAreFatal enabled"));
+					}
+					else if (Token.StartsWith(TEXTVIEW("-dpcvar")) || Token.StartsWith(TEXT("-ForceDPCVars=")))
+					{
+						// TODO: multiprocess cook will not get cvar changes
+						int32 TokenIndex = 0;
+						Token.FindChar(TEXT('='), TokenIndex);
+						const FString CVarParam{ Token.RightChop(TokenIndex + 1) };
+						TArray<FString> CVars;
+						CVarParam.ParseIntoArray(CVars, TEXT(","), true);
+						for (const FString& CVarValueKey : CVars)
+						{
+							FString CVarKey, CVarValue;
+							if (CVarValueKey.Split(TEXT("="), &CVarKey, &CVarValue))
+							{
+								IConsoleVariable* ConsoleVariable = IConsoleManager::Get().FindConsoleVariable(*CVarKey);
+								if (ConsoleVariable)
+								{
+									ConsoleVariable->Set(*CVarValue, ECVF_SetByCode);
+									UE_LOG(LogCook, Display, TEXT("Set CVar %s"), *CVarValueKey);
+								}
+								else
+								{
+									UE_LOG(LogCook, Display, TEXT("Failed to find cvar '%s'"), *CVarKey);
+								}
+							}
+							else
+							{
+								UE_LOG(LogCook, Warning, TEXT("Failed to parse cvar from:%s"), *CVarValueKey);
+							}
+						}
+					}
+					else if (Token.StartsWith(TEXTVIEW("-ini:")))
+					{
+						const FString TokenValue{ Token };
+						TArray<FString> Tokens;
+						TokenValue.ParseIntoArray(Tokens, TEXT(":"), true); //-ini:IniName:[Section1]:Key=Value
+						if (Tokens.Num() == 4)
+						{
+							FConfigBranch* Branch = GConfig->FindBranch(*Tokens[1], TEXT(""));
+							if (Branch)
+							{
+								static const FName CookCmdOverrideName = TEXT("CookCommand");
+								UE::DynamicConfig::PerformDynamicConfig(CookCmdOverrideName, [Tokens, Branch](FConfigModificationTracker* ChangeTracker)
+									{
+										const FString ConfigString = Tokens[2] + TEXT("\n") + Tokens[3];
+										// to force reloading cvars through -ini
+										ChangeTracker->CVars.Add(TEXT("ConsoleVariables")).CVarPriority = (int)ECVF_SetByHotfix;
+										Branch->AddDynamicLayerStringToHierarchy(Tokens[1], ConfigString, CookCmdOverrideName, DynamicLayerPriority::Hotfix, ChangeTracker);
+									});
+							}
+							else
+							{
+								UE_LOG(LogCook, Warning, TEXT("Failed to find config file:%s"), *Tokens[1]);
+							}
+						}
+						else
+						{
+							UE_LOG(LogCook, Warning, TEXT("Failed to parse ini from:%s"), *TokenValue);
+						}
+					}
 					else if (UAssetManager::Get().HandleCookCommand(Token))
 					{
 						// Do nothing - handled by the asset manager
