Improve compiler results structured logging.
Previously, compiler errors logged while running a commandlet used FAssetMsg::FormatPathForAssetLog, which produces an unstructured string. Use UE:FAssetLog instead, which UE_ASSET_LOG utilizes. This is serialized as an object with both the object and local paths.
Update the other logs to use UE_LOGFMT_NSLOC so the compiler message can be isolated in the structured log.

diff --git a/Engine/Source/Editor/UnrealEd/Private/Kismet2/CompilerResultsLog.cpp b/Engine/Source/Editor/UnrealEd/Private/Kismet2/CompilerResultsLog.cpp
--- a/Engine/Source/Editor/UnrealEd/Private/Kismet2/CompilerResultsLog.cpp
+++ b/Engine/Source/Editor/UnrealEd/Private/Kismet2/CompilerResultsLog.cpp
@@ -449,24 +449,18 @@ void FCompilerResultsLog::InternalLogMessage(FName MessageID, const TSharedRef<F
 	{
 		if (Severity == EMessageSeverity::Error)
 		{
-			if (IsRunningCommandlet())
-			{				
-				UE_LOGFMT_NSLOC(LogBlueprint, Error, "Blueprint", "CompilerError", "[AssetLog] {AssetPath}: [Compiler] {ErrorMessage} from Source: {SourceFile}",
-												("AssetPath", FAssetMsg::FormatPathForAssetLog(*SourcePath)), ("ErrorMessage", Message->ToText().ToString()), ("SourceFile", SourcePath));
-			}
-			else
-			{
-				// in editor the compiler log is 'rich' and we don't need to annotate with the full blueprint path, just the name:
-				UE_ASSET_LOG(LogBlueprint, Error, *SourcePath, TEXT("[Compiler] %s"), *Message->ToText().ToString());				
-			}
+			UE_LOGFMT_NSLOC(LogBlueprint, Error, "CompilerResultsLog", "InternalLogMessage", "[AssetLog] {Asset}: [Compiler] {Message}",
+				("Asset", UE::FAssetLog(*SourcePath)), ("Message", Message->ToText().ToString()));
 		}
 		else if (Severity == EMessageSeverity::Warning || Severity == EMessageSeverity::PerformanceWarning)
 		{
-			UE_ASSET_LOG(LogBlueprint, Warning, *SourcePath, TEXT("[Compiler] %s"), *Message->ToText().ToString());			
+			UE_LOGFMT_NSLOC(LogBlueprint, Warning, "CompilerResultsLog", "InternalLogMessage", "[AssetLog] {Asset}: [Compiler] {Message}",
+				("Asset", UE::FAssetLog(*SourcePath)), ("Message", Message->ToText().ToString()));
 		}
 		else
 		{
-			UE_ASSET_LOG(LogBlueprint, Log, *SourcePath, TEXT("[Compiler] %s"), *Message->ToText().ToString());
+			UE_LOGFMT_NSLOC(LogBlueprint, Log, "CompilerResultsLog", "InternalLogMessage", "[AssetLog] {Asset}: [Compiler] {Message}",
+				("Asset", UE::FAssetLog(*SourcePath)), ("Message", Message->ToText().ToString()));
 		}
 	}
 }
