Avoid using UDerivedDataCacheCommandlet::CacheWorldPackages on non-WP world packages since it doesn't deal properly with already loaded worlds (can happen when it's a sublevel in an already loaded non-WP map)
Made FScopedEditorWorld reset the world cleaned up state on init to allow for reuse

diff --git a/Engine/Source/Editor/UnrealEd/Private/Commandlets/DerivedDataCacheCommandlet.cpp b/Engine/Source/Editor/UnrealEd/Private/Commandlets/DerivedDataCacheCommandlet.cpp
--- a/Engine/Source/Editor/UnrealEd/Private/Commandlets/DerivedDataCacheCommandlet.cpp
+++ b/Engine/Source/Editor/UnrealEd/Private/Commandlets/DerivedDataCacheCommandlet.cpp
@@ -431,6 +431,11 @@ void UDerivedDataCacheCommandlet::FinishCachingObjects()
 
 void UDerivedDataCacheCommandlet::CacheWorldPackages(UWorld* World, uint8 PackageFilter, TSet<FName>& OutNewProcessedPackages)
 {
+	if (!UWorld::IsPartitionedWorld(World))
+	{
+		return;
+	}
+
 	// Setup the world
 	UWorld::InitializationValues IVS;
 	IVS.RequiresHitProxies(false);
@@ -443,23 +448,19 @@ void UDerivedDataCacheCommandlet::CacheWorldPackages(UWorld* World, uint8 Packag
 	FScopedEditorWorld EditorWorld(World, IVS);
 
 	// If the world is partitioned
-	bool bResult = true;
-	if (UWorld::IsPartitionedWorld(World))
-	{
-		// Ensure the world has a valid world partition.
-		UWorldPartition* WorldPartition = World->GetWorldPartition();
-		check(WorldPartition);
+	// Ensure the world has a valid world partition.
+	UWorldPartition* WorldPartition = World->GetWorldPartition();
+	check(WorldPartition);
 
-		FWorldPartitionHelpers::ForEachActorWithLoading(WorldPartition, [this, PackageFilter, &OutNewProcessedPackages](const FWorldPartitionActorDescInstance* ActorDescInstance)
+	FWorldPartitionHelpers::ForEachActorWithLoading(WorldPartition, [this, PackageFilter, &OutNewProcessedPackages](const FWorldPartitionActorDescInstance* ActorDescInstance)
+	{
+		if (AActor* Actor = ActorDescInstance->GetActor())
 		{
-			if (AActor* Actor = ActorDescInstance->GetActor())
-			{
-				UE_LOG(LogDerivedDataCacheCommandlet, Display, TEXT("Loaded actor %s"), *Actor->GetName());
-				CacheLoadedPackages(Actor->GetPackage(), PackageFilter, OutNewProcessedPackages);
-			}
-			return true;
-		});
-	}
+			UE_LOG(LogDerivedDataCacheCommandlet, Display, TEXT("Loaded actor %s"), *Actor->GetName());
+			CacheLoadedPackages(Actor->GetPackage(), PackageFilter, OutNewProcessedPackages);
+		}
+		return true;
+	});
 }
 
 int32 UDerivedDataCacheCommandlet::Main( const FString& Params )
diff --git a/Engine/Source/Editor/UnrealEd/Private/EditorWorldUtils.cpp b/Engine/Source/Editor/UnrealEd/Private/EditorWorldUtils.cpp
--- a/Engine/Source/Editor/UnrealEd/Private/EditorWorldUtils.cpp
+++ b/Engine/Source/Editor/UnrealEd/Private/EditorWorldUtils.cpp
@@ -59,6 +59,7 @@ void FScopedEditorWorld::Init(UWorld* InWorld, const UWorld::InitializationValue
 
 	// Initialize the world
 	World->WorldType = InWorldType;
+	World->ResetCleanedUpState(); // We reset the cleanedup state to be able to reuse the world
 	World->InitWorld(InInitializationValues);
 	World->PersistentLevel->UpdateModelComponents();
 	World->UpdateWorldComponents(true /*bRerunConstructionScripts*/, false /*bCurrentLevelOnly*/);
