SKM cloth painting improvements.

Translucent materials don't display as expected when Opacity = 1, and backface culling is disabled, but this is the most common/desired setting. Added new opaque materials to use in these cases.

Also updated the ClothPainter's Brush and Paint settings to all work correctly within and between sessions via configs. Both were duplicating the default configs, which was the wrong behavior. However, the correct behavior is different between the two due to how their base classes are set up. I just left the base classes alone.
* Brush settings were not correctly reloaded within the same session with when reopening an SKM editor. This was because we were duplicating the Default (which is created on program load), but the logic to remember settings between SKM editors happens in the constructor. Fixed this by just creating a new object rather than duplicating the Default.
* The paint setting's VertexPreviewSize didn't work at all to set from the tool's details because the detail view was connected to the ClothPainter's instance, but the tool always read from the Default. Fixed this by making the tool use the ClothPainter's instance (there was already a virtual method for this; the tool just wasn't using it).
* The paint settings weren't stored in the config. We needed to save and load it.

Also changed some default values to match the behavior users would prefer:
* VertexPreviewSize 6->3
* Opacity .8 ->1.
* WidgetLineThickness (e.g., the circle) 0.5 -> 0.

#virtualized

diff --git a/Engine/Config/BaseEditorPerProjectUserSettings.ini b/Engine/Config/BaseEditorPerProjectUserSettings.ini
--- a/Engine/Config/BaseEditorPerProjectUserSettings.ini
+++ b/Engine/Config/BaseEditorPerProjectUserSettings.ini
@@ -955,6 +955,9 @@ IncludeSmoothLODCheck=false
 [/Script/MeshPaint.MeshPaintSettings]
 VertexPreviewSize=6
 
+[/Script/ClothPainter.ClothPainterSettings]
+VertexPreviewSize=3
+
 [/Script/UndoHistory.UndoHistorySettings]
 bShowTransactionDetails=false
 
diff --git a/Engine/Config/BaseEngine.ini b/Engine/Config/BaseEngine.ini
--- a/Engine/Config/BaseEngine.ini
+++ b/Engine/Config/BaseEngine.ini
@@ -195,6 +195,8 @@ TexturePaintingMaskMaterialName=/Engine/EngineDebugMaterials/TexturePaintingMask
 BoneWeightMaterialName=/Engine/EngineDebugMaterials/BoneWeightMaterial.BoneWeightMaterial
 ClothPaintMaterialName=/Engine/EngineDebugMaterials/ClothMaterial.ClothMaterial
 ClothPaintMaterialWireframeName=/Engine/EngineDebugMaterials/ClothMaterial_WF.ClothMaterial_WF
+ClothPaintOpaqueMaterialName=/Engine/EngineDebugMaterials/ClothOpaqueMaterial.ClothOpaqueMaterial
+ClothPaintOpaqueMaterialWireframeName=/Engine/EngineDebugMaterials/ClothOpaqueMaterial_WF.ClothOpaqueMaterial_WF
 DebugEditorMaterialName=/Engine/EngineDebugMaterials/DebugEditorMaterial.DebugEditorMaterial
 InvalidLightmapSettingsMaterialName=/Engine/EngineMaterials/M_InvalidLightmapSettings.M_InvalidLightmapSettings
 PreviewShadowsIndicatorMaterialName=/Engine/EditorMaterials/PreviewShadowIndicatorMaterial.PreviewShadowIndicatorMaterial
diff --git a/Engine/Source/Editor/ClothPainter/Private/ClothPaintSettings.h b/Engine/Source/Editor/ClothPainter/Private/ClothPaintSettings.h
--- a/Engine/Source/Editor/ClothPainter/Private/ClothPaintSettings.h
+++ b/Engine/Source/Editor/ClothPainter/Private/ClothPaintSettings.h
@@ -23,7 +23,7 @@ enum class EPaintableClothProperty
 	BackstopRadius
 };
 
-UCLASS()
+UCLASS(config = EditorPerProjectUserSettings)
 class UClothPainterSettings : public UMeshPaintSettings
 {
 	GENERATED_BODY()
@@ -39,7 +39,7 @@ public:
 		, AutoCalculatedViewMax(0.0f)
 		, bFlipNormal(false)
 		, bCullBackface(false)
-		, Opacity(0.8f)
+		, Opacity(1.f)
 	{}
 
 	float GetViewMin()
@@ -61,11 +61,11 @@ public:
 
 protected:
 	/** When painting float/1D values, this is considered the zero or black point */
-	UPROPERTY(EditAnywhere, Category = View, meta = (UIMin = 0, UIMax = 100000, ClampMin = 0, ClampMax = 100000, EditCondition = "!bAutoViewRange"))
+	UPROPERTY(EditAnywhere, Config, Category = View, meta = (UIMin = 0, UIMax = 100000, ClampMin = 0, ClampMax = 100000, EditCondition = "!bAutoViewRange"))
 	float ViewMin;
 
 	/** When painting float/1D values, this is considered the one or white point */
-	UPROPERTY(EditAnywhere, Category = View, meta = (UIMin = 0, UIMax = 100000, ClampMin = 0, ClampMax = 100000, EditCondition = "!bAutoViewRange"))
+	UPROPERTY(EditAnywhere, Config, Category = View, meta = (UIMin = 0, UIMax = 100000, ClampMin = 0, ClampMax = 100000, EditCondition = "!bAutoViewRange"))
 	float ViewMax;
 
 public:
@@ -87,14 +87,14 @@ public:
 	TArray<TObjectPtr<UClothingAssetCommon>> ClothingAssets;
 
 	/** Whether to flip normals on the mesh preview */
-	UPROPERTY(EditAnywhere, AdvancedDisplay, Category = View)
+	UPROPERTY(EditAnywhere, Config, AdvancedDisplay, Category = View)
 	bool bFlipNormal;
 
 	/** Whether to cull backfacing triangles when rendering the mesh preview */
-	UPROPERTY(EditAnywhere, AdvancedDisplay, Category = View)
+	UPROPERTY(EditAnywhere, Config, AdvancedDisplay, Category = View)
 	bool bCullBackface;
 
 	/** Opacity of the mesh preview */
-	UPROPERTY(EditAnywhere, AdvancedDisplay, Category = View, meta = (UIMin = 0, UIMax = 1, ClampMin = 0, ClampMax = 1))
+	UPROPERTY(EditAnywhere, Config, AdvancedDisplay, Category = View, meta = (UIMin = 0, UIMax = 1, ClampMin = 0, ClampMax = 1))
 	float Opacity;
 };
diff --git a/Engine/Source/Editor/ClothPainter/Private/ClothPainter.cpp b/Engine/Source/Editor/ClothPainter/Private/ClothPainter.cpp
--- a/Engine/Source/Editor/ClothPainter/Private/ClothPainter.cpp
+++ b/Engine/Source/Editor/ClothPainter/Private/ClothPainter.cpp
@@ -62,7 +62,7 @@ FClothPainter::FClothPainter()
 	, BrushSettings(nullptr)
 {
 	VertexPointColor = FLinearColor::White;
-	WidgetLineThickness = .5f;
+	WidgetLineThickness = .0f;
 	bShouldSimulate = false;
 	bShowHiddenVerts = false;
 }
@@ -77,11 +77,12 @@ FClothPainter::~FClothPainter()
 
 void FClothPainter::Init()
 {
-	BrushSettings = DuplicateObject<UPaintBrushSettings>(GetMutableDefault<UPaintBrushSettings>(), GetTransientPackage());	
+	// Brush settings share their values by getting/setting per-property values in the config via Constructor/PostEditChange.
+	BrushSettings = NewObject<UPaintBrushSettings>(GetTransientPackage());
 	BrushSettings->AddToRoot();
 	BrushSettings->bOnlyFrontFacingTriangles = false;
-	PaintSettings = DuplicateObject<UClothPainterSettings>(GetMutableDefault<UClothPainterSettings>(), GetTransientPackage());
-	PaintSettings->AddToRoot();
+	PaintSettings = GetMutableDefault<UClothPainterSettings>();
+	PaintSettings->LoadConfig();
 
 	CommandList = MakeShareable(new FUICommandList);
 
@@ -237,6 +238,7 @@ void FClothPainter::ExitPaintMode()
 	if (PaintSettings)
 	{
 		PaintSettings->ClothingAssets.Reset();
+		PaintSettings->SaveConfig();
 	}
 }
 
diff --git a/Engine/Source/Editor/MeshPaint/Private/IMeshPainter.cpp b/Engine/Source/Editor/MeshPaint/Private/IMeshPainter.cpp
--- a/Engine/Source/Editor/MeshPaint/Private/IMeshPainter.cpp
+++ b/Engine/Source/Editor/MeshPaint/Private/IMeshPainter.cpp
@@ -184,7 +184,8 @@ void IMeshPainter::RenderInteractorWidget(const FVector& InCameraOrigin, const F
 		const FLinearColor NormalLineColor(0.3f, 1.0f, 0.3f);
 		const FLinearColor BrushCueColor = (bArePainting ? FLinearColor(1.0f, 1.0f, 0.3f) : FLinearColor(0.3f, 1.0f, 0.3f));
 		const FLinearColor InnerBrushCueColor = (bArePainting ? FLinearColor(0.5f, 0.5f, 0.1f) : FLinearColor(0.1f, 0.5f, 0.1f));
-		const float PointDrawSize = GetDefault<UMeshPaintSettings>()->VertexPreviewSize;
+		const UMeshPaintSettings* const PaintSettings = GetPainterSettings();
+		const float PointDrawSize = PaintSettings ? PaintSettings->VertexPreviewSize : GetDefault<UMeshPaintSettings>()->VertexPreviewSize;
 
 		FVector BrushXAxis, BrushYAxis;
 		HitResult.Normal.FindBestAxisVectors(BrushXAxis, BrushYAxis);
diff --git a/Engine/Source/Editor/MeshPaint/Public/MeshPaintSettings.h b/Engine/Source/Editor/MeshPaint/Public/MeshPaintSettings.h
--- a/Engine/Source/Editor/MeshPaint/Public/MeshPaintSettings.h
+++ b/Engine/Source/Editor/MeshPaint/Public/MeshPaintSettings.h
@@ -87,7 +87,7 @@ class MESHPAINT_API UMeshPaintSettings : public UObject
 public:
 
 	/** Size of vertex points drawn when mesh painting is active. */
-	UPROPERTY(config, EditAnywhere, Category=Visualization)
+	UPROPERTY(config, EditAnywhere, Category=Visualization, Meta = (ClampMin = 0, UIMin = 0, UIMax = 10))
 	float VertexPreviewSize;
 
 };
\ No newline at end of file
diff --git a/Engine/Source/Editor/UnrealEd/Classes/Animation/DebugSkelMeshComponent.h b/Engine/Source/Editor/UnrealEd/Classes/Animation/DebugSkelMeshComponent.h
--- a/Engine/Source/Editor/UnrealEd/Classes/Animation/DebugSkelMeshComponent.h
+++ b/Engine/Source/Editor/UnrealEd/Classes/Animation/DebugSkelMeshComponent.h
@@ -114,6 +114,7 @@ public:
 
 	bool bFlipNormal;
 
+	bool bDoubleSidedOpaqueCloth;
 	int32 ClothingSimDataIndexWhenPainting;
 	TArray<uint32> ClothingSimIndices;
 
diff --git a/Engine/Source/Editor/UnrealEd/Private/Animation/DebugSkelMeshComponent.cpp b/Engine/Source/Editor/UnrealEd/Private/Animation/DebugSkelMeshComponent.cpp
--- a/Engine/Source/Editor/UnrealEd/Private/Animation/DebugSkelMeshComponent.cpp
+++ b/Engine/Source/Editor/UnrealEd/Private/Animation/DebugSkelMeshComponent.cpp
@@ -1316,6 +1316,8 @@ void UDebugSkelMeshComponent::GetUsedMaterials(TArray<UMaterialInterface *>& Out
 	{
 		OutMaterials.Add(GEngine->ClothPaintMaterialInstance);
 		OutMaterials.Add(GEngine->ClothPaintMaterialWireframeInstance);
+		OutMaterials.Add(GEngine->ClothPaintOpaqueMaterialInstance);
+		OutMaterials.Add(GEngine->ClothPaintOpaqueMaterialWireframeInstance);
 	}
 }
 
@@ -1448,9 +1450,9 @@ void FDebugSkelMeshSceneProxy::GetDynamicMeshElements(const TArray<const FSceneV
 				}
 
 				// Set material params
-				UMaterialInstanceDynamic* SurfaceMID = GEngine->ClothPaintMaterialInstance;
+				UMaterialInstanceDynamic* SurfaceMID = (DynamicData->bDoubleSidedOpaqueCloth && GEngine->ClothPaintOpaqueMaterialInstance) ? GEngine->ClothPaintOpaqueMaterialInstance : GEngine->ClothPaintMaterialInstance;
 				check(SurfaceMID);
-				UMaterialInstanceDynamic* WireMID = GEngine->ClothPaintMaterialWireframeInstance;
+				UMaterialInstanceDynamic* WireMID = (DynamicData->bDoubleSidedOpaqueCloth && GEngine->ClothPaintOpaqueMaterialWireframeInstance) ? GEngine->ClothPaintOpaqueMaterialWireframeInstance : GEngine->ClothPaintMaterialWireframeInstance;
 				check(WireMID);
 
 				FMaterialRenderProxy* MatProxySurface = SurfaceMID->GetRenderProxy();
@@ -1478,6 +1480,7 @@ FDebugSkelMeshDynamicData::FDebugSkelMeshDynamicData(UDebugSkelMeshComponent* In
 	, bDrawBinormals(InComponent->bDrawBinormals)
 	, bDrawClothPaintPreview(InComponent->bShowClothData)
 	, bFlipNormal(InComponent->bClothFlipNormal)
+	, bDoubleSidedOpaqueCloth(!InComponent->bClothCullBackface && InComponent->ClothMeshOpacity == 1.f)
 	, ClothingSimDataIndexWhenPainting(INDEX_NONE)
 	, PropertyViewMin(InComponent->MinClothPropertyView)
 	, PropertyViewMax(InComponent->MaxClothPropertyView)
@@ -1526,7 +1529,6 @@ FDebugSkelMeshDynamicData::FDebugSkelMeshDynamicData(UDebugSkelMeshComponent* In
 		SurfaceMID->SetScalarParameterValue(FName("ClothOpacity"), InComponent->ClothMeshOpacity);
 		SurfaceMID->SetScalarParameterValue(FName("BackfaceCull"), InComponent->bClothCullBackface ? 1.f : 0.f);
 	}
-
 	if (UMaterialInstanceDynamic* const WireMID = GEngine->ClothPaintMaterialWireframeInstance)
 	{
 		WireMID->SetScalarParameterValue(FName("ClothOpacity"), InComponent->ClothMeshOpacity);
diff --git a/Engine/Source/Runtime/Engine/Classes/Engine/Engine.h b/Engine/Source/Runtime/Engine/Classes/Engine/Engine.h
--- a/Engine/Source/Runtime/Engine/Classes/Engine/Engine.h
+++ b/Engine/Source/Runtime/Engine/Classes/Engine/Engine.h
@@ -1078,20 +1078,36 @@ public:
 	UPROPERTY()
 	TObjectPtr<class UMaterial> ClothPaintMaterial;
 	UPROPERTY()
+	TObjectPtr<class UMaterial> ClothPaintOpaqueMaterial;
+	UPROPERTY()
 	TObjectPtr<class UMaterial> ClothPaintMaterialWireframe;
 	UPROPERTY()
+	TObjectPtr<class UMaterial> ClothPaintOpaqueMaterialWireframe;
+	UPROPERTY()
 	TObjectPtr<class UMaterialInstanceDynamic> ClothPaintMaterialInstance;
 	UPROPERTY()
+	TObjectPtr<class UMaterialInstanceDynamic> ClothPaintOpaqueMaterialInstance;
+	UPROPERTY()
 	TObjectPtr<class UMaterialInstanceDynamic> ClothPaintMaterialWireframeInstance;
+	UPROPERTY()
+	TObjectPtr<class UMaterialInstanceDynamic> ClothPaintOpaqueMaterialWireframeInstance;
 
 	/** Name of the material used to render cloth in the clothing tools */
 	UPROPERTY(globalconfig)
 	FSoftObjectPath ClothPaintMaterialName;
 
+	/** Name of the material used to render cloth in the clothing tools with full opacity */
+	UPROPERTY(globalconfig)
+	FSoftObjectPath ClothPaintOpaqueMaterialName;
+
 	/** Name of the material used to render cloth wireframe in the clothing tools */
 	UPROPERTY(globalconfig)
 	FSoftObjectPath ClothPaintMaterialWireframeName;
 
+	/** Name of the material used to render cloth wireframe in the clothing tools with full opacity */
+	UPROPERTY(globalconfig)
+	FSoftObjectPath ClothPaintOpaqueMaterialWireframeName;
+
 	/** A material used to render physical material mask on mesh. */
 	UPROPERTY()
 	TObjectPtr<class UMaterial> PhysicalMaterialMaskMaterial;
diff --git a/Engine/Source/Runtime/Engine/Private/SkinnedMeshSceneProxyDesc.cpp b/Engine/Source/Runtime/Engine/Private/SkinnedMeshSceneProxyDesc.cpp
--- a/Engine/Source/Runtime/Engine/Private/SkinnedMeshSceneProxyDesc.cpp
+++ b/Engine/Source/Runtime/Engine/Private/SkinnedMeshSceneProxyDesc.cpp
@@ -190,6 +190,8 @@ void FSkinnedMeshSceneProxyDesc::GetUsedMaterials(TArray<UMaterialInterface*>& O
 		{
 			OutMaterials.Add(GEngine->ClothPaintMaterialInstance);
 			OutMaterials.Add(GEngine->ClothPaintMaterialWireframeInstance);
+			OutMaterials.Add(GEngine->ClothPaintOpaqueMaterialInstance);
+			OutMaterials.Add(GEngine->ClothPaintOpaqueMaterialWireframeInstance);
 		}
 #endif
 	}
diff --git a/Engine/Source/Runtime/Engine/Private/UnrealEngine.cpp b/Engine/Source/Runtime/Engine/Private/UnrealEngine.cpp
--- a/Engine/Source/Runtime/Engine/Private/UnrealEngine.cpp
+++ b/Engine/Source/Runtime/Engine/Private/UnrealEngine.cpp
@@ -3239,12 +3239,16 @@ void UEngine::InitializeObjectReferences()
 		LoadSpecialMaterial(TEXT("BoneWeightMaterialName"), BoneWeightMaterialName.ToString(), BoneWeightMaterial, false);
 		LoadSpecialMaterial(TEXT("ClothPaintMaterialName"), ClothPaintMaterialName.ToString(), ClothPaintMaterial, false);
 		LoadSpecialMaterial(TEXT("ClothPaintMaterialWireframeName"), ClothPaintMaterialWireframeName.ToString(), ClothPaintMaterialWireframe, false);
+		LoadSpecialMaterial(TEXT("ClothPaintOpaqueMaterialName"), ClothPaintOpaqueMaterialName.ToString(), ClothPaintOpaqueMaterial, false);
+		LoadSpecialMaterial(TEXT("ClothPaintOpaqueMaterialWireframeName"), ClothPaintOpaqueMaterialWireframeName.ToString(), ClothPaintOpaqueMaterialWireframe, false);
 		LoadSpecialMaterial(TEXT("DebugEditorMaterialName"), DebugEditorMaterialName.ToString(), DebugEditorMaterial, false);
 		LoadSpecialMaterial(TEXT("PhysicalMaterialMaskMaterialName"), PhysicalMaterialMaskMaterialName.ToString(), PhysicalMaterialMaskMaterial, false);
 		LoadSpecialMaterial(TEXT("TexturePaintingMaskMaterialName"), TexturePaintingMaskMaterialName.ToString(), TexturePaintingMaskMaterial, false);
 
 		ClothPaintMaterialInstance = UMaterialInstanceDynamic::Create(ClothPaintMaterial, nullptr);
 		ClothPaintMaterialWireframeInstance = UMaterialInstanceDynamic::Create(ClothPaintMaterialWireframe, nullptr);
+		ClothPaintOpaqueMaterialInstance = UMaterialInstanceDynamic::Create(ClothPaintOpaqueMaterial, nullptr);
+		ClothPaintOpaqueMaterialWireframeInstance = UMaterialInstanceDynamic::Create(ClothPaintOpaqueMaterialWireframe, nullptr);
 #endif
 		FString ValidPreviewShadowsIndicatorMaterialName = PreviewShadowsIndicatorMaterialName.ToString();
 		if (ValidPreviewShadowsIndicatorMaterialName.IsEmpty())
