[Viewport Toolbar] Switching UnrealEd FMenuBuilder entries to use UToolMenu

- UE::UnrealEd::CreateCameraMenuWidget
- UE::UnrealEd::CreateViewportToolbarAssetViewerProfileSubmenu

diff --git a/Engine/Source/Editor/UnrealEd/Private/ViewportToolbar/UnrealEdViewportToolbar.cpp b/Engine/Source/Editor/UnrealEd/Private/ViewportToolbar/UnrealEdViewportToolbar.cpp
--- a/Engine/Source/Editor/UnrealEd/Private/ViewportToolbar/UnrealEdViewportToolbar.cpp
+++ b/Engine/Source/Editor/UnrealEd/Private/ViewportToolbar/UnrealEdViewportToolbar.cpp
@@ -2090,22 +2090,37 @@ FToolMenuEntry CreateNumericEntry(
 
 TSharedRef<SWidget> CreateCameraMenuWidget(const TSharedRef<SEditorViewport>& InViewport)
 {
-	constexpr bool bInShouldCloseWindowAfterMenuSelection = true;
-	FMenuBuilder CameraMenuBuilder(bInShouldCloseWindowAfterMenuSelection, InViewport->GetCommandList());
+	// We generate a menu via UToolMenus, so we can use FillShowSubmenu call from both old and new toolbar
+	FName OldShowMenuName = "LevelEditor.OldViewportToolbar.CameraOptions";
 
-	// Camera types
-	CameraMenuBuilder.AddMenuEntry(FEditorViewportCommands::Get().Perspective);
+	if (!UToolMenus::Get()->IsMenuRegistered(OldShowMenuName))
+	{
+		UToolMenu* Menu = UToolMenus::Get()->RegisterMenu(OldShowMenuName, NAME_None, EMultiBoxType::Menu, false);
+		Menu->AddDynamicSection(
+			"BaseSection",
+			FNewToolMenuDelegate::CreateLambda(
+				[](UToolMenu* InMenu)
+				{
+					PopulateCameraMenu(InMenu);
+				}
+			)
+		);
+	}
 
-	CameraMenuBuilder.BeginSection("LevelViewportCameraType_Ortho", LOCTEXT("CameraTypeHeader_Ortho", "Orthographic"));
-	CameraMenuBuilder.AddMenuEntry(FEditorViewportCommands::Get().Top);
-	CameraMenuBuilder.AddMenuEntry(FEditorViewportCommands::Get().Bottom);
-	CameraMenuBuilder.AddMenuEntry(FEditorViewportCommands::Get().Left);
-	CameraMenuBuilder.AddMenuEntry(FEditorViewportCommands::Get().Right);
-	CameraMenuBuilder.AddMenuEntry(FEditorViewportCommands::Get().Front);
-	CameraMenuBuilder.AddMenuEntry(FEditorViewportCommands::Get().Back);
-	CameraMenuBuilder.EndSection();
+	FToolMenuContext MenuContext;
+	{
+		MenuContext.AppendCommandList(InViewport->GetCommandList());
+
+		// Add the UnrealEd viewport toolbar context.
+		{
+			UUnrealEdViewportToolbarContext* const ContextObject =
+				UE::UnrealEd::CreateViewportToolbarDefaultContext(InViewport);
 
-	return CameraMenuBuilder.MakeWidget();
+			MenuContext.AddObject(ContextObject);
+		}
+	}
+
+	return UToolMenus::Get()->GenerateWidget(OldShowMenuName, MenuContext);
 }
 
 TSharedRef<SWidget> CreateFOVMenuWidget(const TSharedRef<SEditorViewport>& InViewport)
@@ -2483,11 +2498,7 @@ FToolMenuEntry CreateViewportToolbarAssetViewerProfileSubmenu()
 								return;
 							}
 
-							FToolMenuSection& UnnamedSection = Submenu->FindOrAddSection(NAME_None);
-
-							constexpr bool bInShouldCloseWindowAfterMenuSelection = true;
-							FMenuBuilder PreviewProfilesSelectionMenuBuilder(bInShouldCloseWindowAfterMenuSelection, nullptr);
-							PreviewProfilesSelectionMenuBuilder.BeginSection(
+							FToolMenuSection& PreviewProfilesSelectionSection = Submenu->FindOrAddSection(
 								"AssetViewerProfileSelectionSection",
 								LOCTEXT("AssetViewerProfileSelectionSectionLabel", "Preview Scene Profiles")
 							);
@@ -2499,7 +2510,8 @@ FToolMenuEntry CreateViewportToolbarAssetViewerProfileSubmenu()
 							for (int32 ProfileIndex = 0; ProfileIndex < PreviewProfiles.Num(); ProfileIndex++)
 							{
 								const FString& ProfileName = PreviewProfiles[ProfileIndex];
-								PreviewProfilesSelectionMenuBuilder.AddMenuEntry(
+								PreviewProfilesSelectionSection.AddMenuEntry(
+									NAME_None,
 									FText::FromString(ProfileName),
 									FText(),
 									FSlateIcon(),
@@ -2532,16 +2544,9 @@ FToolMenuEntry CreateViewportToolbarAssetViewerProfileSubmenu()
 											}
 										)
 									),
-									NAME_None,
 									EUserInterfaceActionType::RadioButton
 								);
 							}
-
-							PreviewProfilesSelectionMenuBuilder.EndSection();
-
-							UnnamedSection.AddEntry(FToolMenuEntry::InitWidget(
-								"AssetViewerProfile", PreviewProfilesSelectionMenuBuilder.MakeWidget(), FText(), true
-							));
 						}
 					)
 				);
