Eliminate unreachable code

#rnx

diff --git a/Engine/Source/Editor/UnrealEd/Private/Fbx/FbxFactory.cpp b/Engine/Source/Editor/UnrealEd/Private/Fbx/FbxFactory.cpp
--- a/Engine/Source/Editor/UnrealEd/Private/Fbx/FbxFactory.cpp
+++ b/Engine/Source/Editor/UnrealEd/Private/Fbx/FbxFactory.cpp
@@ -222,39 +222,36 @@ UObject* UFbxFactory::FactoryCreateFile
 			{
 				bIsObjectSupported = true;
 				// Using RF_NoFlags because we want to include class default object, which are used by the reimport manager.
-				for (TObjectIterator<UReimportFbxStaticMeshFactory> ReimportFactoryIterator(RF_NoFlags); ReimportFactoryIterator; ++ReimportFactoryIterator)
+				if (TObjectIterator<UReimportFbxStaticMeshFactory> ReimportFactoryIterator(RF_NoFlags); ReimportFactoryIterator)
 				{
 					HandlerOriginalImportTask = ReimportFactoryIterator->AssetImportTask;
 					ReimportFactoryIterator->AssetImportTask = AssetImportTask;
 					ReimportHandler = *ReimportFactoryIterator;
 					ReimportHandlerFactory = *ReimportFactoryIterator;
-					break;
 				}
 			}
 			else if (Cast<USkeletalMesh>(ExistingObject))
 			{
 				bIsObjectSupported = true;
 				// Using RF_NoFlags because we want to include class default object, which are used by the reimport manager.
-				for (TObjectIterator<UReimportFbxSkeletalMeshFactory> ReimportFactoryIterator(RF_NoFlags); ReimportFactoryIterator; ++ReimportFactoryIterator)
+				if (TObjectIterator<UReimportFbxSkeletalMeshFactory> ReimportFactoryIterator(RF_NoFlags); ReimportFactoryIterator)
 				{
 					HandlerOriginalImportTask = ReimportFactoryIterator->AssetImportTask;
 					ReimportFactoryIterator->AssetImportTask = AssetImportTask;
 					ReimportHandler = *ReimportFactoryIterator;
 					ReimportHandlerFactory = *ReimportFactoryIterator;
-					break;
 				}
 			}
 			else if (Cast<UAnimSequence>(ExistingObject))
 			{
 				bIsObjectSupported = true;
 				// Using RF_NoFlags because we want to include class default object, which are used by the reimport manager.
-				for (TObjectIterator<UReimportFbxAnimSequenceFactory> ReimportFactoryIterator(RF_NoFlags); ReimportFactoryIterator; ++ReimportFactoryIterator)
+				if (TObjectIterator<UReimportFbxAnimSequenceFactory> ReimportFactoryIterator(RF_NoFlags); ReimportFactoryIterator)
 				{
 					HandlerOriginalImportTask = ReimportFactoryIterator->AssetImportTask;
 					ReimportFactoryIterator->AssetImportTask = AssetImportTask;
 					ReimportHandler = *ReimportFactoryIterator;
 					ReimportHandlerFactory = *ReimportFactoryIterator;
-					break;
 				}
 			}
 
