IncrementalCook: For FCookDependency::Function, add the output function ReportInvalidated to FCookDependencyContext, so that functions that want to handle invalidation themselves can do so.
#rnx

diff --git a/Engine/Source/Editor/UnrealEd/Private/Cooker/CookFunctionLibrary.cpp b/Engine/Source/Editor/UnrealEd/Private/Cooker/CookFunctionLibrary.cpp
--- a/Engine/Source/Editor/UnrealEd/Private/Cooker/CookFunctionLibrary.cpp
+++ b/Engine/Source/Editor/UnrealEd/Private/Cooker/CookFunctionLibrary.cpp
@@ -213,11 +213,17 @@ void UCookFunctionLibrary::CookAsset(UObject* Object, const FString& ForPlatform
 		FGeneratedPackageResultStruct* Todo = nullptr;
 		TArray<FName> RuntimeDepsTodo; // I don't believe we need valid data in here at this time
 
-		FString ErrorMessage;
-		FCookDependencies Deps = FCookDependencies::Collect(Package, TargetPlatform, &Result, Todo, MoveTemp(RuntimeDepsTodo), &ErrorMessage);
-		if(!ErrorMessage.IsEmpty())
+		TArray<TPair<ELogVerbosity::Type, FString>> Messages;
+		FCookDependencies Deps = FCookDependencies::Collect(Package, TargetPlatform, &Result, Todo, MoveTemp(RuntimeDepsTodo), &Messages);
+		if (!Deps.IsValid())
 		{
-			UE_LOG(LogCook, Warning, TEXT("Error collecting cook dependencies: %s"), *ErrorMessage);
+			TStringBuilder<256> ErrorMessage;
+			ErrorMessage << TEXT("Error collecting cook dependencies:");
+			for (TPair<ELogVerbosity::Type, FString>& MessagePair : Messages)
+			{
+				ErrorMessage << TEXT("\n\t") << MessagePair.Value;
+			}
+			UE_LOG(LogCook, Warning, TEXT("%s"), *ErrorMessage);
 		}
 
 		// persist cook dependencies along side the loose files we wrote above..
diff --git a/Engine/Source/Editor/UnrealEd/Private/TargetDomain/TargetDomainUtils.cpp b/Engine/Source/Editor/UnrealEd/Private/TargetDomain/TargetDomainUtils.cpp
--- a/Engine/Source/Editor/UnrealEd/Private/TargetDomain/TargetDomainUtils.cpp
+++ b/Engine/Source/Editor/UnrealEd/Private/TargetDomain/TargetDomainUtils.cpp
@@ -121,7 +121,8 @@ bool FCookDependencies::HasKeyMatch(const FAssetPackageData* OverrideAssetPackag
 	}
 	if (CurrentKey.IsZero())
 	{
-		if (!TryCalculateCurrentKey(OverrideAssetPackageData))
+		ECurrentKeyResult Result = TryCalculateCurrentKey(OverrideAssetPackageData);
+		if (Result != ECurrentKeyResult::Success)
 		{
 			return false;
 		}
@@ -129,25 +130,26 @@ bool FCookDependencies::HasKeyMatch(const FAssetPackageData* OverrideAssetPackag
 	return CurrentKey == StoredKey;
 }
 
-bool FCookDependencies::TryCalculateCurrentKey(const FAssetPackageData* OverrideAssetPackageData,
-	FString* OutErrorMessage)
+FCookDependencies::ECurrentKeyResult FCookDependencies::TryCalculateCurrentKey(
+	const FAssetPackageData* OverrideAssetPackageData,
+	TArray<TPair<ELogVerbosity::Type, FString>>* OutMessages)
 {
 	IAssetRegistry* AssetRegistry = IAssetRegistry::Get();
 	if (PackageName.IsNone())
 	{
-		if (OutErrorMessage) *OutErrorMessage = TEXT("PackageName is not set.");
-		return false;
+		if (OutMessages) OutMessages->Emplace(ELogVerbosity::Error, TEXT("PackageName is not set."));
+		return ECurrentKeyResult::Error;
 	}
 	if (!AssetRegistry)
 	{
-		if (OutErrorMessage) *OutErrorMessage = TEXT("AssetRegistry is unavailable.");
-		return false;
+		if (OutMessages) OutMessages->Emplace(ELogVerbosity::Error, TEXT("AssetRegistry is unavailable."));
+		return ECurrentKeyResult::Error;
 	}
 	FEditorDomain* EditorDomain = FEditorDomain::Get();
 	if (!EditorDomain)
 	{
-		if (OutErrorMessage) *OutErrorMessage = TEXT("EditorDomain is unavailable.");
-		return false;
+		if (OutMessages) OutMessages->Emplace(ELogVerbosity::Error, TEXT("EditorDomain is unavailable."));
+		return ECurrentKeyResult::Error;
 	}
 	FBlake3 KeyBuilder;
 	KeyBuilder.Update(&TargetDomainVersion, sizeof(TargetDomainVersion));
@@ -157,8 +159,8 @@ bool FCookDependencies::TryCalculateCurrentKey(const FAssetPackageData* Override
 		EditorDomain->GetPackageDigest(PackageName);
 	if (!PackageDigest.IsSuccessful())
 	{
-		if (OutErrorMessage) *OutErrorMessage = PackageDigest.GetStatusString();
-		return false;
+		if (OutMessages) OutMessages->Emplace(ELogVerbosity::Error, PackageDigest.GetStatusString());
+		return ECurrentKeyResult::Error;
 	}
 	KeyBuilder.Update(&PackageDigest.Hash, sizeof(PackageDigest.Hash));
 
@@ -171,23 +173,23 @@ bool FCookDependencies::TryCalculateCurrentKey(const FAssetPackageData* Override
 			FTopLevelAssetPath& ClassPath = ClassPaths.Emplace_GetRef(ClassPathStr);
 			if (!ClassPath.IsValid())
 			{
-				if (OutErrorMessage)
+				if (OutMessages)
 				{
-					*OutErrorMessage = FString::Printf(
-						TEXT("ClassDependency failed: %s is not a valid TopLevelAssetPath."), *ClassPathStr);
+					OutMessages->Emplace(ELogVerbosity::Error, FString::Printf(
+						TEXT("ClassDependency failed: %s is not a valid TopLevelAssetPath."), *ClassPathStr));
 				}
-				return false;
+				return ECurrentKeyResult::Error;
 			}
 		}
 		FString AppendDigestError;
 		if (!UE::EditorDomain::TryAppendClassDigests(KeyBuilder, ClassPaths, &AppendDigestError))
 		{
-			if (OutErrorMessage)
+			if (OutMessages)
 			{
-				*OutErrorMessage = FString::Printf(
-					TEXT("ClassDependency failed: %s"), *AppendDigestError);
+				OutMessages->Emplace(ELogVerbosity::Error, FString::Printf(
+					TEXT("ClassDependency failed: %s"), *AppendDigestError));
 			}
-			return false;
+			return ECurrentKeyResult::Error;
 		}
 	}
 
@@ -196,12 +198,12 @@ bool FCookDependencies::TryCalculateCurrentKey(const FAssetPackageData* Override
 		PackageDigest = EditorDomain->GetPackageDigest(PackageDependency);
 		if (!PackageDigest.IsSuccessful())
 		{
-			if (OutErrorMessage)
+			if (OutMessages)
 			{
-				*OutErrorMessage = FString::Printf(TEXT("Could not create PackageDigest for %s: %s"),
-					*PackageDependency.ToString(), *PackageDigest.GetStatusString());
+				OutMessages->Emplace(ELogVerbosity::Error, FString::Printf(TEXT("Could not create PackageDigest for %s: %s"),
+					*PackageDependency.ToString(), *PackageDigest.GetStatusString()));
 			}
-			return false;
+			return ECurrentKeyResult::Error;
 		}
 		KeyBuilder.Update(&PackageDigest.Hash, sizeof(PackageDigest.Hash));
 	}
@@ -230,39 +232,36 @@ bool FCookDependencies::TryCalculateCurrentKey(const FAssetPackageData* Override
 		}
 	}
 
+	ECurrentKeyResult Result = ECurrentKeyResult::Success;
 	if (!CookDependencies.IsEmpty())
 	{
-		bool bError = false;
-		UE::Cook::FCookDependencyContext Context(&KeyBuilder, [&bError, OutErrorMessage](FString&& ErrorMessage)
+		UE::Cook::FCookDependencyContext Context(&KeyBuilder,
+			[&Result, OutMessages](ELogVerbosity::Type Verbosity, FString&& Message)
 			{
-				if (OutErrorMessage)
+				if (OutMessages)
 				{
-					if (bError)
-					{
-						*OutErrorMessage += TEXT("\n");
-						*OutErrorMessage += ErrorMessage;
-					}
-					else
-					{
-						*OutErrorMessage = MoveTemp(ErrorMessage);
-					}
+					OutMessages->Emplace(Verbosity, MoveTemp(Message));
 				}
-				bError = true;
-			}, PackageName);
+			},
+			[&Result](ELogVerbosity::Type Verbosity)
+			{
+				ECurrentKeyResult NewResult = Verbosity <= ELogVerbosity::Error
+					? ECurrentKeyResult::Error : ECurrentKeyResult::Invalidated;
+				Result = FMath::Max(Result, NewResult);
+			},
+			PackageName);
 
 		for (UE::Cook::FCookDependency& CookDependency : CookDependencies)
 		{
 			CookDependency.UpdateHash(Context);
 		}
-		if (bError)
-		{
-			return false;
-		}
 	}
 
-	if (OutErrorMessage) OutErrorMessage->Reset();
-	CurrentKey = KeyBuilder.Finalize();
-	return true;
+	if (Result != ECurrentKeyResult::Error)
+	{
+		CurrentKey = KeyBuilder.Finalize();
+	}
+	return Result;
 }
 
 void FCookDependencies::Reset()
@@ -316,25 +315,42 @@ EPackageMountPoint GetPackageMountPoint(FName PackageName, FName TransientPackag
 	return EPackageMountPoint::Content;
 }
 
+static void JoinMessagesIntoErrorReason(FStringBuilderBase& OutText,
+	TArray<TPair<ELogVerbosity::Type, FString>>& Messages)
+{
+	if (Messages.IsEmpty())
+	{
+		OutText << TEXT(".");
+	}
+	else
+	{
+		OutText << TEXT(":");
+		for (TPair<ELogVerbosity::Type, FString>& MessagePair : Messages)
+		{
+			OutText << TEXT("\n\t") << MessagePair.Value;
+		}
+	}
+}
+
 FCookDependencies FCookDependencies::Collect(UPackage* Package, const ITargetPlatform* TargetPlatform,
 	FSavePackageResultStruct* SaveResult, const FGeneratedPackageResultStruct* GeneratedResult, 
-	TArray<FName>&& RuntimeDependencies, FString* OutErrorMessage)
+	TArray<FName>&& RuntimeDependencies, TArray<TPair<ELogVerbosity::Type, FString>>* OutMessages)
 {
 	if (!Package)
 	{
-		if (OutErrorMessage) *OutErrorMessage = TEXT("Invalid null package.");
+		if (OutMessages) OutMessages->Emplace(ELogVerbosity::Error, TEXT("Invalid null package."));
 		return FCookDependencies();
 	}
 	IAssetRegistry* AssetRegistry = IAssetRegistry::Get();
 	if (!AssetRegistry)
 	{
-		if (OutErrorMessage) *OutErrorMessage = TEXT("AssetRegistry is unavailable.");
+		if (OutMessages) OutMessages->Emplace(ELogVerbosity::Error, TEXT("AssetRegistry is unavailable."));
 		return FCookDependencies();
 	}
 	FEditorDomain* EditorDomain = FEditorDomain::Get();
 	if (!EditorDomain)
 	{
-		if (OutErrorMessage) *OutErrorMessage = TEXT("EditorDomain is unavailable.");
+		if (OutMessages) OutMessages->Emplace(ELogVerbosity::Error, TEXT("EditorDomain is unavailable."));
 		return FCookDependencies();
 	}
 
@@ -465,9 +481,10 @@ FCookDependencies FCookDependencies::Collect(UPackage* Package, const ITargetPla
 				NextIncludedDependencies->Reset();
 				if (IncludedDependenciesRound > MaxNumberOfIncludedDependenciesRounds)
 				{
-					if (OutErrorMessage)
+					if (OutMessages)
 					{
-						*OutErrorMessage = TEXT("More rounds than expected when handling included dependencies.");
+						OutMessages->Emplace(ELogVerbosity::Error,
+							TEXT("More rounds than expected when handling included dependencies."));
 					}
 					return FCookDependencies();
 				}
@@ -516,16 +533,18 @@ FCookDependencies FCookDependencies::Collect(UPackage* Package, const ITargetPla
 						if (!IncludeDependencies.bInitialized)
 						{
 							IncludeDependencies.Dependencies = CollectSettingsObject(SettingsObject,
-								&IncludeDependencies.ErrorMessage);
+								&IncludeDependencies.Messages);
 							IncludeDependencies.bInitialized = true;
 						}
 						if (!IncludeDependencies.Dependencies.IsValid())
 						{
-							if (OutErrorMessage)
+							if (OutMessages)
 							{
-								*OutErrorMessage = FString::Printf(
-									TEXT("Dependencies for SettingsObject %s are unavailable: %s."),
-									*SettingsObject->GetPathName(), *IncludeDependencies.ErrorMessage);
+								TStringBuilder<256> ErrorText;
+								ErrorText << TEXT("Dependencies for SettingsObject ") << SettingsObject->GetPathName()
+									<< TEXT(" are unavailable");
+								JoinMessagesIntoErrorReason(ErrorText, IncludeDependencies.Messages);
+								OutMessages->Emplace(ELogVerbosity::Error, FString(*ErrorText));
 							}
 							return FCookDependencies();
 						}
@@ -605,24 +624,25 @@ FCookDependencies FCookDependencies::Collect(UPackage* Package, const ITargetPla
 	Result.ScriptPackageDependencies.SetNum(Algo::Unique(Result.ScriptPackageDependencies), EAllowShrinking::Yes);
 
 	const FAssetPackageData* AssetPackageData = GeneratedResult ? &GeneratedResult->AssetPackageData : nullptr;
-	if (!Result.TryCalculateCurrentKey(AssetPackageData, OutErrorMessage))
+	ECurrentKeyResult CurrentKeyResult = Result.TryCalculateCurrentKey(AssetPackageData, OutMessages);
+	if (CurrentKeyResult == ECurrentKeyResult::Error)
 	{
 		return FCookDependencies();
 	}
 	Result.StoredKey = Result.CurrentKey;
 	Result.bValid = true;
 
-	if (OutErrorMessage) OutErrorMessage->Reset();
 	return Result;
 }
 
-FCookDependencies FCookDependencies::CollectSettingsObject(const UObject* Object, FString* OutErrorMessage)
+FCookDependencies FCookDependencies::CollectSettingsObject(const UObject* Object,
+	TArray<TPair<ELogVerbosity::Type, FString>>* OutMessages)
 {
 	if (!Object)
 	{
-		if (OutErrorMessage)
+		if (OutMessages)
 		{
-			*OutErrorMessage = TEXT("Invalid null Object.");
+			OutMessages->Emplace(ELogVerbosity::Error, TEXT("Invalid null Object."));
 		}
 		return FCookDependencies();
 	}
@@ -630,18 +650,19 @@ FCookDependencies FCookDependencies::CollectSettingsObject(const UObject* Object
 	UClass* Class = Object->GetClass();
 	if (!Class->HasAnyClassFlags(CLASS_Config | CLASS_PerObjectConfig))
 	{
-		if (OutErrorMessage)
+		if (OutMessages)
 		{
-			*OutErrorMessage = FString::Printf(TEXT("Class %s is not a config class."), *Class->GetPathName());
+			OutMessages->Emplace(ELogVerbosity::Error, FString::Printf(
+				TEXT("Class %s is not a config class."), *Class->GetPathName()));
 		}
 		return FCookDependencies();
 	}
 	if (!Class->HasAnyClassFlags(CLASS_PerObjectConfig) && Object != Class->GetDefaultObject())
 	{
-		if (OutErrorMessage)
+		if (OutMessages)
 		{
-			*OutErrorMessage = FString::Printf(TEXT("Class %s is not a per-object-config class."),
-				*Class->GetPathName());
+			OutMessages->Emplace(ELogVerbosity::Error, FString::Printf(
+				TEXT("Class %s is not a per-object-config class."), *Class->GetPathName()));
 		}
 		return FCookDependencies();
 	}
@@ -819,7 +840,7 @@ FCookDependencyGroups::FRecordedDependencies& FCookDependencyGroups::FindOrCreat
 }
 
 FBuildDefinitionList FBuildDefinitionList::Collect(UPackage* Package, const ITargetPlatform* TargetPlatform,
-	FString* OutErrorMessage)
+	TArray<TPair<ELogVerbosity::Type, FString>>* OutMessages)
 {
 	using namespace UE::DerivedData;
 
@@ -828,7 +849,7 @@ FBuildDefinitionList FBuildDefinitionList::Collect(UPackage* Package, const ITar
 	// TODO_BuildDefinitionList: Calculate and store BuildDefinitionList on the PackageData, or collect it here from some other source.
 	if (Result.Definitions.IsEmpty())
 	{
-		if (OutErrorMessage) *OutErrorMessage = TEXT("Not yet implemented");
+		if (OutMessages) OutMessages->Emplace(ELogVerbosity::Error, TEXT("Not yet implemented"));
 		return FBuildDefinitionList();
 	}
 
@@ -838,7 +859,6 @@ FBuildDefinitionList FBuildDefinitionList::Collect(UPackage* Package, const ITar
 			return A.GetKey().Hash < B.GetKey().Hash;
 		});
 
-	if (OutErrorMessage) OutErrorMessage->Reset();
 	return Result;
 }
 
@@ -914,16 +934,18 @@ bool TryCollectAndStoreCookDependencies(UPackage* Package, const ITargetPlatform
 	FSavePackageResultStruct* SaveResult, const FGeneratedPackageResultStruct* GeneratedResult,
 	TArray<FName>&& RuntimeDependencies, IPackageWriter::FCommitAttachmentInfo& OutResult)
 {
-	FString ErrorMessage;
+	TArray<TPair<ELogVerbosity::Type, FString>> Messages;
 	FCookDependencies CookDependencies = FCookDependencies::Collect(Package, TargetPlatform, SaveResult,
-		GeneratedResult, MoveTemp(RuntimeDependencies), &ErrorMessage);
+		GeneratedResult, MoveTemp(RuntimeDependencies), &Messages);
 	if (!CookDependencies.IsValid())
 	{
+		TStringBuilder<256> LogText;
+		LogText << TEXT("Could not collect CookDependencies for package '") << Package->GetFName() << TEXT("'");
+		JoinMessagesIntoErrorReason(LogText, Messages);
+
 		// IterativeTODO: This error occurs due to dependencies on _Verse and on some transient packages.
-#if 0
-		UE_LOG(LogCook, Error, TEXT("Could not collect CookDependencies for package '%s': %s"),
-			*Package->GetName(), *ErrorMessage);
-#endif
+		// Raise Verbosity to Error once that is fixed.
+		UE_LOG(LogCook, Verbose, TEXT("%s"), *LogText);
 		OutResult.Value = FCbObject();
 		return false;
 	}
diff --git a/Engine/Source/Editor/UnrealEd/Private/TargetDomain/TargetDomainUtils.h b/Engine/Source/Editor/UnrealEd/Private/TargetDomain/TargetDomainUtils.h
--- a/Engine/Source/Editor/UnrealEd/Private/TargetDomain/TargetDomainUtils.h
+++ b/Engine/Source/Editor/UnrealEd/Private/TargetDomain/TargetDomainUtils.h
@@ -7,8 +7,10 @@
 #include "Containers/Array.h"
 #include "Containers/ArrayView.h"
 #include "Containers/ContainersFwd.h"
+#include "Containers/UnrealString.h"
 #include "DerivedDataBuildDefinition.h"
 #include "IO/IoHash.h"
+#include "Logging/LogVerbosity.h"
 #include "Misc/Optional.h"
 #include "Serialization/PackageWriter.h"
 #include "UObject/NameTypes.h"
@@ -16,7 +18,6 @@
 class FCbObject;
 class FCbObjectView;
 class FCbWriter;
-class FString;
 class ITargetPlatform;
 class UPackage;
 struct FSavePackageResultStruct;
@@ -91,7 +92,14 @@ public:
 	bool HasKeyMatch(const FAssetPackageData* OverrideAssetPackageData);
 
 	// Modifying the structure data
-	bool TryCalculateCurrentKey(const FAssetPackageData* OverrideAssetPackageData, FString* OutErrorMessage = nullptr);
+	enum class ECurrentKeyResult
+	{
+		Success,
+		Invalidated,
+		Error,
+	};
+	ECurrentKeyResult TryCalculateCurrentKey(const FAssetPackageData* OverrideAssetPackageData,
+		TArray<TPair<ELogVerbosity::Type, FString>>* OutMessages = nullptr);
 	void Reset();
 	void Empty();
 
@@ -102,8 +110,9 @@ public:
 	 */
 	static FCookDependencies Collect(UPackage* Package, const ITargetPlatform* TargetPlatform,
 		FSavePackageResultStruct* SaveResult, const FGeneratedPackageResultStruct* GeneratedResult,
-		TArray<FName>&& RuntimeDependencies, FString* OutErrorMessage = nullptr);
-	static FCookDependencies CollectSettingsObject(const UObject* Object, FString* OutErrorMessage);
+		TArray<FName>&& RuntimeDependencies, TArray<TPair<ELogVerbosity::Type, FString>>* OutMessages = nullptr);
+	static FCookDependencies CollectSettingsObject(const UObject* Object,
+		TArray<TPair<ELogVerbosity::Type, FString>>* OutMessages);
 
 	// Fetch function to load the dependencies from a PackageStore is not yet implemented independently for
 	// this structure. Use FCookAttachments instead. 
@@ -140,7 +149,7 @@ public:
 	struct FRecordedDependencies
 	{
 		FCookDependencies Dependencies;
-		FString ErrorMessage;
+		TArray<TPair<ELogVerbosity::Type,FString>> Messages;
 		bool bInitialized = false;
 	};
 
@@ -162,7 +171,7 @@ public:
 
 	/** Collect DDC BuildDefinitions that were issued from the load/save of the given package and platform. */
 	static FBuildDefinitionList Collect(UPackage* Package, const ITargetPlatform* TargetPlatform,
-		FString* OutErrorMessage = nullptr);
+		TArray<TPair<ELogVerbosity::Type, FString>>* OutMessages = nullptr);
 
 private:
 	friend bool ::LoadFromCompactBinary(FCbObject&& ObjectView, FBuildDefinitionList& Definitions);
diff --git a/Engine/Source/Runtime/CoreUObject/Private/Cooker/CookDependency.cpp b/Engine/Source/Runtime/CoreUObject/Private/Cooker/CookDependency.cpp
--- a/Engine/Source/Runtime/CoreUObject/Private/Cooker/CookDependency.cpp
+++ b/Engine/Source/Runtime/CoreUObject/Private/Cooker/CookDependency.cpp
@@ -716,13 +716,23 @@ void FCookDependencyContext::Update(const void* Data, uint64 Size)
 	Blake3Hasher.Update(Data, Size);
 }
 
-void FCookDependencyContext::LogError(FString Message)
+void FCookDependencyContext::ReportInvalidated()
+{
+	OnInvalidated(ELogVerbosity::Display);
+}
+
+void FCookDependencyContext::ReportError()
+{
+	OnInvalidated(ELogVerbosity::Error);
+}
+
+void FCookDependencyContext::Log(ELogVerbosity::Type Verbosity, FString Message)
 {
 	for (TUniqueFunction<FString(FString&&)>& Handler : ErrorHandlers)
 	{
 		Message = Handler(MoveTemp(Message));
 	}
-	OnLogError(MoveTemp(Message));
+	OnLog(Verbosity, MoveTemp(Message));
 }
 
 FCookDependencyContext::FErrorHandlerScope::FErrorHandlerScope(FCookDependencyContext& InContext)
diff --git a/Engine/Source/Runtime/CoreUObject/Public/Cooker/CookDependency.h b/Engine/Source/Runtime/CoreUObject/Public/Cooker/CookDependency.h
--- a/Engine/Source/Runtime/CoreUObject/Public/Cooker/CookDependency.h
+++ b/Engine/Source/Runtime/CoreUObject/Public/Cooker/CookDependency.h
@@ -6,6 +6,7 @@
 #include "Containers/StringView.h"
 #include "Containers/UnrealString.h"
 #include "HAL/PreprocessorHelpers.h"
+#include "Logging/LogVerbosity.h"
 #include "Misc/PackageAccessTracking.h"
 #include "Serialization/CompactBinary.h"
 #include "Templates/Function.h"
@@ -45,7 +46,10 @@ struct FCookDependencyContext
 {
 public:
 	/** InHasher is void* to mask the implementation details of the hashbuilder. See Update function. */
-	explicit FCookDependencyContext(void* InHasher, TUniqueFunction<void(FString&&)>&& InOnLogError, FName InPackageName);
+	explicit FCookDependencyContext(void* InHasher,
+		TUniqueFunction<void(ELogVerbosity::Type, FString&&)>&& InOnLog,
+		TUniqueFunction<void(ELogVerbosity::Type)>&& InOnInvalidated,
+		FName InPackageName);
 
 	/**
 	 * Update the hashbuilder for the key being constructed (e.g. TargetDomainKey for cooked packages)
@@ -53,13 +57,30 @@ public:
 	 */
 	COREUOBJECT_API void Update(const void* Data, uint64 Size);
 
+	/**
+	 * Reports that current evaluation of the function is different from all previous evaluations for a reason that
+	 * cannot be reported as data passed into Update. When called while calculating the initial hash this call is
+	 * ignored. When called while testing iteratively skippable the package is marked modified and recooked.
+	 */
+	COREUOBJECT_API void ReportInvalidated();
 	/**
 	 * Reports failure to compute the hash (e.g. because a file cannot be read).
-	 * When calculating the initial hash during package save, this error will be logged as an error
-	 * and the package will be recooked on the next cook. When calculating the hash during an incremental cook
-	 * the message will be logged at Log level and will cause the package to be recooked.
+	 * When called while calculating the initial hash the storage of the key fails and the package will be recooked on
+	 * the next cook. When called while testing iteratively skippable the package is marked modified and recooked.
 	 */
-	COREUOBJECT_API void LogError(FString Message);
+	COREUOBJECT_API void ReportError();
+
+	/**
+	 * Send a message to the cook dependency context with the given severity. The message may be suppressed or
+	 * reduced in verbosity based on the calling context.
+	 */
+	COREUOBJECT_API void Log(ELogVerbosity::Type, FString Message);
+
+	/** Calls Log(Error, Message) and ReportError. */
+	void LogError(FString Message);
+
+	/** Calls Log(Display, Message) and ReportInvalidated. */
+	void LogInvalidated(FString Message);
 
 	/**
 	 * Private implementation struct used for AddErrorHandlerScope. Should only be used via
@@ -89,7 +110,8 @@ public:
 	 */
 	COREUOBJECT_API FName GetPackageName() const { return PackageName; }
 private:
-	TUniqueFunction<void(FString&&)> OnLogError;
+	TUniqueFunction<void(ELogVerbosity::Type, FString&&)> OnLog;
+	TUniqueFunction<void(ELogVerbosity::Type)> OnInvalidated;
 	TArray<TUniqueFunction<FString(FString&&)>, TInlineAllocator<1>> ErrorHandlers;
 	FName PackageName;
 	void* Hasher; // Type is void* to mask the implementation detail
@@ -399,6 +421,18 @@ struct FCookDependencyFunctionRegistration
 namespace UE::Cook
 {
 
+inline void FCookDependencyContext::LogError(FString Message)
+{
+	Log(ELogVerbosity::Error, MoveTemp(Message));
+	ReportError();
+}
+
+inline void FCookDependencyContext::LogInvalidated(FString Message)
+{
+	Log(ELogVerbosity::Display, MoveTemp(Message));
+	ReportInvalidated();
+}
+
 inline ECookDependency FCookDependency::GetType() const
 {
 	return Type;
@@ -613,8 +647,12 @@ inline bool FCookDependency::operator!=(const FCookDependency& Other) const
 	return !(*this == Other);
 }
 
-inline FCookDependencyContext::FCookDependencyContext(void* InHasher, TUniqueFunction<void(FString&&)>&& InOnLogError, FName InPackageName)
-	: OnLogError(MoveTemp(InOnLogError))
+inline FCookDependencyContext::FCookDependencyContext(void* InHasher,
+	TUniqueFunction<void(ELogVerbosity::Type, FString&&)>&& InOnLog,
+	TUniqueFunction<void(ELogVerbosity::Type)>&& InOnInvalidated,
+	FName InPackageName)
+	: OnLog(MoveTemp(InOnLog))
+	, OnInvalidated(MoveTemp(InOnInvalidated))
 	, PackageName(InPackageName)
 	, Hasher(InHasher)
 {
