TargetDomainUtils.cpp: Suppress spurious static analysis error warning V547: Expression 'Result != ECurrentKeyResult::Error' is always true.
#rnx

diff --git a/Engine/Source/Editor/UnrealEd/Private/TargetDomain/TargetDomainUtils.cpp b/Engine/Source/Editor/UnrealEd/Private/TargetDomain/TargetDomainUtils.cpp
--- a/Engine/Source/Editor/UnrealEd/Private/TargetDomain/TargetDomainUtils.cpp
+++ b/Engine/Source/Editor/UnrealEd/Private/TargetDomain/TargetDomainUtils.cpp
@@ -257,7 +257,7 @@ FCookDependencies::ECurrentKeyResult FCookDependencies::TryCalculateCurrentKey(
 		}
 	}
 
-	if (Result != ECurrentKeyResult::Error)
+	if (Result != ECurrentKeyResult::Error) // -V547
 	{
 		CurrentKey = KeyBuilder.Finalize();
 	}
