[ContentBrowser] CIS fix

diff --git a/Engine/Plugins/Editor/ContentBrowser/ContentBrowserAssetDataSource/Source/ContentBrowserAssetDataSource/Private/ContentBrowserAssetDataSource.cpp b/Engine/Plugins/Editor/ContentBrowser/ContentBrowserAssetDataSource/Source/ContentBrowserAssetDataSource/Private/ContentBrowserAssetDataSource.cpp
--- a/Engine/Plugins/Editor/ContentBrowser/ContentBrowserAssetDataSource/Source/ContentBrowserAssetDataSource/Private/ContentBrowserAssetDataSource.cpp
+++ b/Engine/Plugins/Editor/ContentBrowser/ContentBrowserAssetDataSource/Source/ContentBrowserAssetDataSource/Private/ContentBrowserAssetDataSource.cpp
@@ -50,7 +50,7 @@
 
 namespace UE::Editor::ContentBrowser
 {
-	inline bool IsNewStyleEnabled()
+	static bool IsNewStyleEnabled()
 	{
 		static bool bIsNewStyleEnabled = false;
 		UE_CALL_ONCE([&]()
diff --git a/Engine/Source/Editor/ContentBrowser/Private/ContentBrowserStyle.h b/Engine/Source/Editor/ContentBrowser/Private/ContentBrowserStyle.h
--- a/Engine/Source/Editor/ContentBrowser/Private/ContentBrowserStyle.h
+++ b/Engine/Source/Editor/ContentBrowser/Private/ContentBrowserStyle.h
@@ -8,7 +8,7 @@
 
 namespace UE::Editor::ContentBrowser
 {
-	inline bool IsNewStyleEnabled()
+	static bool IsNewStyleEnabled()
 	{
 		static bool bIsNewStyleEnabled = false;
 		UE_CALL_ONCE([&]()
diff --git a/Engine/Source/Editor/UMGEditor/Private/Library/SLibraryView.cpp b/Engine/Source/Editor/UMGEditor/Private/Library/SLibraryView.cpp
--- a/Engine/Source/Editor/UMGEditor/Private/Library/SLibraryView.cpp
+++ b/Engine/Source/Editor/UMGEditor/Private/Library/SLibraryView.cpp
@@ -37,7 +37,7 @@
 
 namespace UE::Editor::ContentBrowser
 {
-	inline bool IsNewStyleEnabled()
+	static bool IsNewStyleEnabled()
 	{
 		static bool bIsNewStyleEnabled = false;
 		UE_CALL_ONCE([&]()
diff --git a/Engine/Source/Editor/UnrealEd/Private/AssetThumbnail.cpp b/Engine/Source/Editor/UnrealEd/Private/AssetThumbnail.cpp
--- a/Engine/Source/Editor/UnrealEd/Private/AssetThumbnail.cpp
+++ b/Engine/Source/Editor/UnrealEd/Private/AssetThumbnail.cpp
@@ -53,24 +53,8 @@
 #include "Widgets/Images/SLayeredImage.h"
 #include "Widgets/Layout/SSeparator.h"
 #include "Widgets/Text/SRichTextBlock.h"
+#include "ContentBrowserUtils.h"
 
-namespace UE::Editor::ContentBrowser
-{
-	inline bool IsNewStyleEnabled()
-	{
-		static bool bIsNewStyleEnabled = false;
-		UE_CALL_ONCE([&]()
-		{
-			if (const IConsoleVariable* CVar = IConsoleManager::Get().FindConsoleVariable(TEXT("ContentBrowser.EnableNewStyle")))
-			{
-				ensureAlwaysMsgf(!EnumHasAnyFlags(CVar->GetFlags(), ECVF_Default), TEXT("The CVar should have already been set from commandline, @see: UnrealEdGlobals.cpp, UE::Editor::ContentBrowser::EnableContentBrowserNewStyleCVarRegistration."));
-				bIsNewStyleEnabled = CVar->GetBool();
-			}
-		});
-
-		return bIsNewStyleEnabled;
-	}
-}
 namespace AssetThumbnailPool
 {
 	const FObjectThumbnail* LoadThumbnailsFromPackage(const FAssetData& AssetData, FThumbnailMap& OutThumbnailMap)
diff --git a/Engine/Source/Editor/UnrealEd/Private/ContentBrowserUtils.h b/Engine/Source/Editor/UnrealEd/Private/ContentBrowserUtils.h
new file mode 100644
--- /dev/null
+++ b/Engine/Source/Editor/UnrealEd/Private/ContentBrowserUtils.h
@@ -0,0 +1,26 @@
+// Copyright Epic Games, Inc. All Rights Reserved.
+
+#pragma once
+
+#include "HAL/IConsoleManager.h"
+#include "Misc/AssertionMacros.h"
+#include "Misc/CoreMiscDefines.h"
+#include "Misc/EnumClassFlags.h"
+
+namespace UE::Editor::ContentBrowser
+{
+	static bool IsNewStyleEnabled()
+	{
+		static bool bIsNewStyleEnabled = false;
+		UE_CALL_ONCE([&]()
+		{
+			if (const IConsoleVariable* CVar = IConsoleManager::Get().FindConsoleVariable(TEXT("ContentBrowser.EnableNewStyle")))
+			{
+				ensureAlwaysMsgf(!EnumHasAnyFlags(CVar->GetFlags(), ECVF_Default), TEXT("The CVar should have already been set from commandline, @see: UnrealEdGlobals.cpp, UE::Editor::ContentBrowser::EnableContentBrowserNewStyleCVarRegistration."));
+				bIsNewStyleEnabled = CVar->GetBool();
+			}
+		});
+
+		return bIsNewStyleEnabled;
+	}
+}
diff --git a/Engine/Source/Editor/UnrealEd/Private/DragAndDrop/AssetDragDropOp.cpp b/Engine/Source/Editor/UnrealEd/Private/DragAndDrop/AssetDragDropOp.cpp
--- a/Engine/Source/Editor/UnrealEd/Private/DragAndDrop/AssetDragDropOp.cpp
+++ b/Engine/Source/Editor/UnrealEd/Private/DragAndDrop/AssetDragDropOp.cpp
@@ -10,24 +10,7 @@
 #include "ClassIconFinder.h"
 #include "ThumbnailRendering/ThumbnailManager.h"
 #include "Engine/Level.h"
-
-namespace UE::Editor::ContentBrowser
-{
-	inline bool IsNewStyleEnabled()
-	{
-		static bool bIsNewStyleEnabled = false;
-		UE_CALL_ONCE([&]()
-		{
-			if (const IConsoleVariable* CVar = IConsoleManager::Get().FindConsoleVariable(TEXT("ContentBrowser.EnableNewStyle")))
-			{
-				ensureAlwaysMsgf(!EnumHasAnyFlags(CVar->GetFlags(), ECVF_Default), TEXT("The CVar should have already been set from commandline, @see: UnrealEdGlobals.cpp, UE::Editor::ContentBrowser::EnableContentBrowserNewStyleCVarRegistration."));
-				bIsNewStyleEnabled = CVar->GetBool();
-			}
-		});
-
-		return bIsNewStyleEnabled;
-	}
-}
+#include "ContentBrowserUtils.h"
 
 TSharedRef<FAssetDragDropOp> FAssetDragDropOp::New(const FAssetData& InAssetData, TScriptInterface<IAssetFactoryInterface> AssetFactory)
 {
