Avoid trying to respawn Actors that belong to cleaned up world when replacing instances due to blueprint changes.

diff --git a/Engine/Source/Editor/UnrealEd/Private/Kismet2/KismetReinstanceUtilities.cpp b/Engine/Source/Editor/UnrealEd/Private/Kismet2/KismetReinstanceUtilities.cpp
--- a/Engine/Source/Editor/UnrealEd/Private/Kismet2/KismetReinstanceUtilities.cpp
+++ b/Engine/Source/Editor/UnrealEd/Private/Kismet2/KismetReinstanceUtilities.cpp
@@ -2298,7 +2298,7 @@ static void ReplaceObjectHelper(UObject*& OldObject, UClass* OldClass, UObject*&
 		Options.bDoDelta = false;
 	}
 	// We only need to copy properties of the pre-created instances, the rest of the default sub object is done inside the UEditorEngine::CopyPropertiesForUnrelatedObjects
-	for (const auto& Pair : OrderedListOfObjectToCopy)
+	for (const auto& Pair : OrderedListOfObjectToCopy)	
 	{
 		UEditorEngine::CopyPropertiesForUnrelatedObjects(Pair.Key, Pair.Value, Options);
 	}
@@ -2855,7 +2855,7 @@ void FBlueprintCompileReinstancer::ReplaceInstancesOfClass_Inner(const TMap<UCla
 						CheckAndSaveOuterPackageToCleanList(OldActor);
 
 						UObject* NewUObject = nullptr;
-						if (OldActor->GetLevel() && OldActor->GetWorld())
+						if (OldActor->GetLevel() && OldActor->GetWorld() && !OldActor->GetWorld()->IsCleanedUp())
 						{
 							// Attached actors will be marked dirty as well when the old instance is destroyed, but
 							// since these attachments will get restored as part of replacement, there's no need to
