World Partition - Landscape Spline Mesh Builder: Fixed typo (from CL 36568205) preventing landscape splines from being partitioned

diff --git a/Engine/Source/Editor/UnrealEd/Private/WorldPartition/WorldPartitionLandscapeSplineMeshesBuilder.cpp b/Engine/Source/Editor/UnrealEd/Private/WorldPartition/WorldPartitionLandscapeSplineMeshesBuilder.cpp
--- a/Engine/Source/Editor/UnrealEd/Private/WorldPartition/WorldPartitionLandscapeSplineMeshesBuilder.cpp
+++ b/Engine/Source/Editor/UnrealEd/Private/WorldPartition/WorldPartitionLandscapeSplineMeshesBuilder.cpp
@@ -399,7 +399,7 @@ bool UWorldPartitionLandscapeSplineMeshesBuilder::RunInternal(UWorld* InWorld, c
 			}
 
 			// Only partition splines if they are spatially loaded & not editor only
-			if (SplineActor->ShouldPartitionSpline())
+			if (!SplineActor->ShouldPartitionSpline())
 			{
 				continue;
 			}
