MPCook: CookPackageSplitter: Fix spurious error reporting of "We are demoting a %s package from save and removing our references that keep its objects loaded". when GC occurs on director after all packages from a CookpackageSplitter were retracted to CookWorkers and one of them had reached the save state before being retracted.
#rnx

diff --git a/Engine/Source/Editor/UnrealEd/Private/Cooker/CookGenerationHelper.cpp b/Engine/Source/Editor/UnrealEd/Private/Cooker/CookGenerationHelper.cpp
--- a/Engine/Source/Editor/UnrealEd/Private/Cooker/CookGenerationHelper.cpp
+++ b/Engine/Source/Editor/UnrealEd/Private/Cooker/CookGenerationHelper.cpp
@@ -1413,10 +1413,22 @@ void FGenerationHelper::ResetSaveState(FCookGenerationInfo& Info, UPackage* Pack
 
 	if (Info.HasTakenOverCachedCookedPlatformData())
 	{
-		if (NewState != EPackageState::Idle &&
-			Info.PackageData->GetCachedObjectsInOuter().Num() != 0 && IsUseInternalReferenceToAvoidGarbageCollect() &&
-			(!IsTerminalStateChange(ReleaseSaveReason) && ReleaseSaveReason != EStateChangeReason::DoneForNow
-				&& ReleaseSaveReason != EStateChangeReason::Retraction))
+		// Demotion possibly causes a failure of our guarantees to CookPackageSplitter authors that we will keep the
+		// UObjects in the packages referenced. But there are also some cases where it does not cause a problem.
+		bool bDemotionAllowed = Info.PackageData->GetCachedObjectsInOuter().Num() == 0
+			|| !IsUseInternalReferenceToAvoidGarbageCollect();
+		// If the package is no longer cooking, the demotion is okay, because we will not be relying on the splitter
+		// for any further work on it. The package is "no longer cooking" if it is moving into the idle or the
+		// AssignedToWorker state. In the AssignedToWorker state we might get it again and need to cook it if it is
+		// retracted back to us, but in the case of moving to AssignedToWorker state, we expect that is only happening
+		// because this GenerationHelper is moving into the Uninitializ, and we will redo all of the splitter's
+		// package loading later when we retract it and reinitialize this GenerationHelper.
+		bDemotionAllowed = bDemotionAllowed ||
+			NewState == EPackageState::Idle || NewState == EPackageState::AssignedToWorker;
+		// Demotion is also allowed in cases of being kicked out of the save state and not expecting to come back.
+		bDemotionAllowed = bDemotionAllowed || IsTerminalStateChange(ReleaseSaveReason)
+			|| ReleaseSaveReason == EStateChangeReason::DoneForNow || ReleaseSaveReason == EStateChangeReason::Retraction;
+		if (!bDemotionAllowed)
 		{
 			UE_LOG(LogCook, Error,
 				TEXT("CookPackageSplitter failure: We are demoting a %s package from save and removing our references that keep its objects loaded.\n")
