Thumbnails for cooked editor can now be loaded from an uncooked editor's directory
#rnx

diff --git a/Engine/Source/Editor/UnrealEd/Private/ObjectTools.cpp b/Engine/Source/Editor/UnrealEd/Private/ObjectTools.cpp
--- a/Engine/Source/Editor/UnrealEd/Private/ObjectTools.cpp
+++ b/Engine/Source/Editor/UnrealEd/Private/ObjectTools.cpp
@@ -5562,10 +5562,10 @@ namespace ThumbnailTools
 
 
 	/** Loads thumbnails from the specified package file name */
-	bool LoadThumbnailsFromPackageInternal( const FString& InPackageFileName, const TSet< FName >& InObjectFullNames, FThumbnailMap& InOutThumbnails )
+	bool LoadThumbnailsFromPackageDirectly(const FString& InPackageFileName, const TSet<FName>& InObjectFullNames, FThumbnailMap& InOutThumbnails, const FString& InPackageFileNameToOpen)
 	{
 		// Create a file reader to load the file
-		TUniquePtr< FArchive > FileReader( IFileManager::Get().CreateFileReader( *InPackageFileName ) );
+		TUniquePtr<FArchive> FileReader(IFileManager::Get().CreateFileReader(*InPackageFileNameToOpen));
 		if( FileReader == nullptr )
 		{
 			// Couldn't open the file
@@ -5762,12 +5762,30 @@ namespace ThumbnailTools
 	}
 
 
+	TArray<FLoadThumbnailsFromPackage>& GetLoadThumbnailsFromPackageDelegate()
+	{
+		static TArray<FLoadThumbnailsFromPackage> LoadThumbnailsFromPackageDelegates;
+		return LoadThumbnailsFromPackageDelegates;
+	}
+
+	FDelegateHandle AddLoadThumbnailsFromPackageDelegate(FLoadThumbnailsFromPackage InDelegate)
+	{
+		return GetLoadThumbnailsFromPackageDelegate().Add_GetRef(InDelegate).GetHandle();
+	}
+
+	void RemoveLoadThumbnailsFromPackageDelegate(const FDelegateHandle InHandle)
+	{
+		GetLoadThumbnailsFromPackageDelegate().RemoveAll([InHandle](const FLoadThumbnailsFromPackage& Entry)
+		{
+			return Entry.GetHandle() == InHandle || !Entry.IsBound();
+		});
+	}
 
 	/** Loads thumbnails from the specified package file name, try loading from external cache file if not found in package file */
 	bool LoadThumbnailsFromPackage( const FString& InPackageFileName, const TSet< FName >& InObjectFullNames, FThumbnailMap& InOutThumbnails )
 	{
 		TRACE_CPUPROFILER_EVENT_SCOPE_STR("ThumbnailTools::LoadThumbnailsFromPackage");
-		if (LoadThumbnailsFromPackageInternal(InPackageFileName, InObjectFullNames, InOutThumbnails))
+		if (LoadThumbnailsFromPackageDirectly(InPackageFileName, InObjectFullNames, InOutThumbnails, InPackageFileName))
 		{
 			return true;
 		}
@@ -5775,6 +5793,30 @@ namespace ThumbnailTools
 		{
 			return true;
 		}
+		else
+		{
+			// Delegates
+			bool bDelegatesNeedCleanup = false;
+			for (const FLoadThumbnailsFromPackage& DelegateEntry : GetLoadThumbnailsFromPackageDelegate())
+			{
+				if (DelegateEntry.IsBound())
+				{
+					if (DelegateEntry.Execute(InPackageFileName, InObjectFullNames, InOutThumbnails))
+					{
+						return true;
+					}
+				}
+				else
+				{
+					bDelegatesNeedCleanup = true;
+				}
+			}
+
+			if (bDelegatesNeedCleanup)
+			{
+				RemoveLoadThumbnailsFromPackageDelegate(FDelegateHandle());
+			}
+		}
 
 		return false;
 	}
diff --git a/Engine/Source/Editor/UnrealEd/Public/ObjectTools.h b/Engine/Source/Editor/UnrealEd/Public/ObjectTools.h
--- a/Engine/Source/Editor/UnrealEd/Public/ObjectTools.h
+++ b/Engine/Source/Editor/UnrealEd/Public/ObjectTools.h
@@ -21,6 +21,8 @@ class UExporter;
 class UFactory;
 class USoundWave;
 
+DECLARE_DELEGATE_RetVal_ThreeParams(bool, FLoadThumbnailsFromPackage, const FString& /*InPackageFileName*/, const TSet<FName>& /*InObjectFullNames*/, FThumbnailMap& /*InOutThumbnails*/);
+
 namespace ObjectTools
 {
 	/** A simple struct to represent the package group name triplet */
@@ -759,6 +761,9 @@ namespace ThumbnailTools
 	/** Loads thumbnails from the specified package file name (or from the external thumbnail cache file if it exists) */
 	UNREALED_API bool LoadThumbnailsFromPackage( const FString& InPackageFileName, const TSet<FName>& InObjectFullNames, FThumbnailMap& InOutThumbnails );
 
+	/** Loads thumbnails from the specified package file name (without using external thumbnail cache or delegates, only attempts to read from InPackageFileNameToOpen provided) */
+	UNREALED_API bool LoadThumbnailsFromPackageDirectly( const FString& InPackageFileName, const TSet<FName>& InObjectFullNames, FThumbnailMap& InOutThumbnails, const FString& InPackageFileNameToOpen);
+
 	/** Loads thumbnails from a package unless they're already cached in that package's thumbnail map */
 	UNREALED_API bool ConditionallyLoadThumbnailsFromPackage( const FString& InPackageFileName, const TSet< FName >& InObjectFullNames, FThumbnailMap& InOutThumbnails );
 
@@ -773,4 +778,13 @@ namespace ThumbnailTools
 	UNREALED_API bool AssetHasCustomThumbnail(const FString& InAssetDataFullName, FObjectThumbnail& OutThumbnail);
 	/** Returns true if the given asset has a custom thumbnail cached or on the disk and if the thumbnail was captured from a viewport. */
 	UNREALED_API bool AssetHasCustomCreatedThumbnail(const FString& InAssetDataFullName);
+
+	/** Add delegate to handle loading of thumbnail from package. */
+	UNREALED_API FDelegateHandle AddLoadThumbnailsFromPackageDelegate(FLoadThumbnailsFromPackage InDelegate);
+
+	/** Remove delegate to handle loading of thumbnail from package. */
+	UNREALED_API void RemoveLoadThumbnailsFromPackageDelegate(const FDelegateHandle InHandle);
+
+	/** Loads thumbnails from the specified package file name */
+	UNREALED_API bool LoadThumbnailsFromPackageUsingExternalFolders(const FString& InPackageFileName, const TSet<FName>& InObjectFullNames, FThumbnailMap& InOutThumbnails, TConstArrayView<FString> ExternalFolders);
 }
