WorldPartitionResaveActorsBuilder: added limited support for resaving non-partitioned maps using one file per actor.
Also fixed bad condition when retrieving actor descriptors bounds during resave on unregistered actors.

diff --git a/Engine/Source/Editor/UnrealEd/Private/WorldPartition/WorldPartitionResaveActorsBuilder.cpp b/Engine/Source/Editor/UnrealEd/Private/WorldPartition/WorldPartitionResaveActorsBuilder.cpp
--- a/Engine/Source/Editor/UnrealEd/Private/WorldPartition/WorldPartitionResaveActorsBuilder.cpp
+++ b/Engine/Source/Editor/UnrealEd/Private/WorldPartition/WorldPartitionResaveActorsBuilder.cpp
@@ -5,6 +5,7 @@
 #include "WorldPartition/WorldPartitionHelpers.h"
 #include "WorldPartition/WorldPartitionClassDescRegistry.h"
 #include "WorldPartition/WorldPartitionActorDescInstance.h"
+#include "WorldPartition/ActorDescContainerSubsystem.h"
 #include "PackageSourceControlHelper.h"
 #include "SourceControlHelpers.h"
 #include "Engine/Level.h"
@@ -199,287 +200,346 @@ bool UWorldPartitionResaveActorsBuilder::RunInternal(UWorld* World, const FCellI
 	if (ActorClasses.IsEmpty())
 	{
 		return false;
-	}	
+	}
+
+	TArray<FString> PackagesToDelete;
+	TArray<FString> DirtyActorDescsOld;
+	TArray<FString> DirtyActorDescsNew;
 
 	UWorldPartition* WorldPartition = World->GetWorldPartition();
 	if (!WorldPartition)
 	{
-		UE_LOG(LogWorldPartitionResaveActorsBuilder, Error, TEXT("Failed to retrieve WorldPartition."));
-		return false;
-	}
-
-	if (bSwitchActorPackagingSchemeToReduced)
-	{
-		if (World->PersistentLevel->GetActorPackagingScheme() == EActorPackagingScheme::Reduced)
+		if (!World->PersistentLevel->bUseExternalActors)
 		{
-			UE_LOG(LogWorldPartitionResaveActorsBuilder, Error, TEXT("World is already using the reduced actor packaging scheme."));
+			UE_LOG(LogWorldPartitionResaveActorsBuilder, Error, TEXT("Cannot work on maps not using one file per actor."));
 			return false;
 		}
 
-		World->PersistentLevel->ActorPackagingScheme = EActorPackagingScheme::Reduced;
-	}
+		if (bSwitchActorPackagingSchemeToReduced)
+		{
+			UE_LOG(LogWorldPartitionResaveActorsBuilder, Error, TEXT("Cannot switch actor packaging scheme on non-partitioned worlds."));
+			return false;
+		}
 
-	if (bEnableActorFolders)
-	{
-		if (World->PersistentLevel->IsUsingActorFolders())
+		if (bEnableActorFolders)
 		{
-			UE_LOG(LogWorldPartitionResaveActorsBuilder, Error, TEXT("World is already using actor folder objects."));
+			UE_LOG(LogWorldPartitionResaveActorsBuilder, Error, TEXT("Cannot enable actor folders non-partitioned worlds."));
 			return false;
 		}
 
-		World->PersistentLevel->SetUseActorFolders(true);
-		World->PersistentLevel->bFixupActorFoldersAtLoad = false;
-	}
+		if (UActorDescContainer* LevelContainer = UActorDescContainerSubsystem::GetChecked().RegisterContainer(UActorDescContainer::FInitializeParams{ World->GetPackage()->GetLoadedPath().GetPackageFName() }))
+		{
+			TArray<UPackage*> PackagesToSave;
+			for (FActorDescList::TIterator<> ActorDescIterator(LevelContainer); ActorDescIterator; ++ActorDescIterator)
+			{
+				if (AActor* Actor = FindObject<AActor>(World->PersistentLevel, *ActorDescIterator->GetActorName().ToString()))
+				{
+					UPackage* Package = Actor->GetExternalPackage();
+					check(Package);
 
-	TArray<FString> PackagesToDelete;
+					if (bResaveDirtyActorDescsOnly)
+					{
+						TUniquePtr<FWorldPartitionActorDesc> NewActorDesc = Actor->CreateActorDesc();
 
-	TArray<FString> DirtyActorDescsOld;
-	TArray<FString> DirtyActorDescsNew;
+						if (!ActorDescIterator->IsResaveNeeded() && !ActorDescIterator->ShouldResave(NewActorDesc.Get()))
+						{
+							continue;
+						}
 
-	if (bSwitchActorPackagingSchemeToReduced)
-	{
-		World->PersistentLevel->ActorPackagingScheme = EActorPackagingScheme::Reduced;
+						if (bDiffDirtyActorDescs)
+						{
+							DirtyActorDescsOld.Add(ActorDescIterator->ToString(FWorldPartitionActorDesc::EToStringMode::ForDiff));
+							DirtyActorDescsNew.Add(NewActorDesc->ToString(FWorldPartitionActorDesc::EToStringMode::ForDiff));
+						}
 
-		// Build clusters
-		TArray<TPair<FGuid, TArray<FGuid>>> ActorsWithRefs;
-		for (FActorDescContainerInstanceCollection::TIterator<> Iterator(WorldPartition); Iterator; ++Iterator)
-		{
-			ActorsWithRefs.Emplace(Iterator->GetGuid(), Iterator->GetReferences());
+						UE_LOG(LogWorldPartitionResaveActorsBuilder, Log, TEXT("Package %s needs to be resaved."), *Package->GetName());
+					}
+
+					if (!bReportOnly)
+					{
+						PackagesToSave.Add(Package);
+					}
+				}
+			}
+
+			UActorDescContainerSubsystem::GetChecked().UnregisterContainer(LevelContainer);
+
+			if (!UWorldPartitionBuilder::SavePackages(PackagesToSave, PackageHelper))
+			{
+				return false;
+			}
 		}
+	}
+	else
+	{
+		if (bSwitchActorPackagingSchemeToReduced)
+		{
+			if (World->PersistentLevel->GetActorPackagingScheme() == EActorPackagingScheme::Reduced)
+			{
+				UE_LOG(LogWorldPartitionResaveActorsBuilder, Error, TEXT("World is already using the reduced actor packaging scheme."));
+				return false;
+			}
 
-		TArray<TArray<FGuid>> Clusters = GenerateObjectsClusters(ActorsWithRefs);
+			World->PersistentLevel->ActorPackagingScheme = EActorPackagingScheme::Reduced;
+		}
 
-		TSet<FWorldPartitionReference> LoadedReferences;
-		TArray<UPackage*> PackagesToSave;
-		for (const TArray<FGuid>& Cluster : Clusters)
+		if (bEnableActorFolders)
 		{
-			// Load actor clusters
-			TArray<FWorldPartitionReference> ActorReferences;
-			ActorReferences.Reserve(Cluster.Num());
-			for (const FGuid& ActorGuid : Cluster)
+			if (World->PersistentLevel->IsUsingActorFolders())
 			{
-				ActorReferences.Emplace(WorldPartition, ActorGuid);
+				UE_LOG(LogWorldPartitionResaveActorsBuilder, Error, TEXT("World is already using actor folder objects."));
+				return false;
 			}
-			LoadedReferences.Append(ActorReferences);
 
-			// Change packaging of all actors in the current cluster
-			for (FWorldPartitionReference& ActorReference : ActorReferences)
+			World->PersistentLevel->SetUseActorFolders(true);
+			World->PersistentLevel->bFixupActorFoldersAtLoad = false;
+		}
+
+		if (bSwitchActorPackagingSchemeToReduced)
+		{
+			World->PersistentLevel->ActorPackagingScheme = EActorPackagingScheme::Reduced;
+
+			// Build clusters
+			TArray<TPair<FGuid, TArray<FGuid>>> ActorsWithRefs;
+			for (FActorDescContainerInstanceCollection::TIterator<> Iterator(WorldPartition); Iterator; ++Iterator)
 			{
-				const FWorldPartitionActorDescInstance* ActorDescInstance = *ActorReference;
-				AActor* Actor = ActorReference.GetActor();
+				ActorsWithRefs.Emplace(Iterator->GetGuid(), Iterator->GetReferences());
+			}
 
-				if (!Actor)
+			TArray<TArray<FGuid>> Clusters = GenerateObjectsClusters(ActorsWithRefs);
+
+			TSet<FWorldPartitionReference> LoadedReferences;
+			TArray<UPackage*> PackagesToSave;
+			for (const TArray<FGuid>& Cluster : Clusters)
+			{
+				// Load actor clusters
+				TArray<FWorldPartitionReference> ActorReferences;
+				ActorReferences.Reserve(Cluster.Num());
+				for (const FGuid& ActorGuid : Cluster)
 				{
-					PackagesToDelete.Add(ActorDescInstance->GetActorPackage().ToString());
-					continue;
+					ActorReferences.Emplace(WorldPartition, ActorGuid);
 				}
-				else
+				LoadedReferences.Append(ActorReferences);
+
+				// Change packaging of all actors in the current cluster
+				for (FWorldPartitionReference& ActorReference : ActorReferences)
 				{
-					UPackage* Package = Actor->GetExternalPackage();
-					check(Package);
+					const FWorldPartitionActorDescInstance* ActorDescInstance = *ActorReference;
+					AActor* Actor = ActorReference.GetActor();
 
-					if (!bReportOnly)
+					if (!Actor)
 					{
-						// Always mark this package for deletion, as it will contain a temporary redirector to fixup references
 						PackagesToDelete.Add(ActorDescInstance->GetActorPackage().ToString());
+						continue;
+					}
+					else
+					{
+						UPackage* Package = Actor->GetExternalPackage();
+						check(Package);
 
-						// Move actor back into world's package
-						Actor->SetPackageExternal(false);
-
-						// Gather dependant objects that also needs to be moved along with the actor
-						TArray<UObject*> DependantObjects;
-						ForEachObjectWithPackage(Package, [&DependantObjects](UObject* Object)
+						if (!bReportOnly)
 						{
-							if (!Cast<UMetaData>(Object))
+							// Always mark this package for deletion, as it will contain a temporary redirector to fixup references
+							PackagesToDelete.Add(ActorDescInstance->GetActorPackage().ToString());
+
+							// Move actor back into world's package
+							Actor->SetPackageExternal(false);
+
+							// Gather dependant objects that also needs to be moved along with the actor
+							TArray<UObject*> DependantObjects;
+							ForEachObjectWithPackage(Package, [&DependantObjects](UObject* Object)
+							{
+								if (!Cast<UMetaData>(Object))
+								{
+									DependantObjects.Add(Object);
+								}
+								return true;
+							}, false);
+
+							// Move dependant objects into the new world package temporarily
+							for (UObject* DependantObject : DependantObjects)
 							{
-								DependantObjects.Add(Object);
+								DependantObject->Rename(nullptr, WorldPackage, REN_NonTransactional | REN_DontCreateRedirectors);
 							}
-							return true;
-						}, false);
 
-						// Move dependant objects into the new world package temporarily
-						for (UObject* DependantObject : DependantObjects)
-						{
-							DependantObject->Rename(nullptr, WorldPackage, REN_NonTransactional | REN_DontCreateRedirectors);
-						}
+							// Move actor in its new package
+							Actor->SetPackageExternal(true);
 
-						// Move actor in its new package
-						Actor->SetPackageExternal(true);
+							// Also move dependant objects into the new package
+							UPackage* NewActorPackage = Actor->GetExternalPackage();
+							for (UObject* DependantObject : DependantObjects)
+							{
+								DependantObject->Rename(nullptr, NewActorPackage, REN_NonTransactional | REN_DontCreateRedirectors);
+							}
 
-						// Also move dependant objects into the new package
-						UPackage* NewActorPackage = Actor->GetExternalPackage();
-						for (UObject* DependantObject : DependantObjects)
-						{
-							DependantObject->Rename(nullptr, NewActorPackage, REN_NonTransactional | REN_DontCreateRedirectors);
+							PackagesToSave.Add(NewActorPackage);
 						}
-
-						PackagesToSave.Add(NewActorPackage);
 					}
 				}
-			}
 
-			if (FWorldPartitionHelpers::ShouldCollectGarbage())
-			{
-				if (!UWorldPartitionBuilder::SavePackages(PackagesToSave, PackageHelper))
+				if (FWorldPartitionHelpers::ShouldCollectGarbage())
 				{
-					return false;
-				}
+					if (!UWorldPartitionBuilder::SavePackages(PackagesToSave, PackageHelper))
+					{
+						return false;
+					}
 
-				PackagesToSave.Empty();
-				LoadedReferences.Empty();
-				FWorldPartitionHelpers::DoCollectGarbage();
+					PackagesToSave.Empty();
+					LoadedReferences.Empty();
+					FWorldPartitionHelpers::DoCollectGarbage();
+				}
 			}
-		}
 
-		// Save last batch
-		if (!UWorldPartitionBuilder::SavePackages(PackagesToSave, PackageHelper))
-		{
-			return false;
+			// Save last batch
+			if (!UWorldPartitionBuilder::SavePackages(PackagesToSave, PackageHelper))
+			{
+				return false;
+			}
+			PackagesToSave.Empty();
+			LoadedReferences.Empty();
+			FWorldPartitionHelpers::DoCollectGarbage();
 		}
-		PackagesToSave.Empty();
-		LoadedReferences.Empty();
-		FWorldPartitionHelpers::DoCollectGarbage();
-	}
-	else
-	{
-		if (bResaveBlueprints)
+		else
 		{
-			// Resave blueprints top-down from the class hierarchy
-			FWorldPartitionClassDescRegistry& ClassDescRegistry = FWorldPartitionClassDescRegistry::Get();
-
-			TArray<FTopLevelAssetPath> ClassesToResave;
-			TMap<FTopLevelAssetPath, TArray<FTopLevelAssetPath>> ClassHierarchy;
-			for (auto& [ClassPath, ParentClassPath] : ClassDescRegistry.GetParentClassMap())
+			if (bResaveBlueprints)
 			{
-				ClassHierarchy.FindOrAdd(ParentClassPath).Add(ClassPath);
+				// Resave blueprints top-down from the class hierarchy
+				FWorldPartitionClassDescRegistry& ClassDescRegistry = FWorldPartitionClassDescRegistry::Get();
 
-				if (!ClassPath.ToString().StartsWith(TEXT("/Script/")) && ParentClassPath.ToString().StartsWith(TEXT("/Script/")))
+				TArray<FTopLevelAssetPath> ClassesToResave;
+				TMap<FTopLevelAssetPath, TArray<FTopLevelAssetPath>> ClassHierarchy;
+				for (auto& [ClassPath, ParentClassPath] : ClassDescRegistry.GetParentClassMap())
 				{
-					ClassesToResave.Add(ClassPath);
+					ClassHierarchy.FindOrAdd(ParentClassPath).Add(ClassPath);
+
+					if (!ClassPath.ToString().StartsWith(TEXT("/Script/")) && ParentClassPath.ToString().StartsWith(TEXT("/Script/")))
+					{
+						ClassesToResave.Add(ClassPath);
+					}
 				}
-			}
 
-			TArray<UPackage*> PackagesToSave;
-			while (!ClassesToResave.IsEmpty())
-			{
-				for (const FTopLevelAssetPath& ClassToResave : ClassesToResave)
+				TArray<UPackage*> PackagesToSave;
+				while (!ClassesToResave.IsEmpty())
 				{
-					UPackage* ClassPackage = LoadPackage(nullptr, *ClassToResave.GetPackageName().ToString(), LOAD_None);
-
-					if (!ClassPackage)
+					for (const FTopLevelAssetPath& ClassToResave : ClassesToResave)
 					{
-						UE_LOG(LogWorldPartitionResaveActorsBuilder, Error, TEXT("Failed to load package for class '%s'"), *ClassToResave.ToString());
-						return false;
+						UPackage* ClassPackage = LoadPackage(nullptr, *ClassToResave.GetPackageName().ToString(), LOAD_None);
+
+						if (!ClassPackage)
+						{
+							UE_LOG(LogWorldPartitionResaveActorsBuilder, Error, TEXT("Failed to load package for class '%s'"), *ClassToResave.ToString());
+							return false;
+						}
+
+						PackagesToSave.Add(ClassPackage);
+
+						if (FWorldPartitionHelpers::HasExceededMaxMemory())
+						{
+							UWorldPartitionBuilder::SavePackages(PackagesToSave, PackageHelper, true);
+							FWorldPartitionHelpers::DoCollectGarbage();
+							PackagesToSave.Empty();
+						}
 					}
 
-					PackagesToSave.Add(ClassPackage);
+					const TArray<FTopLevelAssetPath> ParentClasses = MoveTemp(ClassesToResave);
 
-					if (FWorldPartitionHelpers::HasExceededMaxMemory())
+					for (const FTopLevelAssetPath& ParentClass : ParentClasses)
 					{
-						UWorldPartitionBuilder::SavePackages(PackagesToSave, PackageHelper, true);
-						FWorldPartitionHelpers::DoCollectGarbage();
-						PackagesToSave.Empty();
+						ClassesToResave += ClassHierarchy.FindRef(ParentClass);
 					}
 				}
 
-				const TArray<FTopLevelAssetPath> ParentClasses = MoveTemp(ClassesToResave);
-
-				for (const FTopLevelAssetPath& ParentClass : ParentClasses)
-				{
-					ClassesToResave += ClassHierarchy.FindRef(ParentClass);
-				}
+				UWorldPartitionBuilder::SavePackages(PackagesToSave, PackageHelper);
+				FWorldPartitionHelpers::DoCollectGarbage();
 			}
 
-			UWorldPartitionBuilder::SavePackages(PackagesToSave, PackageHelper);
-			FWorldPartitionHelpers::DoCollectGarbage();
-		}
-
-		TArray<UPackage*> PackagesToSave;
+			TArray<UPackage*> PackagesToSave;
 
-		FWorldPartitionHelpers::FForEachActorWithLoadingParams ForEachActorWithLoadingParams;
+			FWorldPartitionHelpers::FForEachActorWithLoadingParams ForEachActorWithLoadingParams;
 
-		ForEachActorWithLoadingParams.ActorClasses = ActorClasses;
+			ForEachActorWithLoadingParams.ActorClasses = ActorClasses;
 
-		ForEachActorWithLoadingParams.FilterActorDesc = [this](const FWorldPartitionActorDesc* ActorDesc) -> bool
-		{
-			for (const FName& ActorTag : ActorDesc->GetTags())
+			ForEachActorWithLoadingParams.FilterActorDesc = [this](const FWorldPartitionActorDesc* ActorDesc) -> bool
 			{
-				if (ActorTags.Contains(ActorTag))
+				for (const FName& ActorTag : ActorDesc->GetTags())
 				{
-					return true;
+					if (ActorTags.Contains(ActorTag))
+					{
+						return true;
+					}
 				}
-			}
 			
-			for (auto [PropertyName, PropertyValue] : ActorProperties)
-			{
-				FName Value;
-				if (ActorDesc->GetProperty(PropertyName, &Value))
+				for (auto [PropertyName, PropertyValue] : ActorProperties)
 				{
-					if (Value == PropertyValue)
+					FName Value;
+					if (ActorDesc->GetProperty(PropertyName, &Value))
 					{
-						return true;
+						if (Value == PropertyValue)
+						{
+							return true;
+						}
 					}
 				}
-			}
 
-			return ActorTags.IsEmpty() && ActorProperties.IsEmpty();
-		};
-
-		ForEachActorWithLoadingParams.OnPreGarbageCollect = [&PackagesToSave, &PackageHelper]()
-		{
-			UWorldPartitionBuilder::SavePackages(PackagesToSave, PackageHelper, true);
-			PackagesToSave.Empty();
-		};
+				return ActorTags.IsEmpty() && ActorProperties.IsEmpty();
+			};
 
-		FWorldPartitionHelpers::ForEachActorWithLoading(WorldPartition, [this, &PackagesToDelete, &PackagesToSave, &DirtyActorDescsOld, &DirtyActorDescsNew](const FWorldPartitionActorDescInstance* ActorDescInstance)
-		{
-			AActor* Actor = ActorDescInstance->GetActor();
-			
-			if (!Actor)
+			ForEachActorWithLoadingParams.OnPreGarbageCollect = [&PackagesToSave, &PackageHelper]()
 			{
-				PackagesToDelete.Add(ActorDescInstance->GetActorPackage().ToString());
-			}
-			else
+				UWorldPartitionBuilder::SavePackages(PackagesToSave, PackageHelper, true);
+				PackagesToSave.Empty();
+			};
+
+			FWorldPartitionHelpers::ForEachActorWithLoading(WorldPartition, [this, &PackagesToDelete, &PackagesToSave, &DirtyActorDescsOld, &DirtyActorDescsNew](const FWorldPartitionActorDescInstance* ActorDescInstance)
 			{
-				if (bEnableActorFolders)
+				AActor* Actor = ActorDescInstance->GetActor();
+			
+				if (!Actor)
 				{
-					if (!Actor->CreateOrUpdateActorFolder())
+					PackagesToDelete.Add(ActorDescInstance->GetActorPackage().ToString());
+				}
+				else
+				{
+					if (bEnableActorFolders)
 					{
-						UE_LOG(LogWorldPartitionResaveActorsBuilder, Error, TEXT("Failed to create actor folder for actor %s."), *Actor->GetName());
-						return true;
+						if (!Actor->CreateOrUpdateActorFolder())
+						{
+							UE_LOG(LogWorldPartitionResaveActorsBuilder, Error, TEXT("Failed to create actor folder for actor %s."), *Actor->GetName());
+							return true;
+						}
 					}
-				}
-
-				UPackage* Package = Actor->GetExternalPackage();
-				check(Package);
 
-				if (bResaveDirtyActorDescsOnly)
-				{
-					const FWorldPartitionActorDesc* ActorDesc = ActorDescInstance->GetActorDesc();
-					TUniquePtr<FWorldPartitionActorDesc> NewActorDesc = Actor->CreateActorDesc();
+					UPackage* Package = Actor->GetExternalPackage();
+					check(Package);
 
-					if (!ActorDesc->IsResaveNeeded() && !ActorDesc->ShouldResave(NewActorDesc.Get()))
+					if (bResaveDirtyActorDescsOnly)
 					{
-						return true;
+						const FWorldPartitionActorDesc* ActorDesc = ActorDescInstance->GetActorDesc();
+						TUniquePtr<FWorldPartitionActorDesc> NewActorDesc = Actor->CreateActorDesc();
+
+						if (!ActorDesc->IsResaveNeeded() && !ActorDesc->ShouldResave(NewActorDesc.Get()))
+						{
+							return true;
+						}
+
+						if (bDiffDirtyActorDescs)
+						{
+							DirtyActorDescsOld.Add(ActorDesc->ToString(FWorldPartitionActorDesc::EToStringMode::ForDiff));
+							DirtyActorDescsNew.Add(NewActorDesc->ToString(FWorldPartitionActorDesc::EToStringMode::ForDiff));
+						}
+
+						UE_LOG(LogWorldPartitionResaveActorsBuilder, Log, TEXT("Package %s needs to be resaved."), *Package->GetName());
 					}
 
-					if (bDiffDirtyActorDescs)
+					if (!bReportOnly)
 					{
-						DirtyActorDescsOld.Add(ActorDesc->ToString(FWorldPartitionActorDesc::EToStringMode::ForDiff));
-						DirtyActorDescsNew.Add(NewActorDesc->ToString(FWorldPartitionActorDesc::EToStringMode::ForDiff));
+						PackagesToSave.Add(Package);
 					}
-
-					UE_LOG(LogWorldPartitionResaveActorsBuilder, Log, TEXT("Package %s needs to be resaved."), *Package->GetName());
 				}
 
-				if (!bReportOnly)
-				{
-					PackagesToSave.Add(Package);
-				}
-			}
-
-			return true;
-		}, ForEachActorWithLoadingParams);
+				return true;
+			}, ForEachActorWithLoadingParams);
+		}
 	}
 
 	if (bDiffDirtyActorDescs && DirtyActorDescsOld.Num())
diff --git a/Engine/Source/Editor/UnrealEd/Public/WorldPartition/WorldPartitionResaveActorsBuilder.h b/Engine/Source/Editor/UnrealEd/Public/WorldPartition/WorldPartitionResaveActorsBuilder.h
--- a/Engine/Source/Editor/UnrealEd/Public/WorldPartition/WorldPartitionResaveActorsBuilder.h
+++ b/Engine/Source/Editor/UnrealEd/Public/WorldPartition/WorldPartitionResaveActorsBuilder.h
@@ -18,6 +18,7 @@ public:
 	virtual ELoadingMode GetLoadingMode() const override { return ELoadingMode::Custom; }
 	virtual bool PreRun(UWorld* World, FPackageSourceControlHelper& PackageHelper) override;
 	virtual bool RunInternal(UWorld* World, const FCellInfo& InCellInfo, FPackageSourceControlHelper& PackageHelper) override;
+	virtual bool CanProcessNonPartitionedWorlds() const { return true; }
 	// UWorldPartitionBuilder interface end
 
 private:
diff --git a/Engine/Source/Runtime/Engine/Private/WorldPartition/WorldPartitionActorDescUtils.cpp b/Engine/Source/Runtime/Engine/Private/WorldPartition/WorldPartitionActorDescUtils.cpp
--- a/Engine/Source/Runtime/Engine/Private/WorldPartition/WorldPartitionActorDescUtils.cpp
+++ b/Engine/Source/Runtime/Engine/Private/WorldPartition/WorldPartitionActorDescUtils.cpp
@@ -194,7 +194,7 @@ void FWorldPartitionActorDescUtils::AppendAssetDataTagsFromActor(const AActor* I
 	if (!InActor->HasAnyFlags(RF_ArchetypeObject | RF_ClassDefaultObject))
 	{
 		// If the actor is not added to a world, we can't retrieve its bounding volume, so try to get the existing one
-		if (ULevel* Level = InActor->GetLevel(); !Level || !Level->Actors.Contains(InActor))
+		if (!InActor->HasActorRegisteredAllComponents())
 		{
 			IAssetRegistry& AssetRegistry = FModuleManager::LoadModuleChecked<FAssetRegistryModule>(TEXT("AssetRegistry")).Get();
 
diff --git a/Engine/Source/Runtime/Engine/Public/WorldPartition/ActorDescContainerSubsystem.h b/Engine/Source/Runtime/Engine/Public/WorldPartition/ActorDescContainerSubsystem.h
--- a/Engine/Source/Runtime/Engine/Public/WorldPartition/ActorDescContainerSubsystem.h
+++ b/Engine/Source/Runtime/Engine/Public/WorldPartition/ActorDescContainerSubsystem.h
@@ -50,8 +50,8 @@ public:
 	ContainerType* RegisterContainer(const typename UActorDescContainer::FInitializeParams& InitParams) { return ContainerManager.RegisterContainer<ContainerType>(InitParams); }
 	UActorDescContainer* GetActorDescContainer(const FString& InContainerName) { return ContainerManager.GetActorDescContainer(InContainerName); }
 	const UActorDescContainer* GetActorDescContainer(const FString& InContainerName) const { return ContainerManager.GetActorDescContainer(InContainerName); }
-	void RegisterContainer(UActorDescContainer* Container) { ContainerManager.RegisterContainer(Container); }
-	void UnregisterContainer(UActorDescContainer* Container) { ContainerManager.UnregisterContainer(Container); }
+	ENGINE_API void RegisterContainer(UActorDescContainer* Container) { ContainerManager.RegisterContainer(Container); }
+	ENGINE_API void UnregisterContainer(UActorDescContainer* Container) { ContainerManager.UnregisterContainer(Container); }
 	FBox GetContainerBounds(const FString& ContainerName, bool bIsEditorBounds = false) const { return ContainerManager.GetContainerBounds(ContainerName, bIsEditorBounds); }
 	ENGINE_API void UpdateContainerBounds(const FString& ContainerName) { ContainerManager.UpdateContainerBounds(ContainerName); }
 	void SetContainerPackage(UActorDescContainer* Container, FName ContainerPackageName) { ContainerManager.SetContainerPackage(Container, ContainerPackageName); }
@@ -77,20 +77,19 @@ private:
 	FContainerUpdatedEvent OnContainerUpdated;
 	FContainerReplacedEvent OnContainerReplaced;
 
-	class FContainerManager
+	class ENGINE_API FContainerManager
 	{
 		FContainerManager(UActorDescContainerSubsystem* InOwner)
 			: Owner(InOwner)
-		{
-
-		}
+		{}
 
 		FContainerManager()
-			: Owner(nullptr) {}
+			: Owner(nullptr)
+		{}
 
 		friend class UActorDescContainerSubsystem;
 
-		struct FRegisteredContainer
+		struct ENGINE_API FRegisteredContainer
 		{
 			FRegisteredContainer()
 				: Container(nullptr)
