Use already-reinstanced subobjects of reinstanced actors

This adjusts ReplaceActorHelper to match ReplaceObjectHelper, following CL 23181906.

diff --git a/Engine/Source/Editor/UnrealEd/Private/Kismet2/KismetReinstanceUtilities.cpp b/Engine/Source/Editor/UnrealEd/Private/Kismet2/KismetReinstanceUtilities.cpp
--- a/Engine/Source/Editor/UnrealEd/Private/Kismet2/KismetReinstanceUtilities.cpp
+++ b/Engine/Source/Editor/UnrealEd/Private/Kismet2/KismetReinstanceUtilities.cpp
@@ -2496,6 +2496,7 @@ static void ReplaceActorHelper(AActor* OldActor, UClass* OldClass, UObject*& New
 	Params.bAggressiveDefaultSubobjectReplacement = true;
 	PRAGMA_ENABLE_DEPRECATION_WARNINGS
 	Params.bNotifyObjectReplacement = true;
+	Params.OptionalReplacementMappings = &OldToNewInstanceMap;
 	// This shouldn't be possible, but if GetArchetype has a bug we could crash in delta serialization
 	// attempting to use it:
 	if (OldArchetype &&!ensure(OldActor->IsA(OldArchetype->GetClass())))
