ObjectTools : Fix "Advanced Copy" for Data Tables with references in their rows.

FArchiveReplaceObjectRefDataTableRows overides FArchiveReplaceObjectRefBase::SerializeObject to serialize each row struct, allowing for the object references in the row to be replaced.

diff --git a/Engine/Source/Editor/UnrealEd/Private/ObjectTools.cpp b/Engine/Source/Editor/UnrealEd/Private/ObjectTools.cpp
--- a/Engine/Source/Editor/UnrealEd/Private/ObjectTools.cpp
+++ b/Engine/Source/Editor/UnrealEd/Private/ObjectTools.cpp
@@ -110,6 +110,7 @@
 #include "ObjectEditorUtils.h"
 #include "Settings/EditorStyleSettings.h"
 #include "ProfilingDebugging/AssetMetadataTrace.h"
+#include "Engine/DataTable.h"
 
 DEFINE_LOG_CATEGORY_STATIC(LogObjectTools, Log, All);
 
@@ -167,6 +168,42 @@ void ReloadEditorWorldForReferenceReplacementIfNecessary(TArray< TWeakObjectPtr<
 	}
 }
 
+template <class T>
+class FArchiveReplaceObjectRefDataTableRows : public FArchiveReplaceObjectRef<T>
+{
+public:
+
+	FArchiveReplaceObjectRefDataTableRows(UObject* InSearchObject, const TMap<T*, T*>& InReplacementMap)
+		: FArchiveReplaceObjectRef<T>(InSearchObject, InReplacementMap, EArchiveReplaceObjectFlags::DelayStart)
+	{
+		// Note: We intentionally add the 'DelayStart' flag above; otherwise the base class ctor will use the base archive type for serialization, and our overrides won't get called.
+		FArchiveReplaceObjectRef<T>::SerializeSearchObject();
+	}
+
+	virtual void SerializeObject(UObject* ObjectToSerialize) override
+	{
+		UDataTable* DataTableObject = Cast<UDataTable>(ObjectToSerialize);
+
+		if (DataTableObject)
+		{
+			if (DataTableObject->RowStruct != nullptr && DataTableObject->RowStruct->RefLink != nullptr)
+			{
+				for (const TPair<FName, uint8*>& Pair : DataTableObject->GetRowMap())
+				{
+					if (uint8* RowData = Pair.Value)
+					{
+						DataTableObject->RowStruct->SerializeItem(*this, RowData, nullptr);
+					}
+				}
+			}
+		}
+		else
+		{
+			FArchiveReplaceObjectRef<T>::SerializeObject(ObjectToSerialize);
+		}
+	}
+};
+
 namespace ObjectTools
 {
 	static int32 MaxTimesToCheckSameObject = 3;
@@ -1354,6 +1391,13 @@ namespace ObjectTools
 						FArchiveReplaceObjectAndStructPropertyRef<UObject> ReplaceInBPClassObject_Ar(BPObjectToUpdate->GeneratedClass, ReplacementMap, EArchiveReplaceObjectFlags::IncludeClassGeneratedByRef);
 						FArchiveReplaceObjectAndStructPropertyRef<UObject> ReplaceInBPClassDefaultObject_Ar(BPObjectToUpdate->GeneratedClass->ClassDefaultObject, ReplacementMap, EArchiveReplaceObjectFlags::IncludeClassGeneratedByRef);
 					}
+
+					UDataTable* DataTableObjectToUpdate = Cast<UDataTable>(CurObject);
+					if (DataTableObjectToUpdate)
+					{
+						FArchiveReplaceObjectRefDataTableRows<UObject> ReplaceInDataTableRows_Ar(DataTableObjectToUpdate, ReplacementMap);
+					}
+
 					FArchiveReplaceObjectAndStructPropertyRef<UObject> ReplaceAr(CurObject, ReplacementMap, EArchiveReplaceObjectFlags::IncludeClassGeneratedByRef);
 				}
 			}
diff --git a/Engine/Source/Runtime/CoreUObject/Public/Serialization/ArchiveReplaceObjectRef.h b/Engine/Source/Runtime/CoreUObject/Public/Serialization/ArchiveReplaceObjectRef.h
--- a/Engine/Source/Runtime/CoreUObject/Public/Serialization/ArchiveReplaceObjectRef.h
+++ b/Engine/Source/Runtime/CoreUObject/Public/Serialization/ArchiveReplaceObjectRef.h
@@ -94,7 +94,7 @@ protected:
 	/**
 	* Serializes a single object
 	*/
-	COREUOBJECT_API void SerializeObject(UObject* ObjectToSerialize);
+	COREUOBJECT_API virtual void SerializeObject(UObject* ObjectToSerialize);
 
 	/** Initial object to start the reference search from */
 	UObject* SearchObject = nullptr;
