WorldPartitionBuilderCommandlet: added support for comma-separated map list, where each entry can be a map or a directory.
Provided directories will be parsed recursively and maps will be added to the map list.

diff --git a/Engine/Source/Editor/UnrealEd/Private/Commandlets/WorldPartitionBuilderCommandlet.cpp b/Engine/Source/Editor/UnrealEd/Private/Commandlets/WorldPartitionBuilderCommandlet.cpp
--- a/Engine/Source/Editor/UnrealEd/Private/Commandlets/WorldPartitionBuilderCommandlet.cpp
+++ b/Engine/Source/Editor/UnrealEd/Private/Commandlets/WorldPartitionBuilderCommandlet.cpp
@@ -9,6 +9,7 @@
 #include "Logging/LogMacros.h"
 #include "Misc/CommandLine.h"
 #include "Misc/EngineVersion.h"
+#include "HAL/FileManager.h"
 #include "HAL/PlatformFileManager.h"
 #include "ProfilingDebugging/ScopedTimers.h"
 #include "WorldPartition/WorldPartition.h"
@@ -80,7 +81,6 @@ int32 UWorldPartitionBuilderCommandlet::Main(const FString& Params)
 	TSet<FString> MapPackagesNames;
 
 	// Parse map name or maps collection
-	FString MapLongPackageName;
 	if (CollectionManager.CollectionExists(FName(Tokens[0]), ECollectionShareType::CST_All))
 	{
 		MapPackagesNames = GatherMapsFromCollection(Tokens[0]);
@@ -90,11 +90,40 @@ int32 UWorldPartitionBuilderCommandlet::Main(const FString& Params)
 		    return 0;
 	    }
 	}
-	else if (FPackageName::SearchForPackageOnDisk(Tokens[0], &MapLongPackageName))
-	{
-		MapPackagesNames = { MapLongPackageName };
-	}	
 	else
+	{
+		TArray<FString> MapList;
+		Tokens[0].ParseIntoArray(MapList, TEXT(","));
+
+		for (const FString& Map : MapList)
+		{
+			FString MapLongPackageName;
+			if (FPackageName::SearchForPackageOnDisk(Map, &MapLongPackageName))
+			{
+				MapPackagesNames.Add(MapLongPackageName);
+			}
+			else
+			{
+				FString PackageFilename;
+				if (FPackageName::TryConvertLongPackageNameToFilename(Map, PackageFilename, TEXT("")))
+				{
+					IFileManager::Get().IterateDirectoryRecursively(*PackageFilename, [&MapLongPackageName, &MapPackagesNames](const TCHAR* FilenameOrDirectory, bool bIsDirectory)
+					{
+						if (!bIsDirectory && FString(FilenameOrDirectory).EndsWith(FPackageName::GetMapPackageExtension()))
+						{
+							if (FPackageName::TryConvertFilenameToLongPackageName(FilenameOrDirectory, MapLongPackageName))
+							{
+								MapPackagesNames.Add(MapLongPackageName);
+							}
+						}
+						return true;
+					});
+				}
+			}
+		}
+	}
+
+	if (MapPackagesNames.IsEmpty())
 	{
 		UE_LOG(LogWorldPartitionBuilderCommandlet, Error, TEXT("Missing world(s) as the first argument to the commandlet. Either supply the world name directly (WorldName or /Path/To/WorldName), or provide a collection name to have the builder operate on a set of maps."));
 		return 1;
@@ -283,4 +312,4 @@ bool UWorldPartitionBuilderCommandlet::AutoSubmitModifiedFiles() const
 	}
 
 	return bSucceeded;
-}
\ No newline at end of file
+}
