MPCook: Collect crash diagnostics from CookWorkers that we terminate because they take longer than timeout to shutdown.
#rnx

diff --git a/Engine/Source/Editor/UnrealEd/Private/Cooker/CookWorkerServer.cpp b/Engine/Source/Editor/UnrealEd/Private/Cooker/CookWorkerServer.cpp
--- a/Engine/Source/Editor/UnrealEd/Private/Cooker/CookWorkerServer.cpp
+++ b/Engine/Source/Editor/UnrealEd/Private/Cooker/CookWorkerServer.cpp
@@ -641,6 +641,7 @@ void FCookWorkerServer::TickWaitForDisconnect()
 		UE_CLOG(!bTerminateImmediately, LogCook, Warning,
 			TEXT("CookWorker process of CookWorkerServer %d failed to disconnect within %.0f seconds; we will terminate it."),
 			ProfileId, WaitForDisconnectTimeout);
+		bNeedCrashDiagnostics = true;
 		ShutdownRemoteProcess();
 		SendToState(EConnectStatus::LostConnection);
 	}
