Fixed an issue where the NavigationBuildingNotification would show the message that the navigation is building when it isn't. This is because the NavigationBuildingNotification module wasn't taking into account that some navmesh can override the "IsAutoUpdateEnabled" functionality.

#rnx

diff --git a/Engine/Source/Editor/UnrealEd/Private/NavigationBuildingNotification.cpp b/Engine/Source/Editor/UnrealEd/Private/NavigationBuildingNotification.cpp
--- a/Engine/Source/Editor/UnrealEd/Private/NavigationBuildingNotification.cpp
+++ b/Engine/Source/Editor/UnrealEd/Private/NavigationBuildingNotification.cpp
@@ -155,7 +155,7 @@ void FNavigationBuildingNotificationImpl::Tick(float DeltaTime)
 		
 		UNavigationSystemV1* NavSys = FNavigationSystem::GetCurrent<UNavigationSystemV1>(EditorContext.World());
 		const bool bBuildInProgress = NavSys != nullptr
-			&& (GetDefault<ULevelEditorMiscSettings>()->bNavigationAutoUpdate ? NavSys->IsNavigationBuildInProgress() : NavSys->GetNumRunningBuildTasks() > 0)
+			&& (NavSys->GetIsAutoUpdateEnabled() ? NavSys->IsNavigationBuildInProgress() : NavSys->GetNumRunningBuildTasks() > 0)
 			&& NavSys->GetNumRemainingBuildTasks() > 0;
 
 		if (!bPreviouslyDetectedBuild && bBuildInProgress)
