WorldPartitionBuilderCommandlet: can now specify *[,n] as a map name which will execute on all maps in the project, where optional argument n can be used to start at a specific index in case some map crashes during the process.

diff --git a/Engine/Source/Editor/UnrealEd/Classes/Commandlets/WorldPartitionBuilderCommandlet.h b/Engine/Source/Editor/UnrealEd/Classes/Commandlets/WorldPartitionBuilderCommandlet.h
--- a/Engine/Source/Editor/UnrealEd/Classes/Commandlets/WorldPartitionBuilderCommandlet.h
+++ b/Engine/Source/Editor/UnrealEd/Classes/Commandlets/WorldPartitionBuilderCommandlet.h
@@ -19,7 +19,7 @@ class UWorldPartitionBuilderCommandlet : public UCommandlet
 	//~ End UCommandlet Interface
 
 private:
-	TSet<FString> GatherMapsFromCollection(const FString& CollectionName) const;
+	TArray<FString> GatherMapsFromCollection(const FString& CollectionName) const;
 	bool RunBuilder(TSubclassOf<UWorldPartitionBuilder> InBuilderClass, const FString& InWorldPackageName);
 
 	bool OnFilesModified(const TArray<FString>& InModifiedFiles, const FString& InChangeDescription);
diff --git a/Engine/Source/Editor/UnrealEd/Private/Commandlets/WorldPartitionBuilderCommandlet.cpp b/Engine/Source/Editor/UnrealEd/Private/Commandlets/WorldPartitionBuilderCommandlet.cpp
--- a/Engine/Source/Editor/UnrealEd/Private/Commandlets/WorldPartitionBuilderCommandlet.cpp
+++ b/Engine/Source/Editor/UnrealEd/Private/Commandlets/WorldPartitionBuilderCommandlet.cpp
@@ -25,6 +25,7 @@
 #include "SourceControlOperations.h"
 
 #include "AssetRegistry/AssetRegistryHelpers.h"
+#include "AssetRegistry/AssetRegistryModule.h"
 
 DEFINE_LOG_CATEGORY_STATIC(LogWorldPartitionBuilderCommandlet, All, All);
 
@@ -78,7 +79,7 @@ int32 UWorldPartitionBuilderCommandlet::Main(const FString& Params)
 	}
 
 	ICollectionManager& CollectionManager = FModuleManager::LoadModuleChecked<FCollectionManagerModule>("CollectionManager").Get();
-	TSet<FString> MapPackagesNames;
+	TArray<FString> MapPackagesNames;
 
 	// Parse map name or maps collection
 	if (CollectionManager.CollectionExists(FName(Tokens[0]), ECollectionShareType::CST_All))
@@ -90,6 +91,31 @@ int32 UWorldPartitionBuilderCommandlet::Main(const FString& Params)
 		    return 0;
 	    }
 	}
+	else if (Tokens[0].StartsWith(TEXT("*")))
+	{
+		IAssetRegistry& AssetRegistry = FModuleManager::LoadModuleChecked<FAssetRegistryModule>(TEXT("AssetRegistry")).Get();
+		AssetRegistry.WaitForCompletion();
+
+		FARFilter Filter;
+		Filter.bIncludeOnlyOnDiskAssets = true;
+		Filter.ClassPaths.Add(UWorld::StaticClass()->GetClassPathName());
+		
+		TArray<FAssetData> WorldAssets;
+		AssetRegistry.GetAssets(Filter, WorldAssets);
+
+		Algo::Transform(WorldAssets, MapPackagesNames, [](const FAssetData& AssetData) { return AssetData.GetSoftObjectPath().ToString(); });
+		
+		// Sort maps gathered from the Asset Registry in case the commandlet crashes, so we can specify *,100 to start at a specific index
+		MapPackagesNames.Sort();
+
+		TArray<FString> SubTokens;
+		Tokens[0].ParseIntoArray(SubTokens, TEXT(","));
+
+		if (SubTokens.IsValidIndex(1))
+		{
+			MapPackagesNames.RemoveAt(0, FCString::Atoi(*SubTokens[1]));
+		}
+	}
 	else
 	{
 		TArray<FString> MapList;
@@ -173,7 +199,7 @@ int32 UWorldPartitionBuilderCommandlet::Main(const FString& Params)
 	return Result;
 }
 
-TSet<FString> UWorldPartitionBuilderCommandlet::GatherMapsFromCollection(const FString& CollectionName) const
+TArray<FString> UWorldPartitionBuilderCommandlet::GatherMapsFromCollection(const FString& CollectionName) const
 {
 	TSet<FString> MapPackagesNames;
 
@@ -200,7 +226,7 @@ TSet<FString> UWorldPartitionBuilderCommandlet::GatherMapsFromCollection(const F
 		}
 	}
 
-	return MapPackagesNames;
+	return MapPackagesNames.Array();
 }
 
 bool UWorldPartitionBuilderCommandlet::RunBuilder(TSubclassOf<UWorldPartitionBuilder> InBuilderClass, const FString& InWorldPackageName)
