Fixed auto-save failing to update the state of source controlled files when using an uncontrolled changelist

diff --git a/Engine/Source/Editor/UnrealEd/Private/PackageAutoSaver.cpp b/Engine/Source/Editor/UnrealEd/Private/PackageAutoSaver.cpp
--- a/Engine/Source/Editor/UnrealEd/Private/PackageAutoSaver.cpp
+++ b/Engine/Source/Editor/UnrealEd/Private/PackageAutoSaver.cpp
@@ -32,16 +32,24 @@
 #include "ShaderCompiler.h"
 #include "EditorLevelUtils.h"
 #include "ISourceControlModule.h"
+#include "SourceControlOperations.h"
 #include "IVREditorModule.h"
 #include "LevelEditorViewport.h"
 #include "AssetCompilingManager.h"
 #include "AssetDefinition.h"
 #include "AssetDefinitionRegistry.h"
 
-namespace PackageAutoSaverOptions
+namespace PackageAutoSaver
 {
 	static bool bUseUncontrolledChangelists = true;
 	static FAutoConsoleVariableRef CVarUseUncontrolledChangelists(TEXT("Editor.AutoSave.UseUncontrolledChangelists"), bUseUncontrolledChangelists, TEXT("True to use uncontrolled changelists (if possible) when using the BackupAndOverwrite auto-save method."));
+
+	static const FUncontrolledChangelist AutoSaveUncontrolledChangelist(FGuid(0x0D227D90, 0x2AC34632, 0xB1A35533, 0xF3CC134E));
+
+	bool CanUseUncontrolledChangelist()
+	{
+		return bUseUncontrolledChangelists && FUncontrolledChangelistsModule::Get().IsEnabled() && ISourceControlModule::Get().GetProvider().UsesLocalReadOnlyState();
+	}
 }
 
 namespace PackageAutoSaverJson
@@ -103,7 +111,7 @@ namespace PackageAutoSaverJson
 /* FPackageAutoSaver                                                    */
 /************************************************************************/
 
-DEFINE_LOG_CATEGORY_STATIC(PackageAutoSaver, Log, All);
+DEFINE_LOG_CATEGORY_STATIC(LogPackageAutoSaver, Log, All);
 
 FPackageAutoSaver::FPackageAutoSaver()
 	: AutoSaveIndex(0)
@@ -317,12 +325,12 @@ void FPackageAutoSaver::AttemptAutoSave()
 
 				// If possible, auto-save will use an uncontrolled changelist rather than actually checkout the files
 				FUncontrolledChangelistsModule& UncontrolledChangelistsModule = FUncontrolledChangelistsModule::Get();
-				const bool bMakeFilesWritable = PackageAutoSaverOptions::bUseUncontrolledChangelists && UncontrolledChangelistsModule.IsEnabled() && ISourceControlModule::Get().GetProvider().UsesLocalReadOnlyState();
+				const bool bUseUncontrolledChangelist = PackageAutoSaver::CanUseUncontrolledChangelist();
 
 				// Build the complete list of packages to save
 				TArray<UPackage*> PackagesToSave;
 				{
-					auto AppendPackagesToSave = [&PackagesToSave, &UncontrolledChangelistsModule, bMakeFilesWritable](const TSet<TWeakObjectPtr<UPackage>, TWeakObjectPtrSetKeyFuncs<TWeakObjectPtr<UPackage>>>& PotentialPackagesToSave)
+					auto AppendPackagesToSave = [&PackagesToSave, &UncontrolledChangelistsModule, bUseUncontrolledChangelist](const TSet<TWeakObjectPtr<UPackage>, TWeakObjectPtrSetKeyFuncs<TWeakObjectPtr<UPackage>>>& PotentialPackagesToSave)
 					{
 						FNameBuilder PkgName;
 						FString PackageFilename;
@@ -339,7 +347,7 @@ void FPackageAutoSaver::AttemptAutoSave()
 									PackagesToSave.Add(Pkg);
 
 									// If using an uncontrolled changelist to track writable files, then make the file writable now so that we can overwrite it during save
-									if (bMakeFilesWritable)
+									if (bUseUncontrolledChangelist)
 									{
 										if (FPackageName::DoesPackageExist(Pkg->GetPathName(), &PackageFilename) && FPlatformFileManager::Get().GetPlatformFile().IsReadOnly(*PackageFilename))
 										{
@@ -369,34 +377,128 @@ void FPackageAutoSaver::AttemptAutoSave()
 					// Note: The in-place save does a regular save of dirty packages, so it doesn't set the bIsAutoSaving flag since it functions like a user-initiated save
 					bool bSuccess = false;
 					{
+						ISourceControlProvider& Provider = ISourceControlModule::Get().GetProvider();
 						FUncontrolledChangelistStatePtr DefaultUncontrolledChangelistState = UncontrolledChangelistsModule.GetDefaultChangelistState();
 
+						TArray<FString> SavedPackageFilenames;
 						TSet<FString> PreSaveFilesInDefaultUncontrolledChangelist;
-						if (bMakeFilesWritable && DefaultUncontrolledChangelistState)
+
+						if (bUseUncontrolledChangelist)
 						{
-							PreSaveFilesInDefaultUncontrolledChangelist.Append(DefaultUncontrolledChangelistState->GetFilenames());
+							// Track which files were in the default uncontrolled changelist prior to the auto-save
+							if (DefaultUncontrolledChangelistState)
+							{
+								PreSaveFilesInDefaultUncontrolledChangelist.Append(DefaultUncontrolledChangelistState->GetFilenames());
+							}
+
+							// Since we disable source control around UEditorLoadingAndSavingUtils::SavePackages, we need to handle reverting and checking out any source controlled files that are in a deleted state prior to the save
+							// Source controlled files won't be managed by the uncontrolled changelist, so we must keep their source control status correct to avoid work being lost
+							if (Provider.IsEnabled())
+							{
+								TSet<FString> EmptyPackageFilenames;
+								for (UPackage* Pkg : PackagesToSave)
+								{
+									FString PackageFilename;
+									if (FPackageName::TryConvertLongPackageNameToFilename(FNameBuilder(Pkg->GetFName()).ToView(), PackageFilename, Pkg->ContainsMap() ? FPackageName::GetMapPackageExtension() : FPackageName::GetAssetPackageExtension()))
+									{
+										PackageFilename = FPaths::ConvertRelativePathToFull(MoveTemp(PackageFilename));
+
+										if (UPackage::IsEmptyPackage(Pkg))
+										{
+											EmptyPackageFilenames.Add(PackageFilename);
+										}
+
+										SavedPackageFilenames.Add(MoveTemp(PackageFilename));
+									}
+								}
+
+								TArray<FSourceControlStateRef> FileStates;
+								if (Provider.GetState(SavedPackageFilenames, FileStates, EStateCacheUsage::ForceUpdate) == ECommandResult::Succeeded)
+								{
+									TArray<FString> FilesToRevertAndCheckout;
+									for (const FSourceControlStateRef& FileState : FileStates)
+									{
+										if (FileState->IsDeleted() && !EmptyPackageFilenames.Contains(FileState->GetFilename()))
+										{
+											FilesToRevertAndCheckout.Add(FileState->GetFilename());
+										}
+									}
+
+									if (FilesToRevertAndCheckout.Num() > 0)
+									{
+										Provider.Execute(ISourceControlOperation::Create<FRevert>(), FilesToRevertAndCheckout);
+										Provider.Execute(ISourceControlOperation::Create<FCheckOut>(), FilesToRevertAndCheckout);
+
+										// If any files failed to checkout then make them writable so they go into the uncontrolled changelist instead
+										for (const FString& PackageFilename : FilesToRevertAndCheckout)
+										{
+											if (FPlatformFileManager::Get().GetPlatformFile().IsReadOnly(*PackageFilename))
+											{
+												if (FPlatformFileManager::Get().GetPlatformFile().SetReadOnly(*PackageFilename, false))
+												{
+													UncontrolledChangelistsModule.OnMakeWritable(PackageFilename);
+												}
+											}
+										}
+									}
+								}
+							}
 						}
 
-						// Disable source control around UEditorLoadingAndSavingUtils::SavePackages when making files writable to avoid it still checking them them out
-						FScopedDisableSourceControl DisableSourceControl(bMakeFilesWritable);
-						bSuccess = UEditorLoadingAndSavingUtils::SavePackages(PackagesToSave, false);
+						// Disable source control around UEditorLoadingAndSavingUtils::SavePackages when using uncontrolled changelists avoid it still checking writable files out
+						{
+							FScopedDisableSourceControl DisableSourceControl(bUseUncontrolledChangelist);
+							bSuccess = UEditorLoadingAndSavingUtils::SavePackages(PackagesToSave, false);
+						}
 
-						// Move any files that the auto-save added to the default uncontrolled changelist to the auto-save specific uncontrolled changelist
 						// Note: We don't check bSuccess here, as a subset of files might still have been saved
-						if (bMakeFilesWritable && DefaultUncontrolledChangelistState)
+						if (bUseUncontrolledChangelist)
 						{
-							TSet<FString> PostSaveFilesInDefaultUncontrolledChangelist = TSet<FString>(DefaultUncontrolledChangelistState->GetFilenames());
-							TSet<FString> NewFilesInDefaultUncontrolledChangelist = PostSaveFilesInDefaultUncontrolledChangelist.Difference(PreSaveFilesInDefaultUncontrolledChangelist);
+							// Since we disable source control around UEditorLoadingAndSavingUtils::SavePackages, we need to handle reverting and deleting any source controlled files that are added or checked out, but no longer available on disk
+							// Source controlled files won't be managed by the uncontrolled changelist, so we must keep their source control status correct to avoid work being lost
+							if (Provider.IsEnabled())
+							{
+								TArray<FSourceControlStateRef> FileStates;
+								if (Provider.GetState(SavedPackageFilenames, FileStates, EStateCacheUsage::ForceUpdate) == ECommandResult::Succeeded)
+								{
+									TArray<FString> FilesToRevertAndDelete;
+									for (const FSourceControlStateRef& FileState : FileStates)
+									{
+										if ((FileState->IsAdded() || FileState->IsCheckedOut()) && !FPaths::FileExists(FileState->GetFilename()))
+										{
+											FilesToRevertAndDelete.Add(FileState->GetFilename());
+										}
+									}
+
+									if (FilesToRevertAndDelete.Num() > 0)
+									{
+										Provider.Execute(ISourceControlOperation::Create<FRevert>(), FilesToRevertAndDelete);
+										Provider.Execute(ISourceControlOperation::Create<FDelete>(), FilesToRevertAndDelete);
 
-							if (NewFilesInDefaultUncontrolledChangelist.Num() > 0)
+										// Make sure the source control operations didn't leave any files on disk
+										for (const FString& PackageFilename : FilesToRevertAndDelete)
+										{
+											IFileManager::Get().Delete(*PackageFilename);
+										}
+									}
+								}
+							}
+
+							// Move any files that the auto-save added to the default uncontrolled changelist to the auto-save specific uncontrolled changelist
+							if (DefaultUncontrolledChangelistState)
 							{
-								FUncontrolledChangelist AutoSaveUncontrolledChangelist(FGuid(0x0D227D90, 0x2AC34632, 0xB1A35533, 0xF3CC134E));
-								if (UncontrolledChangelistsModule.CreateUncontrolledChangelist(NSLOCTEXT("AutoSave", "UncontrolledChangelistTitle", "Auto Save Uncontrolled Changelist"), AutoSaveUncontrolledChangelist))
+								TSet<FString> PostSaveFilesInDefaultUncontrolledChangelist = TSet<FString>(DefaultUncontrolledChangelistState->GetFilenames());
+								TSet<FString> NewFilesInDefaultUncontrolledChangelist = PostSaveFilesInDefaultUncontrolledChangelist.Difference(PreSaveFilesInDefaultUncontrolledChangelist);
+
+								if (NewFilesInDefaultUncontrolledChangelist.Num() > 0)
 								{
-									TArray<FSourceControlStateRef> FileStates;
-									if (ISourceControlModule::Get().GetProvider().GetState(NewFilesInDefaultUncontrolledChangelist.Array(), FileStates, EStateCacheUsage::Use) == ECommandResult::Succeeded)
+									if (UncontrolledChangelistsModule.CreateUncontrolledChangelist(NSLOCTEXT("AutoSave", "UncontrolledChangelistTitle", "Auto Save Uncontrolled Changelist"), PackageAutoSaver::AutoSaveUncontrolledChangelist))
 									{
-										UncontrolledChangelistsModule.MoveFilesToUncontrolledChangelist(TArray<FSourceControlStateRef>(), FileStates, AutoSaveUncontrolledChangelist);
+										TArray<FSourceControlStateRef> FileStates;
+										if (Provider.GetState(NewFilesInDefaultUncontrolledChangelist.Array(), FileStates, EStateCacheUsage::Use) == ECommandResult::Succeeded)
+										{
+											UncontrolledChangelistsModule.MoveFilesToUncontrolledChangelist(TArray<FSourceControlStateRef>(), FileStates, PackageAutoSaver::AutoSaveUncontrolledChangelist);
+										}
 									}
 								}
 							}
