Convert a delete to a checkout when restoring a deleted actor via an interactive save

CanCheckIn considers a delete to be valid to check-in, but when restoring a deleted actor (eg, via undo) the delete actually needs to be reverted and converted to a checkout.

FEditorFileUtils::CheckoutPackages understands that, and handles it correctly, but FEditorFileUtils::AddCheckoutPackageItems was filtering the deleted package out of the list it was passed.

Given the list of packages to checkout can also be passed to FEditorFileUtils::MakePackagesWritable, that also needs to handle reverting deleted packages prior to making them writable on disk.

diff --git a/Engine/Source/Editor/UnrealEd/Private/FileHelpers.cpp b/Engine/Source/Editor/UnrealEd/Private/FileHelpers.cpp
--- a/Engine/Source/Editor/UnrealEd/Private/FileHelpers.cpp
+++ b/Engine/Source/Editor/UnrealEd/Private/FileHelpers.cpp
@@ -1727,7 +1727,8 @@ bool FEditorFileUtils::AddCheckoutPackageItems(bool bCheckDirty, TArray<UPackage
 		FSourceControlStatePtr SourceControlState = SourceControlProvider.GetState(CurPackage, EStateCacheUsage::Use);
 
 		// Package does not need to be checked out if its already checked out or we are ignoring it for source control
-		bool bSCCCanEdit = !SourceControlState.IsValid() || SourceControlState->CanCheckIn() || SourceControlState->IsIgnored() || SourceControlState->IsUnknown() || (bCareAboutReadOnly && !bPkgReadOnly);
+		bool bNeedsRevert = SourceControlState.IsValid() && SourceControlState->IsDeleted() && !UPackage::IsEmptyPackage(CurPackage); // Need to revert the delete and then checkout
+		bool bSCCCanEdit = !SourceControlState.IsValid() || (SourceControlState->CanCheckIn() && !bNeedsRevert) || SourceControlState->IsIgnored() || SourceControlState->IsUnknown() || (bCareAboutReadOnly && !bPkgReadOnly);
 		bool bIsSourceControlled = SourceControlState.IsValid() && SourceControlState->IsSourceControlled();
 
 		if (!bSCCCanEdit && (bIsSourceControlled && (!bCheckDirty || (bCheckDirty && CurPackage->IsDirty()))) && !SourceControlState->IsCheckedOut())
@@ -2014,6 +2015,27 @@ void FEditorFileUtils::MakePackagesWritable(const TArray<UPackage*>& PackagesToM
 	PackagesMadeWritableSuccess.Reserve(PackagesToMakeWritable.Num());
 	PackagesMadeWritableFailure.Reserve(PackagesToMakeWritable.Num());
 
+	// Packages that are deleted in source control may need to be reverted before we can make them writable
+	if (ISourceControlModule::Get().IsEnabled())
+	{
+		ISourceControlProvider& SourceControlProvider = ISourceControlModule::Get().GetProvider();
+
+		TArray<UPackage*> PackagesToRevert;
+		for (UPackage* PackageToMakeWritable : PackagesToMakeWritable)
+		{
+			if (FSourceControlStatePtr SourceControlState = SourceControlProvider.GetState(PackageToMakeWritable, EStateCacheUsage::Use);
+				SourceControlState && SourceControlState->IsDeleted() && !UPackage::IsEmptyPackage(PackageToMakeWritable))
+			{
+				PackagesToRevert.Add(PackageToMakeWritable);
+			}
+		}
+
+		if (PackagesToRevert.Num() > 0)
+		{
+			SourceControlProvider.Execute(ISourceControlOperation::Create<FRevert>(), PackagesToRevert);
+		}
+	}
+
 	// Attempt to make writable each package the user checked
 	FUncontrolledChangelistsModule& UncontrolledChangelistModule = FUncontrolledChangelistsModule::Get();
 	for (UPackage* PackageToMakeWritable : PackagesToMakeWritable)
