MPCook: CookPackageSplitter: Fix failed assertion when uninitializing a GenerationHelper that has generated packages in the savestalled state. They are demoted before reaching the assertion, but in some states the demotion did not clear out the data that we assert is empty.
Fixed by calling GenerationHelper->ResetSaveState from ReleaseCookedPlatformData when necessary, even in cases where GetSaveSubState == StartSave, rather than early exiting in the StartSave case.
#rnx

diff --git a/Engine/Source/Editor/UnrealEd/Private/CookOnTheFlyServer.cpp b/Engine/Source/Editor/UnrealEd/Private/CookOnTheFlyServer.cpp
--- a/Engine/Source/Editor/UnrealEd/Private/CookOnTheFlyServer.cpp
+++ b/Engine/Source/Editor/UnrealEd/Private/CookOnTheFlyServer.cpp
@@ -4246,12 +4246,6 @@ void UCookOnTheFlyServer::ReleaseCookedPlatformData(UE::Cook::FPackageData& Pack
 {
 	using namespace UE::Cook;
 
-	if (PackageData.GetSaveSubState() == ESaveSubState::StartSave)
-	{
-		PackageData.CheckCookedPlatformDataEmpty();
-		return;
-	}
-
 	TRefCountPtr<FGenerationHelper> GenerationHelper = PackageData.GetGenerationHelper();
 	if (!GenerationHelper)
 	{
@@ -4261,80 +4255,86 @@ void UCookOnTheFlyServer::ReleaseCookedPlatformData(UE::Cook::FPackageData& Pack
 		GenerationHelper->FindInfo(PackageData) : nullptr;
 
 	// For every BeginCacheForCookedPlatformData call we made we need to call ClearAllCachedCookedPlatformData
-	if (ReleaseSaveReason == EStateChangeReason::Completed)
+	// No need to check for CookedPlatformData if in the StartSave state; we can not have any in that case
+	if (PackageData.GetSaveSubState() > ESaveSubState::StartSave)
 	{
-		// Since we have completed CookedPlatformData, we know we called BeginCacheForCookedPlatformData on all
-		// objects in the package, and none are pending
-		UE_SCOPED_HIERARCHICAL_COOKTIMER(ClearAllCachedCookedPlatformData);
-		for (FCachedObjectInOuter& CachedObjectInOuter : PackageData.GetCachedObjectsInOuter())
+		if (ReleaseSaveReason == EStateChangeReason::Completed)
 		{
-			UObject* Object = CachedObjectInOuter.Object.Get();
-			if (Object)
+			// Since we have completed CookedPlatformData, we know we called BeginCacheForCookedPlatformData on all
+			// objects in the package, and none are pending
+			UE_SCOPED_HIERARCHICAL_COOKTIMER(ClearAllCachedCookedPlatformData);
+			for (FCachedObjectInOuter& CachedObjectInOuter : PackageData.GetCachedObjectsInOuter())
 			{
-				FPendingCookedPlatformData::ClearCachedCookedPlatformData(Object, PackageData,
-					true /* bCompletedSuccesfully */);
+				UObject* Object = CachedObjectInOuter.Object.Get();
+				if (Object)
+				{
+					FPendingCookedPlatformData::ClearCachedCookedPlatformData(Object, PackageData,
+						true /* bCompletedSuccesfully */);
+				}
 			}
 		}
-	}
-	else 
-	{
-		// This is a slower but more general flow that can handle releasing whether or not we called SavePackage
-		// Note that even after we return from this function, some objects with pending IsCachedCookedPlatformDataLoaded
-		// calls may still exist for this Package in PendingCookedPlatformDatas
-		// and this PackageData may therefore still have GetNumPendingCookedPlatformData > 0
-		// We have only called BeginCacheForCookedPlatformData on Object,Platform pairs up to GetCookedPlatformDataNextIndex.
-		// Further, some of those calls might still be pending.
-
-		// Find all pending BeginCacheForCookedPlatformData for this FPackageData
-		TMap<UObject*, TArray<FPendingCookedPlatformData*>> PendingObjects;
-		PackageDatas->ForEachPendingCookedPlatformData(
-		[&PendingObjects, &PackageData](FPendingCookedPlatformData& PendingCookedPlatformData)
+		else
 		{
-			if (&PendingCookedPlatformData.PackageData == &PackageData && !PendingCookedPlatformData.PollIsComplete())
-			{
-				UObject* Object = PendingCookedPlatformData.Object.Get();
-				check(Object); // Otherwise PollIsComplete would have returned true
-				check(!PendingCookedPlatformData.bHasReleased); // bHasReleased should be false since PollIsComplete returned false
-				PendingObjects.FindOrAdd(Object).Add(&PendingCookedPlatformData);
-			}
-		});
+			// This is a slower but more general flow that can handle releasing whether or not we called SavePackage
+			// Note that even after we return from this function, some objects with pending IsCachedCookedPlatformDataLoaded
+			// calls may still exist for this Package in PendingCookedPlatformDatas
+			// and this PackageData may therefore still have GetNumPendingCookedPlatformData > 0
+			// We have only called BeginCacheForCookedPlatformData on Object,Platform pairs up to GetCookedPlatformDataNextIndex.
+			// Further, some of those calls might still be pending.
 
+			// Find all pending BeginCacheForCookedPlatformData for this FPackageData
+			TMap<UObject*, TArray<FPendingCookedPlatformData*>> PendingObjects;
+			PackageDatas->ForEachPendingCookedPlatformData(
+				[&PendingObjects, &PackageData](FPendingCookedPlatformData& PendingCookedPlatformData)
+				{
+					if (&PendingCookedPlatformData.PackageData == &PackageData && !PendingCookedPlatformData.PollIsComplete())
+					{
+						UObject* Object = PendingCookedPlatformData.Object.Get();
+						check(Object); // Otherwise PollIsComplete would have returned true
+						check(!PendingCookedPlatformData.bHasReleased); // bHasReleased should be false since PollIsComplete returned false
+						PendingObjects.FindOrAdd(Object).Add(&PendingCookedPlatformData);
+					}
+				});
 
-		// Iterate over all objects in the FPackageData up to GetCookedPlatformDataNextIndex
-		TArray<FCachedObjectInOuter>& CachedObjects = PackageData.GetCachedObjectsInOuter();
-		for (int32 ObjectIndex = 0; ObjectIndex < PackageData.GetCookedPlatformDataNextIndex(); ++ObjectIndex)
-		{
-			UObject* Object = CachedObjects[ObjectIndex].Object.Get();
-			if (!Object)
-			{
-				continue;
-			}
-			TArray<FPendingCookedPlatformData*>* PendingDatas = PendingObjects.Find(Object);
-			if (!PendingDatas || PendingDatas->Num() == 0)
-			{
-				// No pending BeginCacheForCookedPlatformData calls for this object; clear it now.
-				FPendingCookedPlatformData::ClearCachedCookedPlatformData(Object, PackageData,
-					false /* bCompletedSuccesfully */);
-			}
-			else
+
+			// Iterate over all objects in the FPackageData up to GetCookedPlatformDataNextIndex
+			TArray<FCachedObjectInOuter>& CachedObjects = PackageData.GetCachedObjectsInOuter();
+			for (int32 ObjectIndex = 0; ObjectIndex < PackageData.GetCookedPlatformDataNextIndex(); ++ObjectIndex)
 			{
-				// For any pending Objects, we add a CancelManager to the FPendingCookedPlatformData to call
-				// ClearAllCachedCookedPlatformData when the pending Object,Platform pairs for that object completes.
-				FPendingCookedPlatformDataCancelManager* CancelManager = new FPendingCookedPlatformDataCancelManager();
-				CancelManager->NumPendingPlatforms = PendingDatas->Num();
-				for (FPendingCookedPlatformData* PendingCookedPlatformData : *PendingDatas)
+				UObject* Object = CachedObjects[ObjectIndex].Object.Get();
+				if (!Object)
 				{
-					// We never start a new package until after clearing the previous cancels, so all of the
-					// FPendingCookedPlatformData for the PlatformData we are cancelling can not have been cancelled before.
-					// We would leak the CancelManager if we overwrote it here.
-					check(PendingCookedPlatformData->CancelManager == nullptr);
-					// If bHasReleased on the PendingCookedPlatformData were already true, we would leak the CancelManager
-					// because the PendingCookedPlatformData would never call Release on it.
-					check(!PendingCookedPlatformData->bHasReleased);
-					PendingCookedPlatformData->CancelManager = CancelManager;
+					continue;
+				}
+				TArray<FPendingCookedPlatformData*>* PendingDatas = PendingObjects.Find(Object);
+				if (!PendingDatas || PendingDatas->Num() == 0)
+				{
+					// No pending BeginCacheForCookedPlatformData calls for this object; clear it now.
+					FPendingCookedPlatformData::ClearCachedCookedPlatformData(Object, PackageData,
+						false /* bCompletedSuccesfully */);
+				}
+				else
+				{
+					// For any pending Objects, we add a CancelManager to the FPendingCookedPlatformData to call
+					// ClearAllCachedCookedPlatformData when the pending Object,Platform pairs for that object completes.
+					FPendingCookedPlatformDataCancelManager* CancelManager = new FPendingCookedPlatformDataCancelManager();
+					CancelManager->NumPendingPlatforms = PendingDatas->Num();
+					for (FPendingCookedPlatformData* PendingCookedPlatformData : *PendingDatas)
+					{
+						// We never start a new package until after clearing the previous cancels, so all of the
+						// FPendingCookedPlatformData for the PlatformData we are cancelling can not have been cancelled before.
+						// We would leak the CancelManager if we overwrote it here.
+						check(PendingCookedPlatformData->CancelManager == nullptr);
+						// If bHasReleased on the PendingCookedPlatformData were already true, we would leak the CancelManager
+						// because the PendingCookedPlatformData would never call Release on it.
+						check(!PendingCookedPlatformData->bHasReleased);
+						PendingCookedPlatformData->CancelManager = CancelManager;
+					}
 				}
 			}
 		}
+
+		PackageData.ClearCookedPlatformData();
 	}
 
 	if (GenerationInfo)
@@ -4342,8 +4342,6 @@ void UCookOnTheFlyServer::ReleaseCookedPlatformData(UE::Cook::FPackageData& Pack
 		GenerationHelper->ResetSaveState(*GenerationInfo, PackageData.GetPackage(), ReleaseSaveReason, NewState);
 	}
 
-	PackageData.ClearCookedPlatformData();
-
 	if (ReleaseSaveReason != EStateChangeReason::RecreateObjectCache)
 	{
 		if (!IsCookOnTheFlyMode() && !IsCookingInEditor())
diff --git a/Engine/Source/Editor/UnrealEd/Private/Cooker/CookGenerationHelper.cpp b/Engine/Source/Editor/UnrealEd/Private/Cooker/CookGenerationHelper.cpp
--- a/Engine/Source/Editor/UnrealEd/Private/Cooker/CookGenerationHelper.cpp
+++ b/Engine/Source/Editor/UnrealEd/Private/Cooker/CookGenerationHelper.cpp
@@ -1684,10 +1684,10 @@ void FGenerationHelper::PreGarbageCollectGCLifetimeData()
 	}
 
 	// We should either uninitialize or destroy after the garbage collect.
-	// We should not uninitialize unless the Generator package is going to be collected, but we are in a state
-	// where nothing in the cooker is depending on the package anymore (all generator and generated packages are
-	// not in the save state or are stalled) so we do expect the generator package to be garbage collected by the
-	// upcoming GC.
+	// Usually we should not uninitialize unless the Generator package is going to be collected, but in this case we
+	// are in a state where nothing in the cooker is depending on the package anymore (all generator and generated
+	// packages are not in the save state or are stalled) so we do expect the generator package to be garbage
+	// collected by the upcoming GC.
 	// But for that to happen we have to drop our references to it from this FGenerationHelper, so we need
 	// to uninitialize. Also mark that we should check for generator garbage collect in PostGarbageCollect.
 	// Depending on the Splitter class, it may tolerate failure to GC the Generator package, in which case we
@@ -2015,7 +2015,7 @@ void FCookGenerationInfo::Uninitialize()
 {
 	// Check that we have left the save state first, since other assertions assume we have left the save state
 	checkf(PackageData->GetSaveSubState() == ESaveSubState::StartSave,
-		TEXT("Cooker bug: Expected FCookGenerationInfo::Uninitialize to not be called for a package still in the save state, ")
+		TEXT("Cooker bug: Expected FCookGenerationInfo::Uninitialize to not be called for a package in an active save state, ")
 		TEXT("but %s package %s has SaveSubState %s."),
 		bGenerator ? TEXT("generator") : TEXT("generated"), *GetPackageName(),
 		LexToString(PackageData->GetSaveSubState()));
@@ -2027,7 +2027,12 @@ void FCookGenerationInfo::Uninitialize()
 	PackageDependencies.Empty();
 	// Keep PackageData; it is allowed in the uninitialized state
 	KeepReferencedPackages.Empty();
-	check(CachedObjectsInOuterInfo.IsEmpty()); // We can not still be in the save state, so this should be empty
+	// CachedObjectsInOuterInfo should have been cleared before uninitialize, either by leaving the save state, or
+	// calling ReleaseCookedPlatformData.
+	checkf(CachedObjectsInOuterInfo.IsEmpty(),
+		TEXT("Cooker bug: Expected FCookGenerationInfo::Uninitialize to not be called for a package in an active save state, ")
+		TEXT("but %s package %s has non-empty CachedObjectsInOuterInfo."),
+		bGenerator ? TEXT("generator") : TEXT("generated"), *GetPackageName());
 	// Keep SavedOnWorker; it is allowed in the uninitialized state
 	bCreateAsMap = false;
 	bHasCreatedPackage = false;
