Allow IsAutoSaving to query for both transient or persistent auto-saves

Note: It still only queries for transient auto-saves by default, to maintain backwards compatability with the existing API

diff --git a/Engine/Plugins/Editor/DataValidation/Source/DataValidation/Private/EditorValidatorSubsystem.cpp b/Engine/Plugins/Editor/DataValidation/Source/DataValidation/Private/EditorValidatorSubsystem.cpp
--- a/Engine/Plugins/Editor/DataValidation/Source/DataValidation/Private/EditorValidatorSubsystem.cpp
+++ b/Engine/Plugins/Editor/DataValidation/Source/DataValidation/Private/EditorValidatorSubsystem.cpp
@@ -684,7 +684,7 @@ void UEditorValidatorSubsystem::LogAssetValidationSummary(FMessageLog& DataValid
 void UEditorValidatorSubsystem::ValidateOnSave(TArray<FAssetData> AssetDataList, bool bProceduralSave) const
 {
 	// Only validate if enabled and not auto saving
-	if (!GetDefault<UDataValidationSettings>()->bValidateOnSave || GEditor->IsAutosaving())
+	if (!GetDefault<UDataValidationSettings>()->bValidateOnSave || GEditor->IsAutosaving(EPackageAutoSaveType::Any))
 	{
 		return;
 	}
@@ -724,7 +724,7 @@ void UEditorValidatorSubsystem::ValidateOnSave(TArray<FAssetData> AssetDataList,
 void UEditorValidatorSubsystem::ValidateSavedPackage(FName PackageName, bool bProceduralSave)
 {
 	// Only validate if enabled and not auto saving
-	if (!GetDefault<UDataValidationSettings>()->bValidateOnSave || GEditor->IsAutosaving())
+	if (!GetDefault<UDataValidationSettings>()->bValidateOnSave || GEditor->IsAutosaving(EPackageAutoSaveType::Any))
 	{
 		return;
 	}
diff --git a/Engine/Source/Editor/UnrealEd/Classes/Editor/UnrealEdEngine.h b/Engine/Source/Editor/UnrealEd/Classes/Editor/UnrealEdEngine.h
--- a/Engine/Source/Editor/UnrealEd/Classes/Editor/UnrealEdEngine.h
+++ b/Engine/Source/Editor/UnrealEd/Classes/Editor/UnrealEdEngine.h
@@ -194,7 +194,7 @@ public:
 	UNREALED_API virtual void CloseEditor() override;
 
 	
-	UNREALED_API virtual bool IsAutosaving() const override;
+	UNREALED_API virtual bool IsAutosaving(const EPackageAutoSaveType AutoSaveType = EPackageAutoSaveType::Transient) const override;
 	//~ End UEditorEngine Interface 
 	
 	//~ Begin FExec Interface
diff --git a/Engine/Source/Editor/UnrealEd/Private/PackageAutoSaver.cpp b/Engine/Source/Editor/UnrealEd/Private/PackageAutoSaver.cpp
--- a/Engine/Source/Editor/UnrealEd/Private/PackageAutoSaver.cpp
+++ b/Engine/Source/Editor/UnrealEd/Private/PackageAutoSaver.cpp
@@ -116,7 +116,6 @@ DEFINE_LOG_CATEGORY_STATIC(LogPackageAutoSaver, Log, All);
 FPackageAutoSaver::FPackageAutoSaver()
 	: AutoSaveIndex(0)
 	, AutoSaveCount(0.0f)
-	, bIsAutoSaving(false)
 	, bDelayingDueToFailedSave(false)
 	, bAutoDeclineRecovery(FParse::Param(FCommandLine::Get(), TEXT("AutoDeclinePackageRecovery")))
 {
@@ -269,7 +268,7 @@ void FPackageAutoSaver::AttemptAutoSave()
 
 			if (LoadingSavingSettings->AutoSaveMethod == EAutoSaveMethod::BackupAndRestore)
 			{
-				TGuardValue<bool> IsAutoSavingScope(bIsAutoSaving, true);
+				TGuardValue<EPackageAutoSaveType> IsAutoSavingScope(CurrentAutoSaveType, EPackageAutoSaveType::Transient);
 
 				SlowTask.EnterProgressFrame(50);
 
@@ -374,7 +373,8 @@ void FPackageAutoSaver::AttemptAutoSave()
 
 				if (PackagesToSave.Num() > 0)
 				{
-					// Note: The in-place save does a regular save of dirty packages, so it doesn't set the bIsAutoSaving flag since it functions like a user-initiated save
+					TGuardValue<EPackageAutoSaveType> IsAutoSavingScope(CurrentAutoSaveType, EPackageAutoSaveType::Persistent);
+
 					bool bSuccess = false;
 					{
 						ISourceControlProvider& Provider = ISourceControlModule::Get().GetProvider();
@@ -630,7 +630,7 @@ void FPackageAutoSaver::OnPackagesDeleted(const TArray<UPackage*>& DeletedPackag
 
 void FPackageAutoSaver::OnPackageDirtyStateUpdated(UPackage* Pkg)
 {
-	if (!IsAutoSaving())
+	if (!IsAutoSaving(EPackageAutoSaveType::Transient))
 	{
 		PackagesPendingUpdate.Add(Pkg);
 	}
@@ -638,7 +638,7 @@ void FPackageAutoSaver::OnPackageDirtyStateUpdated(UPackage* Pkg)
 
 void FPackageAutoSaver::OnMarkPackageDirty(UPackage* Pkg, bool bWasDirty)
 {
-	if (!IsAutoSaving())
+	if (!IsAutoSaving(EPackageAutoSaveType::Transient))
 	{
 		PackagesPendingUpdate.Add(Pkg);
 	}
@@ -647,7 +647,7 @@ void FPackageAutoSaver::OnMarkPackageDirty(UPackage* Pkg, bool bWasDirty)
 void FPackageAutoSaver::OnPackageSaved(const FString& Filename, UPackage* Pkg, FObjectPostSaveContext ObjectSaveContext)
 {
 	// If this has come from an auto-save, update the last known filename in the user dirty list so that we can offer is up as a restore file later
-	if (IsAutoSaving())
+	if (IsAutoSaving(EPackageAutoSaveType::Transient))
 	{
 		if (FString* const AutoSaveFilename = DirtyPackagesForUserSave.Find(Pkg))
 		{
diff --git a/Engine/Source/Editor/UnrealEd/Private/PackageAutoSaver.h b/Engine/Source/Editor/UnrealEd/Private/PackageAutoSaver.h
--- a/Engine/Source/Editor/UnrealEd/Private/PackageAutoSaver.h
+++ b/Engine/Source/Editor/UnrealEd/Private/PackageAutoSaver.h
@@ -39,7 +39,7 @@ public:
 	virtual bool HasPackagesToRestore() const override;
 	virtual void OfferToRestorePackages() override;
 	virtual void OnPackagesDeleted(const TArray<UPackage*>& DeletedPackages) override;
-	virtual bool IsAutoSaving() const override { return bIsAutoSaving; }
+	virtual bool IsAutoSaving(const EPackageAutoSaveType AutoSaveType = EPackageAutoSaveType::Transient) const override { return EnumHasAnyFlags(CurrentAutoSaveType, AutoSaveType); }
 	virtual void DisableRestorePromptAndDeclinePackageRecovery() override { bAutoDeclineRecovery = true; }
 
 private:
@@ -122,7 +122,7 @@ private:
 	float AutoSaveCount;
 
 	/** If we are currently auto-saving */
-	bool bIsAutoSaving;
+	EPackageAutoSaveType CurrentAutoSaveType = EPackageAutoSaveType::None;
 
 	/** Flag for whether auto-save warning notification has been launched */
 	bool bAutoSaveNotificationLaunched;
@@ -159,4 +159,7 @@ private:
 
 	/** Restore information that was loaded following a crash */
 	TMap<FString, TPair<FString, FString>> PackagesThatCanBeRestored;
+
+	/** The cached state of the uncontrolled packages changelist, the last time OnUncontrolledChangelistModuleChanged was called. This may be stale in-between frames, but is very fast to query */
+	TMap<FName, FString> CachedUncontrolledPackages;
 };
diff --git a/Engine/Source/Editor/UnrealEd/Private/UnrealEdEngine.cpp b/Engine/Source/Editor/UnrealEd/Private/UnrealEdEngine.cpp
--- a/Engine/Source/Editor/UnrealEd/Private/UnrealEdEngine.cpp
+++ b/Engine/Source/Editor/UnrealEd/Private/UnrealEdEngine.cpp
@@ -840,11 +840,11 @@ void UUnrealEdEngine::OnPostWindowsMessage(FViewport* Viewport, uint32 Message)
 	}		
 }
 
-bool UUnrealEdEngine::IsAutosaving() const
+bool UUnrealEdEngine::IsAutosaving(const EPackageAutoSaveType AutoSaveType) const
 {
 	if (PackageAutoSaver)
 	{
-		return PackageAutoSaver->IsAutoSaving();
+		return PackageAutoSaver->IsAutoSaving(AutoSaveType);
 	}
 	
 	return false;
diff --git a/Engine/Source/Editor/UnrealEd/Public/IPackageAutoSaver.h b/Engine/Source/Editor/UnrealEd/Public/IPackageAutoSaver.h
--- a/Engine/Source/Editor/UnrealEd/Public/IPackageAutoSaver.h
+++ b/Engine/Source/Editor/UnrealEd/Public/IPackageAutoSaver.h
@@ -3,6 +3,7 @@
 #pragma once
 
 #include "CoreMinimal.h"
+#include "Engine/EngineTypes.h"
 
 /** An interface to handle the creation, destruction, and restoration of auto-saved packages */
 class IPackageAutoSaver
@@ -39,8 +40,8 @@ public:
 	/** Attempts to auto-save the level and/or content packages, if those features are enabled. */
 	virtual void AttemptAutoSave() = 0;
 
-	/** @return If we are currently auto-saving */
-	virtual bool IsAutoSaving() const = 0;
+	/** @return If we are currently auto-saving (default is to only check for transient auto-saves for legacy reasons) */
+	virtual bool IsAutoSaving(const EPackageAutoSaveType AutoSaveType = EPackageAutoSaveType::Transient) const = 0;
 
 	/** Load the restore file from disk (if present) */
 	virtual void LoadRestoreFile() = 0;
diff --git a/Engine/Source/Runtime/Engine/Classes/Engine/Engine.h b/Engine/Source/Runtime/Engine/Classes/Engine/Engine.h
--- a/Engine/Source/Runtime/Engine/Classes/Engine/Engine.h
+++ b/Engine/Source/Runtime/Engine/Classes/Engine/Engine.h
@@ -3519,8 +3519,8 @@ public:
 	virtual bool AreEditorAnalyticsEnabled() const { return false; }
 	virtual void CreateStartupAnalyticsAttributes( TArray<struct FAnalyticsEventAttribute>& StartSessionAttributes ) const {}
 	
-	/** @return true if the engine is autosaving a package */
-	virtual bool IsAutosaving() const { return false; }
+	/** @return true if the engine is autosaving a package (default is to only check for transient auto-saves for legacy reasons) */
+	virtual bool IsAutosaving(const EPackageAutoSaveType AutoSaveType = EPackageAutoSaveType::Transient) const { return false; }
 
 	virtual bool ShouldDoAsyncEndOfFrameTasks() const { return false; }
 
diff --git a/Engine/Source/Runtime/Engine/Classes/Engine/EngineTypes.h b/Engine/Source/Runtime/Engine/Classes/Engine/EngineTypes.h
--- a/Engine/Source/Runtime/Engine/Classes/Engine/EngineTypes.h
+++ b/Engine/Source/Runtime/Engine/Classes/Engine/EngineTypes.h
@@ -7,6 +7,7 @@
  *	use in multiple files where the enum can't be mapped to a specific file.
  */
 
+#include "Misc/EnumClassFlags.h"
 #include "UObject/ObjectMacros.h"
 #include "UObject/Object.h"
 #include "UObject/Class.h"
@@ -4153,6 +4154,19 @@ enum class ELevelCollectionType : uint8
 	MAX
 };
 
+enum class EPackageAutoSaveType : uint8
+{
+	/** No auto-save is happening */
+	None = 0,
+	/** The auto-save is happening to a transient directory (eg, the Saved/Autosaves directory) and doesn't directly replace the persistent project content */
+	Transient = 1 << 0,
+	/** The auto-save is happening directly to the persistent project content, though a backup may still be made in a transient directory */
+	Persistent = 1 << 1,
+	/** The auto-save is happening for any reason */
+	Any = Transient | Persistent,
+};
+ENUM_CLASS_FLAGS(EPackageAutoSaveType);
+
 #if WITH_EDITOR
 
 ENGINE_API void SerializeNaniteSettingsForDDC(FArchive& Ar, FMeshNaniteSettings& NaniteSettings, bool bIsNaniteForceEnabled);
