ObjectTools::CleanupAfterSuccessfulDelete will now consider packages having a DeletedObjectPlaceholder objects as an external package.
The current code was checking if the package was an external actor or object, but with scene graph entities will be stored anywhere in the folder hierarchy so we can't rely on this.

diff --git a/Engine/Source/Editor/UnrealEd/Private/ObjectTools.cpp b/Engine/Source/Editor/UnrealEd/Private/ObjectTools.cpp
--- a/Engine/Source/Editor/UnrealEd/Private/ObjectTools.cpp
+++ b/Engine/Source/Editor/UnrealEd/Private/ObjectTools.cpp
@@ -74,6 +74,7 @@
 #include "Dialogs/DlgReferenceTree.h"
 #include "AssetRegistry/ARFilter.h"
 #include "AssetDeleteModel.h"
+#include "DeletedObjectPlaceholder.h"
 #include "Dialogs/SPrivateAssetsDialog.h"
 #include "Dialogs/SDeleteAssetsDialog.h"
 #include "AudioDevice.h"
@@ -2463,13 +2464,23 @@ namespace ObjectTools
 		for ( int32 PackageIdx = PackagesToDelete.Num() - 1; PackageIdx >= 0; --PackageIdx )
 		{
 			GWarn->StatusUpdate(OriginalNumPackagesToDelete - PackageIdx, OriginalNumPackagesToDelete, NSLOCTEXT("ObjectTools", "OldPackageCleanupSlowTask", "Cleaning Up Old Assets"));
-			UObject* Package = PackagesToDelete[PackageIdx];
+			UPackage* Package = PackagesToDelete[PackageIdx];
 
 			bool bIsReferenced = false;
 
-			// Skip external object packages when considering whether to clear the transaction buffer, as you should be able to undo deleting an actor, an actor folder, etc. (but not an asset)
-			// If an external object package is kept alive by the transaction buffer then it will be re-marked as "newly created" further down this function
-			if ( Package != nullptr && bPerformReferenceCheck && !Package->GetName().Contains(FPackagePath::GetExternalActorsFolderName()) && !Package->GetName().Contains(FPackagePath::GetExternalObjectsFolderName()))
+			auto ShouldSkipReferenceCheck = [](const UPackage* InPackage)
+			{
+				// Skip external object packages when considering whether to clear the transaction buffer, as you should be able to undo deleting an actor, an actor folder, etc. (but not an asset)
+				// If an external object package is kept alive by the transaction buffer then it will be re-marked as "newly created" further down this function
+				if (InPackage->GetName().Contains(FPackagePath::GetExternalActorsFolderName()) || InPackage->GetName().Contains(FPackagePath::GetExternalObjectsFolderName()))
+				{
+					return true;
+				}
+
+				return !!UDeletedObjectPlaceholder::FindInPackage(InPackage);
+			};
+
+			if ( Package != nullptr && bPerformReferenceCheck && !ShouldSkipReferenceCheck(Package) )
 			{
 				bool bIsReferencedByUndo = false;
 				GatherObjectReferencersForDeletion(Package, bIsReferenced, bIsReferencedByUndo);
diff --git a/Engine/Source/Runtime/Engine/Private/DeletedObjectPlaceholder.cpp b/Engine/Source/Runtime/Engine/Private/DeletedObjectPlaceholder.cpp
--- a/Engine/Source/Runtime/Engine/Private/DeletedObjectPlaceholder.cpp
+++ b/Engine/Source/Runtime/Engine/Private/DeletedObjectPlaceholder.cpp
@@ -55,7 +55,7 @@ UDeletedObjectPlaceholder* UDeletedObjectPlaceholder::RemoveFromPackage(UPackage
 	return nullptr;
 }
 
-UDeletedObjectPlaceholder* UDeletedObjectPlaceholder::FindInPackage(UPackage* InPackage)
+UDeletedObjectPlaceholder* UDeletedObjectPlaceholder::FindInPackage(const UPackage* InPackage)
 {
 	if (IsRunningCommandlet())
 	{
diff --git a/Engine/Source/Runtime/Engine/Public/DeletedObjectPlaceholder.h b/Engine/Source/Runtime/Engine/Public/DeletedObjectPlaceholder.h
--- a/Engine/Source/Runtime/Engine/Public/DeletedObjectPlaceholder.h
+++ b/Engine/Source/Runtime/Engine/Public/DeletedObjectPlaceholder.h
@@ -30,7 +30,7 @@ public:
 	const FExternalDataLayerUID GetExternalDataLayerUID() const { return FExternalDataLayerUID(ExternalDataLayerUID); }
 
 	static ENGINE_API UDeletedObjectPlaceholder* Create(UObject* InOuter, UPackage* InPackage, const UObject* InOriginalObject);
-	static ENGINE_API UDeletedObjectPlaceholder* FindInPackage(UPackage* InPackage);
+	static ENGINE_API UDeletedObjectPlaceholder* FindInPackage(const UPackage* InPackage);
 	static ENGINE_API UDeletedObjectPlaceholder* RemoveFromPackage(UPackage* InPackage);
 
 	DECLARE_MULTICAST_DELEGATE_OneParam(FObjectCreated, const UDeletedObjectPlaceholder*);
