Moved r.MaterialEditor.AllowIgnoringCompilationErrors and r.Material.ContextMenu.* CVars to MaterialEditorSettings.

diff --git a/Engine/Source/Editor/MaterialEditor/Private/MaterialEditor.cpp b/Engine/Source/Editor/MaterialEditor/Private/MaterialEditor.cpp
--- a/Engine/Source/Editor/MaterialEditor/Private/MaterialEditor.cpp
+++ b/Engine/Source/Editor/MaterialEditor/Private/MaterialEditor.cpp
@@ -53,6 +53,7 @@
 #include "MaterialCachedData.h"
 #include "DataDrivenShaderPlatformInfo.h"
 #include "Algo/ForEach.h"
+#include "MaterialEditorSettings.h"
 
 #include "Materials/MaterialExpressionBreakMaterialAttributes.h"
 #include "Materials/MaterialExpressionCollectionParameter.h"
@@ -195,12 +196,7 @@ static TAutoConsoleVariable<int32> CVarMaterialEdMaxDerivedMaterialInstances(
 	TEXT("r.MaterialEditor.MaxDerivedMaterialInstances"),
 	-1,
 	TEXT("Limits amount of derived material instance shown in platform stats. Use negative number to disable the limit. Material editor must be re-opened if changed at runtime."));
-
-TAutoConsoleVariable<bool> CVarMaterialEdAllowIgnoringCompilationErrors(
-	TEXT("r.MaterialEditor.AllowIgnoringCompilationErrors"),
-	true,
-	TEXT("Allow ignoring compilation errors of platform shaders and derived materials."));
-
+ 
 static void InvalidateSubstrateConversionVersion(UMaterial* Material)
 {
 	// The material has been updated, invalid the AutoConversion version
@@ -607,7 +603,7 @@ void FMaterialEditor::InitMaterialEditor( const EToolkitMode::Type Mode, const T
 
 	GEditor->RegisterForUndo(this);
 
-	MaterialStatsManager = FMaterialStatsUtils::CreateMaterialStats(this, true, CVarMaterialEdAllowIgnoringCompilationErrors.GetValueOnGameThread());
+	MaterialStatsManager = FMaterialStatsUtils::CreateMaterialStats(this, true, GetDefault<UMaterialEditorSettings>()->bAllowIgnoringCompilationErrors);
 	MaterialStatsManager->SetMaterialsDisplayNames({OriginalMaterial->GetName()});
 	MaterialStatsManager->GetOldStatsListing()->OnMessageTokenClicked().AddSP(this, &FMaterialEditor::OnMessageLogLinkActivated);
 
diff --git a/Engine/Source/Editor/MaterialEditor/Private/MaterialInstanceEditor.cpp b/Engine/Source/Editor/MaterialEditor/Private/MaterialInstanceEditor.cpp
--- a/Engine/Source/Editor/MaterialEditor/Private/MaterialInstanceEditor.cpp
+++ b/Engine/Source/Editor/MaterialEditor/Private/MaterialInstanceEditor.cpp
@@ -74,8 +74,6 @@ const FName FMaterialInstanceEditor::LayerPropertiesTabId(TEXT("MaterialInstance
 const FName FMaterialInstanceEditor::PreviewSettingsTabId(TEXT("MaterialInstanceEditor_PreviewSettings"));
 const FName FMaterialInstanceEditor::AssetBrowserTabId(TEXT("MaterialInstanceEditor_AssetBrowser"));
 
-extern TAutoConsoleVariable<bool> CVarMaterialEdAllowIgnoringCompilationErrors;
-
 //////////////////////////////////////////////////////////////////////////
 // SMaterialTreeWidgetItem
 class SMaterialTreeWidgetItem : public SMultiColumnTableRow< TWeakObjectPtr<UMaterialInterface> >
@@ -434,7 +432,7 @@ void FMaterialInstanceEditor::InitMaterialInstanceEditor( const EToolkitMode::Ty
 	MaterialEditorInstance->SetSourceInstance(InstanceConstant);
 	MaterialEditorInstance->SetSourceFunction(MaterialFunctionOriginal);
 
-	MaterialStatsManager = FMaterialStatsUtils::CreateMaterialStats(this, false, CVarMaterialEdAllowIgnoringCompilationErrors.GetValueOnGameThread());
+	MaterialStatsManager = FMaterialStatsUtils::CreateMaterialStats(this, false, GetDefault<UMaterialEditorSettings>()->bAllowIgnoringCompilationErrors);
 	MaterialStatsManager->SetMaterialsDisplayNames({MaterialEditorInstance->SourceInstance->GetName()});
 
 	// Register our commands. This will only register them if not previously registered
diff --git a/Engine/Source/Editor/MaterialEditor/Public/MaterialEditorSettings.h b/Engine/Source/Editor/MaterialEditor/Public/MaterialEditorSettings.h
--- a/Engine/Source/Editor/MaterialEditor/Public/MaterialEditorSettings.h
+++ b/Engine/Source/Editor/MaterialEditor/Public/MaterialEditorSettings.h
@@ -105,6 +105,49 @@ class MATERIALEDITOR_API UMaterialEditorSettings : public UObject
 	GENERATED_BODY()
 
 public:
+	
+	/**
+	Allow ignoring compilation errors of platform shaders and derived materials.
+	*/
+	UPROPERTY(config, EditAnywhere, Category = "Editor Defaults", meta = (DisplayName = "Allow Ignoring Compilation Errors"))
+	bool bAllowIgnoringCompilationErrors = true;
+
+	/** The amount of weight placed on the search items title". */
+	UPROPERTY(config, EditAnywhere, Category = "Context Menu", meta = (DisplayName = "Node Title Weight"))
+	float ContextMenuNodeTitleWeight = 20.0f;
+
+	/** The amount of weight placed on search items keyword. */
+	UPROPERTY(config, EditAnywhere, Category = "Context Menu", meta = (DisplayName = "Keyword Weight"))
+	float ContextMenuKeywordWeight = 30.0f;
+
+	/**  amount of weight placed on description that match what the user has typed in. */
+	UPROPERTY(config, EditAnywhere, Category = "Context Menu", meta = (DisplayName = "Menu description Weight"))
+	float ContextMenuDescriptionWeight = 4.0f;
+
+	/** The amount of weight placed on categories that match what the user has typed. */
+	UPROPERTY(config, EditAnywhere, Category = "Context Menu", meta = (DisplayName = "Category Weight"))
+	float ContextMenuCategoryWeight = 4.0f;
+
+	/** The multiplier given if there is an exact localized match to the search. */
+	UPROPERTY(config, EditAnywhere, Category = "Context Menu", meta = (DisplayName = "Whole Match Localized Weight Multiplier"))
+	float ContextMenuWholeMatchLocalizedWeightMultiplier = 0.5f;
+
+	/** The multiplier given if there is an exact match to the search term. */
+	UPROPERTY(config, EditAnywhere, Category = "Context Menu", meta = (DisplayName = "Whole Match Weight Multiplier"))
+	float ContextMenuWholeMatchWeightMultiplier = 0.5f;
+
+	/** The multiplier given if the keyword starts with a term the user typed in. */
+	UPROPERTY(config, EditAnywhere, Category = "Context Menu", meta = (DisplayName = "Starts With Bonus Weight Multiplier"))
+	float ContextMenuStartsWithBonusWeightMultiplier = 4.0f;
+
+	/** A multiplier for how much weight to give something based on the percentage match it is. */
+	UPROPERTY(config, EditAnywhere, Category = "Context Menu", meta = (DisplayName = "Percentage Match Weight Multiplier"))
+	float ContextMenuPercentageMatchWeightMultiplier = 1.0f;
+
+	/** Increasing this weight will make shorter words preferred. */
+	UPROPERTY(config, EditAnywhere, Category = "Context Menu", meta = (DisplayName = "Shorter Match Weight"))
+	float ContextMenuShorterMatchWeight = 10.0f;
+
 	/**
 	Offline shader compiler to use.
 	Mali Offline Compiler: Official website address: https://developer.arm.com/products/software-development-tools/graphics-development-tools/mali-offline-compiler/downloads
@@ -139,6 +182,7 @@ public:
 	*/
 	UPROPERTY(config, EditAnywhere, Category = "Offline Shader Compilers", meta = (DisplayName = "Dump All Stats To Compiler Stats Files", EditCondition = "OfflineCompiler == EOfflineShaderCompiler::Adreno"))
 	bool bDumpAll;
+
 protected:
 	// The width (in pixels) of the preview viewport when a material editor is first opened
 	UPROPERTY(config, EditAnywhere, meta=(ClampMin=1, ClampMax=4096), Category="User Interface Domain")
diff --git a/Engine/Source/Editor/UnrealEd/Private/MaterialGraphSchema.cpp b/Engine/Source/Editor/UnrealEd/Private/MaterialGraphSchema.cpp
--- a/Engine/Source/Editor/UnrealEd/Private/MaterialGraphSchema.cpp
+++ b/Engine/Source/Editor/UnrealEd/Private/MaterialGraphSchema.cpp
@@ -47,6 +47,7 @@
 #include "MaterialEditorActions.h"
 #include "MaterialGraphNode_Knot.h"
 #include "RenderUtils.h"
+#include "MaterialEditorSettings.h"
 
 #define LOCTEXT_NAMESPACE "MaterialGraphSchema"
 
@@ -1239,85 +1240,20 @@ void UMaterialGraphSchema::GetAssetsGraphHoverMessage(const TArray<FAssetData>&
 
 #if WITH_EDITORONLY_DATA
 //////////////////////////////////////////////////////////////////////////
-/** CVars for tweaking how the material editor context menu search picks the best match */
-namespace MaterialEditorContextMenuConsoleVariables
-{
-	/** How much weight the node's title has */
-	static float NodeTitleWeight = 20.0f;
-	static FAutoConsoleVariableRef CVarNodeTitleWeight(
-		TEXT("r.MaterialEditor.ContextMenu.NodeTitleWeight"), NodeTitleWeight,
-		TEXT("The amount of weight placed on the search items title"),
-		ECVF_Default);
-
-	/** Weight used to prefer keywords of actions  */
-	static float KeywordWeight = 30.0f;
-	static FAutoConsoleVariableRef CVarKeywordWeight(
-		TEXT("r.MaterialEditor.ContextMenu.KeywordWeight"), KeywordWeight,
-		TEXT("The amount of weight placed on search items keyword"),
-		ECVF_Default);
-
-	/** Weight that a match to a description search has */
-	static float DescriptionWeight = 4.0f;
-	static FAutoConsoleVariableRef CVarDescriptionWeight(
-		TEXT("r.MaterialEditor.ContextMenu.DescriptionWeight"), DescriptionWeight,
-		TEXT("The amount of weight placed on description that match what the user has typed in"),
-		ECVF_Default);
-
-	/** Weight that a match to a category search has */
-	static float CategoryWeight = 4.0f;
-	static FAutoConsoleVariableRef CVarCategoryWeight(
-		TEXT("r.MaterialEditor.ContextMenu.CategoryWeight"), CategoryWeight,
-		TEXT("The amount of weight placed on categories that match what the user has typed in"),
-		ECVF_Default);
-
-	/** The multiplier given if there is an exact localized match to the search term */
-	static float WholeMatchLocalizedWeightMultiplier = 0.5f;
-	static FAutoConsoleVariableRef CVarWholeMatchLocalizedWeightMultiplier(
-		TEXT("r.MaterialEditor.ContextMenu.WholeMatchLocalizedWeightMultiplier"), WholeMatchLocalizedWeightMultiplier,
-		TEXT("The multiplier given if there is an exact localized match to the search term"),
-		ECVF_Default);
-
-	/** The multiplier given if there is an exact match to the search term */
-	static float WholeMatchWeightMultiplier = 0.5f;
-	static FAutoConsoleVariableRef CVarWholeMatchWeightMultiplier(
-		TEXT("r.MaterialEditor.ContextMenu.WholeMatchWeightMultiplier"), WholeMatchWeightMultiplier,
-		TEXT("The multiplier given if there is an exact match to the search term"),
-		ECVF_Default);
-
-	/** The multiplier given if the keyword starts with a term the user typed in */
-	static float StartsWithBonusWeightMultiplier = 4.0f;
-	static FAutoConsoleVariableRef CVarStartsWithBonusWeightMultiplier(
-		TEXT("r.MaterialEditor.ContextMenu.StartsWithBonusWeightMultiplier"), StartsWithBonusWeightMultiplier,
-		TEXT("The multiplier given if the keyword starts with a term the user typed in"),
-		ECVF_Default);
-
-	/** Increasing this will prefer whole percentage matches when comparing the keyword to what the user has typed in */
-	static float PercentageMatchWeightMultiplier = 1.0f;
-	static FAutoConsoleVariableRef CVarPercentageMatchWeightMultiplier(
-		TEXT("r.MaterialEditor.ContextMenu.PercentageMatchWeightMultiplier"), PercentageMatchWeightMultiplier,
-		TEXT("A multiplier for how much weight to give something based on the percentage match it is"),
-		ECVF_Default);
-
-	/** Increasing this weight will give a bonus to shorter matching words */
-	static float ShorterMatchWeight = 10.0f;
-	static FAutoConsoleVariableRef CVarShorterWeight(
-		TEXT("r.MaterialEditor.ContextMenu.ShorterMatchWeight"), ShorterMatchWeight,
-		TEXT("Increasing this weight will make shorter words preferred"),
-		ECVF_Default);
-}
 
 FGraphSchemaSearchWeightModifiers UMaterialGraphSchema::GetSearchWeightModifiers() const
 {
+	const UMaterialEditorSettings* MaterialSettings = GetDefault<UMaterialEditorSettings>();
 	FGraphSchemaSearchWeightModifiers Modifiers;
-	Modifiers.NodeTitleWeight = MaterialEditorContextMenuConsoleVariables::NodeTitleWeight;
-	Modifiers.KeywordWeight = MaterialEditorContextMenuConsoleVariables::KeywordWeight;
-	Modifiers.DescriptionWeight = MaterialEditorContextMenuConsoleVariables::DescriptionWeight;
-	Modifiers.CategoryWeight = MaterialEditorContextMenuConsoleVariables::DescriptionWeight;
-	Modifiers.WholeMatchLocalizedWeightMultiplier = MaterialEditorContextMenuConsoleVariables::WholeMatchLocalizedWeightMultiplier;
-	Modifiers.WholeMatchWeightMultiplier = MaterialEditorContextMenuConsoleVariables::WholeMatchWeightMultiplier;
-	Modifiers.StartsWithBonusWeightMultiplier = MaterialEditorContextMenuConsoleVariables::StartsWithBonusWeightMultiplier;
-	Modifiers.PercentageMatchWeightMultiplier = MaterialEditorContextMenuConsoleVariables::PercentageMatchWeightMultiplier;
-	Modifiers.ShorterMatchWeight = MaterialEditorContextMenuConsoleVariables::ShorterMatchWeight;
+	Modifiers.NodeTitleWeight = MaterialSettings->ContextMenuNodeTitleWeight;
+	Modifiers.KeywordWeight = MaterialSettings->ContextMenuKeywordWeight;
+	Modifiers.DescriptionWeight = MaterialSettings->ContextMenuDescriptionWeight;
+	Modifiers.CategoryWeight = MaterialSettings->ContextMenuDescriptionWeight;
+	Modifiers.WholeMatchLocalizedWeightMultiplier = MaterialSettings->ContextMenuWholeMatchLocalizedWeightMultiplier;
+	Modifiers.WholeMatchWeightMultiplier = MaterialSettings->ContextMenuWholeMatchWeightMultiplier;
+	Modifiers.StartsWithBonusWeightMultiplier = MaterialSettings->ContextMenuStartsWithBonusWeightMultiplier;
+	Modifiers.PercentageMatchWeightMultiplier = MaterialSettings->ContextMenuPercentageMatchWeightMultiplier;
+	Modifiers.ShorterMatchWeight = MaterialSettings->ContextMenuShorterMatchWeight;
 	return Modifiers;
 }
 
@@ -1333,6 +1269,8 @@ float UMaterialGraphSchema::GetActionFilteredWeight(const FEdGraphSchemaAction&
 	FGraphSchemaSearchWeightModifiers WeightModifiers = GetSearchWeightModifiers();
 	int32 NonLocalizedFirstIndex = CollectSearchTextWeightInfo(InCurrentAction, WeightModifiers, WeightedArrayList, &DebugInfo);
 
+	const UMaterialEditorSettings* MaterialSettings = GetDefault<UMaterialEditorSettings>();
+
 	// Now iterate through all the filter terms and calculate a 'weight' using the values and multipliers
 	for (int32 FilterIndex = 0; FilterIndex < InFilterTerms.Num(); ++FilterIndex)
 	{
@@ -1345,7 +1283,7 @@ float UMaterialGraphSchema::GetActionFilteredWeight(const FEdGraphSchemaAction&
 			const TArray<FString>& WordArray = *WeightedArrayList[iFindCount].Array;
 			float ArrayWeight = WeightedArrayList[iFindCount].WeightModifier;
 			int32 WholeMatchCount = 0;
-			float WholeMatchMultiplier = (iFindCount < NonLocalizedFirstIndex) ? MaterialEditorContextMenuConsoleVariables::WholeMatchLocalizedWeightMultiplier : MaterialEditorContextMenuConsoleVariables::WholeMatchWeightMultiplier;
+			float WholeMatchMultiplier = (iFindCount < NonLocalizedFirstIndex) ? MaterialSettings->ContextMenuWholeMatchLocalizedWeightMultiplier : MaterialSettings->ContextMenuWholeMatchWeightMultiplier;
 
 			// Count of how many words in this array contain a search term that the user has typed in
 			int32 WordMatchCount = 0;
@@ -1365,7 +1303,7 @@ float UMaterialGraphSchema::GetActionFilteredWeight(const FEdGraphSchemaAction&
 					// If the word starts with the search term, give it extra boost of weight
 					if (WordArray[iEachWord].StartsWith(EachTermSanitized, ESearchCase::CaseSensitive) || WordArray[iEachWord].StartsWith(EachTerm, ESearchCase::CaseSensitive))
 					{
-						WeightPerWord += ArrayWeight * MaterialEditorContextMenuConsoleVariables::StartsWithBonusWeightMultiplier;
+						WeightPerWord += ArrayWeight * MaterialSettings->ContextMenuStartsWithBonusWeightMultiplier;
 					}
 				}
 
@@ -1381,14 +1319,14 @@ float UMaterialGraphSchema::GetActionFilteredWeight(const FEdGraphSchemaAction&
 			{
 				// Higher number of matching words contributes to higher weight
 				float PercentMatch = (float)WordMatchCount / (float)WordArray.Num();
-				float PercentMatchWeight = (WeightPerList * PercentMatch * MaterialEditorContextMenuConsoleVariables::PercentageMatchWeightMultiplier);
+				float PercentMatchWeight = (WeightPerList * PercentMatch * MaterialSettings->ContextMenuPercentageMatchWeightMultiplier);
 				WeightPerList += PercentMatchWeight;
 				DebugInfo.PercentMatchWeight += PercentMatchWeight;
 				DebugInfo.PercentMatch += PercentMatch;
 
 				// The shorter the best matched word, the larger bonus it gets
 				float ShorterMatchFactor = (float)EachTerm.Len() / (float)BestMatchCharLength;
-				float ShorterMatchWeight = ShorterMatchFactor * MaterialEditorContextMenuConsoleVariables::ShorterMatchWeight;
+				float ShorterMatchWeight = ShorterMatchFactor * MaterialSettings->ContextMenuShorterMatchWeight;
 				WeightPerList += ShorterMatchWeight;
 				DebugInfo.ShorterMatchWeight += ShorterMatchWeight;
 			}
