Improved UX of using the uncontrolled changelist for auto-save

This change updates the auto-save uncontrolled changelist to be a purgatory that packages should be removed from when a user saves (or doesn't) their pending changes.
This can be controlled by the "Editor.AutoSave.InteractiveSaveManagesUncontrolledChangelist" CVar and is current disabled by default, but will be enabled in the near future.

With this mode enabled auto-saved packages are still saved and added to the auto-save uncontrolled changelist, but:
 * Anything in the auto-save uncontrolled changelist is always marked dirty.
 * Anything in the auto-save uncontrolled changelist is considered read-only for the interactive save dialog, meaning a user gets the choice about whether to check out the package or make it writable.
  * If checking them out, the packages are removed from the auto-saved uncontrolled changelist, and moved to the default controlled changelist.
  * If making them writable, the packages are removed from the auto-saved uncontrolled changelist, and moved to the default uncontrolled changelist.
 * Anything in the auto-save uncontrolled changelist will be discarded on-disk if a user chooses to "Don't Save" on a dialog that allows that (eg, when closing the map, or closing the editor).
  * Any changes will still be in-memory (unless unloaded), so can still be saved later if needed.

These changes make the uncontrolled workflow behave very similarly to the standard UE workflow for an end-user, despite the auto-save itself directly overwriting the on-disk project data so that it isn't lost in the case of a crash.

#rnx

diff --git a/Engine/Source/Developer/UncontrolledChangelists/Private/UncontrolledChangelistsModule.cpp b/Engine/Source/Developer/UncontrolledChangelists/Private/UncontrolledChangelistsModule.cpp
--- a/Engine/Source/Developer/UncontrolledChangelists/Private/UncontrolledChangelistsModule.cpp
+++ b/Engine/Source/Developer/UncontrolledChangelists/Private/UncontrolledChangelistsModule.cpp
@@ -216,6 +216,11 @@ FUncontrolledChangelistStatePtr FUncontrolledChangelistsModule::GetDefaultChange
 	return GetChangelistState(FUncontrolledChangelist(FUncontrolledChangelist::DEFAULT_UNCONTROLLED_CHANGELIST_GUID));
 }
 
+void FUncontrolledChangelistsModule::HandleChangelistStateModified()
+{
+	OnStateChanged();
+}
+
 bool FUncontrolledChangelistsModule::OnMakeWritable(const FString& InFilename)
 {
 	if (!IsEnabled())
@@ -534,14 +539,17 @@ void FUncontrolledChangelistsModule::MoveFilesToUncontrolledChangelist(const TAr
 	}
 
 	TArray<FString> Filenames;
-	Algo::Transform(InControlledFileStates, Filenames, [](const FSourceControlStateRef& State) { return State->GetFilename(); });
+	if (InControlledFileStates.Num() > 0)
+	{
+		Algo::Transform(InControlledFileStates, Filenames, [](const FSourceControlStateRef& State) { return State->GetFilename(); });
 
-	ISourceControlProvider& SourceControlProvider = ISourceControlModule::Get().GetProvider();
-	TSharedRef<FRevert, ESPMode::ThreadSafe> RevertOperation = ISourceControlOperation::Create<FRevert>();
+		ISourceControlProvider& SourceControlProvider = ISourceControlModule::Get().GetProvider();
+		TSharedRef<FRevert, ESPMode::ThreadSafe> RevertOperation = ISourceControlOperation::Create<FRevert>();
 
-	// Revert controlled files
-	RevertOperation->SetSoftRevert(true);
-	SourceControlProvider.Execute(RevertOperation, Filenames);
+		// Revert controlled files
+		RevertOperation->SetSoftRevert(true);
+		SourceControlProvider.Execute(RevertOperation, Filenames);
+	}
 
 	// Removes selected Uncontrolled Files from their Uncontrolled Changelists
 	for (const auto& Pair : UncontrolledChangelistsStateCache)
diff --git a/Engine/Source/Developer/UncontrolledChangelists/Public/UncontrolledChangelistsModule.h b/Engine/Source/Developer/UncontrolledChangelists/Public/UncontrolledChangelistsModule.h
--- a/Engine/Source/Developer/UncontrolledChangelists/Public/UncontrolledChangelistsModule.h
+++ b/Engine/Source/Developer/UncontrolledChangelists/Public/UncontrolledChangelistsModule.h
@@ -60,6 +60,11 @@ public:
 	 */
 	FUncontrolledChangelistStatePtr GetDefaultChangelistState() const;
 
+	/**
+	 * Called if the state of any Uncontrolled Changelist is modified externally, eg, via the mutable accessors above.
+	 */
+	void HandleChangelistStateModified();
+
 	/**
 	 * Called when file has been made writable. Adds the file to the reconcile list because we don't know yet if it will change.
 	 * @param	InFilename		The file to be reconciled.
@@ -102,6 +107,15 @@ public:
 		return FModuleManager::LoadModuleChecked<FUncontrolledChangelistsModule>(GetModuleName());
 	}
 
+	/**
+	 * Gets a pointer to the UncontrolledChangelists module, if loaded
+	 * @return A pointer to the UncontrolledChangelists module.
+	 */
+	static inline FUncontrolledChangelistsModule* GetPtr()
+	{
+		return FModuleManager::GetModulePtr<FUncontrolledChangelistsModule>(GetModuleName());
+	}
+
 	/**
 	 * Checks to see if this module is loaded and ready.  It is only valid to call Get() if IsAvailable() returns true.
 	 *
diff --git a/Engine/Source/Editor/UnrealEd/Private/FileHelpers.cpp b/Engine/Source/Editor/UnrealEd/Private/FileHelpers.cpp
--- a/Engine/Source/Editor/UnrealEd/Private/FileHelpers.cpp
+++ b/Engine/Source/Editor/UnrealEd/Private/FileHelpers.cpp
@@ -2,6 +2,7 @@
 
 
 #include "FileHelpers.h"
+#include "FileHelpersInternal.h"
 #include "HAL/PlatformFileManager.h"
 #include "Misc/MessageDialog.h"
 #include "HAL/FileManager.h"
@@ -1642,7 +1643,7 @@ static bool IsCheckOutSelectedDisabled()
 	return !(ISourceControlModule::Get().IsEnabled() && ISourceControlModule::Get().GetProvider().IsAvailable()) || !ISourceControlModule::Get().GetProvider().UsesCheckout();
 }
 
-bool FEditorFileUtils::AddCheckoutPackageItems(bool bCheckDirty, TArray<UPackage*> PackagesToCheckOut, TArray<UPackage*>* OutPackagesNotNeedingCheckout, bool* bOutHavePackageToCheckOut)
+bool FEditorFileUtils::AddCheckoutPackageItems(bool bCheckDirty, const TArray<UPackage*>& PackagesToCheckOut, const TSet<FName>& ReadOnlyPackages, TArray<UPackage*>* OutPackagesNotNeedingCheckout, bool* bOutHavePackageToCheckOut)
 {
 	TRACE_CPUPROFILER_EVENT_SCOPE(FEditorFileUtils_AddCheckoutPackageItems);
 
@@ -1721,7 +1722,7 @@ bool FEditorFileUtils::AddCheckoutPackageItems(bool bCheckDirty, TArray<UPackage
 		if (bFoundFile)
 		{
 			// determine if the package file is read only
-			bPkgReadOnly = IFileManager::Get().IsReadOnly(*Filename);
+			bPkgReadOnly = ReadOnlyPackages.Contains(CurPackage->GetFName()) || IFileManager::Get().IsReadOnly(*Filename);
 		}
 
 		FSourceControlStatePtr SourceControlState = SourceControlProvider.GetState(CurPackage, EStateCacheUsage::Use);
@@ -1840,9 +1841,9 @@ bool FEditorFileUtils::AddCheckoutPackageItems(bool bCheckDirty, TArray<UPackage
 	return bPackagesAdded;
 }
 
-void FEditorFileUtils::UpdateCheckoutPackageItems(bool bCheckDirty, TArray<UPackage*> PackagesToCheckOut, TArray<UPackage*>* OutPackagesNotNeedingCheckout)
+void FEditorFileUtils::UpdateCheckoutPackageItems(bool bCheckDirty, TArray<UPackage*> PackagesToCheckOut, TSet<FName> ReadOnlyPackages, TArray<UPackage*>* OutPackagesNotNeedingCheckout)
 {
-	AddCheckoutPackageItems(bCheckDirty, PackagesToCheckOut, OutPackagesNotNeedingCheckout, nullptr);
+	AddCheckoutPackageItems(bCheckDirty, PackagesToCheckOut, ReadOnlyPackages, OutPackagesNotNeedingCheckout, nullptr);
 }
 
 bool FEditorFileUtils::PromptToCheckoutPackages(bool bCheckDirty, const TArray<UPackage*>& PackagesToCheckOut, TArray<UPackage*>* OutPackagesCheckedOutOrMadeWritable, TArray<UPackage*>* OutPackagesNotNeedingCheckout, const bool bPromptingAfterModify, const bool bAllowSkip)
@@ -1868,11 +1869,13 @@ bool FEditorFileUtils::PromptToCheckoutPackages(bool bCheckDirty, const TArray<U
 
 bool FEditorFileUtils::PromptToCheckoutPackagesInternal(bool bCheckDirty, const TArray<UPackage*>& PackagesToCheckOut, TArray<UPackage*>* OutPackagesCheckedOutOrMadeWritable, TArray<UPackage*>* OutPackagesNotNeedingCheckout, const bool bPromptingAfterModify, const bool bAllowSkip )
 {
-	bool bResult = true;
+	bool bUserResponse = true;
 
 	ISourceControlProvider& SourceControlProvider = ISourceControlModule::Get().GetProvider();
-	FUncontrolledChangelistsModule& UncontrolledChangelistModule = FUncontrolledChangelistsModule::Get();
 	
+	TSet<FName> ReadOnlyPackages;
+	UE::FileHelpers::Internal::OnPreInteractiveCheckoutPackages.Broadcast(PackagesToCheckOut, ReadOnlyPackages);
+
 	// The checkout dialog to show users if any packages need to be checked out
 	const FText DialogTitle = NSLOCTEXT("PackagesDialogModule", "CheckoutPackagesDialogTitle", "Check Out Assets");
 	const FText DialogHeading = NSLOCTEXT("PackagesDialogModule", "CheckoutPackagesDialogMessage", "Select assets to check out.");
@@ -1893,7 +1896,7 @@ bool FEditorFileUtils::PromptToCheckoutPackagesInternal(bool bCheckDirty, const
 		DialogHeading,
 		bReadOnly,
 		bAllowSourceControlConnection,
-		FSimpleDelegate::CreateStatic(&UpdateCheckoutPackageItems, bCheckDirty, PackagesToCheckOut, OutPackagesNotNeedingCheckout)
+		FSimpleDelegate::CreateStatic(&UpdateCheckoutPackageItems, bCheckDirty, PackagesToCheckOut, ReadOnlyPackages, OutPackagesNotNeedingCheckout)
 		);
 
 	// If we got here and we have one package, it's because someone explicitly saved the asset, therefore remove the package from the ignore list.
@@ -1903,7 +1906,7 @@ bool FEditorFileUtils::PromptToCheckoutPackagesInternal(bool bCheckDirty, const
 		PackagesNotSavedDuringSaveAll.Remove(PackageName);
 	}
 
-	bPackagesAdded = AddCheckoutPackageItems(bCheckDirty, PackagesToCheckOut, OutPackagesNotNeedingCheckout, &bHavePackageToCheckOut);
+	bPackagesAdded = AddCheckoutPackageItems(bCheckDirty, PackagesToCheckOut, ReadOnlyPackages, OutPackagesNotNeedingCheckout, &bHavePackageToCheckOut);
 
 	// If any packages were added to the dialog, show the dialog to the user and allow them to select which files to check out
 	if ( bPackagesAdded )
@@ -1959,6 +1962,8 @@ bool FEditorFileUtils::PromptToCheckoutPackagesInternal(bool bCheckDirty, const
 				}
 				else
 				{
+					UE::FileHelpers::Internal::OnPackagesInteractivelyCheckedOut.Broadcast(PkgsToCheckOut);
+
 					bPerformedOperation = true;
 				}
 			}
@@ -1972,22 +1977,25 @@ bool FEditorFileUtils::PromptToCheckoutPackagesInternal(bool bCheckDirty, const
 
 				MakePackagesWritable(PkgsToMakeWritable, OutPackagesCheckedOutOrMadeWritable, nullptr);
 
+				UE::FileHelpers::Internal::OnPackagesInteractivelyMadeWritable.Broadcast(PkgsToMakeWritable);
+
 				bPerformedOperation = true;
 			}
 			else if (UserResponse == DRT_Save || UserResponse == DRT_Skip)
 			{
-				bResult = true;
 				bPerformedOperation = true;
 			}
 			else if (UserResponse == DRT_Cancel || UserResponse == DRT_None)
 			{
 				// Handle the case of the user canceling out of the dialog
-				bResult = false;
+				bUserResponse = false;
 				bPerformedOperation = true;
 			}
 		}
 	}
 
+	UE::FileHelpers::Internal::OnPostInteractiveCheckoutPackages.Broadcast(PackagesToCheckOut, bUserResponse);
+
 	// Update again to catch potentially new SCC states
 	ISourceControlModule::Get().QueueStatusUpdate(PackagesToCheckOut);
 
@@ -2005,7 +2013,7 @@ bool FEditorFileUtils::PromptToCheckoutPackagesInternal(bool bCheckDirty, const
 		ISourceControlModule::Get().QueueStatusUpdate(*OutPackagesNotNeedingCheckout);
 	}
 
-	return bResult;
+	return bUserResponse;
 }
 
 void FEditorFileUtils::MakePackagesWritable(const TArray<UPackage*>& PackagesToMakeWritable, TArray<UPackage*>* OutPackagesMadeWritable, TArray<UPackage*>* OutPackagesMadeWritableFailed)
@@ -4515,6 +4523,8 @@ FEditorFileUtils::EPromptReturnCode FEditorFileUtils::PromptForCheckoutAndSave(c
 			// If the user has responded they don't wish to save, set the response type accordingly
 			else if ( UserResponse == DRT_DontSave )
 			{
+				UE::FileHelpers::Internal::OnPackagesInteractivelyDiscarded.Broadcast(PackagesToSave);
+
 				ReturnResponse = PR_Declined;
 			}
 			// If the user has cancelled from the dialog, set the response type accordingly
@@ -5213,6 +5223,7 @@ void FEditorFileUtils::GetDirtyWorldPackages(TArray<UPackage*>& OutDirtyPackages
 	const TSharedRef<FPathPermissionList>& WritableFolderFilter = AssetToolsModule.Get().GetWritableFolderPermissionList();
 	const bool bHasWritableFolderFilter = WritableFolderFilter->HasFiltering();
 
+	TOptional<TSet<UPackage*>> AdditionalSaveCandidates; // Lazily queried if needed by the loop below
 	for (TObjectIterator<UWorld> WorldIt; WorldIt; ++WorldIt)
 	{
 		// Filter out pending-delete worlds that may have leaked, e.g. from PIE sessions which were not cleaned up which cleared the PKG_PlayInEditor flag
@@ -5295,10 +5306,15 @@ void FEditorFileUtils::GetDirtyWorldPackages(TArray<UPackage*>& OutDirtyPackages
 						{
 							bool bActorPackageNeedsToSave = true;
 
-							// Skip unsaved packages containing only pending kill actors
-							if (ExternalPackage->HasAnyPackageFlags(PKG_NewlyCreated))
+							// Skip unsaved empty actor packages, unless we've been asked to include them
+							if (ExternalPackage->HasAnyPackageFlags(PKG_NewlyCreated) && UPackage::IsEmptyPackage(ExternalPackage))
 							{
-								bActorPackageNeedsToSave = IsValid(ExternalPackage->FindAssetInPackage());
+								if (!AdditionalSaveCandidates)
+								{
+									AdditionalSaveCandidates.Emplace();
+									UE::FileHelpers::Internal::GetAdditionalInteractiveSavePackageCandidates.Broadcast(AdditionalSaveCandidates.GetValue());
+								}
+								bActorPackageNeedsToSave &= AdditionalSaveCandidates->Contains(ExternalPackage);
 							}
 
 							// Filter out Actors that might be unsaved (/Temp folder)
diff --git a/Engine/Source/Editor/UnrealEd/Private/FileHelpersInternal.cpp b/Engine/Source/Editor/UnrealEd/Private/FileHelpersInternal.cpp
new file mode 100644
--- /dev/null
+++ b/Engine/Source/Editor/UnrealEd/Private/FileHelpersInternal.cpp
@@ -0,0 +1,15 @@
+// Copyright Epic Games, Inc. All Rights Reserved.
+
+#include "FileHelpersInternal.h"
+
+namespace UE::FileHelpers::Internal
+{
+
+FGetAdditionalInteractiveSavePackageCandidates GetAdditionalInteractiveSavePackageCandidates;
+FOnPreInteractiveCheckoutPackages OnPreInteractiveCheckoutPackages;
+FOnPostInteractiveCheckoutPackages OnPostInteractiveCheckoutPackages;
+FOnPackagesInteractivelyCheckedOut OnPackagesInteractivelyCheckedOut;
+FOnPackagesInteractivelyMadeWritable OnPackagesInteractivelyMadeWritable;
+FOnPackagesInteractivelyDiscarded OnPackagesInteractivelyDiscarded;
+
+}
diff --git a/Engine/Source/Editor/UnrealEd/Private/FileHelpersInternal.h b/Engine/Source/Editor/UnrealEd/Private/FileHelpersInternal.h
new file mode 100644
--- /dev/null
+++ b/Engine/Source/Editor/UnrealEd/Private/FileHelpersInternal.h
@@ -0,0 +1,55 @@
+// Copyright Epic Games, Inc. All Rights Reserved.
+
+#pragma once
+
+#include "Delegates/Delegate.h"
+
+namespace UE::FileHelpers::Internal
+{
+
+/**
+ * Low-level delegate called by FEditorFileUtils::GetDirtyXPackages to collect additional candidates to save if dirty.
+ * @note This doesn't forcibly add any candidates to the dirty/save lists, but may allow them to bypass certain filtering that would have otherwise removed them (eg, to allow an empty actor package).
+ * @param OutPackages The set of additional candidates, to be filled by the callback.
+ */
+using FGetAdditionalInteractiveSavePackageCandidates = TMulticastDelegate<void(TSet<UPackage*>& OutPackages)>;
+extern FGetAdditionalInteractiveSavePackageCandidates GetAdditionalInteractiveSavePackageCandidates;
+
+/**
+ * Low-level delegate called by FEditorFileUtils::PromptToCheckoutPackagesInternal prior to deciding which packages need check out.
+ * @param Packages The list of packages that will be considered for checkout.
+ * @param OutReadOnlyPackages The set of packages to be considered read-only by FEditorFileUtils::AddCheckoutPackageItems, even if they're not actually read-only on disk.
+ */
+using FOnPreInteractiveCheckoutPackages = TMulticastDelegate<void(const TArray<UPackage*>& Packages, TSet<FName>& OutReadOnlyPackages)>;
+extern FOnPreInteractiveCheckoutPackages OnPreInteractiveCheckoutPackages;
+
+/**
+ * Low-level delegate called by FEditorFileUtils::PromptToCheckoutPackagesInternal after the dialog has been closed, and after packages may have been checked out or made writable.
+ * @param Packages The list of packages that were considered for checkout.
+ * @param bUserResponse False if the checkout was canceled.
+ */
+using FOnPostInteractiveCheckoutPackages = TMulticastDelegate<void(const TArray<UPackage*>& Packages, bool bUserResponse)>;
+extern FOnPostInteractiveCheckoutPackages OnPostInteractiveCheckoutPackages;
+
+/**
+ * Low-level delegate called by FEditorFileUtils::PromptToCheckoutPackagesInternal after packages have been checked out.
+ * @param Packages The list of packages that were checked out.
+ */
+using FOnPackagesInteractivelyCheckedOut = TMulticastDelegate<void(const TArray<UPackage*>& Packages)>;
+extern FOnPackagesInteractivelyCheckedOut OnPackagesInteractivelyCheckedOut;
+
+/**
+ * Low-level delegate called by FEditorFileUtils::PromptToCheckoutPackagesInternal after packages have been made writable.
+ * @param Packages The list of packages that were made writable.
+ */
+using FOnPackagesInteractivelyMadeWritable = TMulticastDelegate<void(const TArray<UPackage*>& Packages)>;
+extern FOnPackagesInteractivelyMadeWritable OnPackagesInteractivelyMadeWritable;
+
+/**
+ * Low-level delegate called by FEditorFileUtils::PromptForCheckoutAndSave after packages have been discarded for save.
+ * @param Packages The list of packages that were discarded.
+ */
+using FOnPackagesInteractivelyDiscarded = TMulticastDelegate<void(const TArray<UPackage*>& Packages)>;
+extern FOnPackagesInteractivelyDiscarded OnPackagesInteractivelyDiscarded;
+
+}
diff --git a/Engine/Source/Editor/UnrealEd/Private/PackageAutoSaver.cpp b/Engine/Source/Editor/UnrealEd/Private/PackageAutoSaver.cpp
--- a/Engine/Source/Editor/UnrealEd/Private/PackageAutoSaver.cpp
+++ b/Engine/Source/Editor/UnrealEd/Private/PackageAutoSaver.cpp
@@ -22,9 +22,11 @@
 #include "EditorModes.h"
 #include "UnrealEdMisc.h"
 #include "FileHelpers.h"
+#include "FileHelpersInternal.h"
 #include "InterchangeManager.h"
 #include "UncontrolledChangelistsModule.h"
 #include "UnrealEdGlobals.h"
+#include "UObject/Linker.h"
 #include "PackageRestore.h"
 #include "Framework/Notifications/NotificationManager.h"
 #include "Widgets/Notifications/SNotificationList.h"
@@ -44,12 +46,47 @@ namespace PackageAutoSaver
 	static bool bUseUncontrolledChangelists = true;
 	static FAutoConsoleVariableRef CVarUseUncontrolledChangelists(TEXT("Editor.AutoSave.UseUncontrolledChangelists"), bUseUncontrolledChangelists, TEXT("True to use uncontrolled changelists (if possible) when using the BackupAndOverwrite auto-save method."));
 
+	static bool bInteractiveSaveManagesUncontrolledChangelist = false;
+	static FAutoConsoleVariableRef CVarInteractiveSaveManagesUncontrolledChangelist(TEXT("Editor.AutoSave.InteractiveSaveManagesUncontrolledChangelist"), bInteractiveSaveManagesUncontrolledChangelist, TEXT("True if an interactive save should manage which entries stay in the auto-save uncontrolled changelist (ie, promoting them to a controlled save, or moving them to the default uncontrolled changelist)."));
+
 	static const FUncontrolledChangelist AutoSaveUncontrolledChangelist(FGuid(0x0D227D90, 0x2AC34632, 0xB1A35533, 0xF3CC134E));
 
 	bool CanUseUncontrolledChangelist()
 	{
 		return bUseUncontrolledChangelists && FUncontrolledChangelistsModule::Get().IsEnabled() && ISourceControlModule::Get().GetProvider().UsesLocalReadOnlyState();
 	}
+
+	bool ShouldInteractiveSaveManageUncontrolledChangelist(const bool bRequiresSourceControl = true)
+	{
+		return bInteractiveSaveManagesUncontrolledChangelist && (!bRequiresSourceControl || (ISourceControlModule::Get().IsEnabled() && ISourceControlModule::Get().GetProvider().IsAvailable()));
+	}
+
+	TMap<FName, FString> GetPackagesInUncontrolledChangelist(const bool bUpdateStatus = false)
+	{
+		TMap<FName, FString> UncontrolledPackages;
+
+		if (CanUseUncontrolledChangelist())
+		{
+			if (FUncontrolledChangelistStatePtr AutoSaveUncontrolledChangelistState = FUncontrolledChangelistsModule::Get().GetChangelistState(AutoSaveUncontrolledChangelist))
+			{
+				if (bUpdateStatus && AutoSaveUncontrolledChangelistState->UpdateStatus())
+				{
+					FUncontrolledChangelistsModule::Get().HandleChangelistStateModified();
+				}
+
+				for (const FString& UncontrolledFilename : AutoSaveUncontrolledChangelistState->GetFilenames())
+				{
+					FNameBuilder PackageName;
+					if (FPackageName::TryConvertFilenameToLongPackageName(UncontrolledFilename, PackageName))
+					{
+						UncontrolledPackages.Add(FName(PackageName), UncontrolledFilename);
+					}
+				}
+			}
+		}
+
+		return UncontrolledPackages;
+	}
 }
 
 namespace PackageAutoSaverJson
@@ -130,6 +167,19 @@ FPackageAutoSaver::FPackageAutoSaver()
 
 	// Register to detect when an Undo/Redo changes the dirty state of a package
 	FEditorDelegates::PostUndoRedo.AddRaw(this, &FPackageAutoSaver::OnUndoRedo);
+
+	// Register for the dirty sync handling
+	FCoreDelegates::OnEndFrame.AddRaw(this, &FPackageAutoSaver::OnEndFrame);
+	FCoreUObjectDelegates::OnEndLoadPackage.AddRaw(this, &FPackageAutoSaver::OnPackageLoaded);
+	FUncontrolledChangelistsModule::Get().OnUncontrolledChangelistModuleChanged.AddRaw(this, &FPackageAutoSaver::OnUncontrolledChangelistChanged);
+
+	// Register for the interactive checkout handlers
+	UE::FileHelpers::Internal::GetAdditionalInteractiveSavePackageCandidates.AddRaw(this, &FPackageAutoSaver::GetAdditionalInteractiveSavePackageCandidates);
+	UE::FileHelpers::Internal::OnPreInteractiveCheckoutPackages.AddRaw(this, &FPackageAutoSaver::OnPreInteractiveCheckoutPackages);
+	UE::FileHelpers::Internal::OnPostInteractiveCheckoutPackages.AddRaw(this, &FPackageAutoSaver::OnPostInteractiveCheckoutPackages);
+	UE::FileHelpers::Internal::OnPackagesInteractivelyCheckedOut.AddRaw(this, &FPackageAutoSaver::OnPackagesInteractivelyCheckedOut);
+	UE::FileHelpers::Internal::OnPackagesInteractivelyMadeWritable.AddRaw(this, &FPackageAutoSaver::OnPackagesInteractivelyMadeWritable);
+	UE::FileHelpers::Internal::OnPackagesInteractivelyDiscarded.AddRaw(this, &FPackageAutoSaver::OnPackagesInteractivelyDiscarded);
 }
 
 FPackageAutoSaver::~FPackageAutoSaver()
@@ -138,6 +188,20 @@ FPackageAutoSaver::~FPackageAutoSaver()
 	UPackage::PackageMarkedDirtyEvent.RemoveAll(this);
 	UPackage::PackageSavedWithContextEvent.RemoveAll(this);
 	FEditorDelegates::PostUndoRedo.RemoveAll(this);
+
+	FCoreDelegates::OnEndFrame.RemoveAll(this);
+	FCoreUObjectDelegates::OnEndLoadPackage.RemoveAll(this);
+	if (FUncontrolledChangelistsModule* UncontrolledChangelistsModule = FUncontrolledChangelistsModule::GetPtr())
+	{
+		UncontrolledChangelistsModule->OnUncontrolledChangelistModuleChanged.RemoveAll(this);
+	}
+
+	UE::FileHelpers::Internal::GetAdditionalInteractiveSavePackageCandidates.RemoveAll(this);
+	UE::FileHelpers::Internal::OnPreInteractiveCheckoutPackages.RemoveAll(this);
+	UE::FileHelpers::Internal::OnPostInteractiveCheckoutPackages.RemoveAll(this);
+	UE::FileHelpers::Internal::OnPackagesInteractivelyCheckedOut.RemoveAll(this);
+	UE::FileHelpers::Internal::OnPackagesInteractivelyMadeWritable.RemoveAll(this);
+	UE::FileHelpers::Internal::OnPackagesInteractivelyDiscarded.RemoveAll(this);
 }
 
 void FPackageAutoSaver::UpdateAutoSaveCount(const float DeltaSeconds)
@@ -174,7 +238,7 @@ void FPackageAutoSaver::ForceMinimumTimeTillAutoSave(const float TimeTillAutoSav
 	AutoSaveCount = (MinumumTime < AutoSaveCount) ? MinumumTime : AutoSaveCount;
 }
 
-void FPackageAutoSaver::AttemptAutoSave()
+bool FPackageAutoSaver::AttemptAutoSave(const bool bForceAutoSave)
 {
 	const UEditorLoadingSavingSettings* LoadingSavingSettings = GetDefault<UEditorLoadingSavingSettings>();
 	FUnrealEdMisc& UnrealEdMisc = FUnrealEdMisc::Get();
@@ -228,8 +292,9 @@ void FPackageAutoSaver::AttemptAutoSave()
 	}
 
 	// Don't auto-save if disabled or if it is not yet time to auto-save.
-	const bool bTimeToAutosave = (LoadingSavingSettings->bAutoSaveEnable && AutoSaveCount >= LoadingSavingSettings->AutoSaveTimeMinutes * 60.0f);
+	const bool bTimeToAutosave = bForceAutoSave || (LoadingSavingSettings->bAutoSaveEnable && AutoSaveCount >= LoadingSavingSettings->AutoSaveTimeMinutes * 60.0f);
 	bool bAutosaveHandled = false;
+	bool bDidAutoSaveSuccessfully = false;
 
 	if (bTimeToAutosave)
 	{
@@ -240,12 +305,12 @@ void FPackageAutoSaver::AttemptAutoSave()
 		if(!bNeedsAutoSave)
 		{
 			ResetAutoSaveTimer();
-			return;
+			return false;
 		}
 
 		// Don't auto-save during interpolation editing, if there's another slow task
 		// already in progress, or while a PIE world is playing or when doing automated tests.
-		const bool bCanAutosave = CanAutoSave();
+		const bool bCanAutosave = CanAutoSave(bForceAutoSave);
 		if (bCanAutosave)
 		{
 			FScopedSlowTask SlowTask(100.f, NSLOCTEXT("AutoSaveNotify", "PerformingAutoSave_Caption", "Auto-saving out of date packages..."));
@@ -398,6 +463,9 @@ void FPackageAutoSaver::AttemptAutoSave()
 								TSet<FString> EmptyPackageFilenames;
 								for (UPackage* Pkg : PackagesToSave)
 								{
+									// This is processed at the end of the frame, so we can populate it before we run the save
+									PackagesToMakeDirtyIfManagedByAutoSave.Add(Pkg->GetFName());
+
 									FString PackageFilename;
 									if (FPackageName::TryConvertLongPackageNameToFilename(FNameBuilder(Pkg->GetFName()).ToView(), PackageFilename, Pkg->ContainsMap() ? FPackageName::GetMapPackageExtension() : FPackageName::GetAssetPackageExtension()))
 									{
@@ -485,7 +553,7 @@ void FPackageAutoSaver::AttemptAutoSave()
 							}
 
 							// Move any files that the auto-save added to the default uncontrolled changelist to the auto-save specific uncontrolled changelist
-							if (DefaultUncontrolledChangelistState)
+							if (DefaultUncontrolledChangelistState && Provider.IsAvailable())
 							{
 								TSet<FString> PostSaveFilesInDefaultUncontrolledChangelist = TSet<FString>(DefaultUncontrolledChangelistState->GetFilenames());
 								TSet<FString> NewFilesInDefaultUncontrolledChangelist = PostSaveFilesInDefaultUncontrolledChangelist.Difference(PreSaveFilesInDefaultUncontrolledChangelist);
@@ -527,6 +595,8 @@ void FPackageAutoSaver::AttemptAutoSave()
 			const bool bSuccess = (MapsSaveResults != EAutosaveContentPackagesResult::Failure && AssetsSaveResults != EAutosaveContentPackagesResult::Failure && !bNothingToDo);
 			const bool bFailure = (MapsSaveResults == EAutosaveContentPackagesResult::Failure || AssetsSaveResults == EAutosaveContentPackagesResult::Failure);
 
+			bDidAutoSaveSuccessfully = bSuccess;
+
 			// Auto-saved, so close any warning notifications.
 			CloseAutoSaveNotification(
 				bSuccess ?		ECloseNotification::Success		:
@@ -572,6 +642,8 @@ void FPackageAutoSaver::AttemptAutoSave()
 	{
 		UpdateAutoSaveNotification();
 	}
+
+	return bDidAutoSaveSuccessfully;
 }
 
 void FPackageAutoSaver::LoadRestoreFile()
@@ -686,6 +758,351 @@ void FPackageAutoSaver::OnUndoRedo()
 	bSyncWithDirtyPackageList = true;
 }
 
+void FPackageAutoSaver::OnUncontrolledChangelistChanged()
+{
+	// Any loaded packages in the auto-save changelist should be dirty so that they can be considered for an interactive save
+	if (const TMap<FName, FString> UncontrolledPackages = PackageAutoSaver::GetPackagesInUncontrolledChangelist();
+		UncontrolledPackages.Num() > 0)
+	{
+		for (const TTuple<FName, FString>& UncontrolledPackagePair : UncontrolledPackages)
+		{
+			PackagesToMakeDirtyIfManagedByAutoSave.Add(UncontrolledPackagePair.Key);
+		}
+	}
+}
+
+void FPackageAutoSaver::OnPackageLoaded(const FEndLoadPackageContext& Context)
+{
+	// Packages in the auto-save uncontrolled changelist should become dirty post-load so that they can be considered for an interactive save
+	for (UPackage* Pkg : Context.LoadedPackages)
+	{
+		PackagesToMakeDirtyIfManagedByAutoSave.Add(Pkg->GetFName());
+	}
+}
+
+void FPackageAutoSaver::OnEndFrame()
+{
+	if (PackagesToMakeDirtyIfManagedByAutoSave.Num() > 0 && PackageAutoSaver::ShouldInteractiveSaveManageUncontrolledChangelist())
+	{
+		if (const TMap<FName, FString> UncontrolledPackages = PackageAutoSaver::GetPackagesInUncontrolledChangelist();
+			UncontrolledPackages.Num() > 0)
+		{
+			auto FindAndDirtyPackage = [this](const FName PackageName)
+			{
+				if (UPackage* Pkg = FindPackage(nullptr, FNameBuilder(PackageName).ToString());
+					IsValid(Pkg) && !Pkg->IsDirty())
+				{
+					TGuardValue<EPackageAutoSaveType> IsAutoSavingScope(CurrentAutoSaveType, EPackageAutoSaveType::Transient); // Set EPackageAutoSaveType::Transient so this doesn't get re-added to the auto-save queue
+					Pkg->MarkPackageDirty();
+					UE_LOG(LogPackageAutoSaver, Log, TEXT("Package '%s' was marked dirty by the auto-save uncontrolled changelist"), FNameBuilder(Pkg->GetFName()).ToString());
+				}
+			};
+
+			// Iterate whichever is smaller
+			if (PackagesToMakeDirtyIfManagedByAutoSave.Num() <= UncontrolledPackages.Num())
+			{
+				for (const FName PackageName : PackagesToMakeDirtyIfManagedByAutoSave)
+				{
+					if (UncontrolledPackages.Contains(PackageName))
+					{
+						FindAndDirtyPackage(PackageName);
+					}
+				}
+			}
+			else
+			{
+				for (const TTuple<FName, FString>& UncontrolledPackagePair : UncontrolledPackages)
+				{
+					if (PackagesToMakeDirtyIfManagedByAutoSave.Contains(UncontrolledPackagePair.Key))
+					{
+						FindAndDirtyPackage(UncontrolledPackagePair.Key);
+					}
+				}
+			}
+		}
+	}
+
+	PackagesToMakeDirtyIfManagedByAutoSave.Reset();
+}
+
+void FPackageAutoSaver::GetAdditionalInteractiveSavePackageCandidates(TSet<UPackage*>& OutPackages)
+{
+	if (!PackageAutoSaver::ShouldInteractiveSaveManageUncontrolledChangelist())
+	{
+		return;
+	}
+
+	// Strongly hint that any loaded packages in the auto-save changelist should be considered for an interactive save
+	if (const TMap<FName, FString> UncontrolledPackages = PackageAutoSaver::GetPackagesInUncontrolledChangelist(/*bUpdateStatus*/true);
+		UncontrolledPackages.Num() > 0)
+	{
+		for (const TTuple<FName, FString>& UncontrolledPackagePair : UncontrolledPackages)
+		{
+			if (UPackage* Pkg = FindPackage(nullptr, FNameBuilder(UncontrolledPackagePair.Key).ToString());
+				IsValid(Pkg))
+			{
+				OutPackages.Add(Pkg);
+			}
+		}
+	}
+}
+
+void FPackageAutoSaver::OnPreInteractiveCheckoutPackages(const TArray<UPackage*>& Packages, TSet<FName>& OutReadOnlyPackages)
+{
+	if (!PackageAutoSaver::ShouldInteractiveSaveManageUncontrolledChangelist())
+	{
+		return;
+	}
+
+	// Report that packages in the auto-save uncontrolled changelist are read-only so that PromptToCheckoutPackagesInternal will offer to check them out or make them writable
+	if (const TMap<FName, FString> UncontrolledPackages = PackageAutoSaver::GetPackagesInUncontrolledChangelist(/*bUpdateStatus*/true);
+		UncontrolledPackages.Num() > 0)
+	{
+		for (UPackage* Pkg : Packages)
+		{
+			if (UncontrolledPackages.Contains(Pkg->GetFName()))
+			{
+				OutReadOnlyPackages.Add(Pkg->GetFName());
+			}
+		}
+	}
+}
+
+void FPackageAutoSaver::OnPostInteractiveCheckoutPackages(const TArray<UPackage*>& Packages, bool bUserResponse)
+{
+	if (!bUserResponse || !PackageAutoSaver::ShouldInteractiveSaveManageUncontrolledChangelist())
+	{
+		return;
+	}
+
+	// Handle any source control adds, as newly added files don't appear in the checkout list since adds usually happen automatically on save
+	TArray<FString> FilesToPotentiallyAdd;
+	if (const TMap<FName, FString> UncontrolledPackages = PackageAutoSaver::GetPackagesInUncontrolledChangelist();
+		UncontrolledPackages.Num() > 0)
+	{
+		for (UPackage* Pkg : Packages)
+		{
+			if (const FString* PackageFilename = UncontrolledPackages.Find(Pkg->GetFName()))
+			{
+				if (!UPackage::IsEmptyPackage(Pkg) && FPaths::FileExists(*PackageFilename))
+				{
+					FilesToPotentiallyAdd.Add(*PackageFilename);
+				}
+			}
+		}
+	}
+
+	if (FilesToPotentiallyAdd.Num() > 0)
+	{
+		FScopedSlowTask SlowTask(100.0f, NSLOCTEXT("AutoSave", "AddingAutoSavedChanges", "Adding auto-saved changes..."));
+		SlowTask.MakeDialog();
+		SlowTask.EnterProgressFrame(25.0f);
+
+		ISourceControlProvider& Provider = ISourceControlModule::Get().GetProvider();
+
+		TArray<FSourceControlStateRef> FileStates;
+		if (Provider.GetState(FilesToPotentiallyAdd, FileStates, EStateCacheUsage::ForceUpdate) == ECommandResult::Succeeded)
+		{
+			TArray<FString> FilesToAdd;
+			TArray<FSourceControlStateRef> FileStatesToAdd;
+			for (const FSourceControlStateRef& FileState : FileStates)
+			{
+				if (!FileState->IsSourceControlled() && !FileState->IsAdded())
+				{
+					FilesToAdd.Add(FileState->GetFilename());
+					FileStatesToAdd.Add(FileState);
+				}
+			}
+			if (FilesToAdd.Num() > 0)
+			{
+				SlowTask.EnterProgressFrame(25.0f);
+
+				if (Provider.Execute(ISourceControlOperation::Create<FMarkForAdd>(), FilesToAdd) == ECommandResult::Succeeded)
+				{
+					FUncontrolledChangelistStatePtr AutoSaveUncontrolledChangelistState = FUncontrolledChangelistsModule::Get().GetChangelistState(PackageAutoSaver::AutoSaveUncontrolledChangelist);
+					check(AutoSaveUncontrolledChangelistState);
+
+					if (AutoSaveUncontrolledChangelistState->RemoveFiles(FileStatesToAdd))
+					{
+						FUncontrolledChangelistsModule::Get().HandleChangelistStateModified();
+					}
+				}
+			}
+		}
+	}
+}
+
+void FPackageAutoSaver::OnPackagesInteractivelyCheckedOut(const TArray<UPackage*>& Packages)
+{
+	if (!PackageAutoSaver::ShouldInteractiveSaveManageUncontrolledChangelist())
+	{
+		return;
+	}
+
+	// Remove any now checked out packages from the auto-save uncontrolled changelist
+	TArray<FString> FilesToRemove;
+	if (const TMap<FName, FString> UncontrolledPackages = PackageAutoSaver::GetPackagesInUncontrolledChangelist();
+		UncontrolledPackages.Num() > 0)
+	{
+		for (UPackage* Pkg : Packages)
+		{
+			if (const FString* PackageFilename = UncontrolledPackages.Find(Pkg->GetFName()))
+			{
+				FilesToRemove.Add(*PackageFilename);
+			}
+		}
+	}
+
+	if (FilesToRemove.Num() > 0)
+	{
+		TArray<FSourceControlStateRef> FileStates;
+		if (ISourceControlModule::Get().GetProvider().GetState(FilesToRemove, FileStates, EStateCacheUsage::Use) == ECommandResult::Succeeded)
+		{
+			FUncontrolledChangelistStatePtr AutoSaveUncontrolledChangelistState = FUncontrolledChangelistsModule::Get().GetChangelistState(PackageAutoSaver::AutoSaveUncontrolledChangelist);
+			check(AutoSaveUncontrolledChangelistState);
+
+			if (AutoSaveUncontrolledChangelistState->RemoveFiles(FileStates))
+			{
+				FUncontrolledChangelistsModule::Get().HandleChangelistStateModified();
+			}
+		}
+	}
+}
+
+void FPackageAutoSaver::OnPackagesInteractivelyMadeWritable(const TArray<UPackage*>& Packages)
+{
+	if (!PackageAutoSaver::ShouldInteractiveSaveManageUncontrolledChangelist(/*bRequiresSourceControl*/false))
+	{
+		return;
+	}
+
+	// Move any packages the user made writable to the default uncontrolled changelist
+	TArray<FString> FilesToMove;
+	if (const TMap<FName, FString> UncontrolledPackages = PackageAutoSaver::GetPackagesInUncontrolledChangelist();
+		UncontrolledPackages.Num() > 0)
+	{
+		for (UPackage* Pkg : Packages)
+		{
+			if (const FString* PackageFilename = UncontrolledPackages.Find(Pkg->GetFName()))
+			{
+				FilesToMove.Add(*PackageFilename);
+			}
+		}
+	}
+
+	if (FilesToMove.Num() > 0)
+	{
+		TArray<FSourceControlStateRef> FileStates;
+		if (ISourceControlModule::Get().GetProvider().GetState(FilesToMove, FileStates, EStateCacheUsage::Use) == ECommandResult::Succeeded)
+		{
+			FUncontrolledChangelistsModule::Get().MoveFilesToUncontrolledChangelist(TArray<FSourceControlStateRef>(), FileStates, FUncontrolledChangelist(FUncontrolledChangelist::DEFAULT_UNCONTROLLED_CHANGELIST_GUID));
+		}
+	}
+}
+
+void FPackageAutoSaver::OnPackagesInteractivelyDiscarded(const TArray<UPackage*>& Packages)
+{
+	if (!PackageAutoSaver::ShouldInteractiveSaveManageUncontrolledChangelist())
+	{
+		return;
+	}
+
+	// Revert any auto-save uncontrolled changes as the user has asked to discard saving them to disk
+	{
+		FScopedSlowTask SlowTask(100.0f, NSLOCTEXT("AutoSave", "DiscardingAutoSavedChanges", "Discarding auto-saved changes..."));
+		SlowTask.MakeDialog();
+
+		TArray<FString> FilesToDiscard;
+		if (const TMap<FName, FString> UncontrolledPackages = PackageAutoSaver::GetPackagesInUncontrolledChangelist(/*bUpdateStatus*/true);
+			UncontrolledPackages.Num() > 0)
+		{
+			auto DetachLinker = [](UPackage* Package)
+			{
+				if (!Package->IsFullyLoaded())
+				{
+					FlushAsyncLoading();
+					Package->FullyLoad();
+				}
+				ResetLoaders(Package);
+			};
+
+			for (UPackage* Pkg : Packages)
+			{
+				if (const FString* PackageFilename = UncontrolledPackages.Find(Pkg->GetFName()))
+				{
+					DetachLinker(Pkg);
+					FilesToDiscard.Add(*PackageFilename);
+				}
+			}
+		}
+
+		if (FilesToDiscard.Num() > 0)
+		{
+			SlowTask.EnterProgressFrame(50.0f);
+
+			ISourceControlProvider& Provider = ISourceControlModule::Get().GetProvider();
+
+			TArray<FString> FilesToDelete;
+			TArray<FString> FilesToRevert;
+			TArray<FSourceControlStateRef> FileStates;
+			if (Provider.GetState(FilesToDiscard, FileStates, EStateCacheUsage::ForceUpdate) == ECommandResult::Succeeded)
+			{
+				for (const FSourceControlStateRef& FileState : FileStates)
+				{
+					if (FileState->IsSourceControlled())
+					{
+						FilesToRevert.Add(FileState->GetFilename());
+					}
+					else
+					{
+						FilesToDelete.Add(FileState->GetFilename());
+					}
+				}
+			}
+
+			SlowTask.EnterProgressFrame(25.0f);
+
+			bool bUpdatedChangelistState = false;
+
+			if (FilesToRevert.Num() > 0)
+			{
+				TSharedRef<FSync> ForceSyncOperation = ISourceControlOperation::Create<FSync>();
+				ForceSyncOperation->SetForce(true);
+				ForceSyncOperation->SetLastSyncedFlag(true);
+
+				bUpdatedChangelistState |= Provider.Execute(ForceSyncOperation, FilesToRevert) == ECommandResult::Succeeded;
+			}
+
+			if (FilesToDelete.Num() > 0)
+			{
+				IFileManager& FileManager = IFileManager::Get();
+
+				for (const FString& FileToDelete : FilesToDelete)
+				{
+					const bool bRequireExists = true;
+					const bool bEvenReadOnly = false;
+					const bool bQuiet = false;
+
+					bUpdatedChangelistState |= FileManager.Delete(*FileToDelete, bRequireExists, bEvenReadOnly, bQuiet);
+				}
+
+				ISourceControlModule::Get().GetOnFilesDeleted().Broadcast(FilesToDelete);
+			}
+
+			if (bUpdatedChangelistState)
+			{
+				FUncontrolledChangelistStatePtr AutoSaveUncontrolledChangelistState = FUncontrolledChangelistsModule::Get().GetChangelistState(PackageAutoSaver::AutoSaveUncontrolledChangelist);
+				check(AutoSaveUncontrolledChangelistState);
+
+				if (AutoSaveUncontrolledChangelistState->UpdateStatus())
+				{
+					FUncontrolledChangelistsModule::Get().HandleChangelistStateModified();
+				}
+			}
+		}
+	}
+}
+
 void FPackageAutoSaver::UpdateDirtyListsForPackage(UPackage* Pkg)
 {
 	const UPackage* TransientPackage = GetTransientPackage();
@@ -733,18 +1150,28 @@ void FPackageAutoSaver::UpdateDirtyListsForPackage(UPackage* Pkg)
 	}
 }
 
-bool FPackageAutoSaver::CanAutoSave() const
+bool FPackageAutoSaver::CanAutoSave(const bool bForceAutoSave) const
 {
 	// Don't allow auto-saving if the auto-save wouldn't save anything
 	const bool bPackagesNeedAutoSave = DoPackagesNeedAutoSave();
 
-	double LastInteractionTime = FSlateApplication::Get().GetLastUserInteractionTime();
-
 	const UEditorLoadingSavingSettings* LoadingSavingSettings = GetDefault<UEditorLoadingSavingSettings>();
+	const bool bAutosaveEnabled = LoadingSavingSettings->bAutoSaveEnable && bPackagesNeedAutoSave;
+
+	if (!bAutosaveEnabled)
+	{
+		return false;
+	}
+
+	if (bForceAutoSave)
+	{
+		return true;
+	}
+
+	const double LastInteractionTime = FSlateApplication::Get().GetLastUserInteractionTime();
 	const float InteractionDelay = float(LoadingSavingSettings->AutoSaveInteractionDelayInSeconds);
 
 	const bool bDidInteractRecently = (FApp::GetCurrentTime() - LastInteractionTime) < InteractionDelay;
-	const bool bAutosaveEnabled	= LoadingSavingSettings->bAutoSaveEnable && bPackagesNeedAutoSave;
 	const bool bSlowTask = GIsSlowTask;
 	const bool bPlayWorldValid = GUnrealEd->PlayWorld != nullptr;
 	const bool bAnyMenusVisible	= FSlateApplication::Get().AnyMenusVisible();
@@ -769,8 +1196,7 @@ bool FPackageAutoSaver::CanAutoSave() const
 	// query any active editor modes and allow them to prevent autosave
 	const bool bActiveModesAllowAutoSave = GLevelEditorModeTools().CanAutoSave();
 
-	return (bAutosaveEnabled
-		&& !bSlowTask
+	return !bSlowTask
 		&& !bPlayWorldValid
 		&& !bAnyMenusVisible
 		&& !bAutomationTesting
@@ -782,7 +1208,7 @@ bool FPackageAutoSaver::CanAutoSave() const
 		&& !bIsVREditorActive
 		&& !bIsSequencerPlaying
 		&& !bIsInterchangeActive
-		&& bActiveModesAllowAutoSave);
+		&& bActiveModesAllowAutoSave;
 }
 
 bool FPackageAutoSaver::DoPackagesNeedAutoSave() const
diff --git a/Engine/Source/Editor/UnrealEd/Private/PackageAutoSaver.h b/Engine/Source/Editor/UnrealEd/Private/PackageAutoSaver.h
--- a/Engine/Source/Editor/UnrealEd/Private/PackageAutoSaver.h
+++ b/Engine/Source/Editor/UnrealEd/Private/PackageAutoSaver.h
@@ -8,6 +8,7 @@
 #include "IPackageAutoSaver.h"
 
 class FObjectPostSaveContext;
+struct FEndLoadPackageContext;
 struct FTransactionContext;
 
 namespace ECloseNotification
@@ -33,7 +34,7 @@ public:
 	virtual void ResetAutoSaveTimer() override;
 	virtual void ForceAutoSaveTimer() override;
 	virtual void ForceMinimumTimeTillAutoSave(const float TimeTillAutoSave = 10.0f) override;
-	virtual void AttemptAutoSave() override;
+	virtual bool AttemptAutoSave(const bool bForceAutoSave = false) override;
 	virtual void LoadRestoreFile() override;
 	virtual void UpdateRestoreFile(const bool bRestoreEnabled) override;
 	virtual bool HasPackagesToRestore() const override;
@@ -71,6 +72,23 @@ private:
 	 */
 	void OnUndoRedo();
 
+	/**
+	 * Dirty sync handling for the auto-save uncontrolled changelist.
+	 */
+	void OnUncontrolledChangelistChanged();
+	void OnPackageLoaded(const FEndLoadPackageContext& Context);
+	void OnEndFrame();
+
+	/**
+	 * Interactive checkout handling for the auto-save uncontrolled changelist.
+	 */
+	void GetAdditionalInteractiveSavePackageCandidates(TSet<UPackage*>& OutPackages);
+	void OnPreInteractiveCheckoutPackages(const TArray<UPackage*>& Packages, TSet<FName>& OutReadOnlyPackages);
+	void OnPostInteractiveCheckoutPackages(const TArray<UPackage*>& Packages, bool bUserResponse);
+	void OnPackagesInteractivelyCheckedOut(const TArray<UPackage*>& Packages);
+	void OnPackagesInteractivelyMadeWritable(const TArray<UPackage*>& Packages);
+	void OnPackagesInteractivelyDiscarded(const TArray<UPackage*>& Packages);
+
 	/**
 	 * Update the dirty lists for the current package
 	 *
@@ -79,7 +97,7 @@ private:
 	void UpdateDirtyListsForPackage(UPackage* Pkg);
 
 	/** @return Returns whether or not the user is able to auto-save. */
-	bool CanAutoSave() const;
+	bool CanAutoSave(const bool bForceAutoSave = false) const;
 
 	/** @return Returns whether or not we would need to perform an auto-save (note: does not check if we can perform an auto-save, only that we should if we could). */
 	bool DoPackagesNeedAutoSave() const;
@@ -160,6 +178,6 @@ private:
 	/** Restore information that was loaded following a crash */
 	TMap<FString, TPair<FString, FString>> PackagesThatCanBeRestored;
 
-	/** The cached state of the uncontrolled packages changelist, the last time OnUncontrolledChangelistModuleChanged was called. This may be stale in-between frames, but is very fast to query */
-	TMap<FName, FString> CachedUncontrolledPackages;
+	/** Set of packages to make dirty if they're being managed by the auto-save uncontrolled changelist */
+	TSet<FName> PackagesToMakeDirtyIfManagedByAutoSave;
 };
diff --git a/Engine/Source/Editor/UnrealEd/Public/FileHelpers.h b/Engine/Source/Editor/UnrealEd/Public/FileHelpers.h
--- a/Engine/Source/Editor/UnrealEd/Public/FileHelpers.h
+++ b/Engine/Source/Editor/UnrealEd/Public/FileHelpers.h
@@ -674,10 +674,10 @@ public:
 private:
 
 	/** Private method used to build a list of items requiring checkout */
-	static bool AddCheckoutPackageItems(bool bCheckDirty, TArray<UPackage*> PackagesToCheckOut, TArray<UPackage*>* OutPackagesNotNeedingCheckout, bool* bOutHavePackageToCheckOut);
+	static bool AddCheckoutPackageItems(bool bCheckDirty, const TArray<UPackage*>& PackagesToCheckOut, const TSet<FName>& ReadOnlyPackages, TArray<UPackage*>* OutPackagesNotNeedingCheckout, bool* bOutHavePackageToCheckOut);
 
 	/** Callback from PackagesDialog used to update the list of items when the source control state changes */
-	static void UpdateCheckoutPackageItems(bool bCheckDirty, TArray<UPackage*> PackagesToCheckOut, TArray<UPackage*>* OutPackagesNotNeedingCheckout);
+	static void UpdateCheckoutPackageItems(bool bCheckDirty, TArray<UPackage*> PackagesToCheckOut, TSet<FName> ReadOnlyPackages, TArray<UPackage*>* OutPackagesNotNeedingCheckout);
 	
 	/** Prompts for package checkout without reentrance check */
 	static bool PromptToCheckoutPackagesInternal(bool bCheckDirty, const TArray<UPackage*>& PackagesToCheckOut, TArray<UPackage*>* OutPackagesCheckedOutOrMadeWritable, TArray<UPackage*>* OutPackagesNotNeedingCheckout, const bool bPromptingAfterModify, const bool bAllowSkip);
diff --git a/Engine/Source/Editor/UnrealEd/Public/IPackageAutoSaver.h b/Engine/Source/Editor/UnrealEd/Public/IPackageAutoSaver.h
--- a/Engine/Source/Editor/UnrealEd/Public/IPackageAutoSaver.h
+++ b/Engine/Source/Editor/UnrealEd/Public/IPackageAutoSaver.h
@@ -37,8 +37,14 @@ public:
 	 */
 	virtual void ForceMinimumTimeTillAutoSave(const float TimeTillAutoSave = 10.0f) = 0;
 
-	/** Attempts to auto-save the level and/or content packages, if those features are enabled. */
-	virtual void AttemptAutoSave() = 0;
+	/**
+	 * Attempts to auto-save the level and/or content packages, if those features are enabled.
+	 * 
+	 * @param bForceAutoSave Force the auto-save to happen, regardless of the current timer or other mode restrictions/user interaction. The caller is responsible for ensuring the editor is in a good state to auto-save when passing bForceAutoSave as true.
+	 * 
+	 * @return True if an auto-save ran successfully, false otherwise.
+	 */
+	virtual bool AttemptAutoSave(const bool bForceAutoSave = false) = 0;
 
 	/** @return If we are currently auto-saving (default is to only check for transient auto-saves for legacy reasons) */
 	virtual bool IsAutoSaving(const EPackageAutoSaveType AutoSaveType = EPackageAutoSaveType::Transient) const = 0;
