Fix FEditorModeTools::DeactivateAllModes removing editor modes in a loop that affects the modes iterator.
Some modes will deactivate other modes on deactivation, resulting in an invalid index during iteration.

diff --git a/Engine/Source/Editor/UnrealEd/Private/EditorModeManager.cpp b/Engine/Source/Editor/UnrealEd/Private/EditorModeManager.cpp
--- a/Engine/Source/Editor/UnrealEd/Private/EditorModeManager.cpp
+++ b/Engine/Source/Editor/UnrealEd/Private/EditorModeManager.cpp
@@ -807,15 +807,22 @@ void FEditorModeTools::DeactivateMode( FEditorModeID InID )
 			DeactivateModeAtIndex(Index);
 			break;
 		}
-	};
+	}
 }
 
 void FEditorModeTools::DeactivateAllModes()
 {
+	TArray<FEditorModeID> EditorModesToRemove;
 	for (int32 Index = ActiveScriptableModes.Num() - 1; Index >= 0; --Index)
 	{
-		DeactivateModeAtIndex(Index);
-	};
+		UEdMode* Mode = ActiveScriptableModes[Index];
+		EditorModesToRemove.Add(Mode->GetID());
+	}
+
+	for (FEditorModeID EditorModeToRemove : EditorModesToRemove)
+	{
+		DeactivateMode(EditorModeToRemove);
+	}
 }
 
 void FEditorModeTools::DestroyMode( FEditorModeID InID )
