World Partition - Landscape Spline Mesh Builder: Ensure original spline actors are not forced to be editor only during the partitioning step of the builder
- Issue was that by being already forced editor only at that point, they were also prevented from being HLOD relevant
- This fix will allow splines to be assigned to different HLOD layers than the default one

diff --git a/Engine/Source/Editor/UnrealEd/Private/WorldPartition/WorldPartitionLandscapeSplineMeshesBuilder.cpp b/Engine/Source/Editor/UnrealEd/Private/WorldPartition/WorldPartitionLandscapeSplineMeshesBuilder.cpp
--- a/Engine/Source/Editor/UnrealEd/Private/WorldPartition/WorldPartitionLandscapeSplineMeshesBuilder.cpp
+++ b/Engine/Source/Editor/UnrealEd/Private/WorldPartition/WorldPartitionLandscapeSplineMeshesBuilder.cpp
@@ -389,7 +389,12 @@ bool UWorldPartitionLandscapeSplineMeshesBuilder::RunInternal(UWorld* InWorld, c
 			continue;
 		}
 
+		// Reset flag to ensure spline actors/component IsEditorOnly() logic remains valid during the processing below.
+		// Otherwise they'll all be considered editor only, which will cause issues such as IsHLODRelevant() returning false.
+		bool bLandscapeHadSplineActors = Landscape->GetUseGeneratedLandscapeSplineMeshesActors();
 		bool bLandscapeHasSplineActors = false;
+		Landscape->SetUseGeneratedLandscapeSplineMeshesActors(false);
+
 		for (TScriptInterface<ILandscapeSplineInterface> SplineOwner : LandscapeInfo->GetSplineActors())
 		{
 			ALandscapeSplineActor* SplineActor = Cast<ALandscapeSplineActor>(SplineOwner.GetObject());
@@ -443,6 +448,11 @@ bool UWorldPartitionLandscapeSplineMeshesBuilder::RunInternal(UWorld* InWorld, c
 		if (bLandscapeHasSplineActors)
 		{
 			Landscape->SetUseGeneratedLandscapeSplineMeshesActors(true);
+		}
+
+		if (bLandscapeHasSplineActors != bLandscapeHadSplineActors)
+		{
+			Landscape->MarkPackageDirty();
 			AddObjectToSaveIfDirty(Landscape);
 		}
 	}
@@ -577,4 +587,4 @@ bool UWorldPartitionLandscapeSplineMeshesBuilder::RunOnInitializedWorld(UWorld*
 	}
 
 	return bResult;
-}
\ No newline at end of file
+}
diff --git a/Engine/Source/Runtime/Landscape/Private/Landscape.cpp b/Engine/Source/Runtime/Landscape/Private/Landscape.cpp
--- a/Engine/Source/Runtime/Landscape/Private/Landscape.cpp
+++ b/Engine/Source/Runtime/Landscape/Private/Landscape.cpp
@@ -3669,11 +3669,7 @@ FBox ALandscape::GetCompleteBounds() const
 
 void ALandscape::SetUseGeneratedLandscapeSplineMeshesActors(bool bInEnabled)
 {
-	if (bUseGeneratedLandscapeSplineMeshesActors != bInEnabled)
-	{
-		Modify();
-		bUseGeneratedLandscapeSplineMeshesActors = bInEnabled;
-	}
+	bUseGeneratedLandscapeSplineMeshesActors = bInEnabled;
 }
 
 bool ALandscape::GetUseGeneratedLandscapeSplineMeshesActors() const
