ReplaceObjectHelper now preserves the object external package if the old object has one.
Needed when reinstantiating entities saved in their own file.

diff --git a/Engine/Source/Editor/UnrealEd/Private/Kismet2/KismetReinstanceUtilities.cpp b/Engine/Source/Editor/UnrealEd/Private/Kismet2/KismetReinstanceUtilities.cpp
--- a/Engine/Source/Editor/UnrealEd/Private/Kismet2/KismetReinstanceUtilities.cpp
+++ b/Engine/Source/Editor/UnrealEd/Private/Kismet2/KismetReinstanceUtilities.cpp
@@ -2208,6 +2208,7 @@ static void ReplaceObjectHelper(UObject*& OldObject, UClass* OldClass, UObject*&
 	}
 
 	EObjectFlags OldFlags = OldObject->GetFlags();
+	UPackage* OldExternalPackage = OldObject->GetExternalPackage();
 
 	FName OldName(OldObject->GetFName());
 
@@ -2267,7 +2268,7 @@ static void ReplaceObjectHelper(UObject*& OldObject, UClass* OldClass, UObject*&
 		}
 
 		FMakeClassSpawnableOnScope TemporarilySpawnable(NewClass);
-		NewUObject = NewObject<UObject>(DestinationOuter, NewClass, NewName, RF_NoFlags, NewArchetype);
+		NewUObject = NewObject<UObject>(DestinationOuter, NewClass, NewName, RF_NoFlags, NewArchetype, false, nullptr, OldExternalPackage);
 	}
 
 	check(NewUObject != nullptr);
