PR #12486: World Partition Conversion Commandlet: Respect DataLayerAssetFolder value from the config

diff --git a/Engine/Source/Editor/UnrealEd/Private/Commandlets/WorldPartitionConvertCommandlet.cpp b/Engine/Source/Editor/UnrealEd/Private/Commandlets/WorldPartitionConvertCommandlet.cpp
--- a/Engine/Source/Editor/UnrealEd/Private/Commandlets/WorldPartitionConvertCommandlet.cpp
+++ b/Engine/Source/Editor/UnrealEd/Private/Commandlets/WorldPartitionConvertCommandlet.cpp
@@ -758,7 +758,7 @@ int32 UWorldPartitionConvertCommandlet::Main(const FString& Params)
 	{
 		DataLayerAssetFolder = *DataLayerAssetFolderValue;
 	}
-	else
+	else if (DataLayerAssetFolder.IsEmpty())
 	{
 		UPackage* Package = MainWorld->GetPackage();
 		FName PackageMountPoint = FPackageName::GetPackageMountPoint(Package->GetLoadedPath().GetPackageName());
