FEditorFileUtils : In CheckoutPackages if bErrorIfAlreadyCheckedOut is false and the file is already marked for add then don't show an error.

The branches above check CanCheckout() -> CanAdd() -> IsCheckedOut() so this provides symmetry by also checking IsAdded()

This fixes a bug in PackageRestore::PromptToRestorePackages where the user would be shown the checkout error even though the file was already marked for add and the package restore would go on to succeed anyways.

diff --git a/Engine/Source/Editor/UnrealEd/Private/FileHelpers.cpp b/Engine/Source/Editor/UnrealEd/Private/FileHelpers.cpp
--- a/Engine/Source/Editor/UnrealEd/Private/FileHelpers.cpp
+++ b/Engine/Source/Editor/UnrealEd/Private/FileHelpers.cpp
@@ -2380,6 +2380,10 @@ ECommandResult::Type FEditorFileUtils::CheckoutPackages(const TArray<FString>& P
 				{
 					bShowCheckoutError = false;
 				}
+				else if (!bErrorIfAlreadyCheckedOut && SourceControlState->IsAdded())
+				{
+					bShowCheckoutError = false;
+				}
 			}
 
 			// If the package couldn't be checked out, log it so the list of failures can be displayed afterwards
