Fix a bug where bUseZenStore is not reflected in the settings UI properly because we force it to False when it is in conflict with the bUseIoStore setting.  Instead of forcing it to false, leave the value as it is, and just modify places that use it to account for the conflict.

diff --git a/Engine/Source/Developer/DeveloperToolSettings/Classes/Settings/ProjectPackagingSettings.h b/Engine/Source/Developer/DeveloperToolSettings/Classes/Settings/ProjectPackagingSettings.h
--- a/Engine/Source/Developer/DeveloperToolSettings/Classes/Settings/ProjectPackagingSettings.h
+++ b/Engine/Source/Developer/DeveloperToolSettings/Classes/Settings/ProjectPackagingSettings.h
@@ -696,6 +696,11 @@ public:
 		LoadConfig();
 	}
 
+	bool GetUseZenStoreEffective() const
+	{
+		return bUseZenStore && bUseIoStore;
+	}
+
 	FString GetConfigPlatform() const
 	{
 		return ConfigPlatform;
diff --git a/Engine/Source/Developer/DeveloperToolSettings/Private/ProjectPackagingSettings.cpp b/Engine/Source/Developer/DeveloperToolSettings/Private/ProjectPackagingSettings.cpp
--- a/Engine/Source/Developer/DeveloperToolSettings/Private/ProjectPackagingSettings.cpp
+++ b/Engine/Source/Developer/DeveloperToolSettings/Private/ProjectPackagingSettings.cpp
@@ -47,7 +47,6 @@ void UProjectPackagingSettings::PostInitProperties()
 															"This is not a supported combination.  bUseZenStore is being treated as disabled. "
 															"To prevent this warning, either (preferred) set bUseIoStore=True or set bUseZenStore=False"
 															"in your project packaging settings in your DefaultGame.ini"));
-		bUseZenStore = false;
 	}
 }
 
diff --git a/Engine/Source/Editor/TurnkeySupport/Private/TurnkeySupportModule.cpp b/Engine/Source/Editor/TurnkeySupport/Private/TurnkeySupportModule.cpp
--- a/Engine/Source/Editor/TurnkeySupport/Private/TurnkeySupportModule.cpp
+++ b/Engine/Source/Editor/TurnkeySupport/Private/TurnkeySupportModule.cpp
@@ -379,7 +379,7 @@ public:
 			BuildCookRunParams += TEXT(" -installed");
 		}
 
-		if (PackagingSettings->bUseZenStore)
+		if (PackagingSettings->GetUseZenStoreEffective())
 		{
 #if WITH_EDITOR && UE_WITH_ZEN
 			static UE::Zen::FScopeZenService TurnkeyStaticZenService;
diff --git a/Engine/Source/Editor/UnrealEd/Private/Commandlets/CookCommandlet.cpp b/Engine/Source/Editor/UnrealEd/Private/Commandlets/CookCommandlet.cpp
--- a/Engine/Source/Editor/UnrealEd/Private/Commandlets/CookCommandlet.cpp
+++ b/Engine/Source/Editor/UnrealEd/Private/Commandlets/CookCommandlet.cpp
@@ -411,7 +411,7 @@ bool UCookCommandlet::CookByTheBook(const TArray<ITargetPlatform*>& Platforms)
 
 	const bool bUseZenStore =
 		!Switches.Contains(TEXT("SkipZenStore")) &&
-		(Switches.Contains(TEXT("ZenStore")) || PackagingSettings->bUseZenStore);
+		(Switches.Contains(TEXT("ZenStore")) || PackagingSettings->GetUseZenStoreEffective());
 
 	//////////////////////////////////////////////////////////////////////////
 	// start cook by the book 
diff --git a/Engine/Source/Editor/UnrealEd/Private/PlayLevelViaLauncher.cpp b/Engine/Source/Editor/UnrealEd/Private/PlayLevelViaLauncher.cpp
--- a/Engine/Source/Editor/UnrealEd/Private/PlayLevelViaLauncher.cpp
+++ b/Engine/Source/Editor/UnrealEd/Private/PlayLevelViaLauncher.cpp
@@ -318,7 +318,7 @@ void UEditorEngine::StartPlayUsingLauncherSession(FRequestPlaySessionParams& InR
 		LauncherProfile->SetLaunchMode(ELauncherProfileLaunchModes::DefaultRole);
 	}
 
-	const bool bUseZenStore = GetDefault<UProjectPackagingSettings>()->bUseZenStore;
+	const bool bUseZenStore = GetDefault<UProjectPackagingSettings>()->GetUseZenStoreEffective();
 	LauncherProfile->SetUseZenStore(bUseZenStore);
 #if UE_WITH_ZEN
 	if (bUseZenStore)
diff --git a/Engine/Source/Editor/UnrealEd/Private/UnrealEdEngine.cpp b/Engine/Source/Editor/UnrealEd/Private/UnrealEdEngine.cpp
--- a/Engine/Source/Editor/UnrealEd/Private/UnrealEdEngine.cpp
+++ b/Engine/Source/Editor/UnrealEd/Private/UnrealEdEngine.cpp
@@ -380,7 +380,7 @@ void UUnrealEdEngine::StartCookByTheBookInEditor( const TArray<ITargetPlatform*>
 	StartupOptions.CookDirectories = CookDirectories;
 	StartupOptions.CookCultures = CookCultures;
 	StartupOptions.IniMapSections = IniMapSections;
-	if (GetDefault<UProjectPackagingSettings>()->bUseZenStore)
+	if (GetDefault<UProjectPackagingSettings>()->GetUseZenStoreEffective())
 	{
 		StartupOptions.CookOptions |= ECookByTheBookOptions::ZenStore;
 	}
