Viewport Toolbar & New TRS Gizmos: expose parent space when the new trs gizmos are active

#rnx

diff --git a/Engine/Plugins/Animation/ControlRig/Source/ControlRigEditor/Private/Editor/ControlRigViewportToolbarExtensions.cpp b/Engine/Plugins/Animation/ControlRig/Source/ControlRigEditor/Private/Editor/ControlRigViewportToolbarExtensions.cpp
--- a/Engine/Plugins/Animation/ControlRig/Source/ControlRigEditor/Private/Editor/ControlRigViewportToolbarExtensions.cpp
+++ b/Engine/Plugins/Animation/ControlRig/Source/ControlRigEditor/Private/Editor/ControlRigViewportToolbarExtensions.cpp
@@ -141,20 +141,6 @@ void UE::ControlRig::PopulateControlRigViewportToolbarTransformSubmenu(const FNa
 			SelectionSection.AddEntry(Entry);
 		}
 	}
-
-	// Coordinate space submenu.
-	{
-		const FName SubmenuName = UToolMenus::JoinMenuPaths(InMenuName, "CoordinateSystem");
-
-		UToolMenu* const Submenu = UToolMenus::Get()->ExtendMenu(SubmenuName);
-		if (Submenu)
-		{
-			FToolMenuSection& UnnamedSection = Submenu->FindOrAddSection(NAME_None);
-			UnnamedSection.AddEntry(
-				FToolMenuEntry::InitMenuEntry(FEditorViewportCommands::Get().RelativeCoordinateSystem_Parent)
-			);
-		}
-	}
 }
 
 void UE::ControlRig::PopulateControlRigViewportToolbarShowSubmenu(
diff --git a/Engine/Source/Editor/UnrealEd/Private/ViewportToolbar/UnrealEdViewportToolbar.cpp b/Engine/Source/Editor/UnrealEd/Private/ViewportToolbar/UnrealEdViewportToolbar.cpp
--- a/Engine/Source/Editor/UnrealEd/Private/ViewportToolbar/UnrealEdViewportToolbar.cpp
+++ b/Engine/Source/Editor/UnrealEd/Private/ViewportToolbar/UnrealEdViewportToolbar.cpp
@@ -4,6 +4,7 @@
 
 #include "DataDrivenShaderPlatformInfo.h"
 #include "DebugViewModeHelpers.h"
+#include "EditorInteractiveGizmoManager.h"
 #include "Editor/EditorPerformanceSettings.h"
 #include "EditorViewportClient.h"
 #include "EditorViewportCommands.h"
@@ -714,14 +715,20 @@ FToolMenuEntry CreateTransformsSubmenu()
 									FToolMenuEntry WorldCoords = FToolMenuEntry::InitMenuEntry(
 										FEditorViewportCommands::Get().RelativeCoordinateSystem_World
 									);
-									WorldCoords.UserInterfaceActionType = EUserInterfaceActionType::RadioButton;
 									UnnamedSection.AddEntry(WorldCoords);
 
 									FToolMenuEntry LocalCoords = FToolMenuEntry::InitMenuEntry(
 										FEditorViewportCommands::Get().RelativeCoordinateSystem_Local
 									);
-									LocalCoords.UserInterfaceActionType = EUserInterfaceActionType::RadioButton;
 									UnnamedSection.AddEntry(LocalCoords);
+
+									if (UEditorInteractiveGizmoManager::UsesNewTRSGizmos())
+									{
+										FToolMenuEntry ParentCoords = FToolMenuEntry::InitMenuEntry(
+											FEditorViewportCommands::Get().RelativeCoordinateSystem_Parent
+										);
+										UnnamedSection.AddEntry(ParentCoords);
+									}
 								}
 							)
 						);
