Deprecated StudioAnalytics
Superseeded by StudioTelemetry
Deprectation warning since 5.4 so now ready to deleted for 5.6

diff --git a/Engine/Source/Editor/UnrealEd/Private/Cooker/CookProfiling.cpp b/Engine/Source/Editor/UnrealEd/Private/Cooker/CookProfiling.cpp
--- a/Engine/Source/Editor/UnrealEd/Private/Cooker/CookProfiling.cpp
+++ b/Engine/Source/Editor/UnrealEd/Private/Cooker/CookProfiling.cpp
@@ -43,7 +43,6 @@
 #include "AnalyticsET.h"
 #include "AnalyticsEventAttribute.h"
 #include "IAnalyticsProviderET.h"
-#include "StudioAnalytics.h"
 #include "DerivedDataCacheInterface.h"
 #include "DerivedDataCacheUsageStats.h"
 #include "Virtualization/VirtualizationSystem.h"
diff --git a/Engine/Source/Editor/UnrealEd/Private/PlayLevel.cpp b/Engine/Source/Editor/UnrealEd/Private/PlayLevel.cpp
--- a/Engine/Source/Editor/UnrealEd/Private/PlayLevel.cpp
+++ b/Engine/Source/Editor/UnrealEd/Private/PlayLevel.cpp
@@ -118,7 +118,6 @@
 #include "Widgets/SBoxPanel.h"
 #include "Subsystems/AssetEditorSubsystem.h"
 #include "Async/Async.h"
-#include "StudioAnalytics.h"
 #include "UObject/SoftObjectPath.h"
 #include "IAssetViewport.h"
 #include "IPIEAuthorizer.h"
diff --git a/Engine/Source/Editor/UnrealEd/Private/Subsystems/AssetEditorSubsystem.cpp b/Engine/Source/Editor/UnrealEd/Private/Subsystems/AssetEditorSubsystem.cpp
--- a/Engine/Source/Editor/UnrealEd/Private/Subsystems/AssetEditorSubsystem.cpp
+++ b/Engine/Source/Editor/UnrealEd/Private/Subsystems/AssetEditorSubsystem.cpp
@@ -3,7 +3,6 @@
 #include "Subsystems/AssetEditorSubsystem.h"
 #include "AssetEditorMessages.h"
 #include "MessageEndpoint.h"
-#include "StudioAnalytics.h"
 #include "EngineAnalytics.h"
 #include "AnalyticsEventAttribute.h"
 #include "UObject/Package.h"
@@ -25,7 +24,6 @@
 #include "Misc/FeedbackContext.h"
 #include "Misc/ConfigCacheIni.h"
 #include "Misc/NamePermissionList.h"
-#include "StudioAnalytics.h"
 #include "EditorModeRegistry.h"
 #include "Tools/UEdMode.h"
 #include "AssetEditorMessages.h"
diff --git a/Engine/Source/Editor/UnrealEd/Private/UnrealEdMisc.cpp b/Engine/Source/Editor/UnrealEd/Private/UnrealEdMisc.cpp
--- a/Engine/Source/Editor/UnrealEd/Private/UnrealEdMisc.cpp
+++ b/Engine/Source/Editor/UnrealEd/Private/UnrealEdMisc.cpp
@@ -93,7 +93,6 @@
 #include "LauncherPlatformModule.h"
 #include "ILauncherServicesModule.h"
 #include "HAL/PlatformTime.h"
-#include "StudioAnalytics.h"
 #include "DeveloperToolSettingsDelegates.h"
 #include "Cooker/CookConfigAccessTracker.h"
 #include "Cooker/PackageBuildDependencyTracker.h"
diff --git a/Engine/Source/Runtime/Engine/Private/Analytics/StudioAnalytics.cpp b/Engine/Source/Runtime/Engine/Private/Analytics/StudioAnalytics.cpp
deleted file mode 100644
diff --git a/Engine/Source/Runtime/Engine/Private/EngineAnalytics.cpp b/Engine/Source/Runtime/Engine/Private/EngineAnalytics.cpp
--- a/Engine/Source/Runtime/Engine/Private/EngineAnalytics.cpp
+++ b/Engine/Source/Runtime/Engine/Private/EngineAnalytics.cpp
@@ -14,7 +14,6 @@
 #include "RHI.h"
 #include "RHIStats.h"
 #include "GenericPlatform/GenericPlatformCrashContext.h"
-#include "StudioAnalytics.h"
 #include "UObject/Class.h"
 #include "Containers/Set.h"
 #include "Internationalization/Culture.h"
diff --git a/Engine/Source/Runtime/Engine/Private/GameEngine.cpp b/Engine/Source/Runtime/Engine/Private/GameEngine.cpp
--- a/Engine/Source/Runtime/Engine/Private/GameEngine.cpp
+++ b/Engine/Source/Runtime/Engine/Private/GameEngine.cpp
@@ -55,7 +55,6 @@
 #include "Misc/EmbeddedCommunication.h"
 #include "Engine/CoreSettings.h"
 #include "EngineAnalytics.h"
-#include "StudioAnalytics.h"
 #include "Engine/DemoNetDriver.h"
 
 #include "AssetRegistry/AssetRegistryModule.h"
diff --git a/Engine/Source/Runtime/Engine/Private/UnrealEngine.cpp b/Engine/Source/Runtime/Engine/Private/UnrealEngine.cpp
--- a/Engine/Source/Runtime/Engine/Private/UnrealEngine.cpp
+++ b/Engine/Source/Runtime/Engine/Private/UnrealEngine.cpp
@@ -245,7 +245,6 @@ UnrealEngine.cpp: Implements the UEngine class and helpers.
 #include "Particles/ParticleSystemManager.h"
 #include "ObjectTrace.h"
 #include "ObjectPropertyTrace.h"
-#include "StudioAnalytics.h"
 #include "Animation/SkinWeightProfileManager.h"
 
 #include "Particles/ParticlePerfStatsManager.h"
diff --git a/Engine/Source/Runtime/Engine/Public/StudioAnalytics.h b/Engine/Source/Runtime/Engine/Public/StudioAnalytics.h
deleted file mode 100644
