[Viewport Toolbar] Coordinate System selection not working in Static Mesh editor, and potentially other editors

Adds possibility to hide coordinate controls from Editors.
Hides coordinate controls from Static Mesh Editor.

diff --git a/Engine/Source/Editor/StaticMeshEditor/Private/SStaticMeshEditorViewport.cpp b/Engine/Source/Editor/StaticMeshEditor/Private/SStaticMeshEditorViewport.cpp
--- a/Engine/Source/Editor/StaticMeshEditor/Private/SStaticMeshEditorViewport.cpp
+++ b/Engine/Source/Editor/StaticMeshEditor/Private/SStaticMeshEditorViewport.cpp
@@ -1027,6 +1027,8 @@ TSharedPtr<SWidget> SStaticMeshEditorViewport::BuildViewportToolbar()
 			UUnrealEdViewportToolbarContext* const ContextObject =
 				UE::UnrealEd::CreateViewportToolbarDefaultContext(SharedThis(this));
 
+			ContextObject->bShowCoordinateSystemControls = false;
+
 			ViewportToolbarContext.AddObject(ContextObject);
 		}
 	}
diff --git a/Engine/Source/Editor/UnrealEd/Private/ViewportToolbar/UnrealEdViewportToolbar.cpp b/Engine/Source/Editor/UnrealEd/Private/ViewportToolbar/UnrealEdViewportToolbar.cpp
--- a/Engine/Source/Editor/UnrealEd/Private/ViewportToolbar/UnrealEdViewportToolbar.cpp
+++ b/Engine/Source/Editor/UnrealEd/Private/ViewportToolbar/UnrealEdViewportToolbar.cpp
@@ -655,61 +655,70 @@ FToolMenuEntry CreateTransformsSubmenu()
 
 					// Build a submenu for selecting the coordinate system to use.
 					{
-						TransformToolsSection.AddSeparator("CoordinateSystemSeparator");
-
-						FToolMenuEntry& CoordianteSystemSubmenu = TransformToolsSection.AddSubMenu(
-							"CoordinateSystem",
-							LOCTEXT("CoordinateSystemLabel", "Coordinate System"),
-							LOCTEXT("CoordinateSystemTooltip", "Select between coordinate systems"),
-							FNewToolMenuDelegate::CreateLambda(
-								[](UToolMenu* InSubmenu)
-								{
-									FToolMenuSection& UnnamedSection = InSubmenu->FindOrAddSection(NAME_None);
-
-									FToolMenuEntry WorldCoords = FToolMenuEntry::InitMenuEntry(
-										FEditorViewportCommands::Get().RelativeCoordinateSystem_World
-									);
-									UnnamedSection.AddEntry(WorldCoords);
-
-									FToolMenuEntry LocalCoords = FToolMenuEntry::InitMenuEntry(
-										FEditorViewportCommands::Get().RelativeCoordinateSystem_Local
-									);
-									UnnamedSection.AddEntry(LocalCoords);
-
-									if (UEditorInteractiveGizmoManager::UsesNewTRSGizmos())
-									{
-										FToolMenuEntry ParentCoords = FToolMenuEntry::InitMenuEntry(
-											FEditorViewportCommands::Get().RelativeCoordinateSystem_Parent
-										);
-										UnnamedSection.AddEntry(ParentCoords);
-									}
-								}
-							)
-						);
-
-						// Set the icon based on the current coordinate system and fall back to the Local icon.
+						if (const UUnrealEdViewportToolbarContext* const Context =
+								Submenu->FindContext<UUnrealEdViewportToolbarContext>())
 						{
-							TWeakPtr<SEditorViewport> WeakViewport;
-							if (UUnrealEdViewportToolbarContext* const Context =
-									Submenu->FindContext<UUnrealEdViewportToolbarContext>())
+							if (TSharedPtr<SEditorViewport> Viewport = Context->Viewport.Pin())
 							{
-								WeakViewport = Context->Viewport;
-							}
-
-							CoordianteSystemSubmenu.Icon = TAttribute<FSlateIcon>::CreateLambda(
-								[WeakViewport]() -> FSlateIcon
+								if (Context->bShowCoordinateSystemControls)
 								{
-									ECoordSystem CoordSystem = ECoordSystem::COORD_Local;
-									if (const TSharedPtr<SEditorViewport> EditorViewport = WeakViewport.Pin())
+									// Build a submenu for selecting the coordinate system to use.
 									{
-										CoordSystem = EditorViewport->GetViewportClient()->GetWidgetCoordSystemSpace();
+										TransformToolsSection.AddSeparator("CoordinateSystemSeparator");
+
+										FToolMenuEntry& CoordinateSystemSubmenu = TransformToolsSection.AddSubMenu(
+											"CoordinateSystem",
+											LOCTEXT("CoordinateSystemLabel", "Coordinate System"),
+											LOCTEXT("CoordinateSystemTooltip", "Select between coordinate systems"),
+											FNewToolMenuDelegate::CreateLambda(
+												[](UToolMenu* InSubmenu)
+												{
+													FToolMenuSection& UnnamedSection =
+														InSubmenu->FindOrAddSection(NAME_None);
+
+													FToolMenuEntry WorldCoords = FToolMenuEntry::InitMenuEntry(
+														FEditorViewportCommands::Get().RelativeCoordinateSystem_World
+													);
+													UnnamedSection.AddEntry(WorldCoords);
+
+													if (UEditorInteractiveGizmoManager::UsesNewTRSGizmos())
+													{
+														FToolMenuEntry ParentCoords = FToolMenuEntry::InitMenuEntry(
+															FEditorViewportCommands::Get().RelativeCoordinateSystem_Parent
+														);
+														UnnamedSection.AddEntry(ParentCoords);
+													}
+
+													FToolMenuEntry LocalCoords = FToolMenuEntry::InitMenuEntry(
+														FEditorViewportCommands::Get().RelativeCoordinateSystem_Local
+													);
+													UnnamedSection.AddEntry(LocalCoords);
+												}
+											)
+										);
+
+										// Set the icon based on the current coordinate system and fall back to the Local icon.
+										{
+											CoordinateSystemSubmenu.Icon = TAttribute<FSlateIcon>::CreateLambda(
+												[WeakViewport = Viewport.ToWeakPtr()]() -> FSlateIcon
+												{
+													ECoordSystem CoordSystem = ECoordSystem::COORD_Local;
+													if (const TSharedPtr<SEditorViewport> EditorViewport =
+															WeakViewport.Pin())
+													{
+														CoordSystem =
+															EditorViewport->GetViewportClient()->GetWidgetCoordSystemSpace();
+													}
+													return GetIconFromCoordSystem(CoordSystem);
+												}
+											);
+										}
+										CoordinateSystemSubmenu.ToolBarData.LabelOverride = FText();
+										CoordinateSystemSubmenu.SetShowInToolbarTopLevel(true);
 									}
-									return GetIconFromCoordSystem(CoordSystem);
 								}
-							);
+							}
 						}
-						CoordianteSystemSubmenu.ToolBarData.LabelOverride = FText();
-						CoordianteSystemSubmenu.SetShowInToolbarTopLevel(true);
 					}
 				}
 
diff --git a/Engine/Source/Editor/UnrealEd/Public/ViewportToolbar/UnrealEdViewportToolbarContext.h b/Engine/Source/Editor/UnrealEd/Public/ViewportToolbar/UnrealEdViewportToolbarContext.h
--- a/Engine/Source/Editor/UnrealEd/Public/ViewportToolbar/UnrealEdViewportToolbarContext.h
+++ b/Engine/Source/Editor/UnrealEd/Public/ViewportToolbar/UnrealEdViewportToolbarContext.h
@@ -29,4 +29,9 @@ public:
 	TWeakPtr<SEditorViewport> Viewport;
 	UE::UnrealEd::IsViewModeSupportedDelegate IsViewModeSupported;
 	UE::UnrealEd::DoesViewModeMenuShowSectionDelegate DoesViewModeMenuShowSection;
+
+	/**
+	 * Whether the current editor should show menu entries to select between coordinate systems (e.g. Local vs World)
+	 */
+	bool bShowCoordinateSystemControls = true;
 };
