MPCook: CookPackageSplitter: Second attempt at fix for failed assertion checkf(CachedObjectsInOuterInfo.IsEmpty(),...) in FCookGenerationInfo::Uninitialize. Added speculative fix and more feedback.
#rnx

diff --git a/Engine/Source/Editor/UnrealEd/Private/Cooker/CookGenerationHelper.cpp b/Engine/Source/Editor/UnrealEd/Private/Cooker/CookGenerationHelper.cpp
--- a/Engine/Source/Editor/UnrealEd/Private/Cooker/CookGenerationHelper.cpp
+++ b/Engine/Source/Editor/UnrealEd/Private/Cooker/CookGenerationHelper.cpp
@@ -1442,6 +1442,16 @@ void FGenerationHelper::ResetSaveState(FCookGenerationInfo& Info, UPackage* Pack
 		Info.CachedObjectsInOuterInfo.Empty();
 		Info.SetHasTakenOverCachedCookedPlatformData(false);
 	}
+	else if (Info.CachedObjectsInOuterInfo.Num() > 0)
+	{
+		UE_LOG(LogCook, Error,
+			TEXT("CookPackageSplitter logic error: A %s package has !Info.HasTakenOverCachedCookedPlatformData(), but Info.CachedObjectsInOuterInfo.Num() == %d > 0. This should be impossible.\n")
+			TEXT("Package=%s, PackageState=%s, SaveSubState=%s"),
+			Info.IsGenerator() ? TEXT("generator") : TEXT("generated"), Info.CachedObjectsInOuterInfo.Num(),
+			*Info.GetPackageName(), LexToString(Info.PackageData->GetState()), LexToString(Info.PackageData->GetSaveSubState()));
+		Info.CachedObjectsInOuterInfo.Empty();
+	}
+
 	Info.SetHasIssuedUndeclaredMovedObjectsWarning(false);
 
 	// Clear KeepReferencedPackages; we no longer have a contract that we keep them referenced, except for the
@@ -2031,8 +2041,9 @@ void FCookGenerationInfo::Uninitialize()
 	// calling ReleaseCookedPlatformData.
 	checkf(CachedObjectsInOuterInfo.IsEmpty(),
 		TEXT("Cooker bug: Expected FCookGenerationInfo::Uninitialize to not be called for a package in an active save state, ")
-		TEXT("but %s package %s has non-empty CachedObjectsInOuterInfo."),
-		bGenerator ? TEXT("generator") : TEXT("generated"), *GetPackageName());
+		TEXT("but %s package %s has non-empty CachedObjectsInOuterInfo. Package State == %s, SaveSubState == %s"),
+		bGenerator ? TEXT("generator") : TEXT("generated"), *GetPackageName(),
+		LexToString(PackageData->GetState()), LexToString(PackageData->GetSaveSubState()));
 	// Keep SavedOnWorker; it is allowed in the uninitialized state
 	bCreateAsMap = false;
 	bHasCreatedPackage = false;
@@ -2138,6 +2149,7 @@ EPollStatus FCookGenerationInfo::RefreshPackageObjects(FGenerationHelper& Genera
 			}
 		}
 	}
+	SetHasTakenOverCachedCookedPlatformData(true); // Set it again in case we unset it at any point
 	bOutFoundNewObjects = FirstNewObject != nullptr;
 
 	if (FirstNewObject != nullptr && DemotionState != ESaveSubState::Last)
