Enable "Editor.AutoSave.InteractiveSaveManagesUncontrolledChangelist" by default

#rnx

diff --git a/Engine/Source/Editor/UnrealEd/Private/PackageAutoSaver.cpp b/Engine/Source/Editor/UnrealEd/Private/PackageAutoSaver.cpp
--- a/Engine/Source/Editor/UnrealEd/Private/PackageAutoSaver.cpp
+++ b/Engine/Source/Editor/UnrealEd/Private/PackageAutoSaver.cpp
@@ -46,7 +46,7 @@ namespace PackageAutoSaver
 	static bool bUseUncontrolledChangelists = true;
 	static FAutoConsoleVariableRef CVarUseUncontrolledChangelists(TEXT("Editor.AutoSave.UseUncontrolledChangelists"), bUseUncontrolledChangelists, TEXT("True to use uncontrolled changelists (if possible) when using the BackupAndOverwrite auto-save method."));
 
-	static bool bInteractiveSaveManagesUncontrolledChangelist = false;
+	static bool bInteractiveSaveManagesUncontrolledChangelist = true;
 	static FAutoConsoleVariableRef CVarInteractiveSaveManagesUncontrolledChangelist(TEXT("Editor.AutoSave.InteractiveSaveManagesUncontrolledChangelist"), bInteractiveSaveManagesUncontrolledChangelist, TEXT("True if an interactive save should manage which entries stay in the auto-save uncontrolled changelist (ie, promoting them to a controlled save, or moving them to the default uncontrolled changelist)."));
 
 	static const FUncontrolledChangelist AutoSaveUncontrolledChangelist(FGuid(0x0D227D90, 0x2AC34632, 0xB1A35533, 0xF3CC134E));
