Disable inline editing of some attributes in the Root Material Node

diff --git a/Engine/Source/Editor/UnrealEd/Private/MaterialGraphNode_Root.cpp b/Engine/Source/Editor/UnrealEd/Private/MaterialGraphNode_Root.cpp
--- a/Engine/Source/Editor/UnrealEd/Private/MaterialGraphNode_Root.cpp
+++ b/Engine/Source/Editor/UnrealEd/Private/MaterialGraphNode_Root.cpp
@@ -40,20 +40,22 @@ void UMaterialGraphNode_Root::UpdateInputUseConstant(UEdGraphPin* Pin, bool bUse
 	case MP_Metallic:			EditorOnlyData->Metallic.UseConstant = bUseConstant; break;
 	case MP_Specular:			EditorOnlyData->Specular.UseConstant = bUseConstant; break;
 	case MP_Roughness:			EditorOnlyData->Roughness.UseConstant = bUseConstant; break;
-	case MP_Anisotropy:			EditorOnlyData->Anisotropy.UseConstant = bUseConstant; break;
 	case MP_Normal:				EditorOnlyData->Normal.UseConstant = bUseConstant; break;
 	case MP_Tangent:			EditorOnlyData->Tangent.UseConstant = bUseConstant; break;
-	case MP_WorldPositionOffset:EditorOnlyData->WorldPositionOffset.UseConstant = bUseConstant; break;
 	case MP_SubsurfaceColor:	EditorOnlyData->SubsurfaceColor.UseConstant = bUseConstant; break;
 	case MP_CustomData0:		EditorOnlyData->ClearCoat.UseConstant = bUseConstant; break;
 	case MP_CustomData1:		EditorOnlyData->ClearCoatRoughness.UseConstant = bUseConstant; break;
 	case MP_AmbientOcclusion:	EditorOnlyData->AmbientOcclusion.UseConstant = bUseConstant; break;
 	case MP_Refraction:			EditorOnlyData->Refraction.UseConstant = bUseConstant; break;
-	case MP_PixelDepthOffset:	EditorOnlyData->PixelDepthOffset.UseConstant = bUseConstant; break;
 	case MP_ShadingModel:		break; // TODO, see notes in CreateInputPins
 	case MP_FrontMaterial:		break; // TODO, see notes in CreateInputPins
 	case MP_SurfaceThickness:	EditorOnlyData->SurfaceThickness.UseConstant = bUseConstant; break;
-	case MP_Displacement:		EditorOnlyData->Displacement.UseConstant = bUseConstant; break;
+
+	// TODO: Disabling these for the time being, as these values cause issues/artefacts when directly modified inline (UE-219232)
+	case MP_PixelDepthOffset:	EditorOnlyData->PixelDepthOffset.UseConstant = false; break;
+	case MP_Anisotropy:			EditorOnlyData->Anisotropy.UseConstant = false; break;
+	case MP_WorldPositionOffset:EditorOnlyData->WorldPositionOffset.UseConstant = false; break;
+	case MP_Displacement:		EditorOnlyData->Displacement.UseConstant = false; break;
 	default:
 		if (Property >= MP_CustomizedUVs0 && Property <= MP_CustomizedUVs7)
 		{
@@ -108,20 +110,22 @@ void UMaterialGraphNode_Root::PinDefaultValueChanged(UEdGraphPin* Pin)
 	case MP_Metallic:			EditorOnlyData->Metallic.DefaultValueChanged(Pin->DefaultValue); break;
 	case MP_Specular:			EditorOnlyData->Specular.DefaultValueChanged(Pin->DefaultValue); break;
 	case MP_Roughness:			EditorOnlyData->Roughness.DefaultValueChanged(Pin->DefaultValue); break;
-	case MP_Anisotropy:			EditorOnlyData->Anisotropy.DefaultValueChanged(Pin->DefaultValue); break;
 	case MP_Normal:				EditorOnlyData->Normal.DefaultValueChanged(Pin->DefaultValue); break;
 	case MP_Tangent:			EditorOnlyData->Tangent.DefaultValueChanged(Pin->DefaultValue); break;
-	case MP_WorldPositionOffset:EditorOnlyData->WorldPositionOffset.DefaultValueChanged(Pin->DefaultValue); break;
 	case MP_SubsurfaceColor:	EditorOnlyData->SubsurfaceColor.DefaultValueChanged(Pin->DefaultValue); break;
 	case MP_CustomData0:		EditorOnlyData->ClearCoat.DefaultValueChanged(Pin->DefaultValue); break;
 	case MP_CustomData1:		EditorOnlyData->ClearCoatRoughness.DefaultValueChanged(Pin->DefaultValue); break;
 	case MP_AmbientOcclusion:	EditorOnlyData->AmbientOcclusion.DefaultValueChanged(Pin->DefaultValue); break;
 	case MP_Refraction:			EditorOnlyData->Refraction.DefaultValueChanged(Pin->DefaultValue); break;
-	case MP_PixelDepthOffset:	EditorOnlyData->PixelDepthOffset.DefaultValueChanged(Pin->DefaultValue); break;
 	case MP_ShadingModel:		break; // TODO, see notes in CreateInputPins
 	case MP_FrontMaterial:		break; // TODO, see notes in CreateInputPins
 	case MP_SurfaceThickness:	EditorOnlyData->SurfaceThickness.DefaultValueChanged(Pin->DefaultValue); break;
-	case MP_Displacement:		EditorOnlyData->Displacement.DefaultValueChanged(Pin->DefaultValue); break; 
+
+	// TODO: Disabling these for the time being, as these values cause issues/artefacts when directly modified inline (UE-219232)
+	// case MP_Anisotropy:			EditorOnlyData->Anisotropy.DefaultValueChanged(Pin->DefaultValue); break;
+	// case MP_PixelDepthOffset:	EditorOnlyData->PixelDepthOffset.DefaultValueChanged(Pin->DefaultValue); break;
+	// case MP_WorldPositionOffset:EditorOnlyData->WorldPositionOffset.DefaultValueChanged(Pin->DefaultValue); break;
+	// case MP_Displacement:		EditorOnlyData->Displacement.DefaultValueChanged(Pin->DefaultValue); break; 
 	default:
 		if (Property >= MP_CustomizedUVs0 && Property <= MP_CustomizedUVs7)
 		{
@@ -208,31 +212,30 @@ void UMaterialGraphNode_Root::CreateInputPins()
 			case MP_Metallic:
 			case MP_Specular:
 			case MP_Roughness:
-			case MP_Anisotropy:
 			case MP_Opacity:
 			case MP_OpacityMask:
 			case MP_CustomData0:
 			case MP_CustomData1:
 			case MP_AmbientOcclusion:
-			case MP_Displacement:
 			case MP_SurfaceThickness:
 				PinSubCategory = UMaterialGraphSchema::PSC_Red;
 				break;
 
 			case MP_PixelDepthOffset:
-				PinSubCategory = UMaterialGraphSchema::PSC_Red;
-				switch (MaterialGraph->Material->PixelDepthOffsetMode)
-				{
-				case EPixelDepthOffsetMode::PDOM_Legacy:
-					TempStr = TEXT("Legacy");
-					break;
-				case EPixelDepthOffsetMode::PDOM_AlongCameraVector:
-					TempStr = TEXT("Camera Vector");
-					break;
-				default:
-					TempStr = TEXT("UNKNOWN EPixelDepthOffsetMode");
-				}
-				MaterialInputName = *FString::Printf(TEXT("%s (%s)"), *MaterialInputName.ToString(), *TempStr);
+				// TODO: Disabling these for the time being, as these values cause issues/artefacts when directly modified inline (UE-219232)
+				// PinSubCategory = UMaterialGraphSchema::PSC_Red;
+				// switch (MaterialGraph->Material->PixelDepthOffsetMode)
+				// {
+				// case EPixelDepthOffsetMode::PDOM_Legacy:
+				// 	TempStr = TEXT("Legacy");
+				// 	break;
+				// case EPixelDepthOffsetMode::PDOM_AlongCameraVector:
+				// 	TempStr = TEXT("Camera Vector");
+				// 	break;
+				// default:
+				// 	TempStr = TEXT("UNKNOWN EPixelDepthOffsetMode");
+				// }
+				// MaterialInputName = *FString::Printf(TEXT("%s (%s)"), *MaterialInputName.ToString(), *TempStr);
 				break;
 				
 			case MP_Refraction:
@@ -259,7 +262,6 @@ void UMaterialGraphNode_Root::CreateInputPins()
 				
 			case MP_Normal:
 			case MP_Tangent:
-			case MP_WorldPositionOffset:
 				PinSubCategory = UMaterialGraphSchema::PSC_RGB;
 				break;
 				
@@ -276,7 +278,13 @@ void UMaterialGraphNode_Root::CreateInputPins()
 				// If these are based on Enum types, Enum pin construction requires a UEnum object to create UI with
 				// which these constants don't have.
 				break;
-				
+
+			// TODO: Disabling these for the time being, as these values cause issues/artefacts when directly modified inline (UE-219232)
+			case MP_WorldPositionOffset:
+			case MP_Anisotropy:
+			case MP_Displacement:
+				break;
+			
 			default:
 				if (Property >= MP_CustomizedUVs0 && Property <= MP_CustomizedUVs7)
 				{
@@ -297,20 +305,23 @@ void UMaterialGraphNode_Root::CreateInputPins()
 			case MP_Metallic:			InputPin->DefaultValue = EditorOnlyData->Metallic.GetDefaultValue(); break;
 			case MP_Specular:			InputPin->DefaultValue = EditorOnlyData->Specular.GetDefaultValue(); break;
 			case MP_Roughness:			InputPin->DefaultValue = EditorOnlyData->Roughness.GetDefaultValue(); break;
-			case MP_Anisotropy:			InputPin->DefaultValue = EditorOnlyData->Anisotropy.GetDefaultValue(); break;
+			
 			case MP_Normal:				InputPin->DefaultValue = EditorOnlyData->Normal.GetDefaultValue(); break;
 			case MP_Tangent:			InputPin->DefaultValue = EditorOnlyData->Tangent.GetDefaultValue(); break;
-			case MP_WorldPositionOffset:InputPin->DefaultValue = EditorOnlyData->WorldPositionOffset.GetDefaultValue(); break;
 			case MP_SubsurfaceColor:	InputPin->DefaultValue = EditorOnlyData->SubsurfaceColor.GetDefaultValue(); break;
 			case MP_CustomData0:		InputPin->DefaultValue = EditorOnlyData->ClearCoat.GetDefaultValue(); break;
 			case MP_CustomData1:		InputPin->DefaultValue = EditorOnlyData->ClearCoatRoughness.GetDefaultValue(); break;
 			case MP_AmbientOcclusion:	InputPin->DefaultValue = EditorOnlyData->AmbientOcclusion.GetDefaultValue(); break;
 			case MP_Refraction:			InputPin->DefaultValue = EditorOnlyData->Refraction.GetDefaultValue(); break;
-			case MP_PixelDepthOffset:	InputPin->DefaultValue = EditorOnlyData->PixelDepthOffset.GetDefaultValue(); break;
 			case MP_ShadingModel:		break; // TODO
 			case MP_FrontMaterial:		break; // TODO
 			case MP_SurfaceThickness:	InputPin->DefaultValue = EditorOnlyData->SurfaceThickness.GetDefaultValue(); break;
-			case MP_Displacement:		InputPin->DefaultValue = EditorOnlyData->Displacement.GetDefaultValue(); break;
+
+			// TODO: Disabling these for the time being, as these values cause issues/artefacts when directly modified inline (UE-219232)
+			// case MP_WorldPositionOffset:InputPin->DefaultValue = EditorOnlyData->WorldPositionOffset.GetDefaultValue(); break;
+			// case MP_PixelDepthOffset:	InputPin->DefaultValue = EditorOnlyData->PixelDepthOffset.GetDefaultValue(); break;
+			// case MP_Anisotropy:			InputPin->DefaultValue = EditorOnlyData->Anisotropy.GetDefaultValue(); break;
+			// case MP_Displacement:		InputPin->DefaultValue = EditorOnlyData->Displacement.GetDefaultValue(); break;
 			default:
 				if (Property >= MP_CustomizedUVs0 && Property <= MP_CustomizedUVs7)
 				{
