Fix various typos with `parition` instead of `partition`.

diff --git a/Engine/Source/Editor/UnrealEd/Private/StaticLightingSystem/StaticLightingSystem.cpp b/Engine/Source/Editor/UnrealEd/Private/StaticLightingSystem/StaticLightingSystem.cpp
--- a/Engine/Source/Editor/UnrealEd/Private/StaticLightingSystem/StaticLightingSystem.cpp
+++ b/Engine/Source/Editor/UnrealEd/Private/StaticLightingSystem/StaticLightingSystem.cpp
@@ -930,7 +930,7 @@ void FStaticLightingSystem::InvalidateStaticLighting()
 	}
 }
 
-void UpdateStaticLightingWorldParitionHLODTreeIndices(TMultiMap<AActor*, FStaticLightingMesh*>& ActorMeshMap, AWorldPartitionHLOD* LODActor, uint32 HLODTreeIndex, uint32& HLODLeafIndex, const TMap<FGuid, AActor*>& ActorInstanceGuidToActorPtr)
+void UpdateStaticLightingWorldPartitionHLODTreeIndices(TMultiMap<AActor*, FStaticLightingMesh*>& ActorMeshMap, AWorldPartitionHLOD* LODActor, uint32 HLODTreeIndex, uint32& HLODLeafIndex, const TMap<FGuid, AActor*>& ActorInstanceGuidToActorPtr)
 {	
 	check(LODActor && HLODTreeIndex > 0);
 
@@ -971,7 +971,7 @@ void UpdateStaticLightingWorldParitionHLODTreeIndices(TMultiMap<AActor*, FStatic
 			// If sub actor is an WP HLODActor iterate over that too
 			if (AWorldPartitionHLOD* HLODSubActor = Cast<AWorldPartitionHLOD>(SubActor))
 			{
-				UpdateStaticLightingWorldParitionHLODTreeIndices(ActorMeshMap, HLODSubActor, HLODTreeIndex, HLODLeafIndex, ActorInstanceGuidToActorPtr);
+				UpdateStaticLightingWorldPartitionHLODTreeIndices(ActorMeshMap, HLODSubActor, HLODTreeIndex, HLODLeafIndex, ActorInstanceGuidToActorPtr);
 			}
 			else
 			{
@@ -1428,7 +1428,7 @@ void FStaticLightingSystem::GatherStaticLightingInfo(bool bRebuildDirtyGeometryF
 		{
 			WPHLODLeafIndex = 0;
 
-			UpdateStaticLightingWorldParitionHLODTreeIndices(ActorMeshMap, LODActor, WPHLODTreeIndex, WPHLODLeafIndex, ActorInstanceGuidToActorPtr);
+			UpdateStaticLightingWorldPartitionHLODTreeIndices(ActorMeshMap, LODActor, WPHLODTreeIndex, WPHLODLeafIndex, ActorInstanceGuidToActorPtr);
 
 			++WPHLODTreeIndex;
 		}
