Add TRS gizmo toggle to the new viewport toolbar

diff --git a/Engine/Source/Editor/UnrealEd/Private/ViewportToolbar/UnrealEdViewportToolbar.cpp b/Engine/Source/Editor/UnrealEd/Private/ViewportToolbar/UnrealEdViewportToolbar.cpp
--- a/Engine/Source/Editor/UnrealEd/Private/ViewportToolbar/UnrealEdViewportToolbar.cpp
+++ b/Engine/Source/Editor/UnrealEd/Private/ViewportToolbar/UnrealEdViewportToolbar.cpp
@@ -725,6 +725,34 @@ FToolMenuEntry CreateTransformsSubmenu()
 				{
 					FToolMenuSection& GizmoSection = Submenu->FindOrAddSection("Gizmo", LOCTEXT("GizmoLabel", "Gizmo"));
 
+					{
+						FUIAction Action;
+						Action.ExecuteAction = FExecuteAction::CreateLambda(
+							[]()
+							{
+								UEditorInteractiveGizmoManager::SetUsesNewTRSGizmos(
+									!UEditorInteractiveGizmoManager::UsesNewTRSGizmos()
+								);
+							}
+						);
+						Action.GetActionCheckState = FGetActionCheckState::CreateLambda(
+							[]() -> ECheckBoxState
+							{
+								return UEditorInteractiveGizmoManager::UsesNewTRSGizmos() ? ECheckBoxState::Checked
+																						  : ECheckBoxState::Unchecked;
+							}
+						);
+
+						GizmoSection.AddMenuEntry(
+							"TRSGizmoToggle",
+							LOCTEXT("TRSGizmoToggleLabel", "Use New TRS Gizmos"),
+							LOCTEXT("TRSGizmoToggleTooltip", "Whether or not to use the new TRS gizmos"),
+							FSlateIcon(),
+							Action,
+							EUserInterfaceActionType::ToggleButton
+						);
+					}
+
 					TSharedRef<SWidget> GizmoScaleWidget =
 						// clang-format off
 						SNew(SHorizontalBox)
