CleanupAfterSuccessfulDelete will now skip reference checks for all external packages, not only those containing external actors, external objects or deleted object placeholders.
That way, we can include dummy objects in deleted packages to display the proper deleted object type, etc.

diff --git a/Engine/Source/Editor/UnrealEd/Private/ObjectTools.cpp b/Engine/Source/Editor/UnrealEd/Private/ObjectTools.cpp
--- a/Engine/Source/Editor/UnrealEd/Private/ObjectTools.cpp
+++ b/Engine/Source/Editor/UnrealEd/Private/ObjectTools.cpp
@@ -74,7 +74,6 @@
 #include "Dialogs/DlgReferenceTree.h"
 #include "AssetRegistry/ARFilter.h"
 #include "AssetDeleteModel.h"
-#include "DeletedObjectPlaceholder.h"
 #include "Dialogs/SPrivateAssetsDialog.h"
 #include "Dialogs/SDeleteAssetsDialog.h"
 #include "AudioDevice.h"
@@ -2470,14 +2469,20 @@ namespace ObjectTools
 
 			auto ShouldSkipReferenceCheck = [](const UPackage* InPackage)
 			{
-				// Skip external object packages when considering whether to clear the transaction buffer, as you should be able to undo deleting an actor, an actor folder, etc. (but not an asset)
-				// If an external object package is kept alive by the transaction buffer then it will be re-marked as "newly created" further down this function
-				if (InPackage->GetName().Contains(FPackagePath::GetExternalActorsFolderName()) || InPackage->GetName().Contains(FPackagePath::GetExternalObjectsFolderName()))
+				auto IsExternalPackage = [](const UPackage* Package) -> bool
 				{
-					return true;
-				}
+					bool bIsExternalPackage = false;
+					ForEachObjectWithPackage(Package, [Package, &bIsExternalPackage](UObject* InObject)
+					{
+						bIsExternalPackage = InObject->GetOutermostObject()->GetPackage() != Package;
+						return !bIsExternalPackage;
+					});
+					return bIsExternalPackage;
+				};
 
-				return !!UDeletedObjectPlaceholder::FindInPackage(InPackage);
+				// Skip external object packages when considering whether to clear the transaction buffer, as you should be able to undo deleting an actor, an actor folder, etc. (but not an asset)
+				// If an external object package is kept alive by the transaction buffer then it will be re-marked as "newly created" further down this function
+				return IsExternalPackage(InPackage);
 			};
 
 			if ( Package != nullptr && bPerformReferenceCheck && !ShouldSkipReferenceCheck(Package) )
