Improve the behavior of warnings in the new viewport toolbar

* Always apply the StyleNameOverride "ViewportToolbarWarning" to the realtime toggle and the viewport scalability submenu. This will not change how they look when they appear in menus, but when they're raised to the top-level toolbar (i.e. when we're warning) they will look like warnings.
* Make use of the fact that AllowClipping is a TAttribute<bool> and bind that to the logic behind whether or not the realtime toggle and viewport scalability should warn or not. When they warn, their AllowClipping attributes will return false so they stay around in the toolbar and their warnings are clearly visible even if viewport toolbar size is constrained and other entries are clipped.
* Remove the hack where UToolMenus::Get()->RefreshAllWidgets() was called by the realtime toggle when it was toggled. This was needed previously to make sure the style was updated, but since we now always apply the style we no longer need this hack.

diff --git a/Engine/Source/Editor/UnrealEd/Private/ViewportToolbar/UnrealEdViewportToolbar.cpp b/Engine/Source/Editor/UnrealEd/Private/ViewportToolbar/UnrealEdViewportToolbar.cpp
--- a/Engine/Source/Editor/UnrealEd/Private/ViewportToolbar/UnrealEdViewportToolbar.cpp
+++ b/Engine/Source/Editor/UnrealEd/Private/ViewportToolbar/UnrealEdViewportToolbar.cpp
@@ -3065,11 +3065,19 @@ FToolMenuEntry CreateScalabilitySubmenu()
 				ScalabilitySubmenu.SetShowInToolbarTopLevel(TAttribute<bool>::Create(&UE::UnrealEd::IsScalabilityWarningVisible
 				));
 
-				if (UE::UnrealEd::IsScalabilityWarningVisible())
-				{
-					ScalabilitySubmenu.StyleNameOverride = "ViewportToolbarWarning";
-					ScalabilitySubmenu.ToolBarData.ResizeParams.AllowClipping = false;
-				}
+				ScalabilitySubmenu.ToolBarData.ResizeParams.AllowClipping = TAttribute<bool>::CreateLambda(
+					[]() -> bool
+					{
+						return !UE::UnrealEd::IsScalabilityWarningVisible();
+					}
+				);
+
+				// We need the warning style when we've raised this entry to the top-level toolbar to draw
+				// attention to it. However, styles are not driven by an attribute, so they cannot be updated after an
+				// FToolMenuEntry has been created. Therefore we always apply the style instead because the style
+				// will not make this entry look different when it appears in a menu, but when it appears in a toolbar
+				// we get the desired warning effect.
+				ScalabilitySubmenu.StyleNameOverride = "ViewportToolbarWarning";
 			}
 		)
 	);
@@ -3206,11 +3214,6 @@ FToolMenuEntry CreateToggleRealtimeEntry()
 							if (TSharedPtr<SEditorViewport> EditorViewport = EditorViewportWeak.Pin())
 							{
 								EditorViewport->OnToggleRealtime();
-
-								// HACK: We unfortunately have to call UToolMenu::RefreshAllWidgets here to be able
-								// to dynamically set StyleNameOverride below and have that take effect. Without refreshing
-								// widgets, the style name is not applied because the UI has already been built.
-								UToolMenus::Get()->RefreshAllWidgets();
 							}
 						}
 					);
@@ -3309,14 +3312,29 @@ FToolMenuEntry CreateToggleRealtimeEntry()
 					));
 				}
 
-				if (const TSharedPtr<SEditorViewport> EditorViewport = EditorViewportWeak.Pin())
-				{
-					if (!EditorViewport->IsRealtime())
+				ToggleRealtime.ToolBarData.ResizeParams.AllowClipping = TAttribute<bool>::CreateLambda(
+					[EditorViewportWeak]() -> bool
 					{
-						ToggleRealtime.StyleNameOverride = "ViewportToolbarWarning";
-						ToggleRealtime.ToolBarData.ResizeParams.AllowClipping = false;
+						if (const TSharedPtr<SEditorViewport> EditorViewport = EditorViewportWeak.Pin())
+						{
+							// Don't allow this entry to be clipped from the toolbar when the viewport isn't realtime.
+							// This to avoid that the non-relatime-viewport warning is hidden as the toolbar is shrunk.
+							if (!EditorViewport->IsRealtime())
+							{
+								return false;
+							}
+						}
+
+						return true;
 					}
-				}
+				);
+
+				// We need the warning style when we've raised this entry to the top-level toolbar to draw
+				// attention to it. However, styles are not driven by an attribute, so they cannot be updated after an
+				// FToolMenuEntry has been created. Therefore we always apply the style instead because the style
+				// will not make this entry look different when it appears in a menu, but when it appears in a toolbar
+				// we get the desired warning effect.
+				ToggleRealtime.StyleNameOverride = "ViewportToolbarWarning";
 
 				InnerSection.AddEntry(ToggleRealtime);
 			}
