[HWRT] Refactored Ray Tracing Debug Visualization menu.
- Moved logic to a new class FRayTracingVisualizationData, matching the pattern used by other systems.
- Separated visualization modes into groups (Standard, Performance, Other)
- Visualization modes in "Other" group are hidden by default (can be toggled using r.RayTracing.Visualize.Other)

diff --git a/Engine/Source/Editor/UnrealEd/Private/EditorViewportClient.cpp b/Engine/Source/Editor/UnrealEd/Private/EditorViewportClient.cpp
--- a/Engine/Source/Editor/UnrealEd/Private/EditorViewportClient.cpp
+++ b/Engine/Source/Editor/UnrealEd/Private/EditorViewportClient.cpp
@@ -62,6 +62,7 @@
 #include "BufferVisualizationData.h"
 #include "NaniteVisualizationData.h"
 #include "LumenVisualizationData.h"
+#include "RayTracingVisualizationData.h"
 #include "SubstrateVisualizationData.h"
 #include "GroomVisualizationData.h"
 #include "VirtualShadowMapVisualizationData.h"
@@ -4252,7 +4253,7 @@ void FEditorViewportClient::Draw(FViewport* InViewport, FCanvas* Canvas)
 	const bool bVisualizeBufferEnabled = CurrentViewMode == VMI_VisualizeBuffer && CurrentBufferVisualizationMode != NAME_None;
 	const bool bRayTracingDebugEnabled = CurrentViewMode == VMI_RayTracingDebug && CurrentRayTracingDebugVisualizationMode != NAME_None;
 	const bool bVisualizeGPUSkinCache = CurrentViewMode == VMI_VisualizeGPUSkinCache && CurrentGPUSkinCacheVisualizationMode != NAME_None;
-	const bool bCanDisableTonemapper = bVisualizeBufferEnabled || bVisualizeGPUSkinCache || (bRayTracingDebugEnabled && !FRayTracingDebugVisualizationMenuCommands::DebugModeShouldBeTonemapped(CurrentRayTracingDebugVisualizationMode));
+	const bool bCanDisableTonemapper = bVisualizeBufferEnabled || bVisualizeGPUSkinCache || (bRayTracingDebugEnabled && !GetRayTracingVisualizationData().GetModeTonemapped(CurrentRayTracingDebugVisualizationMode));
 	
 	EngineShowFlagOverride(ESFIM_Editor, ViewFamily.ViewMode, ViewFamily.EngineShowFlags, bCanDisableTonemapper);
 	EngineShowFlagOrthographicOverride(IsPerspective(), ViewFamily.EngineShowFlags);
diff --git a/Engine/Source/Editor/UnrealEd/Private/RayTracingDebugVisualizationMenuCommands.cpp b/Engine/Source/Editor/UnrealEd/Private/RayTracingDebugVisualizationMenuCommands.cpp
--- a/Engine/Source/Editor/UnrealEd/Private/RayTracingDebugVisualizationMenuCommands.cpp
+++ b/Engine/Source/Editor/UnrealEd/Private/RayTracingDebugVisualizationMenuCommands.cpp
@@ -6,12 +6,18 @@
 #include "Framework/Commands/InputChord.h"
 #include "Internationalization/Text.h"
 #include "Templates/Function.h"
+#include "RayTracingVisualizationData.h"
 #include "Styling/AppStyle.h"
 #include "EditorViewportClient.h"
 
-#define LOCTEXT_NAMESPACE "RayTracingDebugVisualizationMenuCommands"
+int32 GRayTracingVisualizeOther = 0;
+static FAutoConsoleVariableRef CVarRayTracingVisualizeOther(
+	TEXT("r.RayTracing.Visualize.Other"),
+	GRayTracingVisualizeOther,
+	TEXT("Include 'Other' visualization modes in the in-editor 'Ray Tracing Debug' drop down menu.")
+);
 
-TArray<FText> FRayTracingDebugVisualizationMenuCommands::RayTracingDebugModeNames;
+#define LOCTEXT_NAMESPACE "RayTracingDebugVisualizationMenuCommands"
 
 FRayTracingDebugVisualizationMenuCommands::FRayTracingDebugVisualizationMenuCommands()
 	: TCommands<FRayTracingDebugVisualizationMenuCommands>
@@ -26,112 +32,108 @@ FRayTracingDebugVisualizationMenuCommands::FRayTracingDebugVisualizationMenuComm
 
 void FRayTracingDebugVisualizationMenuCommands::BuildCommandMap()
 {
-	RayTracingDebugVisualizationCommands.Empty();
-	CreateRayTracingDebugVisualizationCommands();
-}
+	const FRayTracingVisualizationData& VisualizationData = GetRayTracingVisualizationData();
+	const FRayTracingVisualizationData::TModeMap& ModeMap = VisualizationData.GetModeMap();
 
-void FRayTracingDebugVisualizationMenuCommands::CreateRayTracingDebugVisualizationCommands()
-{
-	bool bRayTracingShadersSupported = FDataDrivenShaderPlatformInfo::GetSupportsRayTracingShaders(GMaxRHIShaderPlatform);
-	
-	// always supported (as long as either inline RT or RT shaders work)
-	RayTracingDebugModeNames.Add(LOCTEXT("Barycentrics", "Barycentrics"));
-	if (bRayTracingShadersSupported)
+	CommandMap.Empty();
+	for (FRayTracingVisualizationData::TModeMap::TConstIterator It = ModeMap.CreateConstIterator(); It; ++It)
 	{
-		// runs basic lighting calculations on hits
-		RayTracingDebugModeNames.Add(LOCTEXT("PrimaryRays", "PrimaryRays"));
-
-		// all of these fields reflect entries in the payload which require running a CHS
-		RayTracingDebugModeNames.Add(LOCTEXT("Radiance", "Radiance"));
-		RayTracingDebugModeNames.Add(LOCTEXT("World Normal", "World Normal"));
-		RayTracingDebugModeNames.Add(LOCTEXT("BaseColor", "BaseColor"));
-		RayTracingDebugModeNames.Add(LOCTEXT("DiffuseColor", "DiffuseColor"));
-		RayTracingDebugModeNames.Add(LOCTEXT("SpecularColor", "SpecularColor"));
-		RayTracingDebugModeNames.Add(LOCTEXT("Opacity", "Opacity"));
-		RayTracingDebugModeNames.Add(LOCTEXT("Metallic", "Metallic"));
-		RayTracingDebugModeNames.Add(LOCTEXT("Specular", "Specular"));
-		RayTracingDebugModeNames.Add(LOCTEXT("Roughness", "Roughness"));
-		RayTracingDebugModeNames.Add(LOCTEXT("Ior", "Ior"));
-		RayTracingDebugModeNames.Add(LOCTEXT("ShadingModelID", "ShadingModelID"));
-		RayTracingDebugModeNames.Add(LOCTEXT("BlendingMode", "BlendingMode"));
-		RayTracingDebugModeNames.Add(LOCTEXT("PrimitiveLightingChannelMask", "PrimitiveLightingChannelMask"));
-		RayTracingDebugModeNames.Add(LOCTEXT("CustomData", "CustomData"));
-		RayTracingDebugModeNames.Add(LOCTEXT("GBufferAO", "GBufferAO"));
-		RayTracingDebugModeNames.Add(LOCTEXT("IndirectIrradiance", "IndirectIrradiance"));
-		RayTracingDebugModeNames.Add(LOCTEXT("World Position", "World Position"));
-		RayTracingDebugModeNames.Add(LOCTEXT("HitKind", "HitKind"));
-		RayTracingDebugModeNames.Add(LOCTEXT("World Tangent", "World Tangent"));
-		RayTracingDebugModeNames.Add(LOCTEXT("Anisotropy", "Anisotropy"));
-
-		// debugging the geometry itself
-		RayTracingDebugModeNames.Add(LOCTEXT("Instances", "Instances"));
-		RayTracingDebugModeNames.Add(LOCTEXT("Instance Overlap", "Instance Overlap"));
-		RayTracingDebugModeNames.Add(LOCTEXT("Triangle Hit Count", "Triangle Hit Count"));
-		RayTracingDebugModeNames.Add(LOCTEXT("Hit Count Per Instance", "Hit Count Per Instance"));
-	
-		if (GRHISupportsShaderTimestamp)
-		{
-			RayTracingDebugModeNames.Add(LOCTEXT("Performance", "Performance"));
-		}
-	
-		RayTracingDebugModeNames.Add(LOCTEXT("Triangles", "Triangles"));
-		RayTracingDebugModeNames.Add(LOCTEXT("FarField", "FarField"));
-		RayTracingDebugModeNames.Add(LOCTEXT("Dynamic Instances", "Dynamic Instances"));
-		RayTracingDebugModeNames.Add(LOCTEXT("Proxy Type", "Proxy Type"));
-		RayTracingDebugModeNames.Add(LOCTEXT("Picker", "Picker"));
-		RayTracingDebugModeNames.Add(LOCTEXT("Light Grid Occupancy", "Light Grid Occupancy"));
-		if (Substrate::IsSubstrateEnabled())
+		const FRayTracingVisualizationData::FModeRecord& Entry = It.Value();
+
+		if (Entry.bHiddenInEditor)
 		{
-			RayTracingDebugModeNames.Emplace(LOCTEXT("Subtrate Material Properties", "Subtrate Material Properties"));
+			continue;
 		}
-	}
-	for ( int32 RayTracingDebugIndex = 0; RayTracingDebugIndex < RayTracingDebugModeNames.Num(); ++RayTracingDebugIndex)
-	{
-		const FText CommandNameText = RayTracingDebugModeNames[RayTracingDebugIndex];
-		const FName CommandName = FName(*CommandNameText.ToString());
 
-		FRayTracingDebugVisualizationRecord Record;
-		Record.Index = RayTracingDebugIndex;
-		Record.Name = CommandName;
-		Record.Command = FUICommandInfoDecl(this->AsShared(), CommandName, CommandNameText, CommandNameText)
+		FVisualizationRecord& Record = CommandMap.Add(Entry.ModeName, FVisualizationRecord());
+		Record.Name = Entry.ModeName;
+		Record.Command = FUICommandInfoDecl(
+			this->AsShared(),
+			Entry.ModeName,
+			Entry.ModeText,
+			Entry.ModeDesc)
 			.UserInterfaceType(EUserInterfaceActionType::RadioButton)
-			.DefaultChord(FInputChord());
+			.DefaultChord(FInputChord()
+			);
 
-		RayTracingDebugVisualizationCommands.Add(Record);
+		switch (Entry.ModeType)
+		{
+		case FRayTracingVisualizationData::FModeType::Overview:
+			Record.Type = FVisualizationType::Overview;
+			break;
+
+		default:
+		case FRayTracingVisualizationData::FModeType::Standard:
+			Record.Type = FVisualizationType::Standard;
+			break;
+
+		case FRayTracingVisualizationData::FModeType::Performance:
+			Record.Type = FVisualizationType::Performance;
+			break;
+
+		case FRayTracingVisualizationData::FModeType::Other:
+			Record.Type = FVisualizationType::Other;
+			break;
+		}
 	}
 }
 
-bool FRayTracingDebugVisualizationMenuCommands::DebugModeShouldBeTonemapped(const FName& RayTracingDebugModeName)
+void FRayTracingDebugVisualizationMenuCommands::BuildVisualisationSubMenu(FMenuBuilder& Menu)
 {
-	static TArray<FName> TonemappedRayTracingDebugVisualizationModes;
-	if (TonemappedRayTracingDebugVisualizationModes.Num() == 0)
+	const bool bShowOther = GRayTracingVisualizeOther != 0;
+
+	const FRayTracingDebugVisualizationMenuCommands& Commands = FRayTracingDebugVisualizationMenuCommands::Get();
+	if (Commands.IsPopulated())
 	{
-		TonemappedRayTracingDebugVisualizationModes.Add(*LOCTEXT("PrimaryRays", "PrimaryRays").ToString());
-		TonemappedRayTracingDebugVisualizationModes.Add(*LOCTEXT("Radiance", "Radiance").ToString());
-		TonemappedRayTracingDebugVisualizationModes.Add(*LOCTEXT("IndirectIrradiance", "IndirectIrradiance").ToString());
-	}
 
-	return TonemappedRayTracingDebugVisualizationModes.Contains(RayTracingDebugModeName);
-}
+		{
+			Commands.AddCommandTypeToMenu(Menu, FVisualizationType::Overview);
+		}
 
-void FRayTracingDebugVisualizationMenuCommands::BuildVisualisationSubMenu(FMenuBuilder& Menu)
-{
-	const FRayTracingDebugVisualizationMenuCommands& Commands = FRayTracingDebugVisualizationMenuCommands::Get();
+		{
+			Menu.BeginSection("LevelViewportRayTracingVisualizationGeneral", LOCTEXT("RayTracingVisualizationGeneral", "General"));
+			Commands.AddCommandTypeToMenu(Menu, FVisualizationType::Standard);
+			Menu.EndSection();
+		}
+
+		{
+			Menu.BeginSection("LevelViewportRayTracingVisualizationPerformance", LOCTEXT("RayTracingVisualizationPerformance", "Performance"));
+			Commands.AddCommandTypeToMenu(Menu, FVisualizationType::Performance);
+			Menu.EndSection();
+		}
+
+		if (bShowOther)
+		{
+			Menu.BeginSection("LevelViewportRayTracingVisualizationOther", LOCTEXT("RayTracingVisualizationOther", "Other"));
+			Commands.AddCommandTypeToMenu(Menu, FVisualizationType::Other);
+			Menu.EndSection();
+		}
 
-	Menu.BeginSection("RayTracingDebugVisualizationMode", LOCTEXT( "RayTracingDebugVisualizationHeader", "Ray Tracing Debug Viewmodes" ) );
-	Commands.AddRayTracingDebugVisualizationCommandsToMenu(Menu);
-	Menu.EndSection();
+		Menu.EndSection();
+	}
 }
 
-void FRayTracingDebugVisualizationMenuCommands::AddRayTracingDebugVisualizationCommandsToMenu(FMenuBuilder& Menu) const
+bool FRayTracingDebugVisualizationMenuCommands::AddCommandTypeToMenu(FMenuBuilder& Menu, const FVisualizationType Type) const
 {
-	check(RayTracingDebugVisualizationCommands.Num() > 0);
+	bool bAddedCommands = false;
 
-	for (FRayTracingDebugVisualizationRecord Record : RayTracingDebugVisualizationCommands)
+	const TVisualizationModeCommandMap& Commands = CommandMap;
+	for (TCommandConstIterator It = CreateCommandConstIterator(); It; ++It)
 	{
-		FText InName = FText::FromString(Record.Name.GetPlainNameString());
-		Menu.AddMenuEntry(Record.Command, NAME_None, InName);
+		const FVisualizationRecord& Record = It.Value();
+		if (Record.Type == Type)
+		{
+			Menu.AddMenuEntry(Record.Command, NAME_None, Record.Command->GetLabel());
+			bAddedCommands = true;
+		}
 	}
+
+	return bAddedCommands;
+}
+
+FRayTracingDebugVisualizationMenuCommands::TCommandConstIterator FRayTracingDebugVisualizationMenuCommands::CreateCommandConstIterator() const
+{
+	return CommandMap.CreateConstIterator();
 }
 
 void FRayTracingDebugVisualizationMenuCommands::RegisterCommands()
@@ -142,13 +144,15 @@ void FRayTracingDebugVisualizationMenuCommands::RegisterCommands()
 void FRayTracingDebugVisualizationMenuCommands::BindCommands(FUICommandList& CommandList, const TSharedPtr<FEditorViewportClient>& Client) const
 {
 	// Map Buffer visualization mode actions
-	for (FRayTracingDebugVisualizationRecord Record : RayTracingDebugVisualizationCommands)
+	for (FRayTracingDebugVisualizationMenuCommands::TCommandConstIterator It = FRayTracingDebugVisualizationMenuCommands::Get().CreateCommandConstIterator(); It; ++It)
 	{
+		const FRayTracingDebugVisualizationMenuCommands::FVisualizationRecord& Record = It.Value();
 		CommandList.MapAction(
 			Record.Command,
 			FExecuteAction::CreateStatic(&FRayTracingDebugVisualizationMenuCommands::ChangeRayTracingDebugVisualizationMode, Client.ToWeakPtr(), Record.Name),
 			FCanExecuteAction(),
-			FIsActionChecked::CreateStatic(&FRayTracingDebugVisualizationMenuCommands::IsRayTracingDebugVisualizationModeSelected, Client.ToWeakPtr(), Record.Name));
+			FIsActionChecked::CreateStatic(&FRayTracingDebugVisualizationMenuCommands::IsRayTracingDebugVisualizationModeSelected, Client.ToWeakPtr(), Record.Name)
+		);
 	}
 }
 
diff --git a/Engine/Source/Editor/UnrealEd/Public/RayTracingDebugVisualizationMenuCommands.h b/Engine/Source/Editor/UnrealEd/Public/RayTracingDebugVisualizationMenuCommands.h
--- a/Engine/Source/Editor/UnrealEd/Public/RayTracingDebugVisualizationMenuCommands.h
+++ b/Engine/Source/Editor/UnrealEd/Public/RayTracingDebugVisualizationMenuCommands.h
@@ -14,41 +14,52 @@ class FEditorViewportClient;
 class UNREALED_API FRayTracingDebugVisualizationMenuCommands : public TCommands<FRayTracingDebugVisualizationMenuCommands>
 {
 public:
-	struct FRayTracingDebugVisualizationRecord
+	enum class FVisualizationType : uint8
+	{
+		Overview,
+		Standard,
+		Performance,
+		Other,
+	};
+
+	struct FVisualizationRecord
 	{
-		uint32 Index;
 		FName Name;
 		TSharedPtr<FUICommandInfo> Command;
+		FVisualizationType Type;
 
-		FRayTracingDebugVisualizationRecord()
-			: Index(0),
-				Name(),
-				Command()
+		FVisualizationRecord()
+			: Name()
+			, Command()
+			, Type(FVisualizationType::Overview)
 		{
 		}
 	};
 
+	typedef TMultiMap<FName, FVisualizationRecord> TVisualizationModeCommandMap;
+	typedef TVisualizationModeCommandMap::TConstIterator TCommandConstIterator;
+
 	FRayTracingDebugVisualizationMenuCommands();
 
+	TCommandConstIterator CreateCommandConstIterator() const;
+
 	static void BuildVisualisationSubMenu(FMenuBuilder& Menu);
 
 	virtual void RegisterCommands() override;
 
 	void BindCommands(FUICommandList& CommandList, const TSharedPtr<FEditorViewportClient>& Client) const;
 
-	static bool DebugModeShouldBeTonemapped(const FName& RayTracingDebugModeName);
+	inline bool IsPopulated() const
+	{
+		return CommandMap.Num() > 0;
+	}
 
 private:
 	void BuildCommandMap();
-
-	void CreateRayTracingDebugVisualizationCommands();
-
-	void AddRayTracingDebugVisualizationCommandsToMenu(FMenuBuilder& menu) const;
+	bool AddCommandTypeToMenu(FMenuBuilder& Menu, const FVisualizationType Type) const;
 
 	static void ChangeRayTracingDebugVisualizationMode(TWeakPtr<FEditorViewportClient> WeakClient, FName InName);
 	static bool IsRayTracingDebugVisualizationModeSelected(TWeakPtr<FEditorViewportClient> WeakClient, FName InName);
 
-	TArray<FRayTracingDebugVisualizationRecord> RayTracingDebugVisualizationCommands;
-
-	static TArray<FText> RayTracingDebugModeNames;
+	TVisualizationModeCommandMap CommandMap;
 };
diff --git a/Engine/Source/Runtime/Engine/Private/RayTracingVisualizationData.cpp b/Engine/Source/Runtime/Engine/Private/RayTracingVisualizationData.cpp
new file mode 100644
--- /dev/null
+++ b/Engine/Source/Runtime/Engine/Private/RayTracingVisualizationData.cpp
@@ -0,0 +1,154 @@
+// Copyright Epic Games, Inc. All Rights Reserved.
+
+#include "RayTracingVisualizationData.h"
+
+#include "RaytracingDebugDefinitions.h"
+#include "RenderUtils.h"
+#include "DataDrivenShaderPlatformInfo.h"
+#include "RHIGlobals.h"
+
+#define LOCTEXT_NAMESPACE "FRayTracingVisualizationData"
+
+static FRayTracingVisualizationData GRayTracingVisualizationData;
+
+void FRayTracingVisualizationData::Initialize()
+{
+	UE::TScopeLock Lock(Mutex);
+	if (!bIsInitialized)
+	{
+		const bool bRayTracingShadersSupported = FDataDrivenShaderPlatformInfo::GetSupportsRayTracingShaders(GMaxRHIShaderPlatform);
+
+		// always supported (as long as either inline RT or RT shaders work)
+		AddVisualizationMode(TEXT("Barycentrics"), LOCTEXT("Barycentrics", "Barycentrics"), FModeType::Standard, RAY_TRACING_DEBUG_VIZ_BARYCENTRICS, false);
+
+		if (bRayTracingShadersSupported)
+		{
+			// runs basic lighting calculations on hits
+			AddVisualizationMode(TEXT("PrimaryRays"), LOCTEXT("PrimaryRays", "Primary Rays"), FModeType::Standard, RAY_TRACING_DEBUG_VIZ_PRIMARY_RAYS, true);
+
+			// all of these fields reflect entries in the payload which require running a CHS
+			AddVisualizationMode(TEXT("Radiance"), LOCTEXT("Radiance", "Radiance"), FModeType::Other, RAY_TRACING_DEBUG_VIZ_RADIANCE, true);
+			AddVisualizationMode(TEXT("WorldNormal"), LOCTEXT("WorldNormal", "World Normal"), FModeType::Standard, RAY_TRACING_DEBUG_VIZ_WORLD_NORMAL, false);
+			AddVisualizationMode(TEXT("BaseColor"), LOCTEXT("BaseColor", "Base Color"), FModeType::Standard, RAY_TRACING_DEBUG_VIZ_BASE_COLOR, false);
+			AddVisualizationMode(TEXT("DiffuseColor"), LOCTEXT("DiffuseColor", "Diffuse Color"), FModeType::Other, RAY_TRACING_DEBUG_VIZ_DIFFUSE_COLOR, false);
+			AddVisualizationMode(TEXT("SpecularColor"), LOCTEXT("SpecularColor", "Specular Color"), FModeType::Other, RAY_TRACING_DEBUG_VIZ_SPECULAR_COLOR, false);
+			AddVisualizationMode(TEXT("Opacity"), LOCTEXT("Opacity", "Opacity"), FModeType::Other, RAY_TRACING_DEBUG_VIZ_OPACITY, false);
+			AddVisualizationMode(TEXT("Metallic"), LOCTEXT("Metallic", "Metallic"), FModeType::Other, RAY_TRACING_DEBUG_VIZ_METALLIC, false);
+			AddVisualizationMode(TEXT("Specular"), LOCTEXT("Specular", "Specular"), FModeType::Other, RAY_TRACING_DEBUG_VIZ_SPECULAR, false);
+			AddVisualizationMode(TEXT("Roughness"), LOCTEXT("Roughness", "Roughness"), FModeType::Other, RAY_TRACING_DEBUG_VIZ_ROUGHNESS, false);
+			AddVisualizationMode(TEXT("Ior"), LOCTEXT("Ior", "Ior"), FModeType::Other, RAY_TRACING_DEBUG_VIZ_IOR, false);
+			AddVisualizationMode(TEXT("ShadingModelID"), LOCTEXT("ShadingModelID", "Shading Model ID"), FModeType::Other, RAY_TRACING_DEBUG_VIZ_SHADING_MODEL, false);
+			AddVisualizationMode(TEXT("BlendingMode"), LOCTEXT("BlendingMode", "Blending Mode"), FModeType::Other, RAY_TRACING_DEBUG_VIZ_BLENDING_MODE, false);
+			AddVisualizationMode(TEXT("PrimitiveLightingChannelMask"), LOCTEXT("PrimitiveLightingChannelMask", "Primitive Lighting Channel Mask"), FModeType::Other, RAY_TRACING_DEBUG_VIZ_LIGHTING_CHANNEL_MASK, false);
+			AddVisualizationMode(TEXT("CustomData"), LOCTEXT("CustomData", "Custom Data"), FModeType::Other, RAY_TRACING_DEBUG_VIZ_CUSTOM_DATA, false);
+			AddVisualizationMode(TEXT("GBufferAO"), LOCTEXT("GBufferAO", "GBuffer AO"), FModeType::Other, RAY_TRACING_DEBUG_VIZ_GBUFFER_AO, false);
+			AddVisualizationMode(TEXT("IndirectIrradiance"), LOCTEXT("IndirectIrradiance", "Indirect Irradiance"), FModeType::Other, RAY_TRACING_DEBUG_VIZ_INDIRECT_IRRADIANCE, true);
+			AddVisualizationMode(TEXT("WorldPosition"), LOCTEXT("WorldPosition", "World Position"), FModeType::Other, RAY_TRACING_DEBUG_VIZ_WORLD_POSITION, false);
+			AddVisualizationMode(TEXT("HitKind"), LOCTEXT("HitKind", "Hit Kind"), FModeType::Other, RAY_TRACING_DEBUG_VIZ_HITKIND, false);
+			AddVisualizationMode(TEXT("WorldTangent"), LOCTEXT("WorldTangent", "World Tangent"), FModeType::Other, RAY_TRACING_DEBUG_VIZ_WORLD_TANGENT, false);
+			AddVisualizationMode(TEXT("Anisotropy"), LOCTEXT("Anisotropy", "Anisotropy"), FModeType::Other, RAY_TRACING_DEBUG_VIZ_ANISOTROPY, false);
+
+			// debugging the geometry itself
+			AddVisualizationMode(TEXT("Instances"), LOCTEXT("Instances", "Instances"), FModeType::Standard, RAY_TRACING_DEBUG_VIZ_INSTANCES, false);
+			AddVisualizationMode(TEXT("Triangles"), LOCTEXT("Triangles", "Triangles"), FModeType::Standard, RAY_TRACING_DEBUG_VIZ_TRIANGLES, false);
+			AddVisualizationMode(TEXT("FarField"), LOCTEXT("FarField", "Far Field"), FModeType::Standard, RAY_TRACING_DEBUG_VIZ_FAR_FIELD, false);
+			AddVisualizationMode(TEXT("DynamicInstances"), LOCTEXT("DynamicInstances", "Dynamic Instances"), FModeType::Standard, RAY_TRACING_DEBUG_VIZ_DYNAMIC_INSTANCES, false);
+			AddVisualizationMode(TEXT("ProxyType"), LOCTEXT("ProxyType", "Proxy Type"), FModeType::Standard, RAY_TRACING_DEBUG_VIZ_PROXY_TYPE, false);
+			AddVisualizationMode(TEXT("Picker"), LOCTEXT("Picker", "Picker"), FModeType::Standard, RAY_TRACING_DEBUG_VIZ_PICKER, false);
+			AddVisualizationMode(TEXT("LightGridOccupancy"), LOCTEXT("LightGridOccupancy", "Light Grid Occupancy"), FModeType::Standard, RAY_TRACING_DEBUG_VIZ_LIGHT_GRID_COUNT, false);
+
+			// performance
+			AddVisualizationMode(TEXT("InstanceOverlap"), LOCTEXT("InstanceOverlap", "Instance Overlap"), FModeType::Performance, RAY_TRACING_DEBUG_VIZ_INSTANCE_OVERLAP, false);
+			AddVisualizationMode(TEXT("TriangleHitCount"), LOCTEXT("TriangleHitCount", "Triangle Hit Count"), FModeType::Performance, RAY_TRACING_DEBUG_VIZ_TRIANGLE_HITCOUNT, false);
+			AddVisualizationMode(TEXT("HitCountPerInstance"), LOCTEXT("HitCountPerInstance", "Hit Count Per Instance"), FModeType::Performance, RAY_TRACING_DEBUG_VIZ_HITCOUNT_PER_INSTANCE, false);
+
+			if (GRHISupportsShaderTimestamp)
+			{
+				AddVisualizationMode(TEXT("Performance"), LOCTEXT("Performance", "Performance"), FModeType::Performance, RAY_TRACING_DEBUG_VIZ_PERFORMANCE, false);
+			}
+
+			AddVisualizationMode(TEXT("Traversal Node"), LOCTEXT("TraversalNode", "Traversal Node"), FModeType::Performance, RAY_TRACING_DEBUG_VIZ_TRAVERSAL_NODE, false, /*bHiddenInEditor*/ true);
+			AddVisualizationMode(TEXT("Traversal Cluster"), LOCTEXT("TraversalCluster", "Traversal Cluster"), FModeType::Performance, RAY_TRACING_DEBUG_VIZ_TRAVERSAL_CLUSTER, false, /*bHiddenInEditor*/ true);
+			AddVisualizationMode(TEXT("Traversal Triangle"), LOCTEXT("TraversalTriangle", "Traversal Triangle"), FModeType::Performance, RAY_TRACING_DEBUG_VIZ_TRAVERSAL_TRIANGLE, false, /*bHiddenInEditor*/ true);
+			AddVisualizationMode(TEXT("Traversal All"), LOCTEXT("TraversalAll", "Traversal All"), FModeType::Performance, RAY_TRACING_DEBUG_VIZ_TRAVERSAL_ALL, false, /*bHiddenInEditor*/ true);
+			AddVisualizationMode(TEXT("Traversal Statistics"), LOCTEXT("TraversalStatistics", "Traversal Statistics"), FModeType::Performance, RAY_TRACING_DEBUG_VIZ_TRAVERSAL_STATISTICS, false, /*bHiddenInEditor*/ true);
+
+			if (Substrate::IsSubstrateEnabled())
+			{
+				AddVisualizationMode(TEXT("SubtrateMaterialProperties"), LOCTEXT("SubtrateMaterialProperties", "Subtrate Material Properties"), FModeType::Other, RAY_TRACING_DEBUG_VIZ_SUBSTRATE_DATA, false);
+			}
+		}
+
+		bIsInitialized = true;
+	}
+}
+
+void FRayTracingVisualizationData::AddVisualizationMode(
+	const TCHAR* ModeString,
+	const FText& ModeText,
+	const FModeType ModeType,
+	int32 ModeID,
+	bool bTonemapped,
+	bool bHiddenInEditor
+)
+{
+	const FName ModeName = FName(ModeString);
+
+	FModeRecord& Record	= ModeMap.Emplace(ModeName);
+	Record.ModeString			= FString(ModeString);
+	Record.ModeName				= ModeName;
+	Record.ModeText				= ModeText;
+	Record.ModeDesc				= FText::GetEmpty();
+	Record.ModeType				= ModeType;
+	Record.ModeID				= ModeID;
+	Record.bTonemapped			= bTonemapped;
+	Record.bHiddenInEditor		= bHiddenInEditor;
+}
+
+FText FRayTracingVisualizationData::GetModeDisplayName(const FName& InModeName) const
+{
+	if (const FModeRecord* Record = ModeMap.Find(InModeName))
+	{
+		return Record->ModeText;
+	}
+	else
+	{
+		return FText::GetEmpty();
+	}
+}
+
+int32 FRayTracingVisualizationData::GetModeID(const FName& InModeName) const
+{
+	if (const FModeRecord* Record = ModeMap.Find(InModeName))
+	{
+		return Record->ModeID;
+	}
+	else
+	{
+		return INDEX_NONE;
+	}
+}
+
+bool FRayTracingVisualizationData::GetModeTonemapped(const FName& InModeName) const
+{
+	if (const FModeRecord* Record = ModeMap.Find(InModeName))
+	{
+		return Record->bTonemapped;
+	}
+	else
+	{
+		return false;
+	}
+}
+
+FRayTracingVisualizationData& GetRayTracingVisualizationData()
+{
+	if (!GRayTracingVisualizationData.IsInitialized())
+	{
+		GRayTracingVisualizationData.Initialize();
+	}
+
+	return GRayTracingVisualizationData;
+}
+
+#undef LOCTEXT_NAMESPACE
diff --git a/Engine/Source/Runtime/Engine/Public/RayTracingVisualizationData.h b/Engine/Source/Runtime/Engine/Public/RayTracingVisualizationData.h
new file mode 100644
--- /dev/null
+++ b/Engine/Source/Runtime/Engine/Public/RayTracingVisualizationData.h
@@ -0,0 +1,85 @@
+// Copyright Epic Games, Inc. All Rights Reserved.
+
+#pragma once
+
+#include "CoreMinimal.h"
+#include "Async/Mutex.h"
+
+class FRayTracingVisualizationData
+{
+public:
+	enum class FModeType : uint8
+	{
+		Overview,
+		Standard,
+		Performance,
+		Other,
+	};
+
+	/** Describes a single available visualization mode. */
+	struct FModeRecord
+	{
+		FString   ModeString;
+		FName     ModeName;
+		FText     ModeText;
+		FText     ModeDesc;
+		FModeType ModeType;
+		int32     ModeID;
+
+		bool      bTonemapped;
+		bool      bHiddenInEditor;
+	};
+
+	/** Mapping of FName to a visualization mode record. */
+	typedef TMultiMap<FName, FModeRecord> TModeMap;
+
+public:
+	FRayTracingVisualizationData()
+		: bIsInitialized(false)
+	{
+	}
+
+	/** Initialize the system. */
+	void Initialize();
+
+	/** Check if system was initialized. */
+	inline bool IsInitialized() const { return bIsInitialized; }
+
+	/** Get the display name of a named mode from the available mode map. **/
+	ENGINE_API FText GetModeDisplayName(const FName& InModeName) const;
+
+	ENGINE_API int32 GetModeID(const FName& InModeName) const;
+
+	ENGINE_API bool GetModeTonemapped(const FName& InModeName) const;
+
+	inline const TModeMap& GetModeMap() const
+	{
+		return ModeMap;
+	}
+
+private:
+
+	void AddVisualizationMode(
+		const TCHAR* ModeString,
+		const FText& ModeText,
+		const FModeType ModeType,
+		int32 ModeID,
+		bool bTonemapped,
+		bool bHiddenInEditor = false
+	);
+
+private:
+	/** The name->mode mapping table */
+	TModeMap ModeMap;
+	
+	/** Storage for console variable documentation strings. **/
+	FString ConsoleDocumentationVisualizationMode;
+
+	/** Flag indicating if system is initialized. **/
+	std::atomic_bool bIsInitialized;
+
+	/** Mutex for initialization. */
+	UE::FMutex Mutex;
+};
+
+ENGINE_API FRayTracingVisualizationData& GetRayTracingVisualizationData();
diff --git a/Engine/Source/Runtime/Renderer/Private/RayTracing/RayTracingDebug.cpp b/Engine/Source/Runtime/Renderer/Private/RayTracing/RayTracingDebug.cpp
--- a/Engine/Source/Runtime/Renderer/Private/RayTracing/RayTracingDebug.cpp
+++ b/Engine/Source/Runtime/Renderer/Private/RayTracing/RayTracingDebug.cpp
@@ -12,6 +12,7 @@
 #include "PostProcess/SceneRenderTargets.h"
 #include "RenderGraphBuilder.h"
 #include "SceneUtils.h"
+#include "RayTracingVisualizationData.h"
 #include "RaytracingDebugDefinitions.h"
 #include "RayTracingDebugTypes.h"
 #include "RayTracing/RayTracingLighting.h"
@@ -1333,64 +1334,18 @@ static void DrawInstanceOverlap(FRDGBuilder& GraphBuilder, const FScene* Scene,
 
 static uint32 GetRaytracingDebugViewMode(const FSceneView& View)
 {
-	static TMap<FName, uint32> RayTracingDebugVisualizationModes;
-	if (RayTracingDebugVisualizationModes.Num() == 0)
-	{
-		RayTracingDebugVisualizationModes.Emplace(FName(*LOCTEXT("Radiance", "Radiance").ToString()),											RAY_TRACING_DEBUG_VIZ_RADIANCE);
-		RayTracingDebugVisualizationModes.Emplace(FName(*LOCTEXT("World Normal", "World Normal").ToString()),									RAY_TRACING_DEBUG_VIZ_WORLD_NORMAL);
-		RayTracingDebugVisualizationModes.Emplace(FName(*LOCTEXT("BaseColor", "BaseColor").ToString()),											RAY_TRACING_DEBUG_VIZ_BASE_COLOR);
-		RayTracingDebugVisualizationModes.Emplace(FName(*LOCTEXT("DiffuseColor", "DiffuseColor").ToString()),									RAY_TRACING_DEBUG_VIZ_DIFFUSE_COLOR);
-		RayTracingDebugVisualizationModes.Emplace(FName(*LOCTEXT("SpecularColor", "SpecularColor").ToString()),									RAY_TRACING_DEBUG_VIZ_SPECULAR_COLOR);
-		RayTracingDebugVisualizationModes.Emplace(FName(*LOCTEXT("Opacity", "Opacity").ToString()),												RAY_TRACING_DEBUG_VIZ_OPACITY);
-		RayTracingDebugVisualizationModes.Emplace(FName(*LOCTEXT("Metallic", "Metallic").ToString()),											RAY_TRACING_DEBUG_VIZ_METALLIC);
-		RayTracingDebugVisualizationModes.Emplace(FName(*LOCTEXT("Specular", "Specular").ToString()),											RAY_TRACING_DEBUG_VIZ_SPECULAR);
-		RayTracingDebugVisualizationModes.Emplace(FName(*LOCTEXT("Roughness", "Roughness").ToString()),											RAY_TRACING_DEBUG_VIZ_ROUGHNESS);
-		RayTracingDebugVisualizationModes.Emplace(FName(*LOCTEXT("Ior", "Ior").ToString()),														RAY_TRACING_DEBUG_VIZ_IOR);
-		RayTracingDebugVisualizationModes.Emplace(FName(*LOCTEXT("ShadingModelID", "ShadingModelID").ToString()),								RAY_TRACING_DEBUG_VIZ_SHADING_MODEL);
-		RayTracingDebugVisualizationModes.Emplace(FName(*LOCTEXT("BlendingMode", "BlendingMode").ToString()),									RAY_TRACING_DEBUG_VIZ_BLENDING_MODE);
-		RayTracingDebugVisualizationModes.Emplace(FName(*LOCTEXT("PrimitiveLightingChannelMask", "PrimitiveLightingChannelMask").ToString()),	RAY_TRACING_DEBUG_VIZ_LIGHTING_CHANNEL_MASK);
-		RayTracingDebugVisualizationModes.Emplace(FName(*LOCTEXT("CustomData", "CustomData").ToString()),										RAY_TRACING_DEBUG_VIZ_CUSTOM_DATA);
-		RayTracingDebugVisualizationModes.Emplace(FName(*LOCTEXT("GBufferAO", "GBufferAO").ToString()),											RAY_TRACING_DEBUG_VIZ_GBUFFER_AO);
-		RayTracingDebugVisualizationModes.Emplace(FName(*LOCTEXT("IndirectIrradiance", "IndirectIrradiance").ToString()),						RAY_TRACING_DEBUG_VIZ_INDIRECT_IRRADIANCE);
-		RayTracingDebugVisualizationModes.Emplace(FName(*LOCTEXT("World Position", "World Position").ToString()),								RAY_TRACING_DEBUG_VIZ_WORLD_POSITION);
-		RayTracingDebugVisualizationModes.Emplace(FName(*LOCTEXT("HitKind", "HitKind").ToString()),												RAY_TRACING_DEBUG_VIZ_HITKIND);
-		RayTracingDebugVisualizationModes.Emplace(FName(*LOCTEXT("Barycentrics", "Barycentrics").ToString()),									RAY_TRACING_DEBUG_VIZ_BARYCENTRICS);
-		RayTracingDebugVisualizationModes.Emplace(FName(*LOCTEXT("PrimaryRays", "PrimaryRays").ToString()),										RAY_TRACING_DEBUG_VIZ_PRIMARY_RAYS);
-		RayTracingDebugVisualizationModes.Emplace(FName(*LOCTEXT("World Tangent", "World Tangent").ToString()),									RAY_TRACING_DEBUG_VIZ_WORLD_TANGENT);
-		RayTracingDebugVisualizationModes.Emplace(FName(*LOCTEXT("Anisotropy", "Anisotropy").ToString()),										RAY_TRACING_DEBUG_VIZ_ANISOTROPY);
-		RayTracingDebugVisualizationModes.Emplace(FName(*LOCTEXT("Instances", "Instances").ToString()),											RAY_TRACING_DEBUG_VIZ_INSTANCES);
-		RayTracingDebugVisualizationModes.Emplace(FName(*LOCTEXT("Performance", "Performance").ToString()),										RAY_TRACING_DEBUG_VIZ_PERFORMANCE);
-		RayTracingDebugVisualizationModes.Emplace(FName(*LOCTEXT("Triangles", "Triangles").ToString()),											RAY_TRACING_DEBUG_VIZ_TRIANGLES);
-		RayTracingDebugVisualizationModes.Emplace(FName(*LOCTEXT("FarField", "FarField").ToString()),											RAY_TRACING_DEBUG_VIZ_FAR_FIELD);
-		RayTracingDebugVisualizationModes.Emplace(FName(*LOCTEXT("Traversal Node", "Traversal Node").ToString()),								RAY_TRACING_DEBUG_VIZ_TRAVERSAL_NODE);
-		RayTracingDebugVisualizationModes.Emplace(FName(*LOCTEXT("Traversal Cluster", "Traversal Cluster").ToString()),							RAY_TRACING_DEBUG_VIZ_TRAVERSAL_CLUSTER);
-		RayTracingDebugVisualizationModes.Emplace(FName(*LOCTEXT("Traversal Triangle", "Traversal Triangle").ToString()),						RAY_TRACING_DEBUG_VIZ_TRAVERSAL_TRIANGLE);
-		RayTracingDebugVisualizationModes.Emplace(FName(*LOCTEXT("Traversal All", "Traversal All").ToString()),									RAY_TRACING_DEBUG_VIZ_TRAVERSAL_ALL);
-		RayTracingDebugVisualizationModes.Emplace(FName(*LOCTEXT("Traversal Statistics", "Traversal Statistics").ToString()),					RAY_TRACING_DEBUG_VIZ_TRAVERSAL_STATISTICS);
-		RayTracingDebugVisualizationModes.Emplace(FName(*LOCTEXT("Dynamic Instances", "Dynamic Instances").ToString()),							RAY_TRACING_DEBUG_VIZ_DYNAMIC_INSTANCES);
-		RayTracingDebugVisualizationModes.Emplace(FName(*LOCTEXT("Proxy Type", "Proxy Type").ToString()),										RAY_TRACING_DEBUG_VIZ_PROXY_TYPE);
-		RayTracingDebugVisualizationModes.Emplace(FName(*LOCTEXT("Picker", "Picker").ToString()),												RAY_TRACING_DEBUG_VIZ_PICKER);
-		RayTracingDebugVisualizationModes.Emplace(FName(*LOCTEXT("Instance Overlap", "Instance Overlap").ToString()),							RAY_TRACING_DEBUG_VIZ_INSTANCE_OVERLAP);
-		RayTracingDebugVisualizationModes.Emplace(FName(*LOCTEXT("Triangle Hit Count", "Triangle Hit Count").ToString()),						RAY_TRACING_DEBUG_VIZ_TRIANGLE_HITCOUNT);
-		RayTracingDebugVisualizationModes.Emplace(FName(*LOCTEXT("Hit Count Per Instance", "Hit Count Per Instance").ToString()),				RAY_TRACING_DEBUG_VIZ_HITCOUNT_PER_INSTANCE);
-		RayTracingDebugVisualizationModes.Emplace(FName(*LOCTEXT("Light Grid Occupancy", "Light Grid Occupancy").ToString()),					RAY_TRACING_DEBUG_VIZ_LIGHT_GRID_COUNT);
-		if (Substrate::IsSubstrateEnabled())
-		{
-			RayTracingDebugVisualizationModes.Emplace(FName(*LOCTEXT("Subtrate Material Properties", "Subtrate Material Properties").ToString()),RAY_TRACING_DEBUG_VIZ_SUBSTRATE_DATA);
-		}
-	}
+	const FRayTracingVisualizationData& VisualizationData = GetRayTracingVisualizationData();
 
 	uint32 OutMode;
-	
-	FString ConsoleViewMode = CVarRayTracingDebugMode.GetValueOnAnyThread();
 
+	FString ConsoleViewMode = CVarRayTracingDebugMode.GetValueOnAnyThread();
 	if (!ConsoleViewMode.IsEmpty())
 	{
-		OutMode = RayTracingDebugVisualizationModes.FindRef(FName(*ConsoleViewMode));
+		OutMode = VisualizationData.GetModeID(FName(*ConsoleViewMode));
 	}
 	else if(View.CurrentRayTracingDebugVisualizationMode != NAME_None)
 	{
-		OutMode = RayTracingDebugVisualizationModes.FindRef(View.CurrentRayTracingDebugVisualizationMode);
+		OutMode = VisualizationData.GetModeID(View.CurrentRayTracingDebugVisualizationMode);
 	}
 	else
 	{
@@ -1824,25 +1779,10 @@ bool IsRayTracingInstanceDebugDataEnabled(const FViewInfo& View)
 bool IsRayTracingInstanceOverlapEnabled(const FViewInfo& View)
 {
 #if !UE_BUILD_SHIPPING	
-	static TArray<FName> RayTracingDebugVisualizationModes;
-	if (RayTracingDebugVisualizationModes.Num() == 0)
-	{
-		RayTracingDebugVisualizationModes.Add(*LOCTEXT("Instance Overlap", "Instance Overlap").ToString());
-	}
-
-	FString ConsoleViewMode = CVarRayTracingDebugMode.GetValueOnRenderThread();
-	if (!ConsoleViewMode.IsEmpty())
-	{
-		return RayTracingDebugVisualizationModes.Contains(FName(*ConsoleViewMode));
-	}
-	else if (View.CurrentRayTracingDebugVisualizationMode != NAME_None)
-	{
-		return RayTracingDebugVisualizationModes.Contains(View.CurrentRayTracingDebugVisualizationMode);
-	}
+	return GetRaytracingDebugViewMode(View) == RAY_TRACING_DEBUG_VIZ_INSTANCE_OVERLAP;
+#else
+	return false;
 #endif
-	{
-		return false;
-	}
 }
 
 #undef LOCTEXT_NAMESPACE
