- Remove redundant transactions from color picker editor

diff --git a/Engine/Source/Editor/UnrealEd/Private/SColorGradientEditor.cpp b/Engine/Source/Editor/UnrealEd/Private/SColorGradientEditor.cpp
--- a/Engine/Source/Editor/UnrealEd/Private/SColorGradientEditor.cpp
+++ b/Engine/Source/Editor/UnrealEd/Private/SColorGradientEditor.cpp
@@ -731,6 +731,12 @@ void SColorGradientEditor::OpenGradientStopColorPicker()
 
 void SColorGradientEditor::OnSelectedStopColorChanged( FLinearColor InNewColor )
 {
+	// Don't add a transaction if nothing has changed
+	if (SelectedStop.GetColor(*CurveOwner).Equals(InNewColor))
+	{
+		return;
+	}
+
 	FScopedTransaction ColorChange( LOCTEXT("ChangeGradientStopColor", "Change Gradient Stop Color") );
 	CurveOwner->ModifyOwner();
 	SelectedStop.SetColor( InNewColor, *CurveOwner );
@@ -785,6 +791,12 @@ void SColorGradientEditor::OnAlphaValueCommitted( float NewValue, ETextCommit::T
 {
 	if( !GEditor->IsTransactionActive() )
 	{
+		// Don't add a transaction if nothing has changed
+		if (FMath::IsNearlyEqual(SelectedStop.GetColor(*CurveOwner).A, NewValue))
+		{
+			return;
+		}
+
 		// Value was typed in, no transaction is active
 		FScopedTransaction ChangeAlphaTransaction( LOCTEXT("ChangeGradientStopAlpha", "Change Gradient Stop Alpha") );
 		CurveOwner->ModifyOwner();
