[Viewport Toolbar] Move Exposure settings to camera menu

diff --git a/Engine/Source/Editor/UnrealEd/Private/ViewportToolbar/UnrealEdViewportToolbar.cpp b/Engine/Source/Editor/UnrealEd/Private/ViewportToolbar/UnrealEdViewportToolbar.cpp
--- a/Engine/Source/Editor/UnrealEd/Private/ViewportToolbar/UnrealEdViewportToolbar.cpp
+++ b/Engine/Source/Editor/UnrealEd/Private/ViewportToolbar/UnrealEdViewportToolbar.cpp
@@ -39,6 +39,36 @@ namespace UE::UnrealEd::Private
 {
 int32 CVarToolMenusViewportToolbarsValue = 0;
 
+struct ViewModesSubmenu
+{
+	static void AddModeIfSupported(
+		const IsViewModeSupportedDelegate& InIsViewModeSupported,
+		FToolMenuSection& InMenuSection,
+		const TSharedPtr<FUICommandInfo>& InModeCommandInfo,
+		EViewModeIndex InViewModeIndex,
+		const TAttribute<FText>& InToolTipOverride = TAttribute<FText>(),
+		const TAttribute<FSlateIcon>& InIconOverride = TAttribute<FSlateIcon>()
+	)
+	{
+		if (!InIsViewModeSupported.IsBound() || InIsViewModeSupported.Execute(InViewModeIndex))
+		{
+			InMenuSection.AddMenuEntry(
+				InModeCommandInfo, UViewModeUtils::GetViewModeDisplayName(InViewModeIndex), InToolTipOverride, InIconOverride
+			);
+		}
+	}
+
+	static bool IsMenuSectionAvailable(const UUnrealEdViewportToolbarContext* InContext, EHidableViewModeMenuSections InMenuSection)
+	{
+		if (!InContext->DoesViewModeMenuShowSection.IsBound())
+		{
+			return true;
+		}
+
+		return InContext->DoesViewModeMenuShowSection.Execute(InMenuSection);
+	}
+};
+
 FToolUIAction DisabledAction()
 {
 	static FToolUIAction Action;
@@ -1285,43 +1315,17 @@ void PopulateViewModesMenu(UToolMenu* InMenu)
 
 	IsViewModeSupportedDelegate IsViewModeSupported = Context->IsViewModeSupported;
 
-	struct ViewModesSubmenu
-	{
-		static void AddModeIfSupported(
-			const IsViewModeSupportedDelegate& InIsViewModeSupported,
-			FToolMenuSection& InMenuSection,
-			const TSharedPtr<FUICommandInfo>& InModeCommandInfo,
-			EViewModeIndex InViewModeIndex,
-			const TAttribute<FText>& InToolTipOverride = TAttribute<FText>(),
-			const TAttribute<FSlateIcon>& InIconOverride = TAttribute<FSlateIcon>()
-		)
-		{
-			if (!InIsViewModeSupported.IsBound() || InIsViewModeSupported.Execute(InViewModeIndex))
-			{
-				InMenuSection.AddMenuEntry(
-					InModeCommandInfo, UViewModeUtils::GetViewModeDisplayName(InViewModeIndex), InToolTipOverride, InIconOverride
-				);
-			}
-		}
-
-		static bool IsMenuSectionAvailable(const UUnrealEdViewportToolbarContext* InContext, EHidableViewModeMenuSections InMenuSection)
-		{
-			if (!InContext->DoesViewModeMenuShowSection.IsBound())
-			{
-				return true;
-			}
-
-			return InContext->DoesViewModeMenuShowSection.Execute(InMenuSection);
-		}
-	};
-
 	// View modes
 	{
 		FToolMenuSection& Section = InMenu->AddSection("ViewMode", LOCTEXT("ViewModeHeader", "View Mode"));
 		{
-			ViewModesSubmenu::AddModeIfSupported(IsViewModeSupported, Section, BaseViewportActions.LitMode, VMI_Lit);
-			ViewModesSubmenu::AddModeIfSupported(IsViewModeSupported, Section, BaseViewportActions.UnlitMode, VMI_Unlit);
-			ViewModesSubmenu::AddModeIfSupported(
+			UE::UnrealEd::Private::ViewModesSubmenu::AddModeIfSupported(
+				IsViewModeSupported, Section, BaseViewportActions.LitMode, VMI_Lit
+			);
+			UE::UnrealEd::Private::ViewModesSubmenu::AddModeIfSupported(
+				IsViewModeSupported, Section, BaseViewportActions.UnlitMode, VMI_Unlit
+			);
+			UE::UnrealEd::Private::ViewModesSubmenu::AddModeIfSupported(
 				IsViewModeSupported, Section, BaseViewportActions.WireframeMode, VMI_BrushWireframe
 			);
 
@@ -1330,19 +1334,19 @@ void PopulateViewModesMenu(UToolMenu* InMenu)
 				Section.AddEntry(GetLitWireframeEntryToSection(EditorViewport));
 			}
 
-			ViewModesSubmenu::AddModeIfSupported(
+			UE::UnrealEd::Private::ViewModesSubmenu::AddModeIfSupported(
 				IsViewModeSupported, Section, BaseViewportActions.DetailLightingMode, VMI_Lit_DetailLighting
 			);
-			ViewModesSubmenu::AddModeIfSupported(
+			UE::UnrealEd::Private::ViewModesSubmenu::AddModeIfSupported(
 				IsViewModeSupported, Section, BaseViewportActions.LightingOnlyMode, VMI_LightingOnly
 			);
-			ViewModesSubmenu::AddModeIfSupported(
+			UE::UnrealEd::Private::ViewModesSubmenu::AddModeIfSupported(
 				IsViewModeSupported, Section, BaseViewportActions.ReflectionOverrideMode, VMI_ReflectionOverride
 			);
-			ViewModesSubmenu::AddModeIfSupported(
+			UE::UnrealEd::Private::ViewModesSubmenu::AddModeIfSupported(
 				IsViewModeSupported, Section, BaseViewportActions.CollisionPawn, VMI_CollisionPawn
 			);
-			ViewModesSubmenu::AddModeIfSupported(
+			UE::UnrealEd::Private::ViewModesSubmenu::AddModeIfSupported(
 				IsViewModeSupported, Section, BaseViewportActions.CollisionVisibility, VMI_CollisionVisibility
 			);
 		}
@@ -1354,7 +1358,7 @@ void PopulateViewModesMenu(UToolMenu* InMenu)
 			const bool bPathTracingEnabled = PathTracingCvar && PathTracingCvar->GetValueOnAnyThread() != 0;
 			if (bPathTracingSupported && bPathTracingEnabled)
 			{
-				ViewModesSubmenu::AddModeIfSupported(
+				UE::UnrealEd::Private::ViewModesSubmenu::AddModeIfSupported(
 					IsViewModeSupported, Section, BaseViewportActions.PathTracingMode, VMI_PathTracing
 				);
 			}
@@ -1379,23 +1383,23 @@ void PopulateViewModesMenu(UToolMenu* InMenu)
 						);
 						if (FeatureLevel >= ERHIFeatureLevel::SM5)
 						{
-							ViewModesSubmenu::AddModeIfSupported(
+							UE::UnrealEd::Private::ViewModesSubmenu::AddModeIfSupported(
 								IsViewModeSupported, Section, BaseViewportActions.LightComplexityMode, VMI_LightComplexity
 							);
 
 							if (IsStaticLightingAllowed())
 							{
-								ViewModesSubmenu::AddModeIfSupported(
+								UE::UnrealEd::Private::ViewModesSubmenu::AddModeIfSupported(
 									IsViewModeSupported, Section, BaseViewportActions.LightmapDensityMode, VMI_LightmapDensity
 								);
 							}
 
-							ViewModesSubmenu::AddModeIfSupported(
+							UE::UnrealEd::Private::ViewModesSubmenu::AddModeIfSupported(
 								IsViewModeSupported, Section, BaseViewportActions.StationaryLightOverlapMode, VMI_StationaryLightOverlap
 							);
 						}
 
-						ViewModesSubmenu::AddModeIfSupported(
+						UE::UnrealEd::Private::ViewModesSubmenu::AddModeIfSupported(
 							IsViewModeSupported, Section, BaseViewportActions.ShaderComplexityMode, VMI_ShaderComplexity
 						);
 
@@ -1403,7 +1407,7 @@ void PopulateViewModesMenu(UToolMenu* InMenu)
 								DVSM_ShaderComplexityContainedQuadOverhead, GMaxRHIShaderPlatform, FeatureLevel
 							))
 						{
-							ViewModesSubmenu::AddModeIfSupported(
+							UE::UnrealEd::Private::ViewModesSubmenu::AddModeIfSupported(
 								IsViewModeSupported,
 								Section,
 								BaseViewportActions.ShaderComplexityWithQuadOverdrawMode,
@@ -1413,14 +1417,14 @@ void PopulateViewModesMenu(UToolMenu* InMenu)
 
 						if (AllowDebugViewShaderMode(DVSM_QuadComplexity, GMaxRHIShaderPlatform, FeatureLevel))
 						{
-							ViewModesSubmenu::AddModeIfSupported(
+							UE::UnrealEd::Private::ViewModesSubmenu::AddModeIfSupported(
 								IsViewModeSupported, Section, BaseViewportActions.QuadOverdrawMode, VMI_QuadOverdraw
 							);
 						}
 
 						if (AllowDebugViewShaderMode(DVSM_LWCComplexity, GMaxRHIShaderPlatform, FeatureLevel))
 						{
-							ViewModesSubmenu::AddModeIfSupported(
+							UE::UnrealEd::Private::ViewModesSubmenu::AddModeIfSupported(
 								IsViewModeSupported,
 								Section,
 								BaseViewportActions.VisualizeLWCComplexity,
@@ -1438,14 +1442,14 @@ void PopulateViewModesMenu(UToolMenu* InMenu)
 
 						if (AllowDebugViewShaderMode(DVSM_PrimitiveDistanceAccuracy, GMaxRHIShaderPlatform, FeatureLevel))
 						{
-							ViewModesSubmenu::AddModeIfSupported(
+							UE::UnrealEd::Private::ViewModesSubmenu::AddModeIfSupported(
 								IsViewModeSupported, Section, BaseViewportActions.TexStreamAccPrimitiveDistanceMode, VMI_PrimitiveDistanceAccuracy
 							);
 						}
 
 						if (AllowDebugViewShaderMode(DVSM_MeshUVDensityAccuracy, GMaxRHIShaderPlatform, FeatureLevel))
 						{
-							ViewModesSubmenu::AddModeIfSupported(
+							UE::UnrealEd::Private::ViewModesSubmenu::AddModeIfSupported(
 								IsViewModeSupported, Section, BaseViewportActions.TexStreamAccMeshUVDensityMode, VMI_MeshUVDensityAccuracy
 							);
 						}
@@ -1455,7 +1459,7 @@ void PopulateViewModesMenu(UToolMenu* InMenu)
 						if (AllowDebugViewShaderMode(DVSM_MaterialTextureScaleAccuracy, GMaxRHIShaderPlatform, FeatureLevel)
 							&& CVarStreamingUseNewMetrics.GetValueOnAnyThread() != 0)
 						{
-							ViewModesSubmenu::AddModeIfSupported(
+							UE::UnrealEd::Private::ViewModesSubmenu::AddModeIfSupported(
 								IsViewModeSupported,
 								Section,
 								BaseViewportActions.TexStreamAccMaterialTextureScaleMode,
@@ -1465,14 +1469,14 @@ void PopulateViewModesMenu(UToolMenu* InMenu)
 
 						if (AllowDebugViewShaderMode(DVSM_RequiredTextureResolution, GMaxRHIShaderPlatform, FeatureLevel))
 						{
-							ViewModesSubmenu::AddModeIfSupported(
+							UE::UnrealEd::Private::ViewModesSubmenu::AddModeIfSupported(
 								IsViewModeSupported, Section, BaseViewportActions.RequiredTextureResolutionMode, VMI_RequiredTextureResolution
 							);
 						}
 
 						if (AllowDebugViewShaderMode(DVSM_RequiredTextureResolution, GMaxRHIShaderPlatform, FeatureLevel))
 						{
-							ViewModesSubmenu::AddModeIfSupported(
+							UE::UnrealEd::Private::ViewModesSubmenu::AddModeIfSupported(
 								IsViewModeSupported, Section, BaseViewportActions.VirtualTexturePendingMipsMode, VMI_VirtualTexturePendingMips
 							);
 						}
@@ -1539,7 +1543,9 @@ void PopulateViewModesMenu(UToolMenu* InMenu)
 		}
 
 		if (IsRayTracingEnabled()
-			&& ViewModesSubmenu::IsMenuSectionAvailable(Context, EHidableViewModeMenuSections::RayTracingDebug))
+			&& UE::UnrealEd::Private::ViewModesSubmenu::IsMenuSectionAvailable(
+				Context, EHidableViewModeMenuSections::RayTracingDebug
+			))
 		{
 			struct Local
 			{
@@ -1571,11 +1577,11 @@ void PopulateViewModesMenu(UToolMenu* InMenu)
 							"LevelViewportLODColoration", LOCTEXT("LODModesHeader", "Level of Detail Coloration")
 						);
 
-						ViewModesSubmenu::AddModeIfSupported(
+						UE::UnrealEd::Private::ViewModesSubmenu::AddModeIfSupported(
 							IsViewModeSupported, Section, FEditorViewportCommands::Get().LODColorationMode, VMI_LODColoration
 						);
 
-						ViewModesSubmenu::AddModeIfSupported(
+						UE::UnrealEd::Private::ViewModesSubmenu::AddModeIfSupported(
 							IsViewModeSupported, Section, FEditorViewportCommands::Get().HLODColorationMode, VMI_HLODColoration
 						);
 					}
@@ -1612,7 +1618,9 @@ void PopulateViewModesMenu(UToolMenu* InMenu)
 		}
 
 		if (GEnableGPUSkinCache
-			&& ViewModesSubmenu::IsMenuSectionAvailable(Context, EHidableViewModeMenuSections::GPUSkinCache))
+			&& UE::UnrealEd::Private::ViewModesSubmenu::IsMenuSectionAvailable(
+				Context, EHidableViewModeMenuSections::GPUSkinCache
+			))
 		{
 			Section.AddSubMenu(
 				"VisualizeGPUSkinCacheViewMode",
@@ -1638,22 +1646,6 @@ void PopulateViewModesMenu(UToolMenu* InMenu)
 			);
 		}
 	}
-
-	// Auto Exposure
-	if (ViewModesSubmenu::IsMenuSectionAvailable(Context, EHidableViewModeMenuSections::Exposure))
-	{
-		const FEditorViewportCommands& BaseViewportCommands = FEditorViewportCommands::Get();
-
-		TSharedRef<SWidget> FixedEV100Menu = EditorViewport->BuildFixedEV100Menu();
-		TSharedPtr<FEditorViewportClient> EditorViewportClient = EditorViewport->GetViewportClient();
-		const bool bIsLevelEditor = EditorViewportClient.IsValid() && EditorViewportClient->IsLevelEditorClient();
-
-		FToolMenuSection& Section = InMenu->AddSection("Exposure", LOCTEXT("ExposureHeader", "Exposure"));
-		Section.AddMenuEntry(
-			bIsLevelEditor ? BaseViewportCommands.ToggleInGameExposure : BaseViewportCommands.ToggleAutoExposure
-		);
-		Section.AddEntry(FToolMenuEntry::InitWidget("FixedEV100", FixedEV100Menu, LOCTEXT("FixedEV100", "EV100")));
-	}
 }
 
 UUnrealEdViewportToolbarContext* CreateViewportToolbarDefaultContext(const TWeakPtr<SEditorViewport>& InViewport)
@@ -2560,18 +2552,41 @@ void PopulateCameraMenu(UToolMenu* InMenu)
 		return;
 	}
 
-	FToolMenuSection& PerspectiveCameraSection =
-		InMenu->FindOrAddSection("LevelViewportCameraType_Perspective", LOCTEXT("PerspectiveLabel", "Perspective"));
-	PerspectiveCameraSection.AddMenuEntry(FEditorViewportCommands::Get().Perspective);
-
-	FToolMenuSection& OrthographicCameraSection =
-		InMenu->FindOrAddSection("LevelViewportCameraType_Ortho", LOCTEXT("CameraTypeHeader_Ortho", "Orthographic"));
-	OrthographicCameraSection.AddMenuEntry(FEditorViewportCommands::Get().Top);
-	OrthographicCameraSection.AddMenuEntry(FEditorViewportCommands::Get().Bottom);
-	OrthographicCameraSection.AddMenuEntry(FEditorViewportCommands::Get().Left);
-	OrthographicCameraSection.AddMenuEntry(FEditorViewportCommands::Get().Right);
-	OrthographicCameraSection.AddMenuEntry(FEditorViewportCommands::Get().Front);
-	OrthographicCameraSection.AddMenuEntry(FEditorViewportCommands::Get().Back);
+	// Perspective
+	{
+		FToolMenuSection& PerspectiveCameraSection =
+			InMenu->FindOrAddSection("LevelViewportCameraType_Perspective", LOCTEXT("PerspectiveLabel", "Perspective"));
+		PerspectiveCameraSection.AddMenuEntry(FEditorViewportCommands::Get().Perspective);
+
+		FToolMenuSection& OrthographicCameraSection =
+			InMenu->FindOrAddSection("LevelViewportCameraType_Ortho", LOCTEXT("CameraTypeHeader_Ortho", "Orthographic"));
+		OrthographicCameraSection.AddMenuEntry(FEditorViewportCommands::Get().Top);
+		OrthographicCameraSection.AddMenuEntry(FEditorViewportCommands::Get().Bottom);
+		OrthographicCameraSection.AddMenuEntry(FEditorViewportCommands::Get().Left);
+		OrthographicCameraSection.AddMenuEntry(FEditorViewportCommands::Get().Right);
+		OrthographicCameraSection.AddMenuEntry(FEditorViewportCommands::Get().Front);
+		OrthographicCameraSection.AddMenuEntry(FEditorViewportCommands::Get().Back);
+	}
+
+	// Auto Exposure
+	{
+		if (UE::UnrealEd::Private::ViewModesSubmenu::IsMenuSectionAvailable(
+				EditorViewportContext, EHidableViewModeMenuSections::Exposure
+			))
+		{
+			const FEditorViewportCommands& BaseViewportCommands = FEditorViewportCommands::Get();
+
+			TSharedRef<SWidget> FixedEV100Menu = EditorViewport->BuildFixedEV100Menu();
+			TSharedPtr<FEditorViewportClient> EditorViewportClient = EditorViewport->GetViewportClient();
+			const bool bIsLevelEditor = EditorViewportClient.IsValid() && EditorViewportClient->IsLevelEditorClient();
+
+			FToolMenuSection& Section = InMenu->AddSection("Exposure", LOCTEXT("ExposureHeader", "Exposure"));
+			Section.AddMenuEntry(
+				bIsLevelEditor ? BaseViewportCommands.ToggleInGameExposure : BaseViewportCommands.ToggleAutoExposure
+			);
+			Section.AddEntry(FToolMenuEntry::InitWidget("FixedEV100", FixedEV100Menu, LOCTEXT("FixedEV100", "EV100")));
+		}
+	}
 }
 
 void ExtendCameraSubmenu(FName InCameraOptionsSubmenuName, bool bInShowFarViewPlaneEntry /* = true */)
