Fixed uncontrolled auto-save deletes failing to convert to a source control delete when checking out the already deleted files

It would handle this request as a checkout rather than a delete, as usually the file would still exist at this point and would be removed by CleanupAfterSuccessfulDelete. In this case the file is already gone (as the auto-save already ran CleanupAfterSuccessfulDelete), and we need to update the source control state to match.

#rnx

diff --git a/Engine/Source/Editor/UnrealEd/Private/FileHelpers.cpp b/Engine/Source/Editor/UnrealEd/Private/FileHelpers.cpp
--- a/Engine/Source/Editor/UnrealEd/Private/FileHelpers.cpp
+++ b/Engine/Source/Editor/UnrealEd/Private/FileHelpers.cpp
@@ -2121,6 +2121,7 @@ ECommandResult::Type FEditorFileUtils::CheckoutPackages(const TArray<UPackage*>&
 
 	TArray<UPackage*> FinalPackageCheckoutList;
 	TArray<UPackage*> FinalPackageMarkForAddList;
+	TArray<UPackage*> FinalPackageDeleteList;
 
 	// Source control may have been enabled in the package checkout dialog.
 	// Ensure the status is up to date
@@ -2159,7 +2160,7 @@ ECommandResult::Type FEditorFileUtils::CheckoutPackages(const TArray<UPackage*>&
 			for (int Index = 0; Index < SourceControlStates.Num(); ++Index)
 			{
 				const FSourceControlStateRef& SourceControlState = SourceControlStates[Index];
-				if (SourceControlState->IsDeleted())
+				if (SourceControlState->IsDeleted() && !UPackage::IsEmptyPackage(PkgsToCheckOut[Index]))
 				{
 					PkgsToRevert.Add(PkgsToCheckOut[Index]);
 				}
@@ -2184,9 +2185,25 @@ ECommandResult::Type FEditorFileUtils::CheckoutPackages(const TArray<UPackage*>&
 					UPackage* PackageToCheckOut = PkgsToCheckOut[Index];
 					const FSourceControlStateRef& SourceControlState = SourceControlStates[Index];
 
+					FString Filename;
+					const bool bPackageExists = FPackageName::DoesPackageExist(PackageToCheckOut->GetName(), &Filename);
+
 					// Mark the package for check out if possible
 					bool bShowCheckoutError = true;
-					if (SourceControlState->CanCheckout())
+					if (!bPackageExists && UPackage::IsEmptyPackage(PackageToCheckOut))
+					{
+						// This package has already been deleted on disk, so update source control to match
+						if (SourceControlState->CanDelete())
+						{
+							bShowCheckoutError = false;
+							FinalPackageDeleteList.Add(PackageToCheckOut);
+						}
+						else if (SourceControlState->IsDeleted() && !bErrorIfAlreadyCheckedOut)
+						{
+							bShowCheckoutError = false;
+						}
+					}
+					else if (SourceControlState->CanCheckout())
 					{
 						bShowCheckoutError = false;
 						FinalPackageCheckoutList.Add(PackageToCheckOut);
@@ -2194,8 +2211,7 @@ ECommandResult::Type FEditorFileUtils::CheckoutPackages(const TArray<UPackage*>&
 					else if (SourceControlState->CanAdd())
 					{
 						// Cannot add unsaved packages to source control
-						FString Filename;
-						if (FPackageName::DoesPackageExist(PackageToCheckOut->GetName(), &Filename))
+						if (bPackageExists)
 						{
 							bShowCheckoutError = false;
 							FinalPackageMarkForAddList.Add(PackageToCheckOut);
@@ -2236,60 +2252,78 @@ ECommandResult::Type FEditorFileUtils::CheckoutPackages(const TArray<UPackage*>&
 		}
 	}
 
-	// Attempt to check out each package the user specified to be checked out that is not read only
-	if(FinalPackageCheckoutList.Num() > 0)
+	const int32 CombinedPackageCount = FinalPackageCheckoutList.Num() + FinalPackageMarkForAddList.Num() + FinalPackageDeleteList.Num();
+	if (CombinedPackageCount > 0)
 	{
-		FScopedSlowTask SlowTask(static_cast<float>(FinalPackageCheckoutList.Num()), LOCTEXT("CheckingOutPackages", "Checking out packages..."));
+		FScopedSlowTask SlowTask(static_cast<float>(CombinedPackageCount), LOCTEXT("CheckingOutPackages", "Checking out packages..."));
 		SlowTask.MakeDialog();
-		CheckOutResult = SourceControlProvider.Execute(ISourceControlOperation::Create<FCheckOut>(), FinalPackageCheckoutList);
-		SlowTask.EnterProgressFrame(static_cast<float>(FinalPackageCheckoutList.Num()));
-	}
 
-	// Attempt to mark for add each package the user specified that is not already tracked by source control
-	ECommandResult::Type MarkForAddResult = ECommandResult::Cancelled;
-	if (FinalPackageMarkForAddList.Num() > 0)
-	{
-		MarkForAddResult = SourceControlProvider.Execute(ISourceControlOperation::Create<FMarkForAdd>(), FinalPackageMarkForAddList);
-	}
+		// Attempt to check out each package the user specified to be checked out that is not read only
+		if (FinalPackageCheckoutList.Num() > 0)
+		{
+			CheckOutResult = SourceControlProvider.Execute(ISourceControlOperation::Create<FCheckOut>(), FinalPackageCheckoutList);
+			SlowTask.EnterProgressFrame(static_cast<float>(FinalPackageCheckoutList.Num()));
+		}
 
-	TArray<UPackage*> CombinedPackageList = FinalPackageCheckoutList;
-	CombinedPackageList.Append(FinalPackageMarkForAddList);
+		// Attempt to mark for add each package the user specified that is not already tracked by source control
+		ECommandResult::Type MarkForAddResult = ECommandResult::Cancelled;
+		if (FinalPackageMarkForAddList.Num() > 0)
+		{
+			MarkForAddResult = SourceControlProvider.Execute(ISourceControlOperation::Create<FMarkForAdd>(), FinalPackageMarkForAddList);
+			SlowTask.EnterProgressFrame(static_cast<float>(FinalPackageMarkForAddList.Num()));
+		}
 
-	if (CombinedPackageList.Num() > 0)
-	{
+		// Attempt to delete each package the user specified that is empty and missing from disk
+		ECommandResult::Type DeleteResult = ECommandResult::Cancelled;
+		if (FinalPackageDeleteList.Num() > 0)
 		{
-			// Checked out some or all files successfully, so check their state
-			for (int32 i=0; i < CombinedPackageList.Num(); ++i)
+			DeleteResult = SourceControlProvider.Execute(ISourceControlOperation::Create<FDelete>(), FinalPackageDeleteList);
+			SlowTask.EnterProgressFrame(static_cast<float>(FinalPackageDeleteList.Num()));
+		}
+
+		// Checked out some or all files successfully, so check their state
+		auto VerifySourceControlState = [&SourceControlProvider, &OutPackagesCheckedOut, &PkgsWhichFailedCheckout, &CheckOutResult](const TArray<UPackage*>& PackagesToVerify, const ECommandResult::Type CommandResult, TFunctionRef<bool(const FSourceControlStateRef&)> IsValidState)
+		{
+			if (CommandResult != ECommandResult::Cancelled)
 			{
-				const bool bCheckedOut = (i < FinalPackageCheckoutList.Num()) && (CheckOutResult != ECommandResult::Cancelled);
-				const bool bMarkedForAdd = (i >= FinalPackageMarkForAddList.Num()) && (MarkForAddResult != ECommandResult::Cancelled);
-				if (!(bCheckedOut || bMarkedForAdd))
+				for (UPackage* CurPackage : PackagesToVerify)
 				{
-					continue;
-				}
-
-				UPackage* CurPackage = CombinedPackageList[i];
+					// If we're ignoring the package due to the user ignoring it for saving, remove it from the ignore list
+					// as getting here means we've explicitly decided to save the asset.
+					const FString CurPackageName = CurPackage->GetName();
+					PackagesNotSavedDuringSaveAll.Remove(CurPackageName);
 
-				// If we're ignoring the package due to the user ignoring it for saving, remove it from the ignore list
-				// as getting here means we've explicitly decided to save the asset.
-				const FString CurPackageName = CurPackage->GetName();
-				PackagesNotSavedDuringSaveAll.Remove(CurPackageName);
-
-				FSourceControlStatePtr SourceControlState = SourceControlProvider.GetState(CurPackage, EStateCacheUsage::Use);
-				if (SourceControlState.IsValid() && (SourceControlState->IsCheckedOut() || SourceControlState->IsAdded()))
-				{
-					if ( OutPackagesCheckedOut )
+					FSourceControlStatePtr SourceControlState = SourceControlProvider.GetState(CurPackage, EStateCacheUsage::Use);
+					if (SourceControlState && IsValidState(SourceControlState.ToSharedRef()))
 					{
-						OutPackagesCheckedOut->Add(CurPackage);
+						if (OutPackagesCheckedOut)
+						{
+							OutPackagesCheckedOut->Add(CurPackage);
+						}
+					}
+					else
+					{
+						PkgsWhichFailedCheckout += FString::Printf(TEXT("\n%s"), *CurPackageName);
+						CheckOutResult = ECommandResult::Failed;
 					}
-				}
-				else
-				{
-					PkgsWhichFailedCheckout += FString::Printf( TEXT("\n%s"), *CurPackageName );
-					CheckOutResult = ECommandResult::Failed;
 				}
 			}
-		}
+		};
+
+		VerifySourceControlState(FinalPackageCheckoutList, CheckOutResult, [](const FSourceControlStateRef& SourceControlState)
+		{
+			return SourceControlState->IsCheckedOut();
+		});
+
+		VerifySourceControlState(FinalPackageMarkForAddList, MarkForAddResult, [](const FSourceControlStateRef& SourceControlState)
+		{
+			return SourceControlState->IsAdded();
+		});
+
+		VerifySourceControlState(FinalPackageDeleteList, DeleteResult, [](const FSourceControlStateRef& SourceControlState)
+		{
+			return SourceControlState->IsDeleted();
+		});
 	}
 
 	// If any packages failed the check out process, report them to the user so they know
