[Content Browser] Change the default preference for showing "Content" suffix on folders to false.

We have added icons which denote additional content directories from plugins which are the visual differentiator now.

[Content Browser] Change the default preference for showing "Content" suffix on folders to false.

We have added icons which denote additional content directories from plugins which are the visual differentiator now.

diff --git a/Engine/Source/Editor/UnrealEd/Classes/Settings/ContentBrowserSettings.h b/Engine/Source/Editor/UnrealEd/Classes/Settings/ContentBrowserSettings.h
--- a/Engine/Source/Editor/UnrealEd/Classes/Settings/ContentBrowserSettings.h
+++ b/Engine/Source/Editor/UnrealEd/Classes/Settings/ContentBrowserSettings.h
@@ -51,7 +51,7 @@ public:
 
 	/** Whether to append 'Content' text to displayed folder names */
 	UPROPERTY(EditAnywhere, config, Category = ContentBrowser)
-	bool bDisplayContentFolderSuffix = true;
+	bool bDisplayContentFolderSuffix = false;
 
 	/** Whether display friendly name as plugin folder names */
 	UPROPERTY(EditAnywhere, config, Category = ContentBrowser)
