PR #11608: Add support for custom k2 nodes to be more easily detected by macros as latent.

Currently, we hard code any nodes that might be latent and nested in a macro. However, this is restrictive for anyone trying to create a custom node with latent behavior. Instead, nodes can override IsLatentForMacros.

Epic change: moved IsLatentForMacros to UK2Node from UEdGraphNode. Macros are a Blueprint-only concept.

diff --git a/Engine/Source/Editor/BlueprintGraph/Classes/K2Node.h b/Engine/Source/Editor/BlueprintGraph/Classes/K2Node.h
--- a/Engine/Source/Editor/BlueprintGraph/Classes/K2Node.h
+++ b/Engine/Source/Editor/BlueprintGraph/Classes/K2Node.h
@@ -449,6 +449,9 @@ public:
 	/** Helper function for ExpandNode(), allowing other contexts to call pin expansion alone */
 	BLUEPRINTGRAPH_API void ExpandSplitPins(FKismetCompilerContext& CompilerContext, UEdGraph* SourceGraph);
 
+	/** Hook for allowing derived types to determine if they should be treated as a latent node by macros */
+	virtual bool IsLatentForMacros() const { return false; }
+
 protected:
 
 	enum ERedirectType
diff --git a/Engine/Source/Editor/BlueprintGraph/Classes/K2Node_BaseAsyncTask.h b/Engine/Source/Editor/BlueprintGraph/Classes/K2Node_BaseAsyncTask.h
--- a/Engine/Source/Editor/BlueprintGraph/Classes/K2Node_BaseAsyncTask.h
+++ b/Engine/Source/Editor/BlueprintGraph/Classes/K2Node_BaseAsyncTask.h
@@ -60,6 +60,7 @@ class BLUEPRINTGRAPH_API UK2Node_BaseAsyncTask : public UK2Node
 	virtual void GetPinHoverText(const UEdGraphPin& Pin, FString& HoverTextOut) const override;
 	virtual FString GetPinMetaData(FName InPinName, FName InKey) override;
 	virtual bool HasExternalDependencies(TArray<class UStruct*>* OptionalOutput) const override;
+	virtual bool IsLatentForMacros() const override;
 	// End of UEdGraphNode interface
 
 	// UK2Node interface
diff --git a/Engine/Source/Editor/BlueprintGraph/Classes/K2Node_CallFunction.h b/Engine/Source/Editor/BlueprintGraph/Classes/K2Node_CallFunction.h
--- a/Engine/Source/Editor/BlueprintGraph/Classes/K2Node_CallFunction.h
+++ b/Engine/Source/Editor/BlueprintGraph/Classes/K2Node_CallFunction.h
@@ -125,6 +125,7 @@ public:
 	virtual void GetNodeContextMenuActions(class UToolMenu* Menu, class UGraphNodeContextMenuContext* Context) const override;
 	virtual FString GetPinMetaData(FName InPinName, FName InKey) override;
 	virtual bool HasExternalDependencies(TArray<class UStruct*>* OptionalOutput) const override;
+	virtual bool IsLatentForMacros() const override;
 	// End of UEdGraphNode interface
 
 	// UK2Node interface
diff --git a/Engine/Source/Editor/BlueprintGraph/Private/K2Node_BaseAsyncTask.cpp b/Engine/Source/Editor/BlueprintGraph/Private/K2Node_BaseAsyncTask.cpp
--- a/Engine/Source/Editor/BlueprintGraph/Private/K2Node_BaseAsyncTask.cpp
+++ b/Engine/Source/Editor/BlueprintGraph/Private/K2Node_BaseAsyncTask.cpp
@@ -554,6 +554,11 @@ bool UK2Node_BaseAsyncTask::HasExternalDependencies(TArray<class UStruct*>* Opti
 	return bProxyFactoryResult || bProxyResult || bSuperResult;
 }
 
+bool UK2Node_BaseAsyncTask::IsLatentForMacros() const
+{
+	return true;
+}
+
 FName UK2Node_BaseAsyncTask::GetCornerIcon() const
 {
 	return TEXT("Graph.Latent.LatentIcon");
diff --git a/Engine/Source/Editor/BlueprintGraph/Private/K2Node_CallFunction.cpp b/Engine/Source/Editor/BlueprintGraph/Private/K2Node_CallFunction.cpp
--- a/Engine/Source/Editor/BlueprintGraph/Private/K2Node_CallFunction.cpp
+++ b/Engine/Source/Editor/BlueprintGraph/Private/K2Node_CallFunction.cpp
@@ -3483,6 +3483,12 @@ bool UK2Node_CallFunction::HasExternalDependencies(TArray<class UStruct*>* Optio
 	return bSuperResult || bResult;
 }
 
+bool UK2Node_CallFunction::IsLatentForMacros() const
+{
+	UFunction* TargetFunction = GetTargetFunction();
+	return (TargetFunction && TargetFunction->HasMetaData(FBlueprintMetadata::MD_Latent));
+}
+
 UEdGraph* UK2Node_CallFunction::GetFunctionGraph(const UEdGraphNode*& OutGraphNode) const
 {
 	OutGraphNode = nullptr;
diff --git a/Engine/Source/Editor/UnrealEd/Private/Kismet2/BlueprintEditorUtils.cpp b/Engine/Source/Editor/UnrealEd/Private/Kismet2/BlueprintEditorUtils.cpp
--- a/Engine/Source/Editor/UnrealEd/Private/Kismet2/BlueprintEditorUtils.cpp
+++ b/Engine/Source/Editor/UnrealEd/Private/Kismet2/BlueprintEditorUtils.cpp
@@ -9736,43 +9736,36 @@ bool FBlueprintEditorUtils::CheckIfGraphHasLatentFunctions(UEdGraph* InGraph)
 
 				for( const UEdGraphNode* Node : InGraphToCheck->Nodes )
 				{
-					if(const UK2Node_CallFunction* CallFunctionNode = Cast<UK2Node_CallFunction>(Node))
+					const UK2Node* K2Node = Cast<UK2Node>(Node);
+					if (ensure(K2Node))
 					{
-						// Check any function call nodes to see if they are latent.
-						UFunction* TargetFunction = CallFunctionNode->GetTargetFunction();
-						if (TargetFunction && TargetFunction->HasMetaData(FBlueprintMetadata::MD_Latent))
+						if (K2Node->IsLatentForMacros())
 						{
 							TunnelNode->MetaData.HasLatentFunctions = 1;
 							return true;
 						}
-					}
-					else if(const UK2Node_BaseAsyncTask* BaseAsyncNode = Cast<UK2Node_BaseAsyncTask>(Node))
-					{
-						// Async tasks are latent nodes
-						TunnelNode->MetaData.HasLatentFunctions = 1;
-						return true;
-					}
-					else if(const UK2Node_MacroInstance* MacroInstanceNode = Cast<UK2Node_MacroInstance>(Node))
-					{
-						// Any macro graphs that haven't already been checked need to be checked for latent function calls
-						if(InspectedGraphList.Find(MacroInstanceNode->GetMacroGraph()) == INDEX_NONE)
+						else if (const UK2Node_MacroInstance* MacroInstanceNode = Cast<UK2Node_MacroInstance>(K2Node))
 						{
-							if(CheckIfGraphHasLatentFunctions(MacroInstanceNode->GetMacroGraph(), InspectedGraphList))
+							// Any macro graphs that haven't already been checked need to be checked for latent function calls
+							if (InspectedGraphList.Find(MacroInstanceNode->GetMacroGraph()) == INDEX_NONE)
 							{
-								TunnelNode->MetaData.HasLatentFunctions = 1;
-								return true;
+								if (CheckIfGraphHasLatentFunctions(MacroInstanceNode->GetMacroGraph(), InspectedGraphList))
+								{
+									TunnelNode->MetaData.HasLatentFunctions = 1;
+									return true;
+								}
 							}
 						}
-					}
-					else if(const UK2Node_Composite* CompositeNode = Cast<UK2Node_Composite>(Node))
-					{
-						// Any collapsed graphs that haven't already been checked need to be checked for latent function calls
-						if(InspectedGraphList.Find(CompositeNode->BoundGraph) == INDEX_NONE)
+						else if (const UK2Node_Composite* CompositeNode = Cast<UK2Node_Composite>(K2Node))
 						{
-							if(CheckIfGraphHasLatentFunctions(CompositeNode->BoundGraph, InspectedGraphList))
+							// Any collapsed graphs that haven't already been checked need to be checked for latent function calls
+							if (InspectedGraphList.Find(CompositeNode->BoundGraph) == INDEX_NONE)
 							{
-								TunnelNode->MetaData.HasLatentFunctions = 1;
-								return true;
+								if (CheckIfGraphHasLatentFunctions(CompositeNode->BoundGraph, InspectedGraphList))
+								{
+									TunnelNode->MetaData.HasLatentFunctions = 1;
+									return true;
+								}
 							}
 						}
 					}
