FApplicationMode / AssetEditorToolkit- Allow derived classes to specify the config file that should be used for persisting its layout.

diff --git a/Engine/Source/Editor/UnrealEd/Private/ApplicationMode.cpp b/Engine/Source/Editor/UnrealEd/Private/ApplicationMode.cpp
--- a/Engine/Source/Editor/UnrealEd/Private/ApplicationMode.cpp
+++ b/Engine/Source/Editor/UnrealEd/Private/ApplicationMode.cpp
@@ -30,7 +30,7 @@ void FApplicationMode::DeactivateMode(TSharedPtr<FTabManager> InTabManager)
 {
 	// Save the layout to INI
 	check(InTabManager.IsValid());
-	FLayoutSaveRestore::SaveToConfig(GEditorLayoutIni, InTabManager->PersistLayout());
+	FLayoutSaveRestore::SaveToConfig(LayoutIni, InTabManager->PersistLayout());
 
 	// Unregister the tabs
 	/*
@@ -59,7 +59,7 @@ TSharedRef<FTabManager::FLayout> FApplicationMode::ActivateMode(TSharedPtr<FTabM
 
 	// Try loading the layout from INI
 	check(TabLayout.IsValid());
-	return FLayoutSaveRestore::LoadFromConfig(GEditorLayoutIni, TabLayout.ToSharedRef());
+	return FLayoutSaveRestore::LoadFromConfig(LayoutIni, TabLayout.ToSharedRef());
 }
 
 #undef LOCTEXT_NAMESPACE
diff --git a/Engine/Source/Editor/UnrealEd/Private/Toolkits/AssetEditorToolkit.cpp b/Engine/Source/Editor/UnrealEd/Private/Toolkits/AssetEditorToolkit.cpp
--- a/Engine/Source/Editor/UnrealEd/Private/Toolkits/AssetEditorToolkit.cpp
+++ b/Engine/Source/Editor/UnrealEd/Private/Toolkits/AssetEditorToolkit.cpp
@@ -312,7 +312,7 @@ void FAssetEditorToolkit::InitAssetEditor( const EToolkitMode::Type Mode, const
 	{
 		AddMenuExtender(GetSharedMenuExtensibilityManager()->GetAllExtenders(ToolkitCommands, ObjectPtrDecay(EditingObjects)));
 
-		TSharedRef<FTabManager::FLayout> LayoutToUse = FLayoutSaveRestore::LoadFromConfig(GEditorLayoutIni, StandaloneDefaultLayout);
+		TSharedRef<FTabManager::FLayout> LayoutToUse = FLayoutSaveRestore::LoadFromConfig(LayoutIni, StandaloneDefaultLayout);
 
 		// Actually create the widget content
 		if (NewStandaloneHost)
@@ -1524,10 +1524,10 @@ void FAssetEditorToolkit::RestoreFromLayout(const TSharedRef<FTabManager::FLayou
 	if (HostWidget.IsValid())
 	{
 		// Save the old layout
-		FLayoutSaveRestore::SaveToConfig(GEditorLayoutIni, TabManager->PersistLayout());
+		FLayoutSaveRestore::SaveToConfig(LayoutIni, TabManager->PersistLayout());
 
 		// Load the potentially previously saved new layout
-		TSharedRef<FTabManager::FLayout> UserConfiguredNewLayout = FLayoutSaveRestore::LoadFromConfig(GEditorLayoutIni, NewLayout);
+		TSharedRef<FTabManager::FLayout> UserConfiguredNewLayout = FLayoutSaveRestore::LoadFromConfig(LayoutIni, NewLayout);
 
 		for (TSharedPtr<FLayoutExtender> LayoutExtender : LayoutExtenders)
 		{
diff --git a/Engine/Source/Editor/UnrealEd/Public/Toolkits/AssetEditorToolkit.h b/Engine/Source/Editor/UnrealEd/Public/Toolkits/AssetEditorToolkit.h
--- a/Engine/Source/Editor/UnrealEd/Public/Toolkits/AssetEditorToolkit.h
+++ b/Engine/Source/Editor/UnrealEd/Public/Toolkits/AssetEditorToolkit.h
@@ -495,6 +495,9 @@ protected:
 	/** The base category that tabs are registered to, allows for child classes to register to the same point. */
 	TSharedPtr<FWorkspaceItem> AssetEditorTabsCategory;
 
+	/** What ini file should be used to persist and restore the user's layout. Defaults to GEditorLayoutIni. */
+	FString LayoutIni = GEditorLayoutIni;
+
 private:
 	static UNREALED_API const FName DefaultAssetEditorToolBarName;
 	/** The toolkit standalone host; may be nullptr */
diff --git a/Engine/Source/Editor/UnrealEd/Public/WorkflowOrientedApp/ApplicationMode.h b/Engine/Source/Editor/UnrealEd/Public/WorkflowOrientedApp/ApplicationMode.h
--- a/Engine/Source/Editor/UnrealEd/Public/WorkflowOrientedApp/ApplicationMode.h
+++ b/Engine/Source/Editor/UnrealEd/Public/WorkflowOrientedApp/ApplicationMode.h
@@ -37,6 +37,9 @@ protected:
 	/** The workspace menu category for this mode */
 	TSharedPtr<FWorkspaceItem> WorkspaceMenuCategory;
 
+	/** What ini file should be used to persist and restore the user's layout. Defaults to GEditorLayoutIni. */
+	FString LayoutIni = GEditorLayoutIni;
+
 public:
 	UNREALED_API FApplicationMode(FName InModeName);
 
