Remove stray #undef LOCTEXT_NAMESPACE
#rnx

diff --git a/Engine/Source/Editor/UnrealEd/Private/Factories/SparseVolumeTextureMaterialFactory.cpp b/Engine/Source/Editor/UnrealEd/Private/Factories/SparseVolumeTextureMaterialFactory.cpp
--- a/Engine/Source/Editor/UnrealEd/Private/Factories/SparseVolumeTextureMaterialFactory.cpp
+++ b/Engine/Source/Editor/UnrealEd/Private/Factories/SparseVolumeTextureMaterialFactory.cpp
@@ -88,5 +88,3 @@ UObject* USparseVolumeTextureMaterialInstanceFactoryNew::FactoryCreateNew(UClass
 }
 
 #endif //WITH_EDITOR
-
-#undef LOCTEXT_NAMESPACE
