Refresh ToolMenu widgets when the editor setting "Display UI Extension Points" is toggled

diff --git a/Engine/Source/Editor/UnrealEd/Private/Settings/EditorPerProjectUserSettings.cpp b/Engine/Source/Editor/UnrealEd/Private/Settings/EditorPerProjectUserSettings.cpp
--- a/Engine/Source/Editor/UnrealEd/Private/Settings/EditorPerProjectUserSettings.cpp
+++ b/Engine/Source/Editor/UnrealEd/Private/Settings/EditorPerProjectUserSettings.cpp
@@ -4,6 +4,7 @@
 #include "Misc/Paths.h"
 #include "HAL/IConsoleManager.h"
 #include "RHI.h"
+#include "ToolMenus.h"
 #include "UnrealEdMisc.h"
 #include "BlueprintPaletteFavorites.h"
 
@@ -74,6 +75,14 @@ void UEditorPerProjectUserSettings::PostEditChangeProperty( FPropertyChangedEven
 		//GDistributionType == 0 for curves
 		GDistributionType = (bUseCurvesForDistributions) ? 0 : 1;
 	}
+	else if (Name == "bDisplayUIExtensionPoints")
+	{
+		// Refresh all ToolMenu widgets to make them reflect this setting's new value.
+		if (UToolMenus* ToolMenus = UToolMenus::TryGet())
+		{
+			ToolMenus->RefreshAllWidgets();
+		}
+	}
 
 	if (!FUnrealEdMisc::Get().IsDeletePreferences())
 	{
