[Viewport Toolbar] View Modes menu in old viewport toolbar missing exposure settings

- Adds Exposure section back to View Modes menu
- Creates function to create Exposure section, now used by both new and old toolbar

diff --git a/Engine/Source/Editor/UnrealEd/Private/SEditorViewportViewMenu.cpp b/Engine/Source/Editor/UnrealEd/Private/SEditorViewportViewMenu.cpp
--- a/Engine/Source/Editor/UnrealEd/Private/SEditorViewportViewMenu.cpp
+++ b/Engine/Source/Editor/UnrealEd/Private/SEditorViewportViewMenu.cpp
@@ -106,6 +106,7 @@ void SEditorViewportViewMenu::FillViewMenu(UToolMenu* Menu) const
 	}
 
 	UE::UnrealEd::PopulateViewModesMenu(Menu);
+	UE::UnrealEd::AddExposureSection(Menu, Viewport.Pin());
 }
 
 #undef LOCTEXT_NAMESPACE
diff --git a/Engine/Source/Editor/UnrealEd/Private/ViewportToolbar/UnrealEdViewportToolbar.cpp b/Engine/Source/Editor/UnrealEd/Private/ViewportToolbar/UnrealEdViewportToolbar.cpp
--- a/Engine/Source/Editor/UnrealEd/Private/ViewportToolbar/UnrealEdViewportToolbar.cpp
+++ b/Engine/Source/Editor/UnrealEd/Private/ViewportToolbar/UnrealEdViewportToolbar.cpp
@@ -2090,7 +2090,7 @@ FToolMenuEntry CreateNumericEntry(
 	return NumericEntry;
 }
 
-TSharedRef<SWidget> CreateCameraMenuWidget(const TSharedRef<SEditorViewport>& InViewport)
+TSharedRef<SWidget> CreateCameraMenuWidget(const TSharedRef<SEditorViewport>& InViewport, bool bInShowExposureSettings)
 {
 	// We generate a menu via UToolMenus, so we can use FillShowSubmenu call from both old and new toolbar
 	FName OldShowMenuName = "LevelEditor.OldViewportToolbar.CameraOptions";
@@ -2118,6 +2118,7 @@ TSharedRef<SWidget> CreateCameraMenuWidget(const TSharedRef<SEditorViewport>& In
 			UUnrealEdViewportToolbarContext* const ContextObject =
 				UE::UnrealEd::CreateViewportToolbarDefaultContext(InViewport);
 
+			ContextObject->bExposureSectionLivesInCameraMenu = bInShowExposureSettings;
 			MenuContext.AddObject(ContextObject);
 		}
 	}
@@ -2575,6 +2576,19 @@ FToolMenuEntry CreateAssetViewerProfileSubmenu()
 	);
 }
 
+void AddExposureSection(UToolMenu* InMenu, const TSharedPtr<SEditorViewport>& EditorViewport)
+{
+	const FEditorViewportCommands& BaseViewportCommands = FEditorViewportCommands::Get();
+
+	TSharedRef<SWidget> FixedEV100Menu = EditorViewport->BuildFixedEV100Menu();
+	TSharedPtr<FEditorViewportClient> EditorViewportClient = EditorViewport->GetViewportClient();
+	const bool bIsLevelEditor = EditorViewportClient.IsValid() && EditorViewportClient->IsLevelEditorClient();
+
+	FToolMenuSection& Section = InMenu->AddSection("Exposure", LOCTEXT("ExposureHeader", "Exposure"));
+	Section.AddMenuEntry(bIsLevelEditor ? BaseViewportCommands.ToggleInGameExposure : BaseViewportCommands.ToggleAutoExposure);
+	Section.AddEntry(FToolMenuEntry::InitWidget("FixedEV100", FixedEV100Menu, LOCTEXT("FixedEV100", "EV100")));
+}
+
 void PopulateCameraMenu(UToolMenu* InMenu)
 {
 	UUnrealEdViewportToolbarContext* const EditorViewportContext = InMenu->FindContext<UUnrealEdViewportToolbarContext>();
@@ -2606,22 +2620,13 @@ void PopulateCameraMenu(UToolMenu* InMenu)
 	}
 
 	// Auto Exposure
+	//if (EditorViewportContext->bExposureSectionLivesInCameraMenu)
 	{
 		if (UE::UnrealEd::Private::ViewModesSubmenu::IsMenuSectionAvailable(
 				EditorViewportContext, EHidableViewModeMenuSections::Exposure
 			))
 		{
-			const FEditorViewportCommands& BaseViewportCommands = FEditorViewportCommands::Get();
-
-			TSharedRef<SWidget> FixedEV100Menu = EditorViewport->BuildFixedEV100Menu();
-			TSharedPtr<FEditorViewportClient> EditorViewportClient = EditorViewport->GetViewportClient();
-			const bool bIsLevelEditor = EditorViewportClient.IsValid() && EditorViewportClient->IsLevelEditorClient();
-
-			FToolMenuSection& Section = InMenu->AddSection("Exposure", LOCTEXT("ExposureHeader", "Exposure"));
-			Section.AddMenuEntry(
-				bIsLevelEditor ? BaseViewportCommands.ToggleInGameExposure : BaseViewportCommands.ToggleAutoExposure
-			);
-			Section.AddEntry(FToolMenuEntry::InitWidget("FixedEV100", FixedEV100Menu, LOCTEXT("FixedEV100", "EV100")));
+			AddExposureSection(InMenu, EditorViewport);
 		}
 	}
 }
diff --git a/Engine/Source/Editor/UnrealEd/Public/ViewportToolbar/UnrealEdViewportToolbar.h b/Engine/Source/Editor/UnrealEd/Public/ViewportToolbar/UnrealEdViewportToolbar.h
--- a/Engine/Source/Editor/UnrealEd/Public/ViewportToolbar/UnrealEdViewportToolbar.h
+++ b/Engine/Source/Editor/UnrealEd/Public/ViewportToolbar/UnrealEdViewportToolbar.h
@@ -154,6 +154,7 @@ UE_DEPRECATED(5.6, "Use CreateAssetViewerProfileSubmenu() instead.")
 UNREALED_API FToolMenuEntry CreateViewportToolbarAssetViewerProfileSubmenu();
 UNREALED_API FToolMenuEntry CreateAssetViewerProfileSubmenu();
 
+UNREALED_API void AddExposureSection(UToolMenu* InMenu, const TSharedPtr<SEditorViewport>& EditorViewport);
 UNREALED_API void PopulateCameraMenu(UToolMenu* InMenu);
 
 /**
@@ -198,7 +199,7 @@ UNREALED_API void ConstructScreenPercentageMenu(UToolMenu* InMenu);
 UNREALED_API FToolMenuEntry CreateScalabilitySubmenu();
 
 // Camera Menu Widgets
-TSharedRef<SWidget> CreateCameraMenuWidget(const TSharedRef<SEditorViewport>& InViewport);
+TSharedRef<SWidget> CreateCameraMenuWidget(const TSharedRef<SEditorViewport>& InViewport, bool bInShowExposureSettings = false);
 TSharedRef<SWidget> CreateFOVMenuWidget(const TSharedRef<SEditorViewport>& InViewport);
 TSharedRef<SWidget> CreateFarViewPlaneMenuWidget(const TSharedRef<SEditorViewport>& InViewport);
 
diff --git a/Engine/Source/Editor/UnrealEd/Public/ViewportToolbar/UnrealEdViewportToolbarContext.h b/Engine/Source/Editor/UnrealEd/Public/ViewportToolbar/UnrealEdViewportToolbarContext.h
--- a/Engine/Source/Editor/UnrealEd/Public/ViewportToolbar/UnrealEdViewportToolbarContext.h
+++ b/Engine/Source/Editor/UnrealEd/Public/ViewportToolbar/UnrealEdViewportToolbarContext.h
@@ -41,6 +41,12 @@ public:
 	 */
 	bool bShowGridToggleInPreviewSceneSettings = false;
 
+	/**
+	 * Added for compatibility with Old Toolbar, where Exposure settings lived in View Mode menu.
+	 * New toolbars have Exposure settings in Camera menu.
+	 */
+	bool bExposureSectionLivesInCameraMenu = true;
+
 	/**
 	 * Add flags to this array to remove them from the Show Menu
 	 */
