[Sequencer] Selectability: Disallow selection of translucent actors when UEditorPerProjectUserSettings::bAllowSelectTranslucent is false

diff --git a/Engine/Source/Editor/UnrealEd/Private/EditorViewportSelectability.cpp b/Engine/Source/Editor/UnrealEd/Private/EditorViewportSelectability.cpp
--- a/Engine/Source/Editor/UnrealEd/Private/EditorViewportSelectability.cpp
+++ b/Engine/Source/Editor/UnrealEd/Private/EditorViewportSelectability.cpp
@@ -2,6 +2,7 @@
 
 #include "EditorViewportSelectability.h"
 #include "CanvasTypes.h"
+#include "Editor/EditorPerProjectUserSettings.h"
 #include "Engine/Canvas.h"
 #include "EngineUtils.h"
 #include "LevelEditorViewport.h"
@@ -484,6 +485,15 @@ bool FEditorViewportSelectability::HandleClick(FEditorViewportClient* const InVi
 		return false;
 	}
 
+	// Check for translucent actors if we don't want to allow them to be selected
+	const UEditorPerProjectUserSettings* const Settings = GetDefault<UEditorPerProjectUserSettings>();
+	if (!Settings->bAllowSelectTranslucent && HitResult->IsA(HTranslucentActor::StaticGetType()))
+	{
+		// Return true to disable selection of valid translucent actors
+		const HTranslucentActor* const TranslucentActorHitProxy = static_cast<const HTranslucentActor*>(HitResult);
+		return IsValid(TranslucentActorHitProxy->Actor);
+	}
+
 	const HActor* const ActorHitProxy = static_cast<const HActor*>(HitResult);
 	if (!ActorHitProxy || !IsValid(ActorHitProxy->Actor))
 	{
