Cooker: Deprecate ReferencedOnlyByEditorOnlyData: implementation has been replaced by skiponlyeditoronly, and remaining use fits better with ContainsEditorOnlyData.
#rnx

diff --git a/Engine/Source/Developer/IoStoreUtilities/Private/ZenStoreWriter.cpp b/Engine/Source/Developer/IoStoreUtilities/Private/ZenStoreWriter.cpp
--- a/Engine/Source/Developer/IoStoreUtilities/Private/ZenStoreWriter.cpp
+++ b/Engine/Source/Developer/IoStoreUtilities/Private/ZenStoreWriter.cpp
@@ -712,8 +712,7 @@ void FZenStoreWriter::CommitPackage(FCommitPackageInfo&& Info)
 		if (Info.Status == ECommitStatus::Success)
 		{
 			// Only record hashes for successful saves. A single package can be saved unsuccessfully multiple times
-			// during a cook if its rejected for being only referenced by editor-only references and we keep finding
-			// new references to it.
+			// during a cook if it keeps timing out.
 			TRefCountPtr<FPackageHashes>& ExistingPackageHashes = AllPackageHashes.FindOrAdd(Info.PackageName);
 			// This looks weird but we've found the _TRefCountPtr_, not the FPackageHashes. When newly assigned
 			// it will be an empty pointer, which is what we want.
diff --git a/Engine/Source/Editor/UnrealEd/Private/CookOnTheFlyServer.cpp b/Engine/Source/Editor/UnrealEd/Private/CookOnTheFlyServer.cpp
--- a/Engine/Source/Editor/UnrealEd/Private/CookOnTheFlyServer.cpp
+++ b/Engine/Source/Editor/UnrealEd/Private/CookOnTheFlyServer.cpp
@@ -4612,7 +4612,6 @@ private: // Used only by UCookOnTheFlyServer, which has private access
 	const FString PackageName;
 	FString Filename;
 	uint32 SaveFlags = 0;
-	bool bReferencedOnlyByEditorOnlyData = false;
 	bool bHasTimeOut = false;
 	bool bHasRetryErrorCode = false;
 	bool bPlatformAgnosticDependenciesCalculated = false;
@@ -4678,10 +4677,6 @@ void UCookOnTheFlyServer::PumpSaves(UE::Cook::FTickStackData& StackData, uint32
 		UE_LOG(LogCook, Display, TEXT("Processing save for package %s"), *Package->GetName());
 #endif
 
-		// This package is valid, so make sure it wasn't previously marked as being an uncooked editor only package or it would get removed from the
-		// asset registry at the end of the cook
-		PackageTracker->UncookedEditorOnlyPackages.Remove(Package->GetFName());
-
 		if (PackageTracker->NeverCookPackageList.Contains(Package->GetFName()))
 		{
 			// refuse to save this package, it's clearly one of the undesirables
@@ -6650,8 +6645,7 @@ void FSaveCookedPackageContext::SetupPlatform(const ITargetPlatform* InTargetPla
 
 bool IsRetryErrorCode(ESavePackageResult Result)
 {
-	return (Result == ESavePackageResult::ReferencedOnlyByEditorOnlyData) |
-		(Result == ESavePackageResult::Timeout);
+	return Result == ESavePackageResult::Timeout;
 }
 
 void FSaveCookedPackageContext::FinishPlatform()
@@ -6746,18 +6740,6 @@ void FSaveCookedPackageContext::FinishPlatform()
 	}
 
 	// Accumulate results for SaveCookedPackage_Finish
-	bool bLocalReferencedOnlyByEditorOnlyData = SavePackageResult.Result == ESavePackageResult::ReferencedOnlyByEditorOnlyData;
-	if (PlatformIndex > 0 && bLocalReferencedOnlyByEditorOnlyData != bReferencedOnlyByEditorOnlyData)
-	{
-		UE_LOG(LogCook, Error, TEXT("Package %s had different values for IsReferencedOnlyByEditorOnlyData from multiple platforms. ")
-			TEXT("Treating all platforms as IsReferencedOnlyByEditorOnlyData = true; this will cause the package to be ignored on the platforms that need it."),
-			*Package->GetName());
-		bReferencedOnlyByEditorOnlyData = true;
-	}
-	else
-	{
-		bReferencedOnlyByEditorOnlyData = bLocalReferencedOnlyByEditorOnlyData;
-	}
 	if (SavePackageResult.Result == ESavePackageResult::Timeout)
 	{
 		PackageData.FindOrAddPlatformData(TargetPlatform).SetSaveTimedOut(true);
@@ -6851,10 +6833,6 @@ void FSaveCookedPackageContext::FinishPackage()
 			}
 		}
 	}
-	else if (bReferencedOnlyByEditorOnlyData)
-	{
-		COTFS.PackageTracker->UncookedEditorOnlyPackages.AddUnique(Package->GetFName());
-	}
 }
 
 void FSaveCookedPackageContext::CalculatePlatformAgnosticRuntimeDependencies()
@@ -10818,14 +10796,6 @@ void UCookOnTheFlyServer::CookByTheBookFinishedInternal()
 			{
 				IgnorePackageNames.Add(PackageData->GetPackageName());
 			}
-
-			// ignore packages that weren't cooked because they were only referenced by editor-only properties
-			TSet<FName> UncookedEditorOnlyPackageNames;
-			PackageTracker->UncookedEditorOnlyPackages.GetValues(UncookedEditorOnlyPackageNames);
-			for (FName UncookedEditorOnlyPackage : UncookedEditorOnlyPackageNames)
-			{
-				IgnorePackageNames.Add(UncookedEditorOnlyPackage);
-			}
 		}
 		
 		if (bCacheShaderLibraries)
diff --git a/Engine/Source/Editor/UnrealEd/Private/Cooker/CookWorkerServer.cpp b/Engine/Source/Editor/UnrealEd/Private/Cooker/CookWorkerServer.cpp
--- a/Engine/Source/Editor/UnrealEd/Private/Cooker/CookWorkerServer.cpp
+++ b/Engine/Source/Editor/UnrealEd/Private/Cooker/CookWorkerServer.cpp
@@ -969,10 +969,6 @@ void FCookWorkerServer::RecordResults(FPackageResultsMessage& Message)
 					}
 				}
 				COTFS.RecordExternalActorDependencies(Result.GetExternalActorDependencies());
-				if (Result.IsReferencedOnlyByEditorOnlyData())
-				{
-					COTFS.PackageTracker->UncookedEditorOnlyPackages.AddUnique(Result.GetPackageName());
-				}
 			}
 		}
 
diff --git a/Engine/Source/Editor/UnrealEd/Private/Cooker/PackageResultsMessage.cpp b/Engine/Source/Editor/UnrealEd/Private/Cooker/PackageResultsMessage.cpp
--- a/Engine/Source/Editor/UnrealEd/Private/Cooker/PackageResultsMessage.cpp
+++ b/Engine/Source/Editor/UnrealEd/Private/Cooker/PackageResultsMessage.cpp
@@ -136,7 +136,6 @@ void FPackageResultsMessage::Write(FCbWriter& Writer) const
 		{
 			Writer << "N" << Result.PackageName;
 			Writer << "R" << (uint8) Result.SuppressCookReason;
-			Writer << "E" << Result.bReferencedOnlyByEditorOnlyData;
 			WriteMessagesArray(Writer, Result.Messages);
 			Writer.BeginArray("P");
 			for (const FPackageRemoteResult::FPlatformResult& PlatformResult : Result.Platforms)
@@ -191,7 +190,6 @@ bool FPackageResultsMessage::TryRead(FCbObjectView Object)
 			return false;
 		}
 		Result.SuppressCookReason = static_cast<ESuppressCookReason>(LocalSuppressCookReason);
-		Result.bReferencedOnlyByEditorOnlyData = ResultObject["E"].AsBool();
 		if (!TryReadMessagesArray(ResultObject, Result.Messages))
 		{
 			return false;
diff --git a/Engine/Source/Editor/UnrealEd/Private/Cooker/PackageTracker.h b/Engine/Source/Editor/UnrealEd/Private/Cooker/PackageTracker.h
--- a/Engine/Source/Editor/UnrealEd/Private/Cooker/PackageTracker.h
+++ b/Engine/Source/Editor/UnrealEd/Private/Cooker/PackageTracker.h
@@ -199,7 +199,6 @@ public:
 
 	/** Packages to never cook - entries are localpaths in FPaths::MakeStandardFilename format. */
 	FThreadSafeSet<FName> NeverCookPackageList;
-	FThreadSafeSet<FName> UncookedEditorOnlyPackages; // set of packages that have been rejected due to being referenced by editor-only properties
 	TFastPointerMap<const ITargetPlatform*, TSet<FName>> PlatformSpecificNeverCookPackages;
 
 	// Thread-safe enumeration of loaded package. 
diff --git a/Engine/Source/Runtime/CoreUObject/Private/UObject/SavePackage2.cpp b/Engine/Source/Runtime/CoreUObject/Private/UObject/SavePackage2.cpp
--- a/Engine/Source/Runtime/CoreUObject/Private/UObject/SavePackage2.cpp
+++ b/Engine/Source/Runtime/CoreUObject/Private/UObject/SavePackage2.cpp
@@ -106,7 +106,7 @@ ESavePackageResult ValidatePackage(FSaveContext& SaveContext)
 			if (SaveContext.GetPackage()->HasAnyPackageFlags(PKG_EditorOnly))
 			{
 				UE_CLOG(SaveContext.IsGenerateSaveError(), LogSavePackage, Verbose, TEXT("Package marked as editor-only: %s. Package will not be saved."), *SaveContext.GetPackage()->GetName());
-				return ESavePackageResult::ReferencedOnlyByEditorOnlyData;
+				return ESavePackageResult::ContainsEditorOnlyData;
 			}
 		}
 #endif
diff --git a/Engine/Source/Runtime/CoreUObject/Public/UObject/Package.h b/Engine/Source/Runtime/CoreUObject/Public/UObject/Package.h
--- a/Engine/Source/Runtime/CoreUObject/Public/UObject/Package.h
+++ b/Engine/Source/Runtime/CoreUObject/Public/UObject/Package.h
@@ -69,8 +69,7 @@ enum class ESavePackageResult
 	Canceled,
 	/** [When cooking] Package was not saved because it contained editor-only data */
 	ContainsEditorOnlyData,
-	/** [When cooking] Package was not saved because it was referenced by editor-only properties */
-	ReferencedOnlyByEditorOnlyData,
+	ReferencedOnlyByEditorOnlyData UE_DEPRECATED(5.6, "The cooker now uses SkipOnlyEditorOnly to detect ReferencedOnlyByEditorOnlyData instead of detecting it during SavePackage"),
 	/** [When cooking] Package was not saved because it contains assets that were converted into native code */
 	ReplaceCompletely,
 	/** [When cooking] Package was saved, but we should generate a stub so that other converted packages can interface with it*/
