Call UpdateResource on Textures when changing the shaderplatform preview to update the CachePlatformData

diff --git a/Engine/Source/Editor/UnrealEd/Private/EditorEngine.cpp b/Engine/Source/Editor/UnrealEd/Private/EditorEngine.cpp
--- a/Engine/Source/Editor/UnrealEd/Private/EditorEngine.cpp
+++ b/Engine/Source/Editor/UnrealEd/Private/EditorEngine.cpp
@@ -6363,6 +6363,13 @@ void UEditorEngine::SetPreviewPlatform(const FPreviewPlatformInfo& NewPreviewPla
 		{
 			ObjIterator->LoadConfig();
 		}
+
+		//Call UpdateResource on Textures when changing the ShaderPlatformPreview to update the CachePlatformData
+		UTexture* Texture = Cast<UTexture>(*ObjIterator);
+		if (Texture != nullptr)
+		{
+			Texture->UpdateResource();
+		}
 	}
 
 	constexpr bool bUpdateProgressDialog = true;
@@ -6480,6 +6487,13 @@ void UEditorEngine::ToggleFeatureLevelPreview()
 		{
 			ObjIterator->LoadConfig();
 		}
+
+		//Call UpdateResource on Textures when changing the ShaderPlatformPreview to update the CachePlatformData
+		UTexture* Texture = Cast<UTexture>(*ObjIterator);
+		if (Texture != nullptr)
+		{
+			Texture->UpdateResource();
+		}
 	}
 
 	Scalability::ApplyCachedQualityLevelForShaderPlatform(GetActiveShaderPlatform());
