Removed aggressive ensure that was added in CL 38133046.

Blueprint graphs can have comment nodes, which don't inherit from UK2Node. We can just soft-fail silently in these cases.

#rnx

diff --git a/Engine/Source/Editor/UnrealEd/Private/Kismet2/BlueprintEditorUtils.cpp b/Engine/Source/Editor/UnrealEd/Private/Kismet2/BlueprintEditorUtils.cpp
--- a/Engine/Source/Editor/UnrealEd/Private/Kismet2/BlueprintEditorUtils.cpp
+++ b/Engine/Source/Editor/UnrealEd/Private/Kismet2/BlueprintEditorUtils.cpp
@@ -9736,8 +9736,9 @@ bool FBlueprintEditorUtils::CheckIfGraphHasLatentFunctions(UEdGraph* InGraph)
 
 				for( const UEdGraphNode* Node : InGraphToCheck->Nodes )
 				{
+					// BP graphs can sometimes have non-UK2Nodes (eg: UEdGraphNode_Comment)
 					const UK2Node* K2Node = Cast<UK2Node>(Node);
-					if (ensure(K2Node))
+					if (K2Node)
 					{
 						if (K2Node->IsLatentForMacros())
 						{
