Cooker: Deprecate ReferencedOnlyByEditorOnlyData: remove some additional no-longer needed data and functions.
#rnx

diff --git a/Engine/Source/Editor/UnrealEd/Private/Cooker/PackageResultsMessage.h b/Engine/Source/Editor/UnrealEd/Private/Cooker/PackageResultsMessage.h
--- a/Engine/Source/Editor/UnrealEd/Private/Cooker/PackageResultsMessage.h
+++ b/Engine/Source/Editor/UnrealEd/Private/Cooker/PackageResultsMessage.h
@@ -61,9 +61,6 @@ public:
 	ESuppressCookReason GetSuppressCookReason() const;
 	void SetSuppressCookReason(ESuppressCookReason InSuppressCookReason);
 
-	bool IsReferencedOnlyByEditorOnlyData() const;
-	void SetReferencedOnlyByEditorOnlyData(bool bInReferencedOnlyByEditorOnlyData);
-
 	void AddPackageMessage(const FGuid& MessageType, FCbObject&& Object);
 	void AddAsyncPackageMessage(const FGuid& MessageType, TFuture<FCbObject>&& ObjectFuture);
 	void AddPlatformMessage(const ITargetPlatform* TargetPlatform, const FGuid& MessageType, FCbObject&& Object);
@@ -125,7 +122,6 @@ private:
 	FName PackageName;
 	/** If failure reason is NotSuppressed, it was saved. Otherwise, holds the suppression reason */
 	ESuppressCookReason SuppressCookReason;
-	bool bReferencedOnlyByEditorOnlyData = false;
 
 	// Fields guarded by AsyncSupport->AsyncWorkLock. They can only be read or written if either AsyncSupport is nullptr
 	// or if within AsyncSupport->AsyncWorkLock.
@@ -217,16 +213,6 @@ inline void FPackageRemoteResult::SetSuppressCookReason(ESuppressCookReason InSu
 	SuppressCookReason = InSuppressCookReason;
 }
 
-inline bool FPackageRemoteResult::IsReferencedOnlyByEditorOnlyData() const
-{
-	return bReferencedOnlyByEditorOnlyData;
-}
-
-inline void FPackageRemoteResult::SetReferencedOnlyByEditorOnlyData(bool bInReferencedOnlyByEditorOnlyData)
-{
-	bReferencedOnlyByEditorOnlyData = bInReferencedOnlyByEditorOnlyData;
-}
-
 inline TConstArrayView<UE::CompactBinaryTCP::FMarshalledMessage> FPackageRemoteResult::GetMessages() const
 {
 	return Messages;
