[Curve Editor] Fix crash when Curve wrapped by InstanceStruct gets reset/replaced, causing memory address to change.

diff --git a/Engine/Source/Editor/UnrealEd/Private/SCurveEditor.cpp b/Engine/Source/Editor/UnrealEd/Private/SCurveEditor.cpp
--- a/Engine/Source/Editor/UnrealEd/Private/SCurveEditor.cpp
+++ b/Engine/Source/Editor/UnrealEd/Private/SCurveEditor.cpp
@@ -3819,11 +3819,15 @@ void SCurveEditor::OnObjectPropertyChanged(UObject* Object, FPropertyChangedEven
 		// CurveEditor will hold dangling references to CurveData and access them in some cases:
 		// - Curve Data is inline allocated, or wrapped by InstancedStruct, in a container. And we do Array Remove/Clear op, or Add that causes resize
 		// - Curve Data is wrapped by InstancedStruct in a Container, and we do undo/redo(will trigger emptying and refilling the container, causing address changed)
+		// - Curve Data is wrapped by InstancedStruct and InstancedStruct is being reset/replaced, causing actual memory address changed
+		// todo: should only rebuild when it's the InstanceStruct with ValueSet. But currently there is no way to detect if there is an InstanceStruct in the PropertyNode Hierarchy.
 		if (!bIsPendingRebuilt
 			&& (GIsTransacting
 			|| PropertyChangedEvent.ChangeType == EPropertyChangeType::ArrayAdd
 			|| PropertyChangedEvent.ChangeType == EPropertyChangeType::ArrayRemove
-			|| PropertyChangedEvent.ChangeType == EPropertyChangeType::ArrayClear))
+			|| PropertyChangedEvent.ChangeType == EPropertyChangeType::ArrayClear
+			|| PropertyChangedEvent.ChangeType == EPropertyChangeType::ValueSet
+			))
 		{
 			bIsPendingRebuilt = true;
 
