[Sequencer] Limit selectability: Fix marquee selection incorrectly selected camera actor in scene

diff --git a/Engine/Source/Editor/Sequencer/Private/SequencerEdMode.h b/Engine/Source/Editor/Sequencer/Private/SequencerEdMode.h
--- a/Engine/Source/Editor/Sequencer/Private/SequencerEdMode.h
+++ b/Engine/Source/Editor/Sequencer/Private/SequencerEdMode.h
@@ -14,7 +14,6 @@ class FCanvas;
 class FEditorViewportClient;
 class FPrimitiveDrawInterface;
 class FSceneView;
-class FEditorViewportSelectability;
 class ISequencer;
 class FSequencer;
 class FViewport;
diff --git a/Engine/Source/Editor/UnrealEd/Private/EditorViewportSelectability.cpp b/Engine/Source/Editor/UnrealEd/Private/EditorViewportSelectability.cpp
--- a/Engine/Source/Editor/UnrealEd/Private/EditorViewportSelectability.cpp
+++ b/Engine/Source/Editor/UnrealEd/Private/EditorViewportSelectability.cpp
@@ -3,6 +3,9 @@
 #include "EditorViewportSelectability.h"
 #include "CanvasTypes.h"
 #include "Editor/EditorPerProjectUserSettings.h"
+#include "Elements/Framework/EngineElementsLibrary.h"
+#include "Elements/Framework/TypedElementRegistry.h"
+#include "Elements/Interfaces/TypedElementWorldInterface.h"
 #include "Engine/Canvas.h"
 #include "EngineUtils.h"
 #include "LevelEditorViewport.h"
@@ -336,6 +339,23 @@ bool FEditorViewportSelectability::IsActorInLevelHiddenLayer(const AActor& InAct
 	return false;
 }
 
+TTypedElement<ITypedElementWorldInterface> FEditorViewportSelectability::GetTypedWorldElementFromActor(const AActor& InActor)
+{
+	const FTypedElementHandle ActorElementHandle = UEngineElementsLibrary::AcquireEditorActorElementHandle(&InActor);
+	if (!ActorElementHandle)
+	{
+		return TTypedElement<ITypedElementWorldInterface>();
+	}
+
+	const UTypedElementRegistry* const TypedElementRegistry = UTypedElementRegistry::GetInstance();
+	if (!TypedElementRegistry)
+	{
+		return TTypedElement<ITypedElementWorldInterface>();
+	}
+
+	return TypedElementRegistry->GetElement<ITypedElementWorldInterface>(ActorElementHandle);
+}
+
 bool FEditorViewportSelectability::DoesActorIntersectBox(const AActor& InActor, const FBox& InBox, FEditorViewportClient* const InEditorViewportClient, const bool bInUseStrictSelection)
 {
 	if (InActor.IsHiddenEd() || !IsActorSelectableClass(InActor))
@@ -350,21 +370,13 @@ bool FEditorViewportSelectability::DoesActorIntersectBox(const AActor& InActor,
 		return false;
 	}
 
-	// Iterate over all actor components, selecting our primitive components
-	for (const UActorComponent* const Component : InActor.GetComponents())
+	const TTypedElement<ITypedElementWorldInterface> WorldActorElement = GetTypedWorldElementFromActor(InActor);
+	if (!WorldActorElement)
 	{
-		const UPrimitiveComponent* const PrimitiveComponent = Cast<UPrimitiveComponent>(Component);
-		if (PrimitiveComponent
-			&& PrimitiveComponent->IsRegistered()
-			&& PrimitiveComponent->IsVisibleInEditor()
-			&& PrimitiveComponent->IsShown(InEditorViewportClient->EngineShowFlags)
-			&& PrimitiveComponent->ComponentIsTouchingSelectionBox(InBox, false, bInUseStrictSelection))
-		{
-			return true;
-		}
+		return false;
 	}
 
-	return false;
+	return WorldActorElement.IsElementInBox(InBox, bInUseStrictSelection);
 }
 
 bool FEditorViewportSelectability::DoesActorIntersectFrustum(const AActor& InActor, const FConvexVolume& InFrustum, FEditorViewportClient* const InEditorViewportClient, const bool bInUseStrictSelection)
@@ -381,20 +393,13 @@ bool FEditorViewportSelectability::DoesActorIntersectFrustum(const AActor& InAct
 		return false;
 	}
 
-	for (const UActorComponent* const Component : InActor.GetComponents())
+	const TTypedElement<ITypedElementWorldInterface> WorldActorElement = GetTypedWorldElementFromActor(InActor);
+	if (!WorldActorElement)
 	{
-		const UPrimitiveComponent* const PrimitiveComponent = Cast<UPrimitiveComponent>(Component);
-		if (IsValid(PrimitiveComponent)
-			&& PrimitiveComponent->IsRegistered()
-			&& PrimitiveComponent->IsVisibleInEditor()
-			&& PrimitiveComponent->IsShown(InEditorViewportClient->EngineShowFlags)
-			&& PrimitiveComponent->ComponentIsTouchingSelectionFrustum(InFrustum, false, bInUseStrictSelection))
-		{
-			return true;
-		}
+		return false;
 	}
 
-	return false;
+	return WorldActorElement.IsElementInConvexVolume(InFrustum, bInUseStrictSelection);
 }
 
 bool FEditorViewportSelectability::GetCursorForHovered(EMouseCursor::Type& OutCursor) const
diff --git a/Engine/Source/Editor/UnrealEd/Public/EditorViewportSelectability.h b/Engine/Source/Editor/UnrealEd/Public/EditorViewportSelectability.h
--- a/Engine/Source/Editor/UnrealEd/Public/EditorViewportSelectability.h
+++ b/Engine/Source/Editor/UnrealEd/Public/EditorViewportSelectability.h
@@ -3,6 +3,7 @@
 #pragma once
 
 #include "Components/PrimitiveComponent.h"
+#include "Elements/Framework/TypedElementHandle.h"
 #include "GenericPlatform/ICursor.h"
 #include "Internationalization/Text.h"
 #include "Math/Color.h"
@@ -14,6 +15,7 @@ class FEditorViewportClient;
 class FLevelEditorViewportClient;
 class FUICommandInfo;
 class HHitProxy;
+class ITypedElementWorldInterface;
 struct FConvexVolume;
 struct FViewportClick;
 
@@ -84,6 +86,8 @@ protected:
 
 	static bool IsActorInLevelHiddenLayer(const AActor& InActor, FLevelEditorViewportClient* const InLevelEditorViewportClient);
 
+	static TTypedElement<ITypedElementWorldInterface> GetTypedWorldElementFromActor(const AActor& InActor);
+
     static bool DoesActorIntersectBox(const AActor& InActor, const FBox& InBox, FEditorViewportClient* const InEditorViewportClient, const bool bInUseStrictSelection);
     static bool DoesActorIntersectFrustum(const AActor& InActor, const FConvexVolume& InFrustum, FEditorViewportClient* const InEditorViewportClient, const bool bInUseStrictSelection);
 
