Editor: Make component visualizers able to customize whether they want to trigger their component to be selected after any VisProxyHandleClick.

diff --git a/Engine/Source/Editor/UnrealEd/Private/ComponentVisualizer.cpp b/Engine/Source/Editor/UnrealEd/Private/ComponentVisualizer.cpp
--- a/Engine/Source/Editor/UnrealEd/Private/ComponentVisualizer.cpp
+++ b/Engine/Source/Editor/UnrealEd/Private/ComponentVisualizer.cpp
@@ -16,6 +16,17 @@
 
 IMPLEMENT_HIT_PROXY(HComponentVisProxy, HHitProxy);
 
+namespace ComponentVisualizers
+{
+	static bool GAutoSelectComponentWithPointSelection = true;
+	static FAutoConsoleVariableRef CVarAutoSelectComponentWithPointSelection(
+		TEXT("Editor.ComponentVisualizer.AutoSelectComponent"),
+		GAutoSelectComponentWithPointSelection,
+		TEXT("Automatically adds the visualized component to the selection set if "
+			"available when visualization handles a click (e.g. select spline component when a point is selected on the spline).")
+	);
+
+}; // namespace ComponentVisualizers
 
 static FPropertyNameAndIndex GetActorPropertyNameAndIndexForComponent(const AActor* Actor, const UActorComponent* Component)
 {
@@ -200,6 +211,11 @@ bool FComponentPropertyPath::IsValid() const
 }
 
 
+bool FComponentVisualizer::ShouldAutoSelectElementOnHandleClick() const
+{
+	return ComponentVisualizers::GAutoSelectComponentWithPointSelection;
+}
+
 FPropertyNameAndIndex FComponentVisualizer::GetComponentPropertyName(const UActorComponent* Component)
 {
 	if (Component)
diff --git a/Engine/Source/Editor/UnrealEd/Private/ComponentVisualizerManager.cpp b/Engine/Source/Editor/UnrealEd/Private/ComponentVisualizerManager.cpp
--- a/Engine/Source/Editor/UnrealEd/Private/ComponentVisualizerManager.cpp
+++ b/Engine/Source/Editor/UnrealEd/Private/ComponentVisualizerManager.cpp
@@ -11,17 +11,6 @@
 #include "EditorModeManager.h"
 #include "Elements/Framework/TypedElementSelectionSet.h"
 
-namespace ComponentVisualizers
-{
-	static bool GAutoSelectComponentWithPointSelection = true;
-	static FAutoConsoleVariableRef CVarAutoSelectComponentWithPointSelection(
-		TEXT("Editor.ComponentVisualizer.AutoSelectComponent"),
-		GAutoSelectComponentWithPointSelection,
-		TEXT("Automatically adds the spline component to the selection set if avaialable when a point is selected on the spline")
-	);
-
-}; // namespace ComponentVisualizers
-
 
 FComponentVisualizerManager::FComponentVisualizerManager()
 	: EditedVisualizerViewportClient(nullptr)
@@ -83,8 +72,12 @@ bool FComponentVisualizerManager::HandleProxyForComponentVis(FEditorViewportClie
 				bool bIsActive = Visualizer->VisProxyHandleClick(InViewportClient, VisProxy, Click);
 				if (bIsActive)
 				{
-					// Set the selection set to the spline component, so clicking the component doesn't reselect and clear the current editing state.
-					if (ComponentVisualizers::GAutoSelectComponentWithPointSelection && ElementSelectionSet)
+					// For spline components in particular, it is nice to select the component itself when manipulating the spline. Otherwise,
+					//  if we have the actor selected but not the component, move a component point, and then select the component in the 
+					//  component outliner, we lose the point selection state, which feels unexpected.
+					bool bShouldSelectElement = Visualizer->ShouldAutoSelectElementOnHandleClick();
+
+					if (bShouldSelectElement && ElementSelectionSet)
 					{
 						if (!ElementSelectionSet->IsElementSelected(SelectedComponentHandle, FTypedElementIsSelectedOptions()))
 						{
@@ -94,7 +87,7 @@ bool FComponentVisualizerManager::HandleProxyForComponentVis(FEditorViewportClie
 
 					SetActiveComponentVis(InViewportClient, Visualizer);
 
-					if (ComponentVisualizers::GAutoSelectComponentWithPointSelection && ElementSelectionSet)
+					if (bShouldSelectElement && ElementSelectionSet)
 					{
 						ElementSelectionSet->NotifyPendingChanges();
 					}
diff --git a/Engine/Source/Editor/UnrealEd/Public/ComponentVisualizer.h b/Engine/Source/Editor/UnrealEd/Public/ComponentVisualizer.h
--- a/Engine/Source/Editor/UnrealEd/Public/ComponentVisualizer.h
+++ b/Engine/Source/Editor/UnrealEd/Public/ComponentVisualizer.h
@@ -172,6 +172,11 @@ public:
 	virtual void DrawVisualizationHUD(const UActorComponent* Component, const FViewport* Viewport, const FSceneView* View, FCanvas* Canvas) {}
 	/** */
 	virtual bool VisProxyHandleClick(FEditorViewportClient* InViewportClient, HComponentVisProxy* VisProxy, const FViewportClick& Click) { return false; }
+	/** 
+	 * Determines whether, if the visualizer returns true from VisProxyHandleClick, the hitproxy element should
+	 *  be automatically selected. By default, uses the value of Editor.ComponentVisualizer.AutoSelectComponent.
+	 */
+	UNREALED_API virtual bool ShouldAutoSelectElementOnHandleClick() const;
 	/** */
 	virtual void EndEditing() {}
 	/** */
