Fixed Localisation warnings to correctly add a filepath which can be used for identifying suspected CL's which caused the issue.

diff --git a/Engine/Source/Editor/UnrealEd/Private/Commandlets/GatherTextFromSourceCommandlet.cpp b/Engine/Source/Editor/UnrealEd/Private/Commandlets/GatherTextFromSourceCommandlet.cpp
--- a/Engine/Source/Editor/UnrealEd/Private/Commandlets/GatherTextFromSourceCommandlet.cpp
+++ b/Engine/Source/Editor/UnrealEd/Private/Commandlets/GatherTextFromSourceCommandlet.cpp
@@ -13,6 +13,7 @@
 #include "Internationalization/InternationalizationMetadata.h"
 #include "Internationalization/TextNamespaceUtil.h"
 #include "ProfilingDebugging/ScopedTimers.h"
+#include "Logging/StructuredLog.h"
 
 DEFINE_LOG_CATEGORY_STATIC(LogGatherTextFromSourceCommandlet, Log, All);
 
@@ -41,6 +42,8 @@ const FString UGatherTextFromSourceCommandlet::MacroString_NSLOCTEXT(TEXT("NSLOC
 const FString UGatherTextFromSourceCommandlet::MacroString_UI_COMMAND(TEXT("UI_COMMAND"));
 const FString UGatherTextFromSourceCommandlet::MacroString_UI_COMMAND_EXT(TEXT("UI_COMMAND_EXT"));
 
+constexpr int LogIssueIdentifier = 304;
+
 // Nested macro statistics, to track any supported features
 // See https://gcc.gnu.org/onlinedocs/gcc-3.0.1/cpp_3.html
 struct FParsedNestedMacroStats
@@ -1705,7 +1708,7 @@ void UGatherTextFromSourceCommandlet::FSourceFileParseContext::RemoveDefine(cons
 		// #undef LOC_DEFINE_REGION
 		if (!ExcludedRegion)
 		{
-			UE_LOG(LogGatherTextFromSourceCommandlet, Warning, TEXT("%s(%d): Found an '#undef LOC_DEFINE_REGION' without a corresponding '#define LOC_DEFINE_REGION'"), *Filename, LineNumber);
+			UE_LOGFMT(LogGatherTextFromSourceCommandlet, Warning, "{file}({line}): Found an '#undef LOC_DEFINE_REGION' without a corresponding '#define LOC_DEFINE_REGION'", ("id", LogIssueIdentifier), ("file", *Filename), ("line", LineNumber));
 		}
 		else
 		{
@@ -1720,7 +1723,7 @@ void UGatherTextFromSourceCommandlet::FSourceFileParseContext::RemoveDefine(cons
 			// #undef LOCTEXT_NAMESPACE
 			if (WithinNamespaceDefineLineNumber == INDEX_NONE)
 			{
-				UE_LOG(LogGatherTextFromSourceCommandlet, Warning, TEXT("%s(%d): Found an '#undef LOCTEXT_NAMESPACE' without a corresponding '#define LOCTEXT_NAMESPACE'"), *Filename, LineNumber);
+				UE_LOGFMT(LogGatherTextFromSourceCommandlet, Warning, "{file}({line}): Found an '#undef LOCTEXT_NAMESPACE' without a corresponding '#define LOCTEXT_NAMESPACE'", ("id", LogIssueIdentifier), ("file", *Filename), ("line", LineNumber));
 			}
 			else
 			{
diff --git a/Engine/Source/Programs/Shared/EpicGames.Core/JsonLogEvent.cs b/Engine/Source/Programs/Shared/EpicGames.Core/JsonLogEvent.cs
--- a/Engine/Source/Programs/Shared/EpicGames.Core/JsonLogEvent.cs
+++ b/Engine/Source/Programs/Shared/EpicGames.Core/JsonLogEvent.cs
@@ -133,6 +133,7 @@ namespace EpicGames.Core
 				while (reader.Read() && reader.TokenType == JsonTokenType.PropertyName)
 				{
 					ReadOnlySpan<byte> propertyName = reader.ValueSpan;
+
 					if (!reader.Read())
 					{
 						break;
@@ -153,6 +154,25 @@ namespace EpicGames.Core
 					{
 						reader.TryGetInt32(out lineCount);
 					}
+					else if (propertyName.SequenceEqual(LogEventPropertyName.Properties) && reader.TokenType == JsonTokenType.StartObject)
+					{
+						while (reader.Read() && reader.TokenType == JsonTokenType.PropertyName)
+						{
+							ReadOnlySpan<byte> innerPropertyName = reader.ValueSpan;
+
+							if (!reader.Read())
+							{
+								break;
+							}
+
+							// In the case of UE_LOGFMT the event id can be nested within the properties instead
+							if (innerPropertyName.SequenceEqual(LogEventPropertyName.Id) && reader.TokenType == JsonTokenType.Number)
+							{
+								eventId = reader.GetInt32();
+							}
+							reader.Skip();
+						}
+					}
 					reader.Skip();
 				}
 			}
