Add the "Disable realtime override" entry to the new viewport toolbar's Performance & Scalability menu

* Use the same code from the old viewport toolbar (and remove its old special code).
* Use the new entry from the level editor and the blueprint editor and all editors that use the default "Performance & Scalability" menu (static mesh, material, smart object, niagara, water waves).

diff --git a/Engine/Source/Editor/Kismet/Private/SSCSEditorViewport.cpp b/Engine/Source/Editor/Kismet/Private/SSCSEditorViewport.cpp
--- a/Engine/Source/Editor/Kismet/Private/SSCSEditorViewport.cpp
+++ b/Engine/Source/Editor/Kismet/Private/SSCSEditorViewport.cpp
@@ -511,6 +511,12 @@ TSharedPtr<SWidget> SSCSEditorViewport::BuildViewportToolbar()
 						{
 							FToolMenuSection& UnnamedSection = Submenu->FindOrAddSection(NAME_None);
 							UnnamedSection.AddEntry(UE::UnrealEd::CreateToggleRealtimeEntry());
+
+							if (const UUnrealEdViewportToolbarContext* Context =
+									Submenu->FindContext<UUnrealEdViewportToolbarContext>())
+							{
+								UnnamedSection.AddEntry(UE::UnrealEd::CreateRemoveRealtimeOverrideEntry(Context->Viewport));
+							}
 						}
 					)
 				);
diff --git a/Engine/Source/Editor/LevelEditor/Private/SLevelViewportToolBar.cpp b/Engine/Source/Editor/LevelEditor/Private/SLevelViewportToolBar.cpp
--- a/Engine/Source/Editor/LevelEditor/Private/SLevelViewportToolBar.cpp
+++ b/Engine/Source/Editor/LevelEditor/Private/SLevelViewportToolBar.cpp
@@ -497,21 +497,9 @@ void SLevelViewportToolBar::FillOptionsMenu(UToolMenu* Menu)
 
 			// Add an option to disable the temporary override if there is one
 			{
-				FUIAction Action;
-				Action.ExecuteAction.BindSP(this, &SLevelViewportToolBar::OnDisableRealtimeOverride);
-				Action.IsActionVisibleDelegate.BindSP(this, &SLevelViewportToolBar::IsRealtimeOverrideToggleVisible);
-
-				TAttribute<FText> Tooltip(this, &SLevelViewportToolBar::GetRealtimeOverrideTooltip);
-
-				Section.AddMenuEntry(
-					"DisableRealtimeOverride",
-					LOCTEXT("DisableRealtimeOverride", "Disable Realtime Override"),
-					Tooltip,
-					FSlateIcon(),
-					Action);
+				Section.AddEntry(UE::UnrealEd::CreateRemoveRealtimeOverrideEntry(Viewport));
 
 				Section.AddSeparator("DisableRealtimeOverrideSeparator");
-					
 			}
 
 			Section.AddMenuEntry(FEditorViewportCommands::Get().ToggleFPS);
@@ -1007,35 +995,6 @@ void SLevelViewportToolBar::FillViewMenu(UToolMenu* Menu)
 	UE::LevelEditor::PopulateViewModesMenu(Menu);
 }
 
-
-void SLevelViewportToolBar::OnDisableRealtimeOverride()
-{
-	if (TSharedPtr<SLevelViewport> ViewportPinned = Viewport.Pin())
-	{
-		ViewportPinned->GetLevelViewportClient().PopRealtimeOverride();
-	}
-}
-
-bool SLevelViewportToolBar::IsRealtimeOverrideToggleVisible() const
-{
-	if (TSharedPtr<SLevelViewport> ViewportPinned = Viewport.Pin())
-	{
-		return ViewportPinned->GetLevelViewportClient().IsRealtimeOverrideSet();
-	}
-
-	return false;
-}
-
-FText SLevelViewportToolBar::GetRealtimeOverrideTooltip() const
-{
-	if (TSharedPtr<SLevelViewport> ViewportPinned = Viewport.Pin())
-	{
-		return FText::Format(LOCTEXT("DisableRealtimeOverrideToolTip", "Realtime is currently overridden by \"{0}\".  Click to remove that override"), ViewportPinned->GetLevelViewportClient().GetRealtimeOverrideMessage());
-	}
-
-	return FText::GetEmpty();
-}
-
 float SLevelViewportToolBar::GetTransformToolbarWidth() const
 {
 	if (TransformToolbar)
diff --git a/Engine/Source/Editor/LevelEditor/Private/SLevelViewportToolBar.h b/Engine/Source/Editor/LevelEditor/Private/SLevelViewportToolBar.h
--- a/Engine/Source/Editor/LevelEditor/Private/SLevelViewportToolBar.h
+++ b/Engine/Source/Editor/LevelEditor/Private/SLevelViewportToolBar.h
@@ -242,11 +242,6 @@ private:
 	/** Gets the world we are editing */
 	TWeakObjectPtr<UWorld> GetWorld() const;
 
-	/** Called when the user disables realtime override from the toolbar */
-	void OnDisableRealtimeOverride();
-	bool IsRealtimeOverrideToggleVisible() const;
-	FText GetRealtimeOverrideTooltip() const;
-
 	float GetTransformToolbarWidth() const;
 
 private:
diff --git a/Engine/Source/Editor/LevelEditor/Private/ViewportToolbar/LevelEditorSubmenus.cpp b/Engine/Source/Editor/LevelEditor/Private/ViewportToolbar/LevelEditorSubmenus.cpp
--- a/Engine/Source/Editor/LevelEditor/Private/ViewportToolbar/LevelEditorSubmenus.cpp
+++ b/Engine/Source/Editor/LevelEditor/Private/ViewportToolbar/LevelEditorSubmenus.cpp
@@ -1768,6 +1768,12 @@ FToolMenuEntry CreatePerformanceAndScalabilitySubmenu()
 					FToolMenuSection& UnnamedSection = Submenu->FindOrAddSection(NAME_None);
 
 					UnnamedSection.AddEntry(UE::UnrealEd::CreateToggleRealtimeEntry());
+
+					if (const UUnrealEdViewportToolbarContext* Context =
+							Submenu->FindContext<UUnrealEdViewportToolbarContext>())
+					{
+						UnnamedSection.AddEntry(UE::UnrealEd::CreateRemoveRealtimeOverrideEntry(Context->Viewport));
+					}
 				}
 
 				{
diff --git a/Engine/Source/Editor/UnrealEd/Private/ViewportToolbar/UnrealEdViewportToolbar.cpp b/Engine/Source/Editor/UnrealEd/Private/ViewportToolbar/UnrealEdViewportToolbar.cpp
--- a/Engine/Source/Editor/UnrealEd/Private/ViewportToolbar/UnrealEdViewportToolbar.cpp
+++ b/Engine/Source/Editor/UnrealEd/Private/ViewportToolbar/UnrealEdViewportToolbar.cpp
@@ -3180,6 +3180,11 @@ FToolMenuEntry CreatePerformanceAndScalabilitySubmenu()
 
 				UnnamedSection.AddEntry(CreateToggleRealtimeEntry());
 
+				if (const UUnrealEdViewportToolbarContext* Context = Submenu->FindContext<UUnrealEdViewportToolbarContext>())
+				{
+					UnnamedSection.AddEntry(CreateRemoveRealtimeOverrideEntry(Context->Viewport));
+				}
+
 				UnnamedSection.AddSubMenu(
 					"ScreenPercentage",
 					LOCTEXT("ScreenPercentageSubMenu", "Screen Percentage"),
@@ -3430,6 +3435,52 @@ FToolMenuEntry CreateToggleRealtimeEntry()
 	);
 }
 
+FToolMenuEntry CreateRemoveRealtimeOverrideEntry(TWeakPtr<SEditorViewport> WeakViewport)
+{
+	FUIAction Action;
+	Action.ExecuteAction = FExecuteAction::CreateLambda(
+		[WeakViewport]()
+		{
+			if (TSharedPtr<SEditorViewport> ViewportPinned = WeakViewport.Pin())
+			{
+				ViewportPinned->GetViewportClient()->PopRealtimeOverride();
+			}
+		}
+	);
+	Action.IsActionVisibleDelegate = FIsActionButtonVisible::CreateLambda(
+		[WeakViewport]() -> bool
+		{
+			if (TSharedPtr<SEditorViewport> ViewportPinned = WeakViewport.Pin())
+			{
+				return ViewportPinned->GetViewportClient()->IsRealtimeOverrideSet();
+			}
+
+			return false;
+		}
+	);
+
+	TAttribute<FText> Tooltip = TAttribute<FText>::CreateLambda(
+		[WeakViewport]()
+		{
+			if (TSharedPtr<SEditorViewport> ViewportPinned = WeakViewport.Pin())
+			{
+				return FText::Format(
+					LOCTEXT(
+						"DisableRealtimeOverrideToolTip", "Realtime is currently overridden by \"{0}\". Click to remove that override."
+					),
+					ViewportPinned->GetViewportClient()->GetRealtimeOverrideMessage()
+				);
+			}
+
+			return FText::GetEmpty();
+		}
+	);
+
+	return FToolMenuEntry::InitMenuEntry(
+		"DisableRealtimeOverride", LOCTEXT("DisableRealtimeOverride", "Disable Realtime Override"), Tooltip, FSlateIcon(), Action
+	);
+}
+
 FOnViewportClientCamSpeedChanged& OnViewportClientCamSpeedChanged()
 {
 	static FOnViewportClientCamSpeedChanged OnViewportClientCamSpeedChangedDelegate;
diff --git a/Engine/Source/Editor/UnrealEd/Public/ViewportToolbar/UnrealEdViewportToolbar.h b/Engine/Source/Editor/UnrealEd/Public/ViewportToolbar/UnrealEdViewportToolbar.h
--- a/Engine/Source/Editor/UnrealEd/Public/ViewportToolbar/UnrealEdViewportToolbar.h
+++ b/Engine/Source/Editor/UnrealEd/Public/ViewportToolbar/UnrealEdViewportToolbar.h
@@ -188,6 +188,7 @@ UNREALED_API FToolMenuEntry CreateDefaultShowSubmenu();
 UNREALED_API void AddDefaultShowFlags(UToolMenu* InMenu);
 
 UNREALED_API FToolMenuEntry CreateToggleRealtimeEntry();
+UNREALED_API FToolMenuEntry CreateRemoveRealtimeOverrideEntry(TWeakPtr<SEditorViewport> WeakViewport);
 
 UNREALED_API FOnViewportClientCamSpeedChanged& OnViewportClientCamSpeedChanged();
 UNREALED_API FOnViewportClientCamSpeedScalarChanged& OnViewportClientCamSpeedScalarChanged();
