Cooker - Move GarbageCollection and SavePackage functions out of CookOnTheFlyServer into separate files.
#rnx

diff --git a/Engine/Source/Editor/UnrealEd/Classes/CookOnTheSide/CookOnTheFlyServer.h b/Engine/Source/Editor/UnrealEd/Classes/CookOnTheSide/CookOnTheFlyServer.h
--- a/Engine/Source/Editor/UnrealEd/Classes/CookOnTheSide/CookOnTheFlyServer.h
+++ b/Engine/Source/Editor/UnrealEd/Classes/CookOnTheSide/CookOnTheFlyServer.h
@@ -216,6 +216,7 @@ namespace UE::Cook
 	struct FPackageTracker;
 	struct FPendingCookedPlatformData;
 	struct FPlatformManager;
+	struct FScopeFindCookReferences;
 	struct FTickStackData;
 }
 namespace UE::Cook::Private
@@ -1453,16 +1454,7 @@ private:
 
 		UCookOnTheFlyServer& COTFS;
 	};
-	/** Scoped type to call Set/Clear SoftGCPackageToObjectListBuffer. */
-	struct FScopeFindCookReferences
-	{
-		FScopeFindCookReferences(UCookOnTheFlyServer& InCOTFS);
-		~FScopeFindCookReferences();
 
-		UCookOnTheFlyServer& COTFS;
-		TGuardValue<bool> SoftGCGuard;
-		bool bNeedsConstructBuffer;
-	};
 	/** Callback for FGenericCrashContext; provides the current ActivePackage as context. */
 	void DumpCrashContext(FCrashContextExtendedWriter& Writer);
 	/** Callback for analytics when a new UPackage is loaded. */
@@ -1661,6 +1653,7 @@ private:
 	friend UE::Cook::FPlatformManager;
 	friend UE::Cook::FRequestCluster;
 	friend UE::Cook::FSaveCookedPackageContext;
+	friend UE::Cook::FScopeFindCookReferences;
 	friend UE::Cook::FSoftGCHistory;
 	friend UE::Cook::FWorkerRequestsLocal;
 	friend UE::Cook::FWorkerRequestsRemote;
diff --git a/Engine/Source/Editor/UnrealEd/Private/CookOnTheFlyServer.cpp b/Engine/Source/Editor/UnrealEd/Private/CookOnTheFlyServer.cpp
--- a/Engine/Source/Editor/UnrealEd/Private/CookOnTheFlyServer.cpp
+++ b/Engine/Source/Editor/UnrealEd/Private/CookOnTheFlyServer.cpp
@@ -34,6 +34,7 @@
 #include "Cooker/CookRequestCluster.h"
 #include "Cooker/CookRequests.h"
 #include "Cooker/CookSandbox.h"
+#include "Cooker/CookSavePackage.h"
 #include "Cooker/CookTypes.h"
 #include "Cooker/CookWorkerClient.h"
 #include "Cooker/DiffPackageWriter.h"
@@ -44,7 +45,6 @@
 #include "Cooker/NetworkFileCookOnTheFlyRequestManager.h"
 #include "Cooker/OnDemandShaderCompilation.h"
 #include "Cooker/PackageTracker.h"
-#include "Cooker/SoftGC.h"
 #include "Cooker/WorkerRequestsLocal.h"
 #include "Cooker/WorkerRequestsRemote.h"
 #include "CookerSettings.h"
@@ -161,7 +161,6 @@
 #include "UObject/Package.h"
 #include "UObject/ReferenceChainSearch.h"
 #include "UObject/SavePackage.h"
-#include "UObject/UObjectArray.h"
 #include "UObject/UObjectIterator.h"
 #include "UserGeneratedContentLocalization.h"
 #include "ZenCookArtifactReader.h"
@@ -175,8 +174,6 @@ DEFINE_LOG_CATEGORY_STATIC(LogCookList, Log, All);
 
 FName LogCookName(TEXT("LogCook"));
 LLM_DEFINE_TAG(Cooker);
-LLM_DEFINE_TAG(Cooker_SavePackage);
-
 
 int32 GCookProgressDisplay = (int32)ECookProgressDisplayMode::RemainingPackages;
 static FAutoConsoleVariableRef CVarCookDisplayMode(
@@ -2191,38 +2188,6 @@ TRefCountPtr<UCookOnTheFlyServer::FPollable> UCookOnTheFlyServer::CreatePollable
 		[this](UE::Cook::FTickStackData& StackData) { PollGarbageCollection(StackData); }));
 }
 
-void UCookOnTheFlyServer::PollGarbageCollection(UE::Cook::FTickStackData& StackData)
-{
-	NumObjectsHistory.AddInstance(GUObjectArray.GetObjectArrayNumMinusAvailable());
-	VirtualMemoryHistory.AddInstance(FPlatformMemory::GetStats().UsedVirtual);
-
-	if (IsCookFlagSet(ECookInitializationFlags::TestCook))
-	{
-		StackData.ResultFlags |= COSR_RequiresGC | COSR_YieldTick;
-		return;
-	}
-	if (PackagesPerGC > 0 && CookedPackageCountSinceLastGC > PackagesPerGC)
-	{
-		// if we are waiting on things to cache then ignore the PackagesPerGC
-		if (!bSaveBusy)
-		{
-			StackData.ResultFlags |= COSR_RequiresGC | COSR_RequiresGC_PackageCount | COSR_YieldTick;
-			return;
-		}
-	}
-	if (IsCookOnTheFlyMode())
-	{
-		double CurrentTime = FPlatformTime::Seconds();
-		if (IdleStatus == EIdleStatus::Done &&
-			CurrentTime - IdleStatusStartTime > GetIdleTimeToGC() &&
-			IdleStatusStartTime > GetLastGCTime())
-		{
-			StackData.ResultFlags |= COSR_RequiresGC | COSR_RequiresGC_IdleTimer | COSR_YieldTick;
-			return;
-		}
-	}
-}
-
 namespace UE::Cook
 {
 
@@ -4573,72 +4538,6 @@ void UCookOnTheFlyServer::ProcessUnsolicitedPackages(TArray<FName>* OutDiscovere
 	}
 }
 
-namespace UE::Cook
-{
-
-/** Local parameters and helper functions used by SaveCookedPackage */
-class FSaveCookedPackageContext
-{
-private: // Used only by UCookOnTheFlyServer, which has private access
-
-	FSaveCookedPackageContext(UCookOnTheFlyServer& InCOTFS, UE::Cook::FPackageData& InPackageData,
-		TArrayView<const ITargetPlatform*> InPlatformsForPackage, UE::Cook::FTickStackData& StackData);
-
-	// Hooks used by friends
-	void SetupPackage();
-	void SetupPlatform(const ITargetPlatform* InTargetPlatform, int32 InPlatformIndex);
-	void FinishPlatform();
-	void FinishPackage();
-
-	// private helper functions
-	void CalculatePlatformAgnosticRuntimeDependencies();
-	void CalculatePlatformRuntimeDependencies();
-	TArray<FName> GetPlatformRuntimeDependencies() const;
-	TArray<IPackageWriter::FCommitAttachmentInfo> GetCommitAttachments(
-		UE::TargetDomain::FGeneratedPackageResultStruct* GeneratedResult,
-		TArray<FAssetDependency>* OverrideDependencies);
-	IPackageWriter::EWriteOptions GetCommitWriteOptions() const;
-	static void AddDependency(TMap<FPackageData*, EInstigator>& InDependencies, FPackageData* PackageData, bool bHard);
-
-	// General Package Data
-	UCookOnTheFlyServer& COTFS;
-	FPackageData& PackageData;
-	TArrayView<const ITargetPlatform*> PlatformsForPackage;
-	TMap<FPackageData*, EInstigator> PlatformAgnosticDependencies;
-	TArray<TMap<FPackageData*, EInstigator>, TInlineAllocator<1>> PlatformDependencies;
-	FTickStackData& StackData;
-	UPackage* Package;
-	const FString PackageName;
-	FString Filename;
-	uint32 SaveFlags = 0;
-	bool bHasTimeOut = false;
-	bool bHasRetryErrorCode = false;
-	bool bPlatformAgnosticDependenciesCalculated = false;
-	bool bAnySaveSucceeded = false;
-
-	// General Package Data that is delay-loaded the first time we save a platform
-	UWorld* World = nullptr;
-	EObjectFlags FlagsToCook = RF_Public;
-	bool bHasDelayLoaded = false;
-	bool bContainsMap = false;
-
-	// Per-platform data, only valid in PerPlatform callbacks
-	const ITargetPlatform* TargetPlatform = nullptr;
-	FCookSavePackageContext* CookContext = nullptr;
-	// This holds cook data generated by Serialize() that isn't saved in the package
-	TOptional<FArchiveCookContext> ArchiveCookContext;
-	FSavePackageContext* SavePackageContext = nullptr;
-	ICookedPackageWriter* PackageWriter = nullptr;
-	FString PlatFilename;
-	FSavePackageResultStruct SavePackageResult;
-	int32 PlatformIndex = -1;
-	bool bPlatformSetupSuccessful = false;
-	bool bEndianSwap = false;
-
-	friend class ::UCookOnTheFlyServer;
-};
-}
-
 void UCookOnTheFlyServer::PumpSaves(UE::Cook::FTickStackData& StackData, uint32 DesiredQueueLength, int32& OutNumPushed, bool& bOutBusy)
 {
 	using namespace UE::Cook;
@@ -5042,1956 +4941,370 @@ void UCookOnTheFlyServer::TickPrecacheObjectsForPlatforms(const float TimeSlice,
 			return;
 		}
 	}
-
 }
 
-bool UCookOnTheFlyServer::PumpHasExceededMaxMemory(uint32& OutResultFlags)
+void UCookOnTheFlyServer::OnObjectModified( UObject *ObjectMoving )
 {
-	if (GUObjectArray.GetObjectArrayEstimatedAvailable() < MinFreeUObjectIndicesBeforeGC)
+	if (IsGarbageCollecting())
 	{
-		UE_LOG(LogCook, Display, TEXT("Running out of available UObject indices (%d remaining)"), GUObjectArray.GetObjectArrayEstimatedAvailable());
-		static bool bPerformedObjListWhenNearMaxObjects = false;
-		if (GEngine && !bPerformedObjListWhenNearMaxObjects)
-		{
-			UE_LOG(LogCook, Display, TEXT("Performing 'obj list' to show counts of types of objects due to low availability of UObject indices."));
-			GEngine->Exec(nullptr, TEXT("OBJ LIST -COUNTSORT -SKIPMEMORYSIZE"));
-			bPerformedObjListWhenNearMaxObjects = true;
-		}
-		OutResultFlags |= COSR_RequiresGC | COSR_RequiresGC_OOM | COSR_YieldTick;
-		return true;
+		return;
 	}
+	OnObjectUpdated( ObjectMoving );
+}
 
-	TStringBuilder<256> TriggerMessages;
-	const FPlatformMemoryStats MemStats = FPlatformMemory::GetStats();
-	
-	bool bMinFreeTriggered = false;
-	if (MemoryMinFreeVirtual > 0 || MemoryMinFreePhysical > 0)
+void UCookOnTheFlyServer::OnObjectPropertyChanged(UObject* ObjectBeingModified, FPropertyChangedEvent& PropertyChangedEvent)
+{
+	if (IsGarbageCollecting())
 	{
-		// trigger GC if we have less than MemoryMinFreeVirtual OR MemoryMinFreePhysical
-		// the check done in AssetCompilingManager is against the min of the two :
-		//uint64 AvailableMemory = FMath::Min(MemStats.AvailablePhysical, MemStats.AvailableVirtual);
-		// so for consistency the same check should be done here
-		// you can get that by setting the MemoryMinFreeVirtual and MemoryMinFreePhysical config to be the same
-
-		// AvailableVirtual is actually ullAvailPageFile (commit charge available)
-		if (MemoryMinFreeVirtual > 0 && MemStats.AvailableVirtual < MemoryMinFreeVirtual)
-		{
-			TriggerMessages.Appendf(TEXT("\n  CookSettings.MemoryMinFreeVirtual: Available virtual memory %dMiB is less than %dMiB."),
-				static_cast<uint32>(MemStats.AvailableVirtual / 1024 / 1024), static_cast<uint32>(MemoryMinFreeVirtual / 1024 / 1024));
-			bMinFreeTriggered = true;
-		}
-		if (MemoryMinFreePhysical > 0 && MemStats.AvailablePhysical < MemoryMinFreePhysical)
-		{
-			TriggerMessages.Appendf(TEXT("\n  CookSettings.MemoryMinFreePhysical: Available physical memory %dMiB is less than %dMiB."),
-				static_cast<uint32>(MemStats.AvailablePhysical / 1024 / 1024), static_cast<uint32>(MemoryMinFreePhysical / 1024 / 1024));
-			bMinFreeTriggered = true;
-		}
+		return;
 	}
-
-	// if MemoryMaxUsed is set, we won't GC until at least that much mem is used
-	// this can be useful if you demand that amount of memory as your min spec
-	bool bMaxUsedTriggered = false;
-	if (MemoryMaxUsedVirtual > 0 || MemoryMaxUsedPhysical > 0)
+	if ( PropertyChangedEvent.Property == nullptr && 
+		PropertyChangedEvent.MemberProperty == nullptr )
 	{
-		// check validity of trigger :
-		// if the MaxUsed config exceeds the system memory, it can never be triggered and will prevent any GC :
-		uint64 MaxMaxUsed = FMath::Max(MemoryMaxUsedVirtual,MemoryMaxUsedPhysical);
-		if (MaxMaxUsed >= MemStats.TotalPhysical)
-		{
-			UE_CALL_ONCE([&]() {
-				UE_LOG(LogCook, Warning, TEXT("Warning MemoryMaxUsed condition is larger than total memory (%dMiB >= %dMiB).  System does not have enough memory to cook this project."),
-				static_cast<uint32>(MaxMaxUsed / 1024 / 1024), static_cast<uint32>(MemStats.TotalPhysical / 1024 / 1024));
-				});
-		}
-
-		if (MemoryMaxUsedVirtual > 0 && MemStats.UsedVirtual >= MemoryMaxUsedVirtual)
-		{
-			TriggerMessages.Appendf(TEXT("\n  CookSettings.MemoryMaxUsedVirtual: Used virtual memory %dMiB is greater than %dMiB."),
-				static_cast<uint32>(MemStats.UsedVirtual / 1024 / 1024), static_cast<uint32>(MemoryMaxUsedVirtual / 1024 / 1024));
-			bMaxUsedTriggered = true;
-		}
-		if (MemoryMaxUsedPhysical > 0 && MemStats.UsedPhysical >= MemoryMaxUsedPhysical)
-		{
-			TriggerMessages.Appendf(TEXT("\n  CookSettings.MemoryMaxUsedPhysical: Used physical memory %dMiB is greater than %dMiB."),
-				static_cast<uint32>(MemStats.UsedPhysical / 1024 / 1024), static_cast<uint32>(MemoryMaxUsedPhysical / 1024 / 1024));
-			bMaxUsedTriggered = true;
-		}
+		// probably nothing changed... 
+		return;
 	}
 
-	bool bPressureTriggered = false;
-	if (MemoryTriggerGCAtPressureLevel != FPlatformMemoryStats::EMemoryPressureStatus::Unknown)
+	OnObjectUpdated( ObjectBeingModified );
+}
+
+void UCookOnTheFlyServer::OnObjectSaved( UObject* ObjectSaved, FObjectPreSaveContext SaveContext)
+{
+	if (SaveContext.IsProceduralSave() )
 	{
-		FPlatformMemoryStats::EMemoryPressureStatus PressureStatus = MemStats.GetMemoryPressureStatus();
-		if (PressureStatus == FPlatformMemoryStats::EMemoryPressureStatus::Unknown)
-		{
-			UE_CALL_ONCE([&]() {
-				UE_LOG(LogCook, Warning,
-				TEXT("MemoryPressureStatus is not available from the operating system. We may run out of memory due to lack of knowledge of when to collect garbage."));
-				});
-		}
-		else
-		{
-			static_assert(FPlatformMemoryStats::EMemoryPressureStatus::Critical > FPlatformMemoryStats::EMemoryPressureStatus::Nominal,
-				"We expect higher pressure to be higher integer values");
-			int RequiredValue = static_cast<int>(MemoryTriggerGCAtPressureLevel);
-			int CurrentValue = static_cast<int>(PressureStatus);
-			if (CurrentValue >= RequiredValue)
-			{
-				bPressureTriggered = true;
-				TriggerMessages.Appendf(TEXT("\n  Operating system has signalled that memory pressure is high."));
-			}
-		}
+		// This is a procedural save (e.g. our own saving of the cooked package) rather than a user save, ignore
+		return;
 	}
 
-	bool bTriggerGC = false;
-	if (bMinFreeTriggered || bMaxUsedTriggered)
+	UPackage* Package = ObjectSaved->GetOutermost();
+	if (Package == nullptr || Package == GetTransientPackage())
 	{
-		const bool bOnlyTriggerIfBothMinFreeAndMaxUsedTrigger = true;
-		if (!bOnlyTriggerIfBothMinFreeAndMaxUsedTrigger ||
-			((bMinFreeTriggered || (MemoryMinFreeVirtual <= 0 && MemoryMinFreePhysical <= 0)) &&
-				(bMaxUsedTriggered || (MemoryMaxUsedVirtual <= 0 && MemoryMaxUsedPhysical <= 0))))
-		{
-			bTriggerGC = true;
-		}
+		return;
 	}
-	if (bPressureTriggered)
+
+	MarkPackageDirtyForCooker(Package);
+
+	// Register the package filename as modified. We don't use the cache because the file may not exist on disk yet at this point
+	const FString PackageFilename = FPackageName::LongPackageNameToFilename(Package->GetName(), Package->ContainsMap() ? FPackageName::GetMapPackageExtension() : FPackageName::GetAssetPackageExtension());
+	ModifiedAssetFilenames.Add(FName(*PackageFilename));
+}
+
+void UCookOnTheFlyServer::OnObjectUpdated( UObject *Object )
+{
+	// get the outer of the object
+	UPackage *Package = Object->GetOutermost();
+
+	MarkPackageDirtyForCooker( Package );
+}
+
+void UCookOnTheFlyServer::MarkPackageDirtyForCooker(UPackage* Package, bool bAllowInSession)
+{
+	if (Package->RootPackageHasAnyFlags(PKG_PlayInEditor))
 	{
-		bTriggerGC = true;
+		return;
 	}
 
-	// If a normal GC was not triggered, check the SoftGC trigger conditions
-	double CurrentTime = FPlatformTime::Seconds();
-	bool bIsSoftGC = false;
-	if (bUseSoftGC && IsDirectorCookByTheBook() && !IsCookingInEditor())
+	if (Package->HasAnyPackageFlags(PKG_PlayInEditor | PKG_ContainsScript | PKG_InMemoryOnly) == true && !GetClass()->HasAnyClassFlags(CLASS_DefaultConfig | CLASS_Config))
 	{
-		if (!bTriggerGC && SoftGCStartNumerator > 0)
-		{
-			if (SoftGCNextAvailablePhysicalTarget == -1) // Uninitialized
-			{
-				int32 StartNumerator = FMath::Max(SoftGCStartNumerator, 1);
-				int32 Denominator = FMath::Max(SoftGCDenominator, 1);
-				// e.g. Start the target at 5/10, and decrease it by 1/10 each time the target is reached
-				SoftGCNextAvailablePhysicalTarget = (static_cast<int64>(MemStats.TotalPhysical) * StartNumerator)
-					/ Denominator;
-			}
-
-			if (SoftGCNextAvailablePhysicalTarget < -1)
-			{
-				// No further targets, no further SoftGC
-			}
-			else if (static_cast<int64>(MemStats.AvailablePhysical) <= SoftGCNextAvailablePhysicalTarget)
-			{
-				constexpr float SoftGCInstigateCooldown = 5 * 60.f;
-				CurrentTime = FPlatformTime::Seconds();
-				if (LastSoftGCTime + SoftGCInstigateCooldown <= CurrentTime)
-				{
-					TriggerMessages.Appendf(TEXT("\n  CookSettings.SoftGCMemoryTrigger: Available physical memory %dMiB is less than the current target for SoftGC %dMiB."),
-						static_cast<uint32>(MemStats.AvailablePhysical / 1024 / 1024), static_cast<uint32>(SoftGCNextAvailablePhysicalTarget / 1024 / 1024));
-					bTriggerGC = true;
-					bIsSoftGC = true;
-				}
-			}
-		}
+		return;
+	}
 
-		if (!bTriggerGC && SoftGCTimeFractionBudget > 0.)
-		{
-			FString TriggerDiagnostics;
-			if (SoftGCHistory->IsTriggeringWithinBudget(*this, CurrentTime, &TriggerDiagnostics))
-			{
-				TriggerMessages.Appendf(TEXT("\n  CookSettings.SoftGCTimeTrigger: Periodic triggering of SoftGC: %s."),
-					*TriggerDiagnostics);
-				bTriggerGC = true;
-				bIsSoftGC = true;
-			}
-		}
+	if (Package == GetTransientPackage())
+	{
+		return;
 	}
 
-	if (!bTriggerGC)
+	if (Package->GetOuter() != nullptr)
 	{
-		return false;
+		return;
 	}
 
-	// Don't allow a second GC (soft or normal) within the GC cooldown period after a full GC
-	constexpr float GCCooldown = 60.f;
-	if (LastFullGCTime + GCCooldown > CurrentTime)
+	FName PackageName = Package->GetFName();
+	if (FPackageName::IsMemoryPackage(PackageName.ToString()))
 	{
-		if (!bIsSoftGC && !bWarnedExceededMaxMemoryWithinGCCooldown)
-		{
-			bWarnedExceededMaxMemoryWithinGCCooldown = true;
-			// If we are in a cooldown period, return false.
-			UE_LOG(LogCook, Display, TEXT("Garbage collection triggers ignored: Out of memory condition has been detected, but is only %.0fs after the last GC. ")
-				TEXT("It will be prevented until %.0f seconds have passed and we may run out of memory.\n")
-				TEXT("Garbage collection triggered by conditions: %s"),
-				static_cast<float>(CurrentTime - LastFullGCTime), GCCooldown, TriggerMessages.ToString());
-		}
-		return false;
+		return;
 	}
 
-	const TCHAR* TypeMessage = bIsSoftGC ? TEXT("Soft") : (IsCookFlagSet(ECookInitializationFlags::EnablePartialGC) ? TEXT("Partial") : TEXT("Full"));
+	if (bIsSavingPackage)
+	{
+		return;
+	}
 
-	UE_LOG(LogCook, Display, TEXT("Garbage collection triggered (%s). Triggered by conditions:%s"),
-		TypeMessage, TriggerMessages.ToString());
-	OutResultFlags |= COSR_RequiresGC | COSR_YieldTick;
-	OutResultFlags |= COSR_RequiresGC_OOM;
-	if (bIsSoftGC)
+	if (IsInSession() && !bAllowInSession)
 	{
-		OutResultFlags |= COSR_RequiresGC_Soft_OOM;
+		WorkerRequests->AddEditorActionCallback([this, PackageName]() { MarkPackageDirtyForCookerFromSchedulerThread(PackageName); });
+	}
+	else
+	{
+		MarkPackageDirtyForCookerFromSchedulerThread(PackageName);
 	}
-	return true;
 }
 
-void UCookOnTheFlyServer::SetGarbageCollectType(uint32 ResultFlagsFromTick)
+FName GInstigatorMarkPackageDirty(TEXT("MarkPackageDirtyForCooker"));
+void UCookOnTheFlyServer::MarkPackageDirtyForCookerFromSchedulerThread(const FName& PackageName)
 {
-	bGarbageCollectTypeSoft = (ResultFlagsFromTick & COSR_RequiresGC_Soft_OOM);
-}
+	TRACE_CPUPROFILER_EVENT_SCOPE(MarkPackageDirtyForCooker);
 
-void UCookOnTheFlyServer::ClearGarbageCollectType()
-{
-	bGarbageCollectTypeSoft = false;
-}
+	// could have just cooked a file which we might need to write
+	UPackage::WaitForAsyncFileWrites();
 
-/**
- * For every package in memory, add a list of all of its public UObjects into the map
- * used in garbage collection: UPackage::SoftGCPackageToObjectList. This will cause all of
- * its public objects to be referenced if the UPackage is referenced.
- */
-static void ConstructSoftGCPackageToObjectList(TArray<UObject*>& PackageToObjectListBuffer)
-{
-	struct FPackageObjectPair
+	// Update the package's FileName if it has changed
+	UE::Cook::FPackageData* PackageData = PackageDatas->UpdateFileName(PackageName);
+
+	// force the package to be recooked
+	UE_LOG(LogCook, Verbose, TEXT("Modification detected to package %s"), *PackageName.ToString());
+	if ( PackageData && IsCookingInEditor() )
 	{
-		UPackage* Package;
-		UObject* Object;
-		bool operator<(const FPackageObjectPair& Other) const
+		check(IsInGameThread()); // We're editing scheduler data, which is only allowable from the scheduler thread
+		bool bHadCookedPlatforms = PackageData->HasAnyCookedPlatform();
+		PackageData->ClearCookResults();
+		if (PackageData->IsInProgress())
 		{
-			if (Package != Other.Package)
-				return Package < Other.Package;
-			return Object < Other.Object;
+			PackageData->SendToState(UE::Cook::EPackageState::Request,
+				UE::Cook::ESendFlags::QueueAddAndRemove, UE::Cook::EStateChangeReason::ForceRecook);
 		}
-		bool operator==(const FPackageObjectPair& Other) const
+		else if (IsCookByTheBookMode() && IsInSession() && bHadCookedPlatforms)
 		{
-			return Object == Other.Object;
+			QueueDiscoveredPackage(*PackageData,
+				UE::Cook::FInstigator(UE::Cook::EInstigator::Unspecified, GInstigatorMarkPackageDirty),
+				UE::Cook::EDiscoveredPlatformSet::CopyFromInstigator);
 		}
-	};
-
-	PackageToObjectListBuffer.Empty();
-	UPackage::SoftGCPackageToObjectList.Empty();
-
-	// Iterate over all UObjects in memory (in parallel) and for each valid public object, get its package and add a FPackageObjectPair for it
-	int32 MaxNumberOfObjects = GUObjectArray.GetObjectArrayNum();
-	int32 NumThreads = FMath::Clamp(FTaskGraphInterface::Get().GetNumWorkerThreads(), 1, MaxNumberOfObjects);
-	int32 NumberOfObjectsPerThread = (MaxNumberOfObjects + NumThreads - 1) / NumThreads; // ceiling
-	check(NumberOfObjectsPerThread * (NumThreads - 1) <= MaxNumberOfObjects);
 
-	TArray<TArray<FPackageObjectPair>> ThreadContexts;
-	ThreadContexts.SetNum(NumThreads);
-	std::atomic<int32> PackagesNum{ 0 };
-
-	ParallelFor(TEXT("ConstructSoftGCPackageToObjectList"), NumThreads, 1,
-		[&ThreadContexts, &PackagesNum, NumberOfObjectsPerThread, NumThreads, MaxNumberOfObjects](int32 ThreadIndex)
+		if ( IsCookOnTheFlyMode() && FileModifiedDelegate.IsBound())
 		{
-			TArray<FPackageObjectPair>& ThreadContext = ThreadContexts[ThreadIndex];
-			int32 FirstObjectIndex = ThreadIndex * NumberOfObjectsPerThread;
-			int32 NumObjects = (ThreadIndex < (NumThreads - 1)) ? NumberOfObjectsPerThread : (MaxNumberOfObjects - (NumThreads - 1) * NumberOfObjectsPerThread);
-			check(FirstObjectIndex + NumObjects <= MaxNumberOfObjects);
-
-			for (int32 ObjectIndex = 0; ObjectIndex < NumObjects && (FirstObjectIndex + ObjectIndex) < MaxNumberOfObjects; ++ObjectIndex)
+			FString PackageFileNameString = PackageData->GetFileName().ToString();
+			FileModifiedDelegate.Broadcast(PackageFileNameString);
+			if (PackageFileNameString.EndsWith(".uasset") || PackageFileNameString.EndsWith(".umap"))
 			{
-				FUObjectItem& ObjectItem = GUObjectArray.GetObjectItemArrayUnsafe()[FirstObjectIndex + ObjectIndex];
-				if (!ObjectItem.Object)
-				{
-					continue;
-				}
-				if (ObjectItem.IsGarbage())
-				{
-					continue;
-				}
-				UObject* Object = static_cast<UObject*>(ObjectItem.Object);
-				if (!Object->HasAnyFlags(RF_Public))
-				{
-					continue;
-				}
-				UPackage* Package = Object->GetPackage();
-				if (!Package)
-				{
-					continue;
-				}
-				if (Package->HasAnyFlags(RF_Transient) || Package->HasAnyPackageFlags(PKG_CompiledIn))
-				{
-					// Skip any Transient packages (e.g. /Engine/Transient) and script packages
-					// We only need to keep public objects alive in packages that could be saved.
-					continue;
-				}
-				if (Object == Package)
-				{
-					PackagesNum.fetch_add(1, std::memory_order_relaxed);
-				}
-				ThreadContext.Add(FPackageObjectPair{ Package, Object });
+				FileModifiedDelegate.Broadcast( FPaths::ChangeExtension(PackageFileNameString, TEXT(".uexp")) );
+				FileModifiedDelegate.Broadcast( FPaths::ChangeExtension(PackageFileNameString, TEXT(".ubulk")) );
+				FileModifiedDelegate.Broadcast( FPaths::ChangeExtension(PackageFileNameString, TEXT(".ufont")) );
 			}
-		});
-
-	// Accumulate results from the parallel threads into a single array
-	TArray<FPackageObjectPair> PackageObjectPairs = MoveTemp(ThreadContexts[0]);
-	int32 PackageObjectPairsNum = PackageObjectPairs.Num();
-	TArrayView<TArray<FPackageObjectPair>> RemainingThreadContexts = TArrayView<TArray<FPackageObjectPair>>(ThreadContexts).RightChop(1);
-	for (TArray<FPackageObjectPair>& ThreadContext : RemainingThreadContexts)
-	{
-		PackageObjectPairsNum += ThreadContext.Num();
-	}
-	PackageObjectPairs.Reserve(PackageObjectPairsNum);
-	for (TArray<FPackageObjectPair>& ThreadContext : RemainingThreadContexts)
-	{
-		PackageObjectPairs.Append(ThreadContext);
+		}
 	}
-	ThreadContexts.Empty();
-
-	// Sort the array so that all objects for each package are together
-	PackageObjectPairs.Sort();
+}
 
-	// Pull the UObject* out of the array of Pairs into a separate array of just UObject*,
-	// and for each UPackage, add the ArrayView of UObjects matching that package into the UPackage::SoftGCPackageToObjectList.
-	PackageToObjectListBuffer.SetNum(PackageObjectPairsNum);
-	UObject** PackageToObjectListBufferPtr = PackageToObjectListBuffer.GetData();
+bool UCookOnTheFlyServer::IsInSession() const
+{
+	return bSessionRunning;
+}
 
-	UPackage::SoftGCPackageToObjectList.Reserve(PackagesNum);
-	int32 PreviousPackageStartIndex = 0;
-	UPackage* PreviousPackage = nullptr;
-	for (int32 Index = 0; Index < PackageObjectPairsNum; ++Index)
+void UCookOnTheFlyServer::ShutdownCookOnTheFly()
+{
+	if (CookOnTheFlyRequestManager.IsValid())
 	{
-		FPackageObjectPair& Pair = PackageObjectPairs[Index];
-		if (Pair.Package != PreviousPackage)
+		UE_LOG(LogCook, Display, TEXT("Shutting down cook on the fly server"));
+		CookOnTheFlyRequestManager->Shutdown();
+		CookOnTheFlyRequestManager.Reset();
+
+		ShutdownCookSession();
+
+		if (!IsCookingInEditor())
 		{
-			if (Index > PreviousPackageStartIndex)
-			{
-				UPackage::SoftGCPackageToObjectList.Add(PreviousPackage,
-					ObjectPtrWrap(TArrayView<UObject*>(PackageToObjectListBufferPtr + PreviousPackageStartIndex, Index - PreviousPackageStartIndex)));
-			}
-			PreviousPackage = Pair.Package;
-			PreviousPackageStartIndex = Index;
+			GShaderCompilingManager->SkipShaderCompilation(false);
+			GShaderCompilingManager->SetAllowForIncompleteShaderMaps(false);
 		}
-		PackageToObjectListBufferPtr[Index] = Pair.Object;
-	}
-	if (PackageObjectPairsNum > PreviousPackageStartIndex)
-	{
-		UPackage::SoftGCPackageToObjectList.Add(PreviousPackage,
-			ObjectPtrWrap(TArrayView<UObject*>(PackageToObjectListBufferPtr + PreviousPackageStartIndex, PackageObjectPairsNum - PreviousPackageStartIndex)));
 	}
 }
 
-UCookOnTheFlyServer::FScopeFindCookReferences::FScopeFindCookReferences(UCookOnTheFlyServer& InCOTFS)
-	: COTFS(InCOTFS)
-	, SoftGCGuard(UPackage::bSupportCookerSoftGC, true)
-	, bNeedsConstructBuffer(COTFS.SoftGCPackageToObjectListBuffer.IsEmpty())
+uint32 UCookOnTheFlyServer::GetPackagesPerGC() const
 {
-	if (bNeedsConstructBuffer)
-	{
-		ConstructSoftGCPackageToObjectList(COTFS.SoftGCPackageToObjectListBuffer);
-	}
+	return PackagesPerGC;
 }
 
-UCookOnTheFlyServer::FScopeFindCookReferences::~FScopeFindCookReferences()
+uint32 UCookOnTheFlyServer::GetPackagesPerPartialGC() const
 {
-	if (bNeedsConstructBuffer)
-	{
-		UPackage::SoftGCPackageToObjectList.Empty();
-		COTFS.SoftGCPackageToObjectListBuffer.Empty();
-	}
+	return MaxNumPackagesBeforePartialGC;
 }
 
-void UCookOnTheFlyServer::PreGarbageCollect()
+double UCookOnTheFlyServer::GetIdleTimeToGC() const
 {
-	using namespace UE::Cook;
-
-	if (!IsInSession())
+	if (IsCookOnTheFlyMode() && !IsCookingInEditor())
 	{
-		return;
+		// For COTF outside of the editor we want to release open linker file handles promptly but still give some time for new requests to come in
+		return 0.5;
 	}
-
-	NumObjectsHistory.AddInstance(GUObjectArray.GetObjectArrayNumMinusAvailable());
-	VirtualMemoryHistory.AddInstance(FPlatformMemory::GetStats().UsedVirtual);
-	TArray<UPackage*> GCKeepPackages;
-	TArray<UE::Cook::FPackageData*> GCKeepPackageDatas;
-
-#if COOK_CHECKSLOW_PACKAGEDATA
-	// Verify that only packages in the saving states have pointers to objects
-	PackageDatas->LockAndEnumeratePackageDatas([](const FPackageData* PackageData)
-		{
-			check(PackageData->IsInStateProperty(EPackageStateProperty::Saving) || !PackageData->HasReferencedObjects());
-		});
-#endif
-	if (SavingPackageData)
+	else
 	{
-		check(SavingPackageData->GetPackage());
-		GCKeepObjects.Add(SavingPackageData->GetPackage());
-		GCKeepPackageDatas.Add(SavingPackageData);
+		return IdleTimeToGC;
 	}
+}
 
-	// Notify every FGenerationHelper of the garbage collect
-	PackageDatas->LockAndEnumeratePackageDatas([this, &GCKeepPackages, &GCKeepPackageDatas](FPackageData* PackageData)
-		{
-			TRefCountPtr<FGenerationHelper> GenerationHelper = PackageData->GetGenerationHelper();
-			if (!GenerationHelper)
-			{
-				GenerationHelper = PackageData->GetParentGenerationHelper();
-			}
-			if (GenerationHelper)
-			{
-				bool bShouldDemote;
-				GenerationHelper->PreGarbageCollect(GenerationHelper, *PackageData, GCKeepObjects, GCKeepPackages,
-					GCKeepPackageDatas, bShouldDemote);
-				if (bShouldDemote && PackageData->IsInStateProperty(EPackageStateProperty::Saving))
-				{
-					// Demote any Generated/Generator packages we called PreSave on so they call their PostSave before the GC
-					// or prevent them from being garbage collected if the splitter wants to keep them referenced
-					ReleaseCookedPlatformData(*PackageData, UE::Cook::EStateChangeReason::GeneratorPreGarbageCollected,
-						EPackageState::Request);
-				}
-			}
-			if (PackageData->GetIsCookLast() && PackageData->IsInStateProperty(EPackageStateProperty::Saving))
-			{
-				GCKeepPackages.Add(PackageData->GetPackage());
-				GCKeepPackageDatas.Add(PackageData);
-			}
-		});
-	
-	// Find the packages that are waiting on async jobs to finish cooking data
-	// and make sure that they are not garbage collected until the jobs have
-	// completed.
-	{
-		TMap<FPackageData*, UPackage*> UniquePendingPackages;
-		PackageDatas->ForEachPendingCookedPlatformData(
-		[&UniquePendingPackages](const FPendingCookedPlatformData& PendingData)
-		{
-			if (UObject* Object = PendingData.Object.Get())
-			{	
-				if (UPackage* Package = Object->GetPackage())
-				{
-					UniquePendingPackages.Add(&PendingData.PackageData, Package);
-				}	
-			}
-		});
-
-		GCKeepPackages.Reserve(GCKeepPackages.Num() + UniquePendingPackages.Num());
-		for (const TPair<FPackageData*,UPackage*>& Pair : UniquePendingPackages)
-		{
-			GCKeepPackages.Add(Pair.Value);
-			GCKeepPackageDatas.Add(Pair.Key);
-		}
-	}
-
-	// Prevent GC of any objects on which we are still waiting for IsCachedCookedPlatformData
-	PackageDatas->ForEachPendingCookedPlatformData(
-	[this](UE::Cook::FPendingCookedPlatformData& Pending)
-	{
-		if (!Pending.PollIsComplete())
-		{
-			UObject* Object = Pending.Object.Get();
-			check(Object); // Otherwise PollIsComplete would have returned true
-			GCKeepObjects.Add(Object);
-		}
-	});
-
-	const bool bPartialGC = IsCookFlagSet(ECookInitializationFlags::EnablePartialGC);
-	if (bGarbageCollectTypeSoft || bPartialGC)
-	{
-		// Keep referenced all packages in requestqueue, loadqueue, and savequeue, and any packages they depend on
-		TArray<FName> Queue;
-		TSet<FName> Visited;
-		auto AddPackageName = [&Visited, &Queue](FName PackageName)
-		{
-			bool bAlreadyExists;
-			Visited.Add(PackageName, &bAlreadyExists);
-			if (!bAlreadyExists)
-			{
-				Queue.Add(PackageName);
-			}
-		};
-		for (FPackageData* PackageData : PackageDatas->GetRequestQueue().GetReadyRequestsUrgent())
-		{
-			AddPackageName(PackageData->GetPackageName());
-		}
-		for (FPackageData* PackageData : PackageDatas->GetRequestQueue().GetReadyRequestsNormal())
-		{
-			AddPackageName(PackageData->GetPackageName());
-		}
-		for (FPackageData* PackageData : PackageDatas->GetLoadQueue())
-		{
-			AddPackageName(PackageData->GetPackageName());
-		}
-		for (FPackageData* PackageData : PackageDatas->GetSaveQueue())
-		{
-			AddPackageName(PackageData->GetPackageName());
-		}
-		for (FPackageData* PackageData : PackageDatas->GetSaveStalledSet())
-		{
-			AddPackageName(PackageData->GetPackageName());
-		}
-
-		TArray<FName> Dependencies;
-		while (!Queue.IsEmpty())
-		{
-			FName PackageName = Queue.Pop();
-			Dependencies.Reset();
-			AssetRegistry->GetDependencies(PackageName, Dependencies, UE::AssetRegistry::EDependencyCategory::Package,
-				UE::AssetRegistry::EDependencyQuery::Hard);
-			for (FName DependencyName : Dependencies)
-			{
-				AddPackageName(DependencyName);
-			};
-		}
-
-		TSet<UPackage*> GCKeepPackagesSet;
-		GCKeepPackagesSet.Append(GCKeepPackages);
-		for (FName PackageName : Visited)
-		{
-			UPackage* Package = FindPackage(nullptr, *WriteToString<256>(PackageName));
-			if (Package)
-			{
-				bool bAlreadyInSet;
-				GCKeepPackagesSet.Add(Package, &bAlreadyInSet);
-				if (!bAlreadyInSet)
-				{
-					GCKeepPackages.Add(Package);
-					FPackageData* PackageData = PackageDatas->FindPackageDataByPackageName(Package->GetFName());
-					if (PackageData)
-					{
-						GCKeepPackageDatas.Add(PackageData);
-					}
-				}
-			}
-		}
-		ExpectedFreedPackageNames.Empty(PackageTracker->NumLoadedPackages());
-		PackageTracker->ForEachLoadedPackage(
-			[this, &GCKeepPackagesSet](UPackage* Package)
-			{
-				if (!GCKeepPackagesSet.Contains(Package))
-				{
-					ExpectedFreedPackageNames.Add(Package->GetFName());
-				}
-			});
-	}
-
-	// Add packages to GCKeepObjects. 
-	TArray<UObject*> ObjectsWithOuter;
-	for (UPackage* Package : GCKeepPackages)
-	{
-		GCKeepObjects.Add(Package);
-	}
-	for (FPackageData* PackageData : GCKeepPackageDatas)
-	{
-		PackageData->SetKeepReferencedDuringGC(true);
-	}
-
-	// Add all public objects within every package in memory to the UPackage::SoftGCPackageToObjectList container,
-	// so they will be kept in memory if the package is kept in memory.
-	ConstructSoftGCPackageToObjectList(this->SoftGCPackageToObjectListBuffer);
-}
-
-void UCookOnTheFlyServer::CookerAddReferencedObjects(FReferenceCollector& Collector)
-{
-	using namespace UE::Cook;
-
-	// GCKeepObjects are the objects that we want to keep loaded but we only have a WeakPtr to
-	Collector.AddReferencedObjects(GCKeepObjects);
-}
-
-void UCookOnTheFlyServer::PostGarbageCollect()
-{
-	using namespace UE::Cook;
-
-	NumObjectsHistory.AddInstance(GUObjectArray.GetObjectArrayNumMinusAvailable());
-	VirtualMemoryHistory.AddInstance(FPlatformMemory::GetStats().UsedVirtual);
-
-	TSet<UObject*> SaveQueueObjectsThatStillExist;
-
-	// If garbage collection deleted a UPackage WHILE WE WERE SAVING IT, then we have problems.
-	check(!SavingPackageData || SavingPackageData->GetPackage() != nullptr);
-
-	// If there was a GarbageCollect after we already started calling BeginCacheCookedPlatformData, then we
-	// have a list of the WeakObjectPtr to all objects in the package (FPackageData::CachedObjejectsInOuter)
-	// and some of those objects may have been set to null. We declare a reference to prevent GC for the RF_Public
-	// objects in that list, but we do not declare that reference for private objects. The private objects may
-	// therefore have been deleted and set to null 
-	// Side note: because objects can be marked as pending kill at any time and we use FObjectWeakPtr.Get(),
-	// which returns null if pending kill, we need to skip nulls in the array at any point, not just after GC.
-	// 
-	// We do not want to prevent GC of private objects in case there is the expectation by some
-	// systems (blueprints, licensee code) that removing references to an object during PreCollectGarbage will cause
-	// it to be deleted by GC and be replaceable afterwards. We add any new private objects after the garbage collect
-	// and continue with the save. Public objects have a different contract; they are not replaceable across a
-	// GC because anything outside the package could be referring to them. So we keep them referenced. But GC may
-	// force delete them despite our reference, and the package is then in an unknown state. If that happens we
-	// demote the package back to request and start its load and save over.
-	TArray<FPackageData*> Demotes;
-	auto UpdateSavingPackageAfterGarbageCollect =
-		[&Demotes, &SaveQueueObjectsThatStillExist](FPackageData* PackageData)
-	{
-		bool bOutDemote;
-		PackageData->UpdateSaveAfterGarbageCollect(bOutDemote);
-		if (bOutDemote)
-		{
-			Demotes.Add(PackageData);
-		}
-		else
-		{
-			// Mark that the objects for this package should be kept in CachedCookedPlatformData records
-			for (FCachedObjectInOuter& CachedObjectInOuter : PackageData->GetCachedObjectsInOuter())
-			{
-				UObject* Object = CachedObjectInOuter.Object.Get();
-				if (Object)
-				{
-					SaveQueueObjectsThatStillExist.Add(Object);
-				}
-			}
-		}
-	};
-	for (FPackageData* PackageData : PackageDatas->GetSaveQueue())
-	{
-		UpdateSavingPackageAfterGarbageCollect(PackageData);
-	}
-	for (FPackageData* PackageData : PackageDatas->GetSaveStalledSet())
-	{
-		UpdateSavingPackageAfterGarbageCollect(PackageData);
-	}
-	for (FPackageData* PackageData : Demotes)
-	{
-		FGenerationHelper::ValidateSaveStalledState(*this, *PackageData, TEXT("PostGarbageCollect"));
-
-		switch (PackageData->GetState())
-		{
-		case EPackageState::SaveActive:
-			PackageData->SendToState(EPackageState::Request, ESendFlags::QueueRemove, EStateChangeReason::GarbageCollected);
-			if (PackageData->GetIsCookLast())
-			{
-				// CookLast packages in SaveState have had their urgency removed. Add it back if we need to demote them.
-				PackageData->SetUrgency(EUrgency::Blocking, ESendFlags::QueueNone);
-			}
-			PackageDatas->GetRequestQueue().AddRequest(PackageData, /* bForceUrgent */ true);
-			break;
-		case EPackageState::SaveStalledAssignedToWorker:
-			PackageData->SendToState(EPackageState::AssignedToWorker, ESendFlags::QueueAddAndRemove, EStateChangeReason::GarbageCollected);
-			break;
-		case EPackageState::SaveStalledRetracted:
-			DemoteToIdle(*PackageData, ESendFlags::QueueAddAndRemove, ESuppressCookReason::RetractedByCookDirector);
-			break;
-		default:
-			checkf(false, TEXT("State %s not handled in a demoted package."), LexToString(PackageData->GetState()));
-			break;
-		}
-	}
-
-	// Mark that any objects in PendingCookedPlatformDatas should be kept in CachedCookedPlatformData records
-	PackageDatas->ForEachPendingCookedPlatformData(
-		[&SaveQueueObjectsThatStillExist](FPendingCookedPlatformData& CookedPlatformData)
-		{
-			UObject* Object = CookedPlatformData.Object.Get();
-			if (Object)
-			{
-				SaveQueueObjectsThatStillExist.Add(Object);
-			}
-			else
-			{
-				CookedPlatformData.Release();
-			}
-		});
-
-	// Remove objects that were deleted by garbage collection from our containers that track raw object pointers
-	PackageDatas->CachedCookedPlatformDataObjectsPostGarbageCollect(SaveQueueObjectsThatStillExist);
-
-	PackageDatas->LockAndEnumeratePackageDatas([this](FPackageData* PackageData)
-	{
-		if (TRefCountPtr<FGenerationHelper> GenerationHelper = PackageData->GetGenerationHelper())
-		{
-			GenerationHelper->PostGarbageCollect(GenerationHelper, *GCDiagnosticContext);
-		}
-	});
-
-	// Second pass over all PackageDatas, combine a few operations
-	PackageDatas->LockAndEnumeratePackageDatas([](FPackageData* PackageData)
-	{
-		// Mark that the PackageData no longer needs to be keepreferenced.
-		// This can only be done after all GenerationHelper->PostGarbageCollect have been called.
-		PackageData->SetKeepReferencedDuringGC(false);
-
-		// Reset the completion flags for FPreloadPackage, since the UPackage might be no longer loaded.
-		TRefCountPtr<FPackagePreloader> Preloader = PackageData->GetPackagePreloader();
-		if (Preloader)
-		{
-			Preloader->PostGarbageCollect();
-		}
-	});
-
-	// Only after running all possible callbacks that need our links for diagnostics, clear the list of temporary
-	// references that we created for the garbage collection.
-	GCKeepObjects.Empty();
-	UPackage::SoftGCPackageToObjectList.Empty();
-	SoftGCPackageToObjectListBuffer.Empty();
-
-	CookedPackageCountSinceLastGC = 0;
-
-	// Whenever we collect garbage, reset the counter for how many busy reports with an
-	// idle shadercompiler we need before we issue a warning
-	bShaderCompilerWasActiveeOnPreviousBusyReport = true;
-}
-
-bool UCookOnTheFlyServer::NeedsDiagnosticSecondGC() const
-{
-	return GCDiagnosticContext->NeedsDiagnosticSecondGC();
-}
-
-void UCookOnTheFlyServer::OnCookerStartCollectGarbage(uint32& ResultFlagsFromTick)
-{
-	GCDiagnosticContext->OnCookerStartCollectGarbage(*this, ResultFlagsFromTick);
-}
-
-void UCookOnTheFlyServer::OnCookerEndCollectGarbage(uint32& ResultFlagsFromTick)
-{
-	GCDiagnosticContext->OnCookerEndCollectGarbage(*this, ResultFlagsFromTick);
-}
-
-void UCookOnTheFlyServer::EvaluateGarbageCollectionResults(bool bWasDueToOOM, bool bWasPartialGC, uint32 ResultFlags,
-	int32 NumObjectsBeforeGC, const FPlatformMemoryStats& MemStatsBeforeGC,
-	const FGenericMemoryStats& AllocatorStatsBeforeGC,
-	int32 NumObjectsAfterGC, const FPlatformMemoryStats& MemStatsAfterGC,
-	const FGenericMemoryStats& AllocatorStatsAfterGC,
-	float GCDurationSeconds)
-{
-	using namespace UE::Cook;
-
-	ON_SCOPE_EXIT
-	{
-		ExpectedFreedPackageNames.Empty();
-		GCDiagnosticContext->OnEvaluateResultsComplete();
-	};
-	bWarnedExceededMaxMemoryWithinGCCooldown = false;
-	LastGCTime = FPlatformTime::Seconds();
-	bool bWasSoftGC = ResultFlags & COSR_RequiresGC_Soft_OOM;
-	if (bWasSoftGC)
-	{
-		LastSoftGCTime = LastGCTime;
-		if (SoftGCStartNumerator)
-		{
-			int32 StartNumerator = FMath::Max(SoftGCStartNumerator, 1);
-			int32 Denominator = FMath::Max(SoftGCDenominator, 1);
-			// Calculate the new SoftGCNextAvailablePhysicalTarget. Use the floor of NewAvailableMemory/Denominator,
-			// unless we are already 50% of the way through that level, in which case use the next value below that
-			int64 PhysicalMemoryQuantum = static_cast<int64>(MemStatsAfterGC.TotalPhysical) / Denominator;
-			int32 NextTarget =
-				static_cast<int64>(MemStatsAfterGC.AvailablePhysical - PhysicalMemoryQuantum / 2) / PhysicalMemoryQuantum;
-			NextTarget = FMath::Min(NextTarget, StartNumerator);
-			if (NextTarget <= 0)
-			{
-				SoftGCNextAvailablePhysicalTarget = -2; // disabled, no further targets
-			}
-			else
-			{
-				SoftGCNextAvailablePhysicalTarget = static_cast<int64>((MemStatsAfterGC.TotalPhysical) * NextTarget)
-					/ Denominator;
-			}
-		}
-	}
-	else
-	{
-		LastSoftGCTime = LastGCTime;
-		LastFullGCTime = LastGCTime;
-	}
-	if (SoftGCHistory)
-	{
-		SoftGCHistory->AddDurationMeasurement(GCDurationSeconds);
-	}
-
-	if (IsCookingInEditor())
-	{
-		return;
-	}
-	if (IsCookOnTheFlyMode() && !bWasDueToOOM)
-	{
-		return;
-	}
-
-	int64 NumObjectsMin = NumObjectsHistory.GetMinimum();
-	int64 NumObjectsMax = NumObjectsHistory.GetMaximum();
-	int64 NumObjectsSpread = NumObjectsMax - NumObjectsMin;
-	int64 NumObjectsFreed = NumObjectsBeforeGC - NumObjectsAfterGC;
-	int64 NumObjectsCapacity = GUObjectArray.GetObjectArrayEstimatedAvailable() + GUObjectArray.GetObjectArrayNumMinusAvailable();
-	int64 VirtualMemMin = VirtualMemoryHistory.GetMinimum();
-	int64 VirtualMemMax = VirtualMemoryHistory.GetMaximum();
-	int64 VirtualMemSpread = VirtualMemMax - VirtualMemMin;
-	int64 VirtualMemBeforeGC = MemStatsBeforeGC.UsedVirtual;
-	int64 VirtualMemAfterGC = MemStatsAfterGC.UsedVirtual;
-	int64 VirtualMemFreed = MemStatsBeforeGC.UsedVirtual - MemStatsAfterGC.UsedVirtual;
-
-	int64 ExpectedObjectsFreed = static_cast<int64>(MemoryExpectedFreedToSpreadRatio * static_cast<float>(NumObjectsSpread));
-	double ExpectedMemFreed = MemoryExpectedFreedToSpreadRatio * VirtualMemSpread;
-	static bool bCookMemoryAnalysis = FParse::Param(FCommandLine::Get(), TEXT("CookMemoryAnalysis"));
-#if ENABLE_LOW_LEVEL_MEM_TRACKER
-	// When tracking memory with LLM, always show the memory status.
-	const bool bAlwaysShowAnalysis = FLowLevelMemTracker::Get().IsEnabled() || bCookMemoryAnalysis;
-#else
-	const bool bAlwaysShowAnalysis = bCookMemoryAnalysis;
-#endif		
-	constexpr int32 BytesPerMeg = 1000000;
-	auto DisplaySimpleSummary = [&]()
-	{
-		UE_LOG(LogCook, Display, TEXT("GarbageCollection Results:\n")
-			TEXT("\tType: %s\n")
-			TEXT("\tDuration: %.3fs\n")
-			TEXT("\tNumObjects:\n")
-			TEXT("\t\tCapacity:         %10d\n")
-			TEXT("\t\tBefore GC:        %10d\n")
-			TEXT("\t\tAfter GC:         %10d\n")
-			TEXT("\t\tFreed by GC:      %10d\n")
-			TEXT("\tVirtual Memory:\n")
-			TEXT("\t\tBefore GC:        %10" INT64_FMT " MB\n")
-			TEXT("\t\tAfter GC:         %10" INT64_FMT " MB\n")
-			TEXT("\t\tFreed by GC:      %10" INT64_FMT " MB"),
-			(bWasSoftGC ? TEXT("Soft") : (bWasPartialGC ? TEXT("Partial") : TEXT("Full"))),
-			GCDurationSeconds,
-			NumObjectsCapacity, (int64)NumObjectsBeforeGC, (int64)NumObjectsAfterGC, NumObjectsFreed,
-			VirtualMemBeforeGC / BytesPerMeg, VirtualMemAfterGC / BytesPerMeg, VirtualMemFreed / BytesPerMeg
-		);
-	};
-
-	if (!bWasSoftGC)
-	{
-		bool bWasImpactful =
-			(NumObjectsFreed >= ExpectedObjectsFreed || NumObjectsBeforeGC - NumObjectsMin < ExpectedObjectsFreed) &&
-			(VirtualMemFreed >= ExpectedMemFreed || VirtualMemBeforeGC - VirtualMemMin <= ExpectedMemFreed);
-
-		if ((!bWasDueToOOM || bWasImpactful) && !bAlwaysShowAnalysis)
-		{
-			DisplaySimpleSummary();
-			return;
-		}
-
-		if (bWasDueToOOM && !bWasImpactful)
-		{
-			UE_LOG(LogCook, Display, TEXT("GarbageCollection Results: Garbage Collection was not very impactful."));
-		}
-		else
-		{
-			UE_LOG(LogCook, Display, TEXT("GarbageCollection Results:"));
-		}
-		UE_LOG(LogCook, Display, TEXT("\tMemoryAnalysis: General:\n")
-			TEXT("\t\tType: %s\n")
-			TEXT("\tDuration: %.3fs"),
-			(bWasSoftGC ? TEXT("Soft") : (bWasPartialGC ? TEXT("Partial") : TEXT("Full"))),
-			GCDurationSeconds);
-		UE_LOG(LogCook, Display, TEXT("\tMemoryAnalysis: NumObjects:\n")
-			TEXT("\t\tCapacity:         %10" INT64_FMT "\n")
-			TEXT("\t\tProcess Min:      %10" INT64_FMT "\n")
-			TEXT("\t\tProcess Max:      %10" INT64_FMT "\n")
-			TEXT("\t\tProcess Spread:   %10" INT64_FMT "\n")
-			TEXT("\t\tBefore GC:        %10" INT64_FMT "\n")
-			TEXT("\t\tAfter GC:         %10" INT64_FMT "\n")
-			TEXT("\t\tFreed by GC:      %10" INT64_FMT ""),
-			NumObjectsCapacity, NumObjectsMin, NumObjectsMax, NumObjectsSpread,
-			(int64)NumObjectsBeforeGC, (int64)NumObjectsAfterGC, NumObjectsFreed);
-		UE_LOG(LogCook, Display, TEXT("\tMemoryAnalysis: Virtual Memory:\n")
-			TEXT("\t\tProcess Min:      %10" INT64_FMT " MB\n")
-			TEXT("\t\tProcess Max:      %10" INT64_FMT " MB\n")
-			TEXT("\t\tProcess Spread:   %10" INT64_FMT " MB\n")
-			TEXT("\t\tBefore GC:        %10" INT64_FMT " MB\n")
-			TEXT("\t\tAfter GC:         %10" INT64_FMT " MB\n")
-			TEXT("\t\tFreed by GC:      %10" INT64_FMT " MB"),
-			VirtualMemMin / BytesPerMeg, VirtualMemMax / BytesPerMeg, VirtualMemSpread / BytesPerMeg,
-			VirtualMemBeforeGC / BytesPerMeg, VirtualMemAfterGC / BytesPerMeg, VirtualMemFreed / BytesPerMeg);
-		auto AllocatorStatsToString = [](const FGenericMemoryStats& AllocatorStats)
-		{
-			TStringBuilder<256> Writer;
-			for (const TPair<FStringView, SIZE_T>& Item : AllocatorStats)
-			{
-				Writer << TEXT("\n\t\tItem ") << Item.Key << TEXT(" ") << (uint64)Item.Value;
-			}
-			return FString(*Writer);
-		};
-		UE_LOG(LogCook, Display, TEXT("\tMemoryAnalysis: Allocator Stats Before:%s"),
-			*AllocatorStatsToString(AllocatorStatsBeforeGC));
-		UE_LOG(LogCook, Display, TEXT("\tMemoryAnalysis: Allocator Stats After:%s"),
-			*AllocatorStatsToString(AllocatorStatsAfterGC));
-
-
-		UE_LOG(LogCook, Display, TEXT("See log for memory use information for UObject classes and LLM tags."));
-
-		{
-			TGuardValue<bool> SoftGCGuard(UPackage::bSupportCookerSoftGC, true);
-			ConstructSoftGCPackageToObjectList(this->SoftGCPackageToObjectListBuffer);
-			UE::Cook::DumpObjClassList(CookByTheBookOptions->SessionStartupObjects);
-			UPackage::SoftGCPackageToObjectList.Empty();
-			SoftGCPackageToObjectListBuffer.Empty();
-		}
-		GLog->Logf(TEXT("Memory Analysis: LLM Tags:"));
-#if ENABLE_LOW_LEVEL_MEM_TRACKER
-		if (FLowLevelMemTracker::Get().IsEnabled())
-		{
-			FLowLevelMemTracker::Get().DumpToLog();
-		}
-		else
-#endif
-		{
-			GLog->Logf(TEXT("LLM Tags are not displayed because llm is disabled. Run with -llm or -trace=memtag to see llm tags."));
-		}
-	}
-	else
-	{
-		DisplaySimpleSummary();
-
-		// Mark the packages we freed so we can give a warning to diagnose why they are still referenced if they
-		// get loaded again.
-		PackageTracker->AddExpectedNeverLoadPackages(ExpectedFreedPackageNames);
-
-		// Only show diagnostics if LLM is on, because they are somewhat expensive. We could add a separate setting
-		// for this, but it's more convenient to combine it with the LLM enabled setting
-#if ENABLE_LOW_LEVEL_MEM_TRACKER
-		bool bShowDiagnostics = FLowLevelMemTracker::Get().IsEnabled();
-#else
-		constexpr bool bShowDiagnostics = false;
-#endif
-		if (bShowDiagnostics)
-		{
-			// If some packages we expected to be freed were not freed, show the reference chains for why
-			// they were not freed.
-			TArray<UPackage*> PackagesReferencedOutsideOfCooker;
-			for (FWeakObjectPtr& WeakPtr : CookByTheBookOptions->SessionStartupObjects)
-			{
-				UObject* Object = WeakPtr.Get();
-				if (!Object)
-				{
-					continue;
-				}
-				UPackage* Package = Object->GetPackage();
-				ExpectedFreedPackageNames.Remove(Package->GetFName());
-			}
-			PackageTracker->ForEachLoadedPackage(
-				[this, &PackagesReferencedOutsideOfCooker](UPackage* Package)
-				{
-					if (ExpectedFreedPackageNames.Contains(Package->GetFName()))
-					{
-						PackagesReferencedOutsideOfCooker.Add(Package);
-					}
-				});
-			if (PackagesReferencedOutsideOfCooker.Num() > 0)
-			{
-				UE::Cook::DumpPackageReferencers(PackagesReferencedOutsideOfCooker);
-			}
-		}
-	}
-}
-
-void UCookOnTheFlyServer::OnObjectModified( UObject *ObjectMoving )
-{
-	if (IsGarbageCollecting())
-	{
-		return;
-	}
-	OnObjectUpdated( ObjectMoving );
-}
-
-void UCookOnTheFlyServer::OnObjectPropertyChanged(UObject* ObjectBeingModified, FPropertyChangedEvent& PropertyChangedEvent)
-{
-	if (IsGarbageCollecting())
-	{
-		return;
-	}
-	if ( PropertyChangedEvent.Property == nullptr && 
-		PropertyChangedEvent.MemberProperty == nullptr )
-	{
-		// probably nothing changed... 
-		return;
-	}
-
-	OnObjectUpdated( ObjectBeingModified );
-}
-
-void UCookOnTheFlyServer::OnObjectSaved( UObject* ObjectSaved, FObjectPreSaveContext SaveContext)
-{
-	if (SaveContext.IsProceduralSave() )
-	{
-		// This is a procedural save (e.g. our own saving of the cooked package) rather than a user save, ignore
-		return;
-	}
-
-	UPackage* Package = ObjectSaved->GetOutermost();
-	if (Package == nullptr || Package == GetTransientPackage())
-	{
-		return;
-	}
-
-	MarkPackageDirtyForCooker(Package);
-
-	// Register the package filename as modified. We don't use the cache because the file may not exist on disk yet at this point
-	const FString PackageFilename = FPackageName::LongPackageNameToFilename(Package->GetName(), Package->ContainsMap() ? FPackageName::GetMapPackageExtension() : FPackageName::GetAssetPackageExtension());
-	ModifiedAssetFilenames.Add(FName(*PackageFilename));
-}
-
-void UCookOnTheFlyServer::OnObjectUpdated( UObject *Object )
-{
-	// get the outer of the object
-	UPackage *Package = Object->GetOutermost();
-
-	MarkPackageDirtyForCooker( Package );
-}
-
-void UCookOnTheFlyServer::MarkPackageDirtyForCooker(UPackage* Package, bool bAllowInSession)
-{
-	if (Package->RootPackageHasAnyFlags(PKG_PlayInEditor))
-	{
-		return;
-	}
-
-	if (Package->HasAnyPackageFlags(PKG_PlayInEditor | PKG_ContainsScript | PKG_InMemoryOnly) == true && !GetClass()->HasAnyClassFlags(CLASS_DefaultConfig | CLASS_Config))
-	{
-		return;
-	}
-
-	if (Package == GetTransientPackage())
-	{
-		return;
-	}
-
-	if (Package->GetOuter() != nullptr)
-	{
-		return;
-	}
-
-	FName PackageName = Package->GetFName();
-	if (FPackageName::IsMemoryPackage(PackageName.ToString()))
-	{
-		return;
-	}
-
-	if (bIsSavingPackage)
-	{
-		return;
-	}
-
-	if (IsInSession() && !bAllowInSession)
-	{
-		WorkerRequests->AddEditorActionCallback([this, PackageName]() { MarkPackageDirtyForCookerFromSchedulerThread(PackageName); });
-	}
-	else
-	{
-		MarkPackageDirtyForCookerFromSchedulerThread(PackageName);
-	}
-}
-
-FName GInstigatorMarkPackageDirty(TEXT("MarkPackageDirtyForCooker"));
-void UCookOnTheFlyServer::MarkPackageDirtyForCookerFromSchedulerThread(const FName& PackageName)
-{
-	TRACE_CPUPROFILER_EVENT_SCOPE(MarkPackageDirtyForCooker);
-
-	// could have just cooked a file which we might need to write
-	UPackage::WaitForAsyncFileWrites();
-
-	// Update the package's FileName if it has changed
-	UE::Cook::FPackageData* PackageData = PackageDatas->UpdateFileName(PackageName);
-
-	// force the package to be recooked
-	UE_LOG(LogCook, Verbose, TEXT("Modification detected to package %s"), *PackageName.ToString());
-	if ( PackageData && IsCookingInEditor() )
-	{
-		check(IsInGameThread()); // We're editing scheduler data, which is only allowable from the scheduler thread
-		bool bHadCookedPlatforms = PackageData->HasAnyCookedPlatform();
-		PackageData->ClearCookResults();
-		if (PackageData->IsInProgress())
-		{
-			PackageData->SendToState(UE::Cook::EPackageState::Request,
-				UE::Cook::ESendFlags::QueueAddAndRemove, UE::Cook::EStateChangeReason::ForceRecook);
-		}
-		else if (IsCookByTheBookMode() && IsInSession() && bHadCookedPlatforms)
-		{
-			QueueDiscoveredPackage(*PackageData,
-				UE::Cook::FInstigator(UE::Cook::EInstigator::Unspecified, GInstigatorMarkPackageDirty),
-				UE::Cook::EDiscoveredPlatformSet::CopyFromInstigator);
-		}
-
-		if ( IsCookOnTheFlyMode() && FileModifiedDelegate.IsBound())
-		{
-			FString PackageFileNameString = PackageData->GetFileName().ToString();
-			FileModifiedDelegate.Broadcast(PackageFileNameString);
-			if (PackageFileNameString.EndsWith(".uasset") || PackageFileNameString.EndsWith(".umap"))
-			{
-				FileModifiedDelegate.Broadcast( FPaths::ChangeExtension(PackageFileNameString, TEXT(".uexp")) );
-				FileModifiedDelegate.Broadcast( FPaths::ChangeExtension(PackageFileNameString, TEXT(".ubulk")) );
-				FileModifiedDelegate.Broadcast( FPaths::ChangeExtension(PackageFileNameString, TEXT(".ufont")) );
-			}
-		}
-	}
-}
-
-bool UCookOnTheFlyServer::IsInSession() const
-{
-	return bSessionRunning;
-}
-
-void UCookOnTheFlyServer::ShutdownCookOnTheFly()
-{
-	if (CookOnTheFlyRequestManager.IsValid())
-	{
-		UE_LOG(LogCook, Display, TEXT("Shutting down cook on the fly server"));
-		CookOnTheFlyRequestManager->Shutdown();
-		CookOnTheFlyRequestManager.Reset();
-
-		ShutdownCookSession();
-
-		if (!IsCookingInEditor())
-		{
-			GShaderCompilingManager->SkipShaderCompilation(false);
-			GShaderCompilingManager->SetAllowForIncompleteShaderMaps(false);
-		}
-	}
-
-}
-
-uint32 UCookOnTheFlyServer::GetPackagesPerGC() const
-{
-	return PackagesPerGC;
-}
-
-uint32 UCookOnTheFlyServer::GetPackagesPerPartialGC() const
-{
-	return MaxNumPackagesBeforePartialGC;
-}
-
-
-double UCookOnTheFlyServer::GetIdleTimeToGC() const
-{
-	if (IsCookOnTheFlyMode() && !IsCookingInEditor())
-	{
-		// For COTF outside of the editor we want to release open linker file handles promptly but still give some time for new requests to come in
-		return 0.5;
-	}
-	else
-	{
-		return IdleTimeToGC;
-	}
-}
-
-void UCookOnTheFlyServer::BeginDestroy()
-{
-	ShutdownCookOnTheFly();
-
-	Super::BeginDestroy();
-}
-
-void UCookOnTheFlyServer::TickRequestManager()
-{
-	if (CookOnTheFlyRequestManager)
-	{
-		CookOnTheFlyRequestManager->Tick();
-	}
-}
-
-void UCookOnTheFlyServer::TickRecompileShaderRequestsPrivate(UE::Cook::FTickStackData& StackData)
-{
-	// try to pull off a request
-	UE::Cook::FRecompileShaderRequest RecompileShaderRequest;
-	bool bProcessedRequests = false;
-	if (PackageTracker->RecompileRequests.Dequeue(&RecompileShaderRequest))
-	{
-		if (RecompileShaderRequest.RecompileArguments.CommandType != ODSCRecompileCommand::ResetMaterialCache)
-		{
-			RecompileShadersForRemote(RecompileShaderRequest.RecompileArguments, GetSandboxDirectory(RecompileShaderRequest.RecompileArguments.PlatformName));
-		}
-
-		RecompileShaderRequest.CompletionCallback();
-		bProcessedRequests = true;
-	}
-	if (PackageTracker->RecompileRequests.HasItems())
-	{
-		RecompileRequestsPollable->Trigger(*this);
-	}
-
-	if (bProcessedRequests)
-	{
-		// Ask for GC to run again when we processed some shaders requests to ensure material get evicted and we don't keep their package open
-		StackData.ResultFlags |= COSR_RequiresGC | COSR_RequiresGC_IdleTimer | COSR_YieldTick;
-	}
-
-}
-
-class FDiffModeCookServerUtils
-{
-public:
-	enum class EDiffMode
-	{
-		None,
-		DiffOnly,
-		LinkerDiff,
-		IterativeValidate,
-		IterativeValidatePhase1,
-		IterativeValidatePhase2,
-	};
-
-	void InitializePackageWriter(UCookOnTheFlyServer& COTFS, ICookedPackageWriter*& CookedPackageWriter,
-		const FString& ResolvedMetadataPath, UE::Cook::FDeterminismManager* InDeterminismManager)
-	{
-		Initialize();
-		if (DiffMode == EDiffMode::None)
-		{
-			return;
-		}
-
-		ICookedPackageWriter::FCookCapabilities Capabilities = CookedPackageWriter->GetCookCapabilities();
-		if (!Capabilities.bDiffModeSupported)
-		{
-			// All current PackageWriters support bDiffModeSupported; log a fatal error in case a new one is added.
-			UE_LOG(LogCook, Fatal, 
-				TEXT("A DiffMode was enabled, but the current PackageWriter has bDiffModeSupported=false."));
-		}
-
-		// Wrap the incoming writer inside the feature-specific-functionality writer
-		switch (DiffMode)
-		{
-		case EDiffMode::DiffOnly:
-			CookedPackageWriter = new FDiffPackageWriter(TUniquePtr<ICookedPackageWriter>(CookedPackageWriter),
-				InDeterminismManager);
-			break;
-		case EDiffMode::LinkerDiff:
-			CookedPackageWriter = new FLinkerDiffPackageWriter(TUniquePtr<ICookedPackageWriter>(CookedPackageWriter));
-			break;
-		case EDiffMode::IterativeValidate:
-			CookedPackageWriter = new FIterativeValidatePackageWriter(COTFS,
-				TUniquePtr<ICookedPackageWriter>(CookedPackageWriter), FIterativeValidatePackageWriter::EPhase::AllInOnePhase,
-				ResolvedMetadataPath, InDeterminismManager);
-			break;
-		case EDiffMode::IterativeValidatePhase1:
-			CookedPackageWriter = new FIterativeValidatePackageWriter(COTFS,
-				TUniquePtr<ICookedPackageWriter>(CookedPackageWriter), FIterativeValidatePackageWriter::EPhase::Phase1,
-				ResolvedMetadataPath, InDeterminismManager);
-			break;
-		case EDiffMode::IterativeValidatePhase2:
-			CookedPackageWriter = new FIterativeValidatePackageWriter(COTFS,
-				TUniquePtr<ICookedPackageWriter>(CookedPackageWriter), FIterativeValidatePackageWriter::EPhase::Phase2,
-				ResolvedMetadataPath, InDeterminismManager);
-			break;
-		default:
-			checkNoEntry();
-			break;
-		}
-	}
-
-	bool IsDeterminismDebug() const
-	{
-		switch (DiffMode)
-		{
-		case EDiffMode::None: return false;
-		case EDiffMode::DiffOnly: return true;
-		case EDiffMode::LinkerDiff: return false;
-		case EDiffMode::IterativeValidate: return true;
-		case EDiffMode::IterativeValidatePhase1: return true;
-		case EDiffMode::IterativeValidatePhase2: return true;
-		default:
-			checkNoEntry();
-			return false;
-		}
-	}
-
-	void Initialize()
-	{
-		if (bInitialized)
-		{
-			return;
-		}
-
-		DiffMode = EDiffMode::None;
-		const TCHAR* CommandLine = FCommandLine::Get();
-		auto EnsureMutualExclusion = [this]()
-		{
-			if (DiffMode != EDiffMode::None)
-			{
-				UE_LOG(LogCook, Fatal, TEXT("-DiffOnly, -LinkerDiff, and -IterativeValidate* are mutually exclusive."));
-			}
-		};
-
-		bool bValue;
-		if (FParse::Param(CommandLine, TEXT("DIFFONLY")))
-		{
-			EnsureMutualExclusion();
-			DiffMode = EDiffMode::DiffOnly;
-		}
-		if (FParse::Bool(CommandLine, TEXT("-LINKERDIFF="), bValue) && bValue)
-		{
-			EnsureMutualExclusion();
-			DiffMode = EDiffMode::LinkerDiff;
-		}
-		if (FParse::Param(CommandLine, TEXT("IterativeValidate")))
-		{
-			EnsureMutualExclusion();
-			DiffMode = EDiffMode::IterativeValidate;
-		}
-		if (FParse::Param(CommandLine, TEXT("IterativeValidatePhase1")))
-		{
-			EnsureMutualExclusion();
-			DiffMode = EDiffMode::IterativeValidatePhase1;
-		}
-		if (FParse::Param(CommandLine, TEXT("IterativeValidatePhase2")))
-		{
-			EnsureMutualExclusion();
-			DiffMode = EDiffMode::IterativeValidatePhase2;
-		}
-		bInitialized = true;
-	}
-
-private:
-	bool bInitialized = false;
-	EDiffMode DiffMode = EDiffMode::None;
-};
-
-#if OUTPUT_COOKTIMING
-
-UE_TRACE_EVENT_BEGIN(UE_CUSTOM_COOKTIMER_LOG, SaveCookedPackage, NoSync)
-	UE_TRACE_EVENT_FIELD(UE::Trace::WideString, PackageName)
-UE_TRACE_EVENT_END()
-
-#endif //OUTPUT_COOKTIMING
-
-void UCookOnTheFlyServer::SaveCookedPackage(UE::Cook::FSaveCookedPackageContext& Context)
-{
-	using namespace UE::Cook;
-
-	UE_SCOPED_HIERARCHICAL_CUSTOM_COOKTIMER_AND_DURATION(SaveCookedPackage, DetailedCookStats::TickCookOnTheSideSaveCookedPackageTimeSec)
-		UE_ADD_CUSTOM_COOKTIMER_META(SaveCookedPackage, PackageName, *WriteToString<256>(Context.PackageData.GetFileName()));
-	double SaveStartTime = FPlatformTime::Seconds();
-
-	UPackage* Package = Context.Package;
-	FOnScopeExit ScopedPackageFlags([Package, OriginalPackageFlags = Package->GetPackageFlags()]()
-		{ Package->SetPackageFlagsTo(OriginalPackageFlags); });
-
-	Context.SetupPackage();
-
-	ITargetPlatformManagerModule& TPM = GetTargetPlatformManagerRef();
-	TGuardValue<bool> ScopedOutputCookerWarnings(GOutputCookingWarnings, IsCookFlagSet(ECookInitializationFlags::OutputVerboseCookerWarnings));
-	// SavePackage can CollectGarbage, so we need to store the currently-unqueued PackageData in a separate variable that we register for garbage collection
-	TGuardValue<UE::Cook::FPackageData*> ScopedSavingPackageData(SavingPackageData, &Context.PackageData);
-	TGuardValue<bool> ScopedIsSavingPackage(bIsSavingPackage, true);
-	// For legacy reasons we set GIsCookerLoadingPackage == true during save. Some classes use it to conditionally execute cook operations in both save and load
-	TGuardValue<bool> ScopedIsCookerLoadingPackage(GIsCookerLoadingPackage, true);
-
-	for (int32 PlatformIndex = 0; PlatformIndex < Context.PlatformsForPackage.Num(); ++PlatformIndex)
-	{
-		const ITargetPlatform* TargetPlatform = Context.PlatformsForPackage[PlatformIndex];
-		Context.SetupPlatform(TargetPlatform, PlatformIndex);
-		if (Context.bPlatformSetupSuccessful)
-		{
-			UE_SCOPED_HIERARCHICAL_COOKTIMER(GEditorSavePackage);
-			UE_TRACK_REFERENCING_PLATFORM_SCOPED(TargetPlatform);
-
-			TMap<UObject*, FObjectSaveOverride> SaveOverrides;
-			FArchiveCookData CookData(*TargetPlatform, *Context.ArchiveCookContext);
-			FSavePackageArgs SaveArgs;
-			SaveArgs.TopLevelFlags = Context.FlagsToCook;
-			SaveArgs.bForceByteSwapping = Context.bEndianSwap;
-			SaveArgs.bWarnOfLongFilename = false;
-			SaveArgs.SaveFlags = Context.SaveFlags;
-			SaveArgs.ArchiveCookData = &CookData;
-			SaveArgs.bSlowTask = false;
-			SaveArgs.SavePackageContext = Context.SavePackageContext;
-			SaveArgs.InOutSaveOverrides = &SaveOverrides;
-
-			Context.PackageWriter->UpdateSaveArguments(SaveArgs);
-			FSavePackageResultStruct AuthoritativeResult = ESavePackageResult::Error;
-			bool IsFirstPass = true;
-			for (;;) 
-			{
-				try
-				{
-					LLM_SCOPE_BYTAG(Cooker_SavePackage);
-					FScopedActivePackage ScopedActivePackage(*this, Context.PackageData.GetPackageName(), NAME_None);
-					if (bSkipSave)
-					{
-						Context.SavePackageResult = ESavePackageResult::Success;
-					}
-					else
-					{
-						Context.SavePackageResult = GEditor->Save(Package, Context.World, *Context.PlatFilename, SaveArgs);
-					}
-				}
-				catch (std::exception&)
-				{
-					UE_LOG(LogCook, Warning, TEXT("Tried to save package %s for target platform %s but threw an exception"),
-						*Package->GetName(), *TargetPlatform->PlatformName());
-					Context.SavePackageResult = ESavePackageResult::Error;
-				}
-
-				bool IsAnotherSaveNeeded = Context.PackageWriter->IsAnotherSaveNeeded(Context.SavePackageResult, SaveArgs);
-				if (IsFirstPass)
-				{
-					AuthoritativeResult = MoveTemp(Context.SavePackageResult);
-					IsFirstPass = false;
-				}
-				if (IsAnotherSaveNeeded)
-				{
-					// We must not try a second save of a package while the first save is still in flight.
-					// The optimal solution is to wait for ONLY the package that needs a second save, but we don't
-					// have the bookkeeping data to do that, so we have to wait for all async package writes to complete.
-					UPackage::WaitForAsyncFileWrites();
-				}
-				else
-				{
-					break;
-				}
-			}
-			Context.SavePackageResult = MoveTemp(AuthoritativeResult);
-
-			// If package was actually saved check with asset manager to make sure it wasn't excluded for being a
-			// development or never cook package. But skip sending the warnings from this check if it was editor-only.
-			if (Context.SavePackageResult == ESavePackageResult::Success)
-			{
-				UE_SCOPED_HIERARCHICAL_COOKTIMER(VerifyCanCookPackage);
-				if (!UAssetManager::Get().VerifyCanCookPackage(this, Package->GetFName()))
-				{
-					Context.SavePackageResult = ESavePackageResult::Error;
-				}
-			}
-
-			++this->StatSavedPackageCount;
-		}
-
-		Context.FinishPlatform();
-	}
-
-	// Need to restore flags before calling FinishPackage because it might need to save again
-	ScopedPackageFlags.ExitEarly();
-	Context.FinishPackage();
-
-	constexpr double SavePackageMinDurationLogTimeSeconds = 600.;
-	float SaveDurationSeconds = static_cast<float>(FPlatformTime::Seconds() - SaveStartTime);
-	UE_CLOG(SaveDurationSeconds >= SavePackageMinDurationLogTimeSeconds,
-		LogCook, Display, TEXT("SavePackagePerformance: Package %s took %.0fs to save."),
-		*WriteToString<256>(Package->GetFName()), SaveDurationSeconds);
-}
-
-bool UCookOnTheFlyServer::IsDebugRecordUnsolicited() const
-{
-	return bOnlyEditorOnlyDebug | bHiddenDependenciesDebug;
-}
-
-namespace UE::Cook
-{
-FSaveCookedPackageContext::FSaveCookedPackageContext(UCookOnTheFlyServer& InCOTFS, UE::Cook::FPackageData& InPackageData,
-	TArrayView<const ITargetPlatform*> InPlatformsForPackage, UE::Cook::FTickStackData& InStackData)
-	: COTFS(InCOTFS)
-	, PackageData(InPackageData)
-	, PlatformsForPackage(InPlatformsForPackage)
-	, StackData(InStackData)
-	, Package(PackageData.GetPackage())
-	, PackageName(Package ? Package->GetName() : FString())
-	, Filename(PackageData.GetFileName().ToString())
-{
-	PlatformDependencies.SetNum(InPlatformsForPackage.Num());
-}
-
-void FSaveCookedPackageContext::SetupPackage()
-{
-	check(Package && Package->IsFullyLoaded()); // PackageData should not be in the save state if Package is not fully loaded
-	check(Package->GetPathName().Equals(PackageName)); // We should only be saving outermost packages, so the path name should be the same as the package name
-	check(!Filename.IsEmpty()); // PackageData guarantees FileName is non-empty; if not found it should never make it into save state
-	if (Package->HasAnyPackageFlags(PKG_ReloadingForCooker))
-	{
-		UE_LOG(LogCook, Warning, TEXT("Package %s marked as reloading for cook was requested to save"), *PackageName);
-		UE_LOG(LogCook, Fatal, TEXT("Package %s marked as reloading for cook was requested to save"), *PackageName);
-	}
-
-	SaveFlags = SAVE_KeepGUID | SAVE_Async
-		| (COTFS.IsCookFlagSet(ECookInitializationFlags::Unversioned) ? SAVE_Unversioned : 0);
-	SaveFlags |= COTFS.IsCookFlagSet(ECookInitializationFlags::CookEditorOptional) ? SAVE_Optional : SAVE_None;
-
-	if (COTFS.CookByTheBookOptions->bCookSoftPackageReferences)
-	{
-		SaveFlags |= SAVE_CookSoftPackageReferences;
-	}
-
-	// Use SandboxFile to do path conversion to properly handle sandbox paths (outside of standard paths in particular).
-	Filename = COTFS.ConvertToFullSandboxPath(*Filename, true);
-}
-
-void FSaveCookedPackageContext::SetupPlatform(const ITargetPlatform* InTargetPlatform, int32 InPlatformIndex)
-{
-	PlatformIndex = InPlatformIndex;
-	TargetPlatform = InTargetPlatform;
-	PlatFilename = Filename.Replace(TEXT("[Platform]"), *TargetPlatform->PlatformName());
-	bPlatformSetupSuccessful = false;
-
-	ArchiveCookContext.Emplace(Package,
-		COTFS.IsDirectorCookByTheBook() ? UE::Cook::ECookType::ByTheBook : UE::Cook::ECookType::OnTheFly,
-		COTFS.IsCookingDLC() ? UE::Cook::ECookingDLC::Yes : UE::Cook::ECookingDLC::No,
-		TargetPlatform);
-
-	// don't save Editor resources from the Engine if the target doesn't have editoronly data
-	if (COTFS.IsCookFlagSet(ECookInitializationFlags::SkipEditorContent) &&
-		(PackageName.StartsWith(TEXT("/Engine/Editor")) || PackageName.StartsWith(TEXT("/Engine/VREditor"))) &&
-		!TargetPlatform->HasEditorOnlyData())
-	{
-		SavePackageResult = ESavePackageResult::ContainsEditorOnlyData;
-		const TCHAR* RejectedReason = TEXT("EngineEditorContent");
-		if (GCookProgressDisplay & (int32)ECookProgressDisplayMode::Instigators)
-		{
-			UE_LOG(LogCook, Display, TEXT("Cooking %s, Instigator: { %s } -> Rejected %s"), *PackageName,
-				*(PackageData.GetInstigator().ToString()), RejectedReason);
-		}
-		else
-		{
-			UE_LOG(LogCook, Display, TEXT("Cooking %s -> Rejected %s"), *PackageName, RejectedReason);
-		}
-		return;
-	}
-	// Check whether or not game-specific behaviour should prevent this package from being cooked for the target platform
-	else if (!UAssetManager::Get().ShouldCookForPlatform(Package, TargetPlatform))
-	{
-		SavePackageResult = ESavePackageResult::ContainsEditorOnlyData;
-		const TCHAR* RejectedReason = TEXT("NotAssetManagerShouldCookForPlatform");
-		if (GCookProgressDisplay & (int32)ECookProgressDisplayMode::Instigators)
-		{
-			UE_LOG(LogCook, Display, TEXT("Cooking %s, Instigator: { %s } -> Rejected %s"), *PackageName,
-				*(PackageData.GetInstigator().ToString()), RejectedReason);
-		}
-		else
-		{
-			UE_LOG(LogCook, Display, TEXT("Cooking %s -> Rejected %s"), *PackageName, RejectedReason);
-		}
-		return;
-	}
-	// check if this package is unsupported for the target platform (typically plugin content)
-	else
-	{
-		if (TSet<FName>* NeverCookPackages = COTFS.PackageTracker->PlatformSpecificNeverCookPackages.Find(TargetPlatform))
-		{
-			FGenerationHelper* GenerationHelper =
-				PackageData.IsGenerated() ? PackageData.GetParentGenerationHelper() : nullptr;
+void UCookOnTheFlyServer::BeginDestroy()
+{
+	ShutdownCookOnTheFly();
 
-			if (NeverCookPackages->Find(Package->GetFName()) ||
-				(GenerationHelper && NeverCookPackages->Find(GenerationHelper->GetOwner().GetPackageName())))
-			{
-				SavePackageResult = ESavePackageResult::ContainsEditorOnlyData;
-				const TCHAR* RejectedReason = TEXT("PlatformSpecificNeverCook");
-				if (GCookProgressDisplay & (int32)ECookProgressDisplayMode::Instigators)
-				{
-					UE_LOG(LogCook, Display, TEXT("Cooking %s, Instigator: { %s } -> Rejected %s"), *PackageName,
-						*(PackageData.GetInstigator().ToString()), RejectedReason);
-				}
-				else
-				{
-					UE_LOG(LogCook, Display, TEXT("Cooking %s -> Rejected %s"), *PackageName, RejectedReason);
-				}
-				return;
-			}
-		}
-	}
+	Super::BeginDestroy();
+}
 
-	const FString FullFilename = FPaths::ConvertRelativePathToFull(PlatFilename);
-	if (FullFilename.Len() >= FPlatformMisc::GetMaxPathLength())
+void UCookOnTheFlyServer::TickRequestManager()
+{
+	if (CookOnTheFlyRequestManager)
 	{
-		LogCookerMessage(FString::Printf(TEXT("Couldn't save package, filename is too long (%d >= %d): %s"),
-			FullFilename.Len(), FPlatformMisc::GetMaxPathLength(), *FullFilename), EMessageSeverity::Error);
-		SavePackageResult = ESavePackageResult::Error;
-		return;
+		CookOnTheFlyRequestManager->Tick();
 	}
+}
 
-	if (!bHasDelayLoaded)
+void UCookOnTheFlyServer::TickRecompileShaderRequestsPrivate(UE::Cook::FTickStackData& StackData)
+{
+	// try to pull off a request
+	UE::Cook::FRecompileShaderRequest RecompileShaderRequest;
+	bool bProcessedRequests = false;
+	if (PackageTracker->RecompileRequests.Dequeue(&RecompileShaderRequest))
 	{
-		// look for a world object in the package (if there is one, there's a map)
-		World = UWorld::FindWorldInPackage(Package);
-		if (World)
+		if (RecompileShaderRequest.RecompileArguments.CommandType != ODSCRecompileCommand::ResetMaterialCache)
 		{
-			FlagsToCook = RF_NoFlags;
+			RecompileShadersForRemote(RecompileShaderRequest.RecompileArguments, GetSandboxDirectory(RecompileShaderRequest.RecompileArguments.PlatformName));
 		}
-		bContainsMap = Package->ContainsMap();
-		bHasDelayLoaded = true;
-	}
 
-	UE_CLOG((GCookProgressDisplay & (int32)ECookProgressDisplayMode::Instigators) && PlatformIndex == 0,
-		LogCook, Display, TEXT("Cooking %s, Instigator: { %s }"),
-		*PackageName, *(PackageData.GetInstigator().ToString()));
-	UE_CLOG(GCookProgressDisplay & (int32)ECookProgressDisplayMode::PackageNames, LogCook, Display,
-		TEXT("Cooking %s"), *PackageName);
-
-	bEndianSwap = (!TargetPlatform->IsLittleEndian()) ^ (!PLATFORM_LITTLE_ENDIAN);
-
-	if (!TargetPlatform->HasEditorOnlyData())
-	{
-		Package->SetPackageFlags(PKG_FilterEditorOnly);
+		RecompileShaderRequest.CompletionCallback();
+		bProcessedRequests = true;
 	}
-	else
+	if (PackageTracker->RecompileRequests.HasItems())
 	{
-		Package->ClearPackageFlags(PKG_FilterEditorOnly);
+		RecompileRequestsPollable->Trigger(*this);
 	}
 
-	CookContext = &COTFS.FindOrCreateSaveContext(TargetPlatform);
-	SavePackageContext = &CookContext->SaveContext;
-	PackageWriter = CookContext->PackageWriter;
-	ICookedPackageWriter::FBeginPackageInfo Info;
-	Info.PackageName = Package->GetFName();
-	Info.LooseFilePath = PlatFilename;
-	PackageWriter->BeginPackage(Info);
-	if (CookContext->DeterminismManager)
-	{
-		CookContext->DeterminismManager->BeginPackage(Package, TargetPlatform, PackageWriter);
-	}
-	// Set platform-specific save flags
-	FPackagePlatformData& PlatformData = PackageData.FindOrAddPlatformData(TargetPlatform);
-	uint32 PlatformSaveFlagsMask = SAVE_AllowTimeout;
-	SaveFlags = (SaveFlags & ~PlatformSaveFlagsMask);
-	if (!PlatformData.IsSaveTimedOut())
+	if (bProcessedRequests)
 	{
-		// If we timedout before, do not allow another timeout, otherwise do allow it
-		SaveFlags |= SAVE_AllowTimeout;
+		// Ask for GC to run again when we processed some shaders requests to ensure material get evicted and we don't keep their package open
+		StackData.ResultFlags |= COSR_RequiresGC | COSR_RequiresGC_IdleTimer | COSR_YieldTick;
 	}
-
-	// Indicate Setup was successful
-	bPlatformSetupSuccessful = true;
-	SavePackageResult = ESavePackageResult::Success;
-}
-
-bool IsRetryErrorCode(ESavePackageResult Result)
-{
-	return Result == ESavePackageResult::Timeout;
 }
 
-void FSaveCookedPackageContext::FinishPlatform()
+class FDiffModeCookServerUtils
 {
-	TRACE_CPUPROFILER_EVENT_SCOPE(FSaveCookedPackageContext::FinishPlatform);
-	check(PlatformIndex >= 0 && PlatformDependencies.IsValidIndex(PlatformIndex));
-
-	bool bSuccessful = SavePackageResult.IsSuccessful();
-	ECookResult CookResult = bSuccessful ? ECookResult::Succeeded : ECookResult::Failed;
-
-	// Calculate up-to-date assetregistry data for Generator and Generated packages
-	TOptional<TArray<FAssetDependency>> OverridePackageDependencies;
-	TOptional<UE::TargetDomain::FGeneratedPackageResultStruct> GeneratedResult;
-	TOptional<FAssetPackageData> AssetPackageDataBuffer;
-	TOptional<FAssetPackageData> OverrideAssetPackageData;
-	const FAssetPackageData* AssetPackageData = nullptr;
-	if (TRefCountPtr<FGenerationHelper> GenerationHelper = PackageData.GetGenerationHelper();
-		GenerationHelper)
-	{
-		OverridePackageDependencies.Emplace();
-		GenerationHelper->FinishGeneratorPlatformSave(PackageData, PlatformIndex == 0, *OverridePackageDependencies);
-	}
-	else if (TRefCountPtr<FGenerationHelper> ParentGenerationHelper = PackageData.GetParentGenerationHelper();
-		ParentGenerationHelper)
-	{
-		GeneratedResult.Emplace();
-		ParentGenerationHelper->FinishGeneratedPlatformSave(PackageData, *GeneratedResult);
-		OverridePackageDependencies.Emplace(GeneratedResult->PackageDependencies);
-		OverrideAssetPackageData.Emplace(GeneratedResult->AssetPackageData);
-		AssetPackageData = &GeneratedResult->AssetPackageData;
-	}
-	if (!AssetPackageData)
+public:
+	enum class EDiffMode
 	{
-		AssetPackageDataBuffer = COTFS.AssetRegistry->GetAssetPackageDataCopy(Package->GetFName());
-		AssetPackageData = AssetPackageDataBuffer.GetPtrOrNull();
-	}
+		None,
+		DiffOnly,
+		LinkerDiff,
+		IterativeValidate,
+		IterativeValidatePhase1,
+		IterativeValidatePhase2,
+	};
 
-	// Commit the saved bytes and the iterative cook data to the PackageWriter
-	if (bPlatformSetupSuccessful)
+	void InitializePackageWriter(UCookOnTheFlyServer& COTFS, ICookedPackageWriter*& CookedPackageWriter,
+		const FString& ResolvedMetadataPath, UE::Cook::FDeterminismManager* InDeterminismManager)
 	{
-		CalculatePlatformAgnosticRuntimeDependencies();
-		CalculatePlatformRuntimeDependencies();
-
-		ICookedPackageWriter::FCommitPackageInfo Info;
-		// Note GetCommitAttachments mutates the SaveResult; it moves CookDependencies out of it
-		Info.Attachments = GetCommitAttachments(GeneratedResult.GetPtrOrNull(), OverridePackageDependencies.GetPtrOrNull());
-		Info.Status = PackageResultToCommitStatus(SavePackageResult.Result);
-		Info.PackageName = Package->GetFName();
-		Info.PackageHash = AssetPackageData ? AssetPackageData->GetPackageSavedHash() : FIoHash();
-		Info.WriteOptions = GetCommitWriteOptions();
-
-		PackageWriter->CommitPackage(MoveTemp(Info));
-		if (CookContext->DeterminismManager)
+		Initialize();
+		if (DiffMode == EDiffMode::None)
 		{
-			CookContext->DeterminismManager->EndPackage();
+			return;
 		}
-	}
 
-	// Update asset registry
-	if (COTFS.IsDirectorCookByTheBook())
-	{
-		IAssetRegistryReporter& Reporter = *(COTFS.PlatformManager->GetPlatformData(TargetPlatform)->RegistryReporter);
-		TOptional<TArray<FAssetData>> AssetDatasFromSave;
-		if (bSuccessful)
+		ICookedPackageWriter::FCookCapabilities Capabilities = CookedPackageWriter->GetCookCapabilities();
+		if (!Capabilities.bDiffModeSupported)
 		{
-			AssetDatasFromSave.Emplace(MoveTemp(SavePackageResult.SavedAssets));
+			// All current PackageWriters support bDiffModeSupported; log a fatal error in case a new one is added.
+			UE_LOG(LogCook, Fatal, 
+				TEXT("A DiffMode was enabled, but the current PackageWriter has bDiffModeSupported=false."));
 		}
-		Reporter.UpdateAssetRegistryData(Package->GetFName(), Package, CookResult, &SavePackageResult,
-			MoveTemp(AssetDatasFromSave), MoveTemp(OverrideAssetPackageData), MoveTemp(OverridePackageDependencies),
-			COTFS);
-	}
-
-	// If not retrying, mark the package as cooked, either successfully or with failure
-	bool bIsRetryErrorCode = IsRetryErrorCode(SavePackageResult.Result);
-	if (!bIsRetryErrorCode)
-	{
-		PackageData.SetPlatformCooked(TargetPlatform, CookResult);
-	}
 
-	// Update flags used to determine garbage collection.
-	if (bSuccessful)
-	{
-		if (bContainsMap)
-		{
-			StackData.ResultFlags |= UCookOnTheFlyServer::COSR_CookedMap;
-		}
-		else
+		// Wrap the incoming writer inside the feature-specific-functionality writer
+		switch (DiffMode)
 		{
-			++COTFS.CookedPackageCountSinceLastGC;
-			StackData.ResultFlags |= UCookOnTheFlyServer::COSR_CookedPackage;
+		case EDiffMode::DiffOnly:
+			CookedPackageWriter = new FDiffPackageWriter(TUniquePtr<ICookedPackageWriter>(CookedPackageWriter),
+				InDeterminismManager);
+			break;
+		case EDiffMode::LinkerDiff:
+			CookedPackageWriter = new FLinkerDiffPackageWriter(TUniquePtr<ICookedPackageWriter>(CookedPackageWriter));
+			break;
+		case EDiffMode::IterativeValidate:
+			CookedPackageWriter = new FIterativeValidatePackageWriter(COTFS,
+				TUniquePtr<ICookedPackageWriter>(CookedPackageWriter), FIterativeValidatePackageWriter::EPhase::AllInOnePhase,
+				ResolvedMetadataPath, InDeterminismManager);
+			break;
+		case EDiffMode::IterativeValidatePhase1:
+			CookedPackageWriter = new FIterativeValidatePackageWriter(COTFS,
+				TUniquePtr<ICookedPackageWriter>(CookedPackageWriter), FIterativeValidatePackageWriter::EPhase::Phase1,
+				ResolvedMetadataPath, InDeterminismManager);
+			break;
+		case EDiffMode::IterativeValidatePhase2:
+			CookedPackageWriter = new FIterativeValidatePackageWriter(COTFS,
+				TUniquePtr<ICookedPackageWriter>(CookedPackageWriter), FIterativeValidatePackageWriter::EPhase::Phase2,
+				ResolvedMetadataPath, InDeterminismManager);
+			break;
+		default:
+			checkNoEntry();
+			break;
 		}
 	}
 
-	// Accumulate results for SaveCookedPackage_Finish
-	if (SavePackageResult.Result == ESavePackageResult::Timeout)
-	{
-		PackageData.FindOrAddPlatformData(TargetPlatform).SetSaveTimedOut(true);
-		bHasTimeOut = true;
-	}
-	bAnySaveSucceeded |= bSuccessful;
-
-	bHasRetryErrorCode |= bIsRetryErrorCode;
-
-	ArchiveCookContext.Reset();
-	PlatformIndex = -1;
-}
-
-void FSaveCookedPackageContext::FinishPackage()
-{
-	// If any save succeeded, add all dependencies from all platforms to the cook for the platform
-	if (bAnySaveSucceeded && !COTFS.CookByTheBookOptions->bSkipSoftReferences)
+	bool IsDeterminismDebug() const
 	{
-		FName PackageFName = Package->GetFName();
-		if (PlatformsForPackage.Num() == 1)
-		{
-			TConstArrayView<const ITargetPlatform*> ReachablePlatforms(&TargetPlatform, 1);
-			for (const TPair<FPackageData*, EInstigator>& DependencyPair : PlatformDependencies[0])
-			{
-				COTFS.QueueDiscoveredPackage(*DependencyPair.Key,
-					UE::Cook::FInstigator(DependencyPair.Value, PackageFName), FDiscoveredPlatformSet(ReachablePlatforms));
-			}
-		}
-		else
+		switch (DiffMode)
 		{
-			// Merge the platform dependencies for each package into a single QueueDiscoveredPackage call, if possible
-			// It will not be possible if the different platforms have different instigators, so track a list of
-			// platforms for each package for each instigator type.
-			TMap<FPackageData*, TMap<EInstigator, TArray<const ITargetPlatform*>>> PackagePlatformsForInstigator;
-			for (int32 LocalIndex = 0; LocalIndex < PlatformsForPackage.Num(); ++LocalIndex)
-			{
-				const TMap<FPackageData*, EInstigator>& CurrentDependencies = PlatformDependencies[LocalIndex];
-				const ITargetPlatform* CurrentPlatform = PlatformsForPackage[LocalIndex];
-				for (const TPair<FPackageData*, EInstigator>& PackagePlatformPair : CurrentDependencies)
-				{
-					TMap<EInstigator, TArray<const ITargetPlatform*>>& TargetMap =
-						PackagePlatformsForInstigator.FindOrAdd(PackagePlatformPair.Key);
-					TargetMap.FindOrAdd(PackagePlatformPair.Value).Add(CurrentPlatform);
-				}
-			}
-			for (const TPair<FPackageData*, TMap<EInstigator, TArray<const ITargetPlatform*>>>& PackagePair
-				: PackagePlatformsForInstigator)
-			{
-				for (const TPair<EInstigator, TArray<const ITargetPlatform*>>& InstigatorPair : PackagePair.Value)
-				{
-					COTFS.QueueDiscoveredPackage(*PackagePair.Key,
-						UE::Cook::FInstigator(InstigatorPair.Key, PackageFName),
-						FDiscoveredPlatformSet(InstigatorPair.Value));
-				}
-			}
+		case EDiffMode::None: return false;
+		case EDiffMode::DiffOnly: return true;
+		case EDiffMode::LinkerDiff: return false;
+		case EDiffMode::IterativeValidate: return true;
+		case EDiffMode::IterativeValidatePhase1: return true;
+		case EDiffMode::IterativeValidatePhase2: return true;
+		default:
+			checkNoEntry();
+			return false;
 		}
 	}
 
-	if (COTFS.IsDebugRecordUnsolicited())
+	void Initialize()
 	{
-		COTFS.ProcessUnsolicitedPackages();
-		TMap<FPackageData*, EInstigator> AllPlatformDependenciesBuffer;
-		TMap<FPackageData*, EInstigator>* AllPlatformDependencies = &PlatformDependencies[0];
-		if (PlatformsForPackage.Num() > 1)
+		if (bInitialized)
 		{
-			AllPlatformDependencies = &AllPlatformDependenciesBuffer;
-			for (int32 LocalIndex = 0; LocalIndex < PlatformsForPackage.Num(); ++LocalIndex)
-			{
-				AllPlatformDependenciesBuffer.Append(PlatformDependencies[LocalIndex]);
-			}
+			return;
 		}
 
-		FDiagnostics::AnalyzeHiddenDependencies(COTFS, PackageData, PackageData.DetachUnsolicited(),
-			*AllPlatformDependencies, PlatformsForPackage, COTFS.bOnlyEditorOnlyDebug, COTFS.bHiddenDependenciesDebug);
-	}
-
-	if (!bHasRetryErrorCode)
-	{
-		if (COTFS.IsCookOnTheFlyMode() && PackageData.GetUrgency() != EUrgency::Blocking &&
-			(!COTFS.CookOnTheFlyRequestManager || COTFS.CookOnTheFlyRequestManager->ShouldUseLegacyScheduling()))
+		DiffMode = EDiffMode::None;
+		const TCHAR* CommandLine = FCommandLine::Get();
+		auto EnsureMutualExclusion = [this]()
 		{
-			// this is an unsolicited package
-			if (FPaths::FileExists(Filename))
+			if (DiffMode != EDiffMode::None)
 			{
-				COTFS.PackageTracker->UnsolicitedCookedPackages.AddCookedPackage(
-					FFilePlatformRequest(PackageData.GetFileName(), EInstigator::Unspecified, PlatformsForPackage));
-
-#if DEBUG_COOKONTHEFLY
-				UE_LOG(LogCook, Display, TEXT("UnsolicitedCookedPackages: %s"), *Filename);
-#endif
+				UE_LOG(LogCook, Fatal, TEXT("-DiffOnly, -LinkerDiff, and -IterativeValidate* are mutually exclusive."));
 			}
-		}
-	}
-}
-
-void FSaveCookedPackageContext::CalculatePlatformAgnosticRuntimeDependencies()
-{
-	if (bPlatformAgnosticDependenciesCalculated)
-	{
-		return;
-	}
-	bPlatformAgnosticDependenciesCalculated = true;
+		};
 
-	FName PackageFName = PackageData.GetPackageName();
-	for (FName LocalizedPackageName : FRequestCluster::GetLocalizationReferences(PackageFName, COTFS))
-	{
-		UE::Cook::FPackageData* LocalizedPackageData = COTFS.PackageDatas->TryAddPackageDataByPackageName(LocalizedPackageName);
-		if (LocalizedPackageData)
-		{
-			AddDependency(PlatformAgnosticDependencies, LocalizedPackageData, false /* bHard */);
-		}
-	}
-	// Also add any references from the package that are required by the AssetManager
-	for (FName AMPackageName : FRequestCluster::GetAssetManagerReferences(PackageFName))
-	{
-		UE::Cook::FPackageData* AMPackageData = COTFS.PackageDatas->TryAddPackageDataByPackageName(AMPackageName);
-		if (AMPackageData)
+		bool bValue;
+		if (FParse::Param(CommandLine, TEXT("DIFFONLY")))
 		{
-			AddDependency(PlatformAgnosticDependencies, AMPackageData, false /* bHard */);
+			EnsureMutualExclusion();
+			DiffMode = EDiffMode::DiffOnly;
 		}
-	}
-
-	// When using legacy WhatGetsCookedRules, add all the SoftObjectPaths discovered during the package's load, plus any added on
-	// during save, to the cook for all platforms
-	TSet<FName> SoftObjectPackages;
-	GRedirectCollector.ProcessSoftObjectPathPackageList(PackageFName, false /* bGetEditorOnly */, SoftObjectPackages);
-	for (FName SoftObjectPackage : SoftObjectPackages)
-	{
-		TMap<FSoftObjectPath, FSoftObjectPath> RedirectedPaths;
-
-		// If this is a redirector, extract destination from asset registry
-		if (COTFS.ContainsRedirector(SoftObjectPackage, RedirectedPaths))
+		if (FParse::Bool(CommandLine, TEXT("-LINKERDIFF="), bValue) && bValue)
 		{
-			for (TPair<FSoftObjectPath, FSoftObjectPath>& RedirectedPath : RedirectedPaths)
-			{
-				GRedirectCollector.AddAssetPathRedirection(RedirectedPath.Key, RedirectedPath.Value);
-			}
+			EnsureMutualExclusion();
+			DiffMode = EDiffMode::LinkerDiff;
 		}
-
-		UE::Cook::FPackageData* SoftObjectPackageData = COTFS.PackageDatas->TryAddPackageDataByPackageName(SoftObjectPackage);
-		if (SoftObjectPackageData)
+		if (FParse::Param(CommandLine, TEXT("IterativeValidate")))
 		{
-			if (!COTFS.bSkipOnlyEditorOnly)
-			{
-				AddDependency(PlatformAgnosticDependencies, SoftObjectPackageData, false /* bHard */);
-			}
-
-			if (COTFS.IsDebugRecordUnsolicited())
-			{
-				PackageData.CreateOrGetUnsolicited().Add(SoftObjectPackageData, EInstigator::SoftDependency);
-			}
+			EnsureMutualExclusion();
+			DiffMode = EDiffMode::IterativeValidate;
 		}
-	}
-}
-
-void FSaveCookedPackageContext::CalculatePlatformRuntimeDependencies()
-{
-	PlatformDependencies.IsValidIndex(PlatformIndex);
-	TMap<FPackageData*, EInstigator>& CurrentDependencies = PlatformDependencies[PlatformIndex];
-	CurrentDependencies.Reset();
-	// Add PlatformAgnostic dependencies to the PlatformDependencies
-	CurrentDependencies.Append(PlatformAgnosticDependencies);
-
-	// Add imports and softobjectpaths from the save to the SaveDependencies for the current Platform and for all Platforms
-	if (COTFS.bSkipOnlyEditorOnly)
-	{
-		FName PackageFName = Package->GetFName();
-		TConstArrayView<const ITargetPlatform*> ReachablePlatforms(&TargetPlatform, 1);
-		for (const TArray<FName>* DependencyNames : { &SavePackageResult.ImportPackages, &SavePackageResult.SoftPackageReferences })
+		if (FParse::Param(CommandLine, TEXT("IterativeValidatePhase1")))
 		{
-			bool bHard = DependencyNames == &SavePackageResult.ImportPackages;
-			for (FName DependencyName : *DependencyNames)
-			{
-				UE::Cook::FPackageData* DependencyData = COTFS.PackageDatas->TryAddPackageDataByPackageName(DependencyName);
-				if (DependencyData)
-				{
-					AddDependency(CurrentDependencies, DependencyData, bHard);
-				}
-			}
+			EnsureMutualExclusion();
+			DiffMode = EDiffMode::IterativeValidatePhase1;
 		}
-	}
-}
-
-TArray<FName> FSaveCookedPackageContext::GetPlatformRuntimeDependencies() const
-{
-	TArray<FName> PlatformDependencyNames;
-	PlatformDependencyNames.Reserve(PlatformDependencies[PlatformIndex].Num());
-	for (const TPair<FPackageData*, EInstigator>& DependencyPair : PlatformDependencies[PlatformIndex])
-	{
-		PlatformDependencyNames.Add(DependencyPair.Key->GetPackageName());
-	}
-	return PlatformDependencyNames;
-}
-
-TArray<IPackageWriter::FCommitAttachmentInfo> FSaveCookedPackageContext::GetCommitAttachments(
-	UE::TargetDomain::FGeneratedPackageResultStruct* GeneratedResult,
-	TArray<FAssetDependency>* OverrideDependencies)
-{
-	TArray<IPackageWriter::FCommitAttachmentInfo> Result;
-	if (COTFS.bHybridIterativeEnabled)
-	{
-		UE_SCOPED_HIERARCHICAL_COOKTIMER(TargetDomainDependencies);
-		TArray<FName> PlatformRuntimeDependencies = GetPlatformRuntimeDependencies();
-		if (OverrideDependencies)
+		if (FParse::Param(CommandLine, TEXT("IterativeValidatePhase2")))
 		{
-			PlatformRuntimeDependencies.Reserve(PlatformRuntimeDependencies.Num() + OverrideDependencies->Num());
-			for (const FAssetDependency& Dependency : *OverrideDependencies)
-			{
-				PlatformRuntimeDependencies.Add(Dependency.AssetId.PackageName);
-			}
+			EnsureMutualExclusion();
+			DiffMode = EDiffMode::IterativeValidatePhase2;
 		}
-		UE::TargetDomain::CollectAndStoreCookAttachments(Package, TargetPlatform, &SavePackageResult,
-			GeneratedResult, MoveTemp(PlatformRuntimeDependencies), Result);
-	}
-	if (CookContext->DeterminismManager)
-	{
-		CookContext->DeterminismManager->AppendCommitAttachments(Result);
+		bInitialized = true;
 	}
 
-	return Result;
-}
-
-IPackageWriter::EWriteOptions FSaveCookedPackageContext::GetCommitWriteOptions() const
-{
-	IPackageWriter::EWriteOptions Result = IPackageWriter::EWriteOptions::None;
-	if (!COTFS.bSkipSave)
-	{
-		Result |= IPackageWriter::EWriteOptions::Write;
-
-		if (COTFS.IsDirectorCookByTheBook())
-		{
-			Result |= IPackageWriter::EWriteOptions::ComputeHash;
-		}
-	}
-	return Result;
-}
+private:
+	bool bInitialized = false;
+	EDiffMode DiffMode = EDiffMode::None;
+};
 
-void FSaveCookedPackageContext::AddDependency(TMap<FPackageData*, EInstigator>& InDependencies,
-	FPackageData* PackageData, bool bHard)
+bool UCookOnTheFlyServer::IsDebugRecordUnsolicited() const
 {
-	EInstigator& Existing = InDependencies.FindOrAdd(PackageData, EInstigator::Unspecified);
-	if (bHard)
-	{
-		Existing = EInstigator::SaveTimeHardDependency;
-	}
-	else if (Existing == EInstigator::Unspecified)
-	{
-		Existing = EInstigator::SoftDependency;
-	}
+	return bOnlyEditorOnlyDebug | bHiddenDependenciesDebug;
 }
 
-} // namespace UE::Cook
-
 void UCookOnTheFlyServer::RecordExternalActorDependencies(TConstArrayView<FName> ExternalActorDependencies)
 {
 	using namespace UE::Cook;
diff --git a/Engine/Source/Editor/UnrealEd/Private/Cooker/CookGarbageCollect.cpp b/Engine/Source/Editor/UnrealEd/Private/Cooker/CookGarbageCollect.cpp
--- a/Engine/Source/Editor/UnrealEd/Private/Cooker/CookGarbageCollect.cpp
+++ b/Engine/Source/Editor/UnrealEd/Private/Cooker/CookGarbageCollect.cpp
@@ -2,13 +2,65 @@
 
 #include "Cooker/CookGarbageCollect.h"
 
+#include "AssetRegistry/IAssetRegistry.h"
+#include "Async/ParallelFor.h"
+#include "Async/TaskGraphInterfaces.h"
+#include "Containers/ArrayView.h"
+#include "Containers/Set.h"
+#include "Cooker/CookGenerationHelper.h"
+#include "Cooker/CookPackageData.h"
+#include "Cooker/CookPackagePreloader.h"
+#include "Cooker/CookProfiling.h"
+#include "Cooker/CookTypes.h"
+#include "Cooker/PackageTracker.h"
+#include "CookOnTheSide/CookLog.h"
 #include "CookOnTheSide/CookOnTheFlyServer.h"
+#include "CoreGlobals.h"
+#include "Engine/Engine.h"
+#include "HAL/LowLevelMemTracker.h"
+#include "HAL/MemoryMisc.h"
+#include "HAL/PlatformMemory.h"
+#include "HAL/PlatformTime.h"
+#include "Logging/LogMacros.h"
+#include "Math/UnrealMathUtility.h"
+#include "Misc/AssertionMacros.h"
+#include "Misc/CoreMiscDefines.h"
 #include "Misc/EnumClassFlags.h"
+#include "Misc/ScopeExit.h"
+#include "Misc/StringBuilder.h"
+#include "Templates/RefCounting.h"
+#include "UObject/GarbageCollection.h"
 #include "UObject/GarbageCollectionHistory.h"
+#include "UObject/Package.h"
+#include "UObject/UObjectArray.h"
+#include "UObject/UObjectGlobals.h"
+#include "UObject/WeakObjectPtr.h"
+
+#include <atomic>
 
 namespace UE::Cook
 {
 
+FScopeFindCookReferences::FScopeFindCookReferences(UCookOnTheFlyServer& InCOTFS)
+	: COTFS(InCOTFS)
+	, SoftGCGuard(UPackage::bSupportCookerSoftGC, true)
+	, bNeedsConstructBuffer(COTFS.SoftGCPackageToObjectListBuffer.IsEmpty())
+{
+	if (bNeedsConstructBuffer)
+	{
+		ConstructSoftGCPackageToObjectList(COTFS.SoftGCPackageToObjectListBuffer);
+	}
+}
+
+FScopeFindCookReferences::~FScopeFindCookReferences()
+{
+	if (bNeedsConstructBuffer)
+	{
+		UPackage::SoftGCPackageToObjectList.Empty();
+		COTFS.SoftGCPackageToObjectListBuffer.Empty();
+	}
+}
+
 FCookGCDiagnosticContext::~FCookGCDiagnosticContext()
 {
 	SetGCWithHistoryRequested(false);
@@ -107,4 +159,1024 @@ void FCookGCDiagnosticContext::SetGCWithHistoryRequested(bool bValue)
 #endif
 }
 
-} // namespace UE::Cook
\ No newline at end of file
+void FSoftGCHistory::AddDurationMeasurement(float DurationSeconds)
+{
+	int32 HistoryLength = DurationHistory.Num();
+	while (HistoryLength >= MaxHistoryLength)
+	{
+		if (HistoryLength <= 1)
+		{
+			if (HistoryLength > 0)
+			{
+				DurationHistory.PopFront();
+			}
+			AverageDurationSeconds = 0.0;
+			HistoryLength = 0;
+		}
+		else
+		{
+			AverageDurationSeconds = (AverageDurationSeconds * HistoryLength - DurationHistory.PopFrontValue()) / (HistoryLength - 1);
+			--HistoryLength;
+		}
+	}
+	if (HistoryLength < MaxHistoryLength)
+	{
+		if (HistoryLength == 0)
+		{
+			AverageDurationSeconds = DurationSeconds;
+		}
+		else
+		{
+			AverageDurationSeconds = (AverageDurationSeconds * HistoryLength + DurationSeconds) / (HistoryLength + 1);
+		}
+		DurationHistory.Add(DurationSeconds);
+	}
+}
+
+bool FSoftGCHistory::IsTriggeringWithinBudget(UCookOnTheFlyServer& COTFS, double CurrentTimeSeconds, FString* OutDiagnostics)
+{
+	if (OutDiagnostics)
+	{
+		OutDiagnostics->Reset();
+	}
+	if (COTFS.CookedPackageCountSinceLastGC == 0)
+	{
+		return false;
+	}
+	if (DurationHistory.Num() == 0)
+	{
+		if (OutDiagnostics)
+		{
+			*OutDiagnostics = TEXT("No duration data");
+		}
+		return true;
+	}
+	// TimeBudget/(Time + TimeBudget) == BudgetFraction
+	// TimeBudget == (BudgetFraction/(1 - BudgetFraction))*Time
+	if (COTFS.SoftGCTimeFractionBudget > .99f)
+	{
+		if (OutDiagnostics)
+		{
+			*OutDiagnostics = FString::Printf(
+				TEXT("SoftGCTimeFractionBudget == %.3f, above threshold to always trigger"),
+				COTFS.SoftGCTimeFractionBudget);
+		}
+		return true;
+	}
+	float TimeSinceLastGCSeconds = CurrentTimeSeconds - COTFS.LastSoftGCTime;
+	float CurrentTimeBudget = TimeSinceLastGCSeconds *
+		COTFS.SoftGCTimeFractionBudget / (1 - COTFS.SoftGCTimeFractionBudget);
+	if (CurrentTimeBudget >= AverageDurationSeconds)
+	{
+		if (OutDiagnostics)
+		{
+			*OutDiagnostics = FString::Printf(
+				TEXT("SoftGCTimeFractionBudget == %.3f. TimeSinceLastGCSeconds == %.3f. CurrentTimeBudget == %.3f. ExpectedDuration == %.3f"),
+				COTFS.SoftGCTimeFractionBudget, TimeSinceLastGCSeconds, CurrentTimeBudget, AverageDurationSeconds);
+		}
+		return true;
+	}
+	return false;
+}
+
+void ConstructSoftGCPackageToObjectList(TArray<UObject*>& PackageToObjectListBuffer)
+{
+	struct FPackageObjectPair
+	{
+		UPackage* Package;
+		UObject* Object;
+		bool operator<(const FPackageObjectPair& Other) const
+		{
+			if (Package != Other.Package)
+				return Package < Other.Package;
+			return Object < Other.Object;
+		}
+		bool operator==(const FPackageObjectPair& Other) const
+		{
+			return Object == Other.Object;
+		}
+	};
+
+	PackageToObjectListBuffer.Empty();
+	UPackage::SoftGCPackageToObjectList.Empty();
+
+	// Iterate over all UObjects in memory (in parallel) and for each valid public object, get its package and add a FPackageObjectPair for it
+	int32 MaxNumberOfObjects = GUObjectArray.GetObjectArrayNum();
+	int32 NumThreads = FMath::Clamp(FTaskGraphInterface::Get().GetNumWorkerThreads(), 1, MaxNumberOfObjects);
+	int32 NumberOfObjectsPerThread = (MaxNumberOfObjects + NumThreads - 1) / NumThreads; // ceiling
+	check(NumberOfObjectsPerThread * (NumThreads - 1) <= MaxNumberOfObjects);
+
+	TArray<TArray<FPackageObjectPair>> ThreadContexts;
+	ThreadContexts.SetNum(NumThreads);
+	std::atomic<int32> PackagesNum{ 0 };
+
+	ParallelFor(TEXT("ConstructSoftGCPackageToObjectList"), NumThreads, 1,
+		[&ThreadContexts, &PackagesNum, NumberOfObjectsPerThread, NumThreads, MaxNumberOfObjects](int32 ThreadIndex)
+		{
+			TArray<FPackageObjectPair>& ThreadContext = ThreadContexts[ThreadIndex];
+			int32 FirstObjectIndex = ThreadIndex * NumberOfObjectsPerThread;
+			int32 NumObjects = (ThreadIndex < (NumThreads - 1)) ? NumberOfObjectsPerThread : (MaxNumberOfObjects - (NumThreads - 1) * NumberOfObjectsPerThread);
+			check(FirstObjectIndex + NumObjects <= MaxNumberOfObjects);
+
+			for (int32 ObjectIndex = 0; ObjectIndex < NumObjects && (FirstObjectIndex + ObjectIndex) < MaxNumberOfObjects; ++ObjectIndex)
+			{
+				FUObjectItem& ObjectItem = GUObjectArray.GetObjectItemArrayUnsafe()[FirstObjectIndex + ObjectIndex];
+				if (!ObjectItem.Object)
+				{
+					continue;
+				}
+				if (ObjectItem.IsGarbage())
+				{
+					continue;
+				}
+				UObject* Object = static_cast<UObject*>(ObjectItem.Object);
+				if (!Object->HasAnyFlags(RF_Public))
+				{
+					continue;
+				}
+				UPackage* Package = Object->GetPackage();
+				if (!Package)
+				{
+					continue;
+				}
+				if (Package->HasAnyFlags(RF_Transient) || Package->HasAnyPackageFlags(PKG_CompiledIn))
+				{
+					// Skip any Transient packages (e.g. /Engine/Transient) and script packages
+					// We only need to keep public objects alive in packages that could be saved.
+					continue;
+				}
+				if (Object == Package)
+				{
+					PackagesNum.fetch_add(1, std::memory_order_relaxed);
+				}
+				ThreadContext.Add(FPackageObjectPair{ Package, Object });
+			}
+		});
+
+	// Accumulate results from the parallel threads into a single array
+	TArray<FPackageObjectPair> PackageObjectPairs = MoveTemp(ThreadContexts[0]);
+	int32 PackageObjectPairsNum = PackageObjectPairs.Num();
+	TArrayView<TArray<FPackageObjectPair>> RemainingThreadContexts = TArrayView<TArray<FPackageObjectPair>>(ThreadContexts).RightChop(1);
+	for (TArray<FPackageObjectPair>& ThreadContext : RemainingThreadContexts)
+	{
+		PackageObjectPairsNum += ThreadContext.Num();
+	}
+	PackageObjectPairs.Reserve(PackageObjectPairsNum);
+	for (TArray<FPackageObjectPair>& ThreadContext : RemainingThreadContexts)
+	{
+		PackageObjectPairs.Append(ThreadContext);
+	}
+	ThreadContexts.Empty();
+
+	// Sort the array so that all objects for each package are together
+	PackageObjectPairs.Sort();
+
+	// Pull the UObject* out of the array of Pairs into a separate array of just UObject*,
+	// and for each UPackage, add the ArrayView of UObjects matching that package into the UPackage::SoftGCPackageToObjectList.
+	PackageToObjectListBuffer.SetNum(PackageObjectPairsNum);
+	UObject** PackageToObjectListBufferPtr = PackageToObjectListBuffer.GetData();
+
+	UPackage::SoftGCPackageToObjectList.Reserve(PackagesNum);
+	int32 PreviousPackageStartIndex = 0;
+	UPackage* PreviousPackage = nullptr;
+	for (int32 Index = 0; Index < PackageObjectPairsNum; ++Index)
+	{
+		FPackageObjectPair& Pair = PackageObjectPairs[Index];
+		if (Pair.Package != PreviousPackage)
+		{
+			if (Index > PreviousPackageStartIndex)
+			{
+				UPackage::SoftGCPackageToObjectList.Add(PreviousPackage,
+					ObjectPtrWrap(TArrayView<UObject*>(PackageToObjectListBufferPtr + PreviousPackageStartIndex, Index - PreviousPackageStartIndex)));
+			}
+			PreviousPackage = Pair.Package;
+			PreviousPackageStartIndex = Index;
+		}
+		PackageToObjectListBufferPtr[Index] = Pair.Object;
+	}
+	if (PackageObjectPairsNum > PreviousPackageStartIndex)
+	{
+		UPackage::SoftGCPackageToObjectList.Add(PreviousPackage,
+			ObjectPtrWrap(TArrayView<UObject*>(PackageToObjectListBufferPtr + PreviousPackageStartIndex, PackageObjectPairsNum - PreviousPackageStartIndex)));
+	}
+}
+
+} // namespace UE::Cook
+
+
+void UCookOnTheFlyServer::PollGarbageCollection(UE::Cook::FTickStackData& StackData)
+{
+	NumObjectsHistory.AddInstance(GUObjectArray.GetObjectArrayNumMinusAvailable());
+	VirtualMemoryHistory.AddInstance(FPlatformMemory::GetStats().UsedVirtual);
+
+	if (IsCookFlagSet(ECookInitializationFlags::TestCook))
+	{
+		StackData.ResultFlags |= COSR_RequiresGC | COSR_YieldTick;
+		return;
+	}
+	if (PackagesPerGC > 0 && CookedPackageCountSinceLastGC > PackagesPerGC)
+	{
+		// if we are waiting on things to cache then ignore the PackagesPerGC
+		if (!bSaveBusy)
+		{
+			StackData.ResultFlags |= COSR_RequiresGC | COSR_RequiresGC_PackageCount | COSR_YieldTick;
+			return;
+		}
+	}
+	if (IsCookOnTheFlyMode())
+	{
+		double CurrentTime = FPlatformTime::Seconds();
+		if (IdleStatus == EIdleStatus::Done &&
+			CurrentTime - IdleStatusStartTime > GetIdleTimeToGC() &&
+			IdleStatusStartTime > GetLastGCTime())
+		{
+			StackData.ResultFlags |= COSR_RequiresGC | COSR_RequiresGC_IdleTimer | COSR_YieldTick;
+			return;
+		}
+	}
+}
+
+bool UCookOnTheFlyServer::PumpHasExceededMaxMemory(uint32& OutResultFlags)
+{
+	if (GUObjectArray.GetObjectArrayEstimatedAvailable() < MinFreeUObjectIndicesBeforeGC)
+	{
+		UE_LOG(LogCook, Display, TEXT("Running out of available UObject indices (%d remaining)"), GUObjectArray.GetObjectArrayEstimatedAvailable());
+		static bool bPerformedObjListWhenNearMaxObjects = false;
+		if (GEngine && !bPerformedObjListWhenNearMaxObjects)
+		{
+			UE_LOG(LogCook, Display, TEXT("Performing 'obj list' to show counts of types of objects due to low availability of UObject indices."));
+			GEngine->Exec(nullptr, TEXT("OBJ LIST -COUNTSORT -SKIPMEMORYSIZE"));
+			bPerformedObjListWhenNearMaxObjects = true;
+		}
+		OutResultFlags |= COSR_RequiresGC | COSR_RequiresGC_OOM | COSR_YieldTick;
+		return true;
+	}
+
+	TStringBuilder<256> TriggerMessages;
+	const FPlatformMemoryStats MemStats = FPlatformMemory::GetStats();
+	
+	bool bMinFreeTriggered = false;
+	if (MemoryMinFreeVirtual > 0 || MemoryMinFreePhysical > 0)
+	{
+		// trigger GC if we have less than MemoryMinFreeVirtual OR MemoryMinFreePhysical
+		// the check done in AssetCompilingManager is against the min of the two :
+		//uint64 AvailableMemory = FMath::Min(MemStats.AvailablePhysical, MemStats.AvailableVirtual);
+		// so for consistency the same check should be done here
+		// you can get that by setting the MemoryMinFreeVirtual and MemoryMinFreePhysical config to be the same
+
+		// AvailableVirtual is actually ullAvailPageFile (commit charge available)
+		if (MemoryMinFreeVirtual > 0 && MemStats.AvailableVirtual < MemoryMinFreeVirtual)
+		{
+			TriggerMessages.Appendf(TEXT("\n  CookSettings.MemoryMinFreeVirtual: Available virtual memory %dMiB is less than %dMiB."),
+				static_cast<uint32>(MemStats.AvailableVirtual / 1024 / 1024), static_cast<uint32>(MemoryMinFreeVirtual / 1024 / 1024));
+			bMinFreeTriggered = true;
+		}
+		if (MemoryMinFreePhysical > 0 && MemStats.AvailablePhysical < MemoryMinFreePhysical)
+		{
+			TriggerMessages.Appendf(TEXT("\n  CookSettings.MemoryMinFreePhysical: Available physical memory %dMiB is less than %dMiB."),
+				static_cast<uint32>(MemStats.AvailablePhysical / 1024 / 1024), static_cast<uint32>(MemoryMinFreePhysical / 1024 / 1024));
+			bMinFreeTriggered = true;
+		}
+	}
+
+	// if MemoryMaxUsed is set, we won't GC until at least that much mem is used
+	// this can be useful if you demand that amount of memory as your min spec
+	bool bMaxUsedTriggered = false;
+	if (MemoryMaxUsedVirtual > 0 || MemoryMaxUsedPhysical > 0)
+	{
+		// check validity of trigger :
+		// if the MaxUsed config exceeds the system memory, it can never be triggered and will prevent any GC :
+		uint64 MaxMaxUsed = FMath::Max(MemoryMaxUsedVirtual,MemoryMaxUsedPhysical);
+		if (MaxMaxUsed >= MemStats.TotalPhysical)
+		{
+			UE_CALL_ONCE([&]() {
+				UE_LOG(LogCook, Warning, TEXT("Warning MemoryMaxUsed condition is larger than total memory (%dMiB >= %dMiB).  System does not have enough memory to cook this project."),
+				static_cast<uint32>(MaxMaxUsed / 1024 / 1024), static_cast<uint32>(MemStats.TotalPhysical / 1024 / 1024));
+				});
+		}
+
+		if (MemoryMaxUsedVirtual > 0 && MemStats.UsedVirtual >= MemoryMaxUsedVirtual)
+		{
+			TriggerMessages.Appendf(TEXT("\n  CookSettings.MemoryMaxUsedVirtual: Used virtual memory %dMiB is greater than %dMiB."),
+				static_cast<uint32>(MemStats.UsedVirtual / 1024 / 1024), static_cast<uint32>(MemoryMaxUsedVirtual / 1024 / 1024));
+			bMaxUsedTriggered = true;
+		}
+		if (MemoryMaxUsedPhysical > 0 && MemStats.UsedPhysical >= MemoryMaxUsedPhysical)
+		{
+			TriggerMessages.Appendf(TEXT("\n  CookSettings.MemoryMaxUsedPhysical: Used physical memory %dMiB is greater than %dMiB."),
+				static_cast<uint32>(MemStats.UsedPhysical / 1024 / 1024), static_cast<uint32>(MemoryMaxUsedPhysical / 1024 / 1024));
+			bMaxUsedTriggered = true;
+		}
+	}
+
+	bool bPressureTriggered = false;
+	if (MemoryTriggerGCAtPressureLevel != FPlatformMemoryStats::EMemoryPressureStatus::Unknown)
+	{
+		FPlatformMemoryStats::EMemoryPressureStatus PressureStatus = MemStats.GetMemoryPressureStatus();
+		if (PressureStatus == FPlatformMemoryStats::EMemoryPressureStatus::Unknown)
+		{
+			UE_CALL_ONCE([&]() {
+				UE_LOG(LogCook, Warning,
+				TEXT("MemoryPressureStatus is not available from the operating system. We may run out of memory due to lack of knowledge of when to collect garbage."));
+				});
+		}
+		else
+		{
+			static_assert(FPlatformMemoryStats::EMemoryPressureStatus::Critical > FPlatformMemoryStats::EMemoryPressureStatus::Nominal,
+				"We expect higher pressure to be higher integer values");
+			int RequiredValue = static_cast<int>(MemoryTriggerGCAtPressureLevel);
+			int CurrentValue = static_cast<int>(PressureStatus);
+			if (CurrentValue >= RequiredValue)
+			{
+				bPressureTriggered = true;
+				TriggerMessages.Appendf(TEXT("\n  Operating system has signalled that memory pressure is high."));
+			}
+		}
+	}
+
+	bool bTriggerGC = false;
+	if (bMinFreeTriggered || bMaxUsedTriggered)
+	{
+		const bool bOnlyTriggerIfBothMinFreeAndMaxUsedTrigger = true;
+		if (!bOnlyTriggerIfBothMinFreeAndMaxUsedTrigger ||
+			((bMinFreeTriggered || (MemoryMinFreeVirtual <= 0 && MemoryMinFreePhysical <= 0)) &&
+				(bMaxUsedTriggered || (MemoryMaxUsedVirtual <= 0 && MemoryMaxUsedPhysical <= 0))))
+		{
+			bTriggerGC = true;
+		}
+	}
+	if (bPressureTriggered)
+	{
+		bTriggerGC = true;
+	}
+
+	// If a normal GC was not triggered, check the SoftGC trigger conditions
+	double CurrentTime = FPlatformTime::Seconds();
+	bool bIsSoftGC = false;
+	if (bUseSoftGC && IsDirectorCookByTheBook() && !IsCookingInEditor())
+	{
+		if (!bTriggerGC && SoftGCStartNumerator > 0)
+		{
+			if (SoftGCNextAvailablePhysicalTarget == -1) // Uninitialized
+			{
+				int32 StartNumerator = FMath::Max(SoftGCStartNumerator, 1);
+				int32 Denominator = FMath::Max(SoftGCDenominator, 1);
+				// e.g. Start the target at 5/10, and decrease it by 1/10 each time the target is reached
+				SoftGCNextAvailablePhysicalTarget = (static_cast<int64>(MemStats.TotalPhysical) * StartNumerator)
+					/ Denominator;
+			}
+
+			if (SoftGCNextAvailablePhysicalTarget < -1)
+			{
+				// No further targets, no further SoftGC
+			}
+			else if (static_cast<int64>(MemStats.AvailablePhysical) <= SoftGCNextAvailablePhysicalTarget)
+			{
+				constexpr float SoftGCInstigateCooldown = 5 * 60.f;
+				CurrentTime = FPlatformTime::Seconds();
+				if (LastSoftGCTime + SoftGCInstigateCooldown <= CurrentTime)
+				{
+					TriggerMessages.Appendf(TEXT("\n  CookSettings.SoftGCMemoryTrigger: Available physical memory %dMiB is less than the current target for SoftGC %dMiB."),
+						static_cast<uint32>(MemStats.AvailablePhysical / 1024 / 1024), static_cast<uint32>(SoftGCNextAvailablePhysicalTarget / 1024 / 1024));
+					bTriggerGC = true;
+					bIsSoftGC = true;
+				}
+			}
+		}
+
+		if (!bTriggerGC && SoftGCTimeFractionBudget > 0.)
+		{
+			FString TriggerDiagnostics;
+			if (SoftGCHistory->IsTriggeringWithinBudget(*this, CurrentTime, &TriggerDiagnostics))
+			{
+				TriggerMessages.Appendf(TEXT("\n  CookSettings.SoftGCTimeTrigger: Periodic triggering of SoftGC: %s."),
+					*TriggerDiagnostics);
+				bTriggerGC = true;
+				bIsSoftGC = true;
+			}
+		}
+	}
+
+	if (!bTriggerGC)
+	{
+		return false;
+	}
+
+	// Don't allow a second GC (soft or normal) within the GC cooldown period after a full GC
+	constexpr float GCCooldown = 60.f;
+	if (LastFullGCTime + GCCooldown > CurrentTime)
+	{
+		if (!bIsSoftGC && !bWarnedExceededMaxMemoryWithinGCCooldown)
+		{
+			bWarnedExceededMaxMemoryWithinGCCooldown = true;
+			// If we are in a cooldown period, return false.
+			UE_LOG(LogCook, Display, TEXT("Garbage collection triggers ignored: Out of memory condition has been detected, but is only %.0fs after the last GC. ")
+				TEXT("It will be prevented until %.0f seconds have passed and we may run out of memory.\n")
+				TEXT("Garbage collection triggered by conditions: %s"),
+				static_cast<float>(CurrentTime - LastFullGCTime), GCCooldown, TriggerMessages.ToString());
+		}
+		return false;
+	}
+
+	const TCHAR* TypeMessage = bIsSoftGC ? TEXT("Soft") : (IsCookFlagSet(ECookInitializationFlags::EnablePartialGC) ? TEXT("Partial") : TEXT("Full"));
+
+	UE_LOG(LogCook, Display, TEXT("Garbage collection triggered (%s). Triggered by conditions:%s"),
+		TypeMessage, TriggerMessages.ToString());
+	OutResultFlags |= COSR_RequiresGC | COSR_YieldTick;
+	OutResultFlags |= COSR_RequiresGC_OOM;
+	if (bIsSoftGC)
+	{
+		OutResultFlags |= COSR_RequiresGC_Soft_OOM;
+	}
+	return true;
+}
+
+void UCookOnTheFlyServer::SetGarbageCollectType(uint32 ResultFlagsFromTick)
+{
+	bGarbageCollectTypeSoft = (ResultFlagsFromTick & COSR_RequiresGC_Soft_OOM);
+}
+
+void UCookOnTheFlyServer::ClearGarbageCollectType()
+{
+	bGarbageCollectTypeSoft = false;
+}
+
+void UCookOnTheFlyServer::PreGarbageCollect()
+{
+	using namespace UE::Cook;
+
+	if (!IsInSession())
+	{
+		return;
+	}
+
+	NumObjectsHistory.AddInstance(GUObjectArray.GetObjectArrayNumMinusAvailable());
+	VirtualMemoryHistory.AddInstance(FPlatformMemory::GetStats().UsedVirtual);
+	TArray<UPackage*> GCKeepPackages;
+	TArray<UE::Cook::FPackageData*> GCKeepPackageDatas;
+
+#if COOK_CHECKSLOW_PACKAGEDATA
+	// Verify that only packages in the saving states have pointers to objects
+	PackageDatas->LockAndEnumeratePackageDatas([](const FPackageData* PackageData)
+		{
+			check(PackageData->IsInStateProperty(EPackageStateProperty::Saving) || !PackageData->HasReferencedObjects());
+		});
+#endif
+	if (SavingPackageData)
+	{
+		check(SavingPackageData->GetPackage());
+		GCKeepObjects.Add(SavingPackageData->GetPackage());
+		GCKeepPackageDatas.Add(SavingPackageData);
+	}
+
+	// Notify every FGenerationHelper of the garbage collect
+	PackageDatas->LockAndEnumeratePackageDatas([this, &GCKeepPackages, &GCKeepPackageDatas](FPackageData* PackageData)
+		{
+			TRefCountPtr<FGenerationHelper> GenerationHelper = PackageData->GetGenerationHelper();
+			if (!GenerationHelper)
+			{
+				GenerationHelper = PackageData->GetParentGenerationHelper();
+			}
+			if (GenerationHelper)
+			{
+				bool bShouldDemote;
+				GenerationHelper->PreGarbageCollect(GenerationHelper, *PackageData, GCKeepObjects, GCKeepPackages,
+					GCKeepPackageDatas, bShouldDemote);
+				if (bShouldDemote && PackageData->IsInStateProperty(EPackageStateProperty::Saving))
+				{
+					// Demote any Generated/Generator packages we called PreSave on so they call their PostSave before the GC
+					// or prevent them from being garbage collected if the splitter wants to keep them referenced
+					ReleaseCookedPlatformData(*PackageData, UE::Cook::EStateChangeReason::GeneratorPreGarbageCollected,
+						EPackageState::Request);
+				}
+			}
+			if (PackageData->GetIsCookLast() && PackageData->IsInStateProperty(EPackageStateProperty::Saving))
+			{
+				GCKeepPackages.Add(PackageData->GetPackage());
+				GCKeepPackageDatas.Add(PackageData);
+			}
+		});
+	
+	// Find the packages that are waiting on async jobs to finish cooking data
+	// and make sure that they are not garbage collected until the jobs have
+	// completed.
+	{
+		TMap<FPackageData*, UPackage*> UniquePendingPackages;
+		PackageDatas->ForEachPendingCookedPlatformData(
+		[&UniquePendingPackages](const FPendingCookedPlatformData& PendingData)
+		{
+			if (UObject* Object = PendingData.Object.Get())
+			{	
+				if (UPackage* Package = Object->GetPackage())
+				{
+					UniquePendingPackages.Add(&PendingData.PackageData, Package);
+				}	
+			}
+		});
+
+		GCKeepPackages.Reserve(GCKeepPackages.Num() + UniquePendingPackages.Num());
+		for (const TPair<FPackageData*,UPackage*>& Pair : UniquePendingPackages)
+		{
+			GCKeepPackages.Add(Pair.Value);
+			GCKeepPackageDatas.Add(Pair.Key);
+		}
+	}
+
+	// Prevent GC of any objects on which we are still waiting for IsCachedCookedPlatformData
+	PackageDatas->ForEachPendingCookedPlatformData(
+	[this](UE::Cook::FPendingCookedPlatformData& Pending)
+	{
+		if (!Pending.PollIsComplete())
+		{
+			UObject* Object = Pending.Object.Get();
+			check(Object); // Otherwise PollIsComplete would have returned true
+			GCKeepObjects.Add(Object);
+		}
+	});
+
+	const bool bPartialGC = IsCookFlagSet(ECookInitializationFlags::EnablePartialGC);
+	if (bGarbageCollectTypeSoft || bPartialGC)
+	{
+		// Keep referenced all packages in requestqueue, loadqueue, and savequeue, and any packages they depend on
+		TArray<FName> Queue;
+		TSet<FName> Visited;
+		auto AddPackageName = [&Visited, &Queue](FName PackageName)
+		{
+			bool bAlreadyExists;
+			Visited.Add(PackageName, &bAlreadyExists);
+			if (!bAlreadyExists)
+			{
+				Queue.Add(PackageName);
+			}
+		};
+		for (FPackageData* PackageData : PackageDatas->GetRequestQueue().GetReadyRequestsUrgent())
+		{
+			AddPackageName(PackageData->GetPackageName());
+		}
+		for (FPackageData* PackageData : PackageDatas->GetRequestQueue().GetReadyRequestsNormal())
+		{
+			AddPackageName(PackageData->GetPackageName());
+		}
+		for (FPackageData* PackageData : PackageDatas->GetLoadQueue())
+		{
+			AddPackageName(PackageData->GetPackageName());
+		}
+		for (FPackageData* PackageData : PackageDatas->GetSaveQueue())
+		{
+			AddPackageName(PackageData->GetPackageName());
+		}
+		for (FPackageData* PackageData : PackageDatas->GetSaveStalledSet())
+		{
+			AddPackageName(PackageData->GetPackageName());
+		}
+
+		TArray<FName> Dependencies;
+		while (!Queue.IsEmpty())
+		{
+			FName PackageName = Queue.Pop();
+			Dependencies.Reset();
+			AssetRegistry->GetDependencies(PackageName, Dependencies, UE::AssetRegistry::EDependencyCategory::Package,
+				UE::AssetRegistry::EDependencyQuery::Hard);
+			for (FName DependencyName : Dependencies)
+			{
+				AddPackageName(DependencyName);
+			};
+		}
+
+		TSet<UPackage*> GCKeepPackagesSet;
+		GCKeepPackagesSet.Append(GCKeepPackages);
+		for (FName PackageName : Visited)
+		{
+			UPackage* Package = FindPackage(nullptr, *WriteToString<256>(PackageName));
+			if (Package)
+			{
+				bool bAlreadyInSet;
+				GCKeepPackagesSet.Add(Package, &bAlreadyInSet);
+				if (!bAlreadyInSet)
+				{
+					GCKeepPackages.Add(Package);
+					FPackageData* PackageData = PackageDatas->FindPackageDataByPackageName(Package->GetFName());
+					if (PackageData)
+					{
+						GCKeepPackageDatas.Add(PackageData);
+					}
+				}
+			}
+		}
+		ExpectedFreedPackageNames.Empty(PackageTracker->NumLoadedPackages());
+		PackageTracker->ForEachLoadedPackage(
+			[this, &GCKeepPackagesSet](UPackage* Package)
+			{
+				if (!GCKeepPackagesSet.Contains(Package))
+				{
+					ExpectedFreedPackageNames.Add(Package->GetFName());
+				}
+			});
+	}
+
+	// Add packages to GCKeepObjects. 
+	TArray<UObject*> ObjectsWithOuter;
+	for (UPackage* Package : GCKeepPackages)
+	{
+		GCKeepObjects.Add(Package);
+	}
+	for (FPackageData* PackageData : GCKeepPackageDatas)
+	{
+		PackageData->SetKeepReferencedDuringGC(true);
+	}
+
+	// Add all public objects within every package in memory to the UPackage::SoftGCPackageToObjectList container,
+	// so they will be kept in memory if the package is kept in memory.
+	ConstructSoftGCPackageToObjectList(this->SoftGCPackageToObjectListBuffer);
+}
+
+void UCookOnTheFlyServer::CookerAddReferencedObjects(FReferenceCollector& Collector)
+{
+	using namespace UE::Cook;
+
+	// GCKeepObjects are the objects that we want to keep loaded but we only have a WeakPtr to
+	Collector.AddReferencedObjects(GCKeepObjects);
+}
+
+void UCookOnTheFlyServer::PostGarbageCollect()
+{
+	using namespace UE::Cook;
+
+	NumObjectsHistory.AddInstance(GUObjectArray.GetObjectArrayNumMinusAvailable());
+	VirtualMemoryHistory.AddInstance(FPlatformMemory::GetStats().UsedVirtual);
+
+	TSet<UObject*> SaveQueueObjectsThatStillExist;
+
+	// If garbage collection deleted a UPackage WHILE WE WERE SAVING IT, then we have problems.
+	check(!SavingPackageData || SavingPackageData->GetPackage() != nullptr);
+
+	// If there was a GarbageCollect after we already started calling BeginCacheCookedPlatformData, then we
+	// have a list of the WeakObjectPtr to all objects in the package (FPackageData::CachedObjejectsInOuter)
+	// and some of those objects may have been set to null. We declare a reference to prevent GC for the RF_Public
+	// objects in that list, but we do not declare that reference for private objects. The private objects may
+	// therefore have been deleted and set to null 
+	// Side note: because objects can be marked as pending kill at any time and we use FObjectWeakPtr.Get(),
+	// which returns null if pending kill, we need to skip nulls in the array at any point, not just after GC.
+	// 
+	// We do not want to prevent GC of private objects in case there is the expectation by some
+	// systems (blueprints, licensee code) that removing references to an object during PreCollectGarbage will cause
+	// it to be deleted by GC and be replaceable afterwards. We add any new private objects after the garbage collect
+	// and continue with the save. Public objects have a different contract; they are not replaceable across a
+	// GC because anything outside the package could be referring to them. So we keep them referenced. But GC may
+	// force delete them despite our reference, and the package is then in an unknown state. If that happens we
+	// demote the package back to request and start its load and save over.
+	TArray<FPackageData*> Demotes;
+	auto UpdateSavingPackageAfterGarbageCollect =
+		[&Demotes, &SaveQueueObjectsThatStillExist](FPackageData* PackageData)
+	{
+		bool bOutDemote;
+		PackageData->UpdateSaveAfterGarbageCollect(bOutDemote);
+		if (bOutDemote)
+		{
+			Demotes.Add(PackageData);
+		}
+		else
+		{
+			// Mark that the objects for this package should be kept in CachedCookedPlatformData records
+			for (FCachedObjectInOuter& CachedObjectInOuter : PackageData->GetCachedObjectsInOuter())
+			{
+				UObject* Object = CachedObjectInOuter.Object.Get();
+				if (Object)
+				{
+					SaveQueueObjectsThatStillExist.Add(Object);
+				}
+			}
+		}
+	};
+	for (FPackageData* PackageData : PackageDatas->GetSaveQueue())
+	{
+		UpdateSavingPackageAfterGarbageCollect(PackageData);
+	}
+	for (FPackageData* PackageData : PackageDatas->GetSaveStalledSet())
+	{
+		UpdateSavingPackageAfterGarbageCollect(PackageData);
+	}
+	for (FPackageData* PackageData : Demotes)
+	{
+		FGenerationHelper::ValidateSaveStalledState(*this, *PackageData, TEXT("PostGarbageCollect"));
+
+		switch (PackageData->GetState())
+		{
+		case EPackageState::SaveActive:
+			PackageData->SendToState(EPackageState::Request, ESendFlags::QueueRemove, EStateChangeReason::GarbageCollected);
+			if (PackageData->GetIsCookLast())
+			{
+				// CookLast packages in SaveState have had their urgency removed. Add it back if we need to demote them.
+				PackageData->SetUrgency(EUrgency::Blocking, ESendFlags::QueueNone);
+			}
+			PackageDatas->GetRequestQueue().AddRequest(PackageData, /* bForceUrgent */ true);
+			break;
+		case EPackageState::SaveStalledAssignedToWorker:
+			PackageData->SendToState(EPackageState::AssignedToWorker, ESendFlags::QueueAddAndRemove, EStateChangeReason::GarbageCollected);
+			break;
+		case EPackageState::SaveStalledRetracted:
+			DemoteToIdle(*PackageData, ESendFlags::QueueAddAndRemove, ESuppressCookReason::RetractedByCookDirector);
+			break;
+		default:
+			checkf(false, TEXT("State %s not handled in a demoted package."), LexToString(PackageData->GetState()));
+			break;
+		}
+	}
+
+	// Mark that any objects in PendingCookedPlatformDatas should be kept in CachedCookedPlatformData records
+	PackageDatas->ForEachPendingCookedPlatformData(
+		[&SaveQueueObjectsThatStillExist](FPendingCookedPlatformData& CookedPlatformData)
+		{
+			UObject* Object = CookedPlatformData.Object.Get();
+			if (Object)
+			{
+				SaveQueueObjectsThatStillExist.Add(Object);
+			}
+			else
+			{
+				CookedPlatformData.Release();
+			}
+		});
+
+	// Remove objects that were deleted by garbage collection from our containers that track raw object pointers
+	PackageDatas->CachedCookedPlatformDataObjectsPostGarbageCollect(SaveQueueObjectsThatStillExist);
+
+	PackageDatas->LockAndEnumeratePackageDatas([this](FPackageData* PackageData)
+	{
+		if (TRefCountPtr<FGenerationHelper> GenerationHelper = PackageData->GetGenerationHelper())
+		{
+			GenerationHelper->PostGarbageCollect(GenerationHelper, *GCDiagnosticContext);
+		}
+	});
+
+	// Second pass over all PackageDatas, combine a few operations
+	PackageDatas->LockAndEnumeratePackageDatas([](FPackageData* PackageData)
+	{
+		// Mark that the PackageData no longer needs to be keepreferenced.
+		// This can only be done after all GenerationHelper->PostGarbageCollect have been called.
+		PackageData->SetKeepReferencedDuringGC(false);
+
+		// Reset the completion flags for FPreloadPackage, since the UPackage might be no longer loaded.
+		TRefCountPtr<FPackagePreloader> Preloader = PackageData->GetPackagePreloader();
+		if (Preloader)
+		{
+			Preloader->PostGarbageCollect();
+		}
+	});
+
+	// Only after running all possible callbacks that need our links for diagnostics, clear the list of temporary
+	// references that we created for the garbage collection.
+	GCKeepObjects.Empty();
+	UPackage::SoftGCPackageToObjectList.Empty();
+	SoftGCPackageToObjectListBuffer.Empty();
+
+	CookedPackageCountSinceLastGC = 0;
+
+	// Whenever we collect garbage, reset the counter for how many busy reports with an
+	// idle shadercompiler we need before we issue a warning
+	bShaderCompilerWasActiveeOnPreviousBusyReport = true;
+}
+
+bool UCookOnTheFlyServer::NeedsDiagnosticSecondGC() const
+{
+	return GCDiagnosticContext->NeedsDiagnosticSecondGC();
+}
+
+void UCookOnTheFlyServer::OnCookerStartCollectGarbage(uint32& ResultFlagsFromTick)
+{
+	GCDiagnosticContext->OnCookerStartCollectGarbage(*this, ResultFlagsFromTick);
+}
+
+void UCookOnTheFlyServer::OnCookerEndCollectGarbage(uint32& ResultFlagsFromTick)
+{
+	GCDiagnosticContext->OnCookerEndCollectGarbage(*this, ResultFlagsFromTick);
+}
+
+void UCookOnTheFlyServer::EvaluateGarbageCollectionResults(bool bWasDueToOOM, bool bWasPartialGC, uint32 ResultFlags,
+	int32 NumObjectsBeforeGC, const FPlatformMemoryStats& MemStatsBeforeGC,
+	const FGenericMemoryStats& AllocatorStatsBeforeGC,
+	int32 NumObjectsAfterGC, const FPlatformMemoryStats& MemStatsAfterGC,
+	const FGenericMemoryStats& AllocatorStatsAfterGC,
+	float GCDurationSeconds)
+{
+	using namespace UE::Cook;
+
+	ON_SCOPE_EXIT
+	{
+		ExpectedFreedPackageNames.Empty();
+		GCDiagnosticContext->OnEvaluateResultsComplete();
+	};
+	bWarnedExceededMaxMemoryWithinGCCooldown = false;
+	LastGCTime = FPlatformTime::Seconds();
+	bool bWasSoftGC = ResultFlags & COSR_RequiresGC_Soft_OOM;
+	if (bWasSoftGC)
+	{
+		LastSoftGCTime = LastGCTime;
+		if (SoftGCStartNumerator)
+		{
+			int32 StartNumerator = FMath::Max(SoftGCStartNumerator, 1);
+			int32 Denominator = FMath::Max(SoftGCDenominator, 1);
+			// Calculate the new SoftGCNextAvailablePhysicalTarget. Use the floor of NewAvailableMemory/Denominator,
+			// unless we are already 50% of the way through that level, in which case use the next value below that
+			int64 PhysicalMemoryQuantum = static_cast<int64>(MemStatsAfterGC.TotalPhysical) / Denominator;
+			int32 NextTarget =
+				static_cast<int64>(MemStatsAfterGC.AvailablePhysical - PhysicalMemoryQuantum / 2) / PhysicalMemoryQuantum;
+			NextTarget = FMath::Min(NextTarget, StartNumerator);
+			if (NextTarget <= 0)
+			{
+				SoftGCNextAvailablePhysicalTarget = -2; // disabled, no further targets
+			}
+			else
+			{
+				SoftGCNextAvailablePhysicalTarget = static_cast<int64>((MemStatsAfterGC.TotalPhysical) * NextTarget)
+					/ Denominator;
+			}
+		}
+	}
+	else
+	{
+		LastSoftGCTime = LastGCTime;
+		LastFullGCTime = LastGCTime;
+	}
+	if (SoftGCHistory)
+	{
+		SoftGCHistory->AddDurationMeasurement(GCDurationSeconds);
+	}
+
+	if (IsCookingInEditor())
+	{
+		return;
+	}
+	if (IsCookOnTheFlyMode() && !bWasDueToOOM)
+	{
+		return;
+	}
+
+	int64 NumObjectsMin = NumObjectsHistory.GetMinimum();
+	int64 NumObjectsMax = NumObjectsHistory.GetMaximum();
+	int64 NumObjectsSpread = NumObjectsMax - NumObjectsMin;
+	int64 NumObjectsFreed = NumObjectsBeforeGC - NumObjectsAfterGC;
+	int64 NumObjectsCapacity = GUObjectArray.GetObjectArrayEstimatedAvailable() + GUObjectArray.GetObjectArrayNumMinusAvailable();
+	int64 VirtualMemMin = VirtualMemoryHistory.GetMinimum();
+	int64 VirtualMemMax = VirtualMemoryHistory.GetMaximum();
+	int64 VirtualMemSpread = VirtualMemMax - VirtualMemMin;
+	int64 VirtualMemBeforeGC = MemStatsBeforeGC.UsedVirtual;
+	int64 VirtualMemAfterGC = MemStatsAfterGC.UsedVirtual;
+	int64 VirtualMemFreed = MemStatsBeforeGC.UsedVirtual - MemStatsAfterGC.UsedVirtual;
+
+	int64 ExpectedObjectsFreed = static_cast<int64>(MemoryExpectedFreedToSpreadRatio * static_cast<float>(NumObjectsSpread));
+	double ExpectedMemFreed = MemoryExpectedFreedToSpreadRatio * VirtualMemSpread;
+	static bool bCookMemoryAnalysis = FParse::Param(FCommandLine::Get(), TEXT("CookMemoryAnalysis"));
+#if ENABLE_LOW_LEVEL_MEM_TRACKER
+	// When tracking memory with LLM, always show the memory status.
+	const bool bAlwaysShowAnalysis = FLowLevelMemTracker::Get().IsEnabled() || bCookMemoryAnalysis;
+#else
+	const bool bAlwaysShowAnalysis = bCookMemoryAnalysis;
+#endif		
+	constexpr int32 BytesPerMeg = 1000000;
+	auto DisplaySimpleSummary = [&]()
+	{
+		UE_LOG(LogCook, Display, TEXT("GarbageCollection Results:\n")
+			TEXT("\tType: %s\n")
+			TEXT("\tDuration: %.3fs\n")
+			TEXT("\tNumObjects:\n")
+			TEXT("\t\tCapacity:         %10d\n")
+			TEXT("\t\tBefore GC:        %10d\n")
+			TEXT("\t\tAfter GC:         %10d\n")
+			TEXT("\t\tFreed by GC:      %10d\n")
+			TEXT("\tVirtual Memory:\n")
+			TEXT("\t\tBefore GC:        %10" INT64_FMT " MB\n")
+			TEXT("\t\tAfter GC:         %10" INT64_FMT " MB\n")
+			TEXT("\t\tFreed by GC:      %10" INT64_FMT " MB"),
+			(bWasSoftGC ? TEXT("Soft") : (bWasPartialGC ? TEXT("Partial") : TEXT("Full"))),
+			GCDurationSeconds,
+			NumObjectsCapacity, (int64)NumObjectsBeforeGC, (int64)NumObjectsAfterGC, NumObjectsFreed,
+			VirtualMemBeforeGC / BytesPerMeg, VirtualMemAfterGC / BytesPerMeg, VirtualMemFreed / BytesPerMeg
+		);
+	};
+
+	if (!bWasSoftGC)
+	{
+		bool bWasImpactful =
+			(NumObjectsFreed >= ExpectedObjectsFreed || NumObjectsBeforeGC - NumObjectsMin < ExpectedObjectsFreed) &&
+			(VirtualMemFreed >= ExpectedMemFreed || VirtualMemBeforeGC - VirtualMemMin <= ExpectedMemFreed);
+
+		if ((!bWasDueToOOM || bWasImpactful) && !bAlwaysShowAnalysis)
+		{
+			DisplaySimpleSummary();
+			return;
+		}
+
+		if (bWasDueToOOM && !bWasImpactful)
+		{
+			UE_LOG(LogCook, Display, TEXT("GarbageCollection Results: Garbage Collection was not very impactful."));
+		}
+		else
+		{
+			UE_LOG(LogCook, Display, TEXT("GarbageCollection Results:"));
+		}
+		UE_LOG(LogCook, Display, TEXT("\tMemoryAnalysis: General:\n")
+			TEXT("\t\tType: %s\n")
+			TEXT("\tDuration: %.3fs"),
+			(bWasSoftGC ? TEXT("Soft") : (bWasPartialGC ? TEXT("Partial") : TEXT("Full"))),
+			GCDurationSeconds);
+		UE_LOG(LogCook, Display, TEXT("\tMemoryAnalysis: NumObjects:\n")
+			TEXT("\t\tCapacity:         %10" INT64_FMT "\n")
+			TEXT("\t\tProcess Min:      %10" INT64_FMT "\n")
+			TEXT("\t\tProcess Max:      %10" INT64_FMT "\n")
+			TEXT("\t\tProcess Spread:   %10" INT64_FMT "\n")
+			TEXT("\t\tBefore GC:        %10" INT64_FMT "\n")
+			TEXT("\t\tAfter GC:         %10" INT64_FMT "\n")
+			TEXT("\t\tFreed by GC:      %10" INT64_FMT ""),
+			NumObjectsCapacity, NumObjectsMin, NumObjectsMax, NumObjectsSpread,
+			(int64)NumObjectsBeforeGC, (int64)NumObjectsAfterGC, NumObjectsFreed);
+		UE_LOG(LogCook, Display, TEXT("\tMemoryAnalysis: Virtual Memory:\n")
+			TEXT("\t\tProcess Min:      %10" INT64_FMT " MB\n")
+			TEXT("\t\tProcess Max:      %10" INT64_FMT " MB\n")
+			TEXT("\t\tProcess Spread:   %10" INT64_FMT " MB\n")
+			TEXT("\t\tBefore GC:        %10" INT64_FMT " MB\n")
+			TEXT("\t\tAfter GC:         %10" INT64_FMT " MB\n")
+			TEXT("\t\tFreed by GC:      %10" INT64_FMT " MB"),
+			VirtualMemMin / BytesPerMeg, VirtualMemMax / BytesPerMeg, VirtualMemSpread / BytesPerMeg,
+			VirtualMemBeforeGC / BytesPerMeg, VirtualMemAfterGC / BytesPerMeg, VirtualMemFreed / BytesPerMeg);
+		auto AllocatorStatsToString = [](const FGenericMemoryStats& AllocatorStats)
+		{
+			TStringBuilder<256> Writer;
+			for (const TPair<FStringView, SIZE_T>& Item : AllocatorStats)
+			{
+				Writer << TEXT("\n\t\tItem ") << Item.Key << TEXT(" ") << (uint64)Item.Value;
+			}
+			return FString(*Writer);
+		};
+		UE_LOG(LogCook, Display, TEXT("\tMemoryAnalysis: Allocator Stats Before:%s"),
+			*AllocatorStatsToString(AllocatorStatsBeforeGC));
+		UE_LOG(LogCook, Display, TEXT("\tMemoryAnalysis: Allocator Stats After:%s"),
+			*AllocatorStatsToString(AllocatorStatsAfterGC));
+
+
+		UE_LOG(LogCook, Display, TEXT("See log for memory use information for UObject classes and LLM tags."));
+
+		{
+			TGuardValue<bool> SoftGCGuard(UPackage::bSupportCookerSoftGC, true);
+			ConstructSoftGCPackageToObjectList(this->SoftGCPackageToObjectListBuffer);
+			UE::Cook::DumpObjClassList(CookByTheBookOptions->SessionStartupObjects);
+			UPackage::SoftGCPackageToObjectList.Empty();
+			SoftGCPackageToObjectListBuffer.Empty();
+		}
+		GLog->Logf(TEXT("Memory Analysis: LLM Tags:"));
+#if ENABLE_LOW_LEVEL_MEM_TRACKER
+		if (FLowLevelMemTracker::Get().IsEnabled())
+		{
+			FLowLevelMemTracker::Get().DumpToLog();
+		}
+		else
+#endif
+		{
+			GLog->Logf(TEXT("LLM Tags are not displayed because llm is disabled. Run with -llm or -trace=memtag to see llm tags."));
+		}
+	}
+	else
+	{
+		DisplaySimpleSummary();
+
+		// Mark the packages we freed so we can give a warning to diagnose why they are still referenced if they
+		// get loaded again.
+		PackageTracker->AddExpectedNeverLoadPackages(ExpectedFreedPackageNames);
+
+		// Only show diagnostics if LLM is on, because they are somewhat expensive. We could add a separate setting
+		// for this, but it's more convenient to combine it with the LLM enabled setting
+#if ENABLE_LOW_LEVEL_MEM_TRACKER
+		bool bShowDiagnostics = FLowLevelMemTracker::Get().IsEnabled();
+#else
+		constexpr bool bShowDiagnostics = false;
+#endif
+		if (bShowDiagnostics)
+		{
+			// If some packages we expected to be freed were not freed, show the reference chains for why
+			// they were not freed.
+			TArray<UPackage*> PackagesReferencedOutsideOfCooker;
+			for (FWeakObjectPtr& WeakPtr : CookByTheBookOptions->SessionStartupObjects)
+			{
+				UObject* Object = WeakPtr.Get();
+				if (!Object)
+				{
+					continue;
+				}
+				UPackage* Package = Object->GetPackage();
+				ExpectedFreedPackageNames.Remove(Package->GetFName());
+			}
+			PackageTracker->ForEachLoadedPackage(
+				[this, &PackagesReferencedOutsideOfCooker](UPackage* Package)
+				{
+					if (ExpectedFreedPackageNames.Contains(Package->GetFName()))
+					{
+						PackagesReferencedOutsideOfCooker.Add(Package);
+					}
+				});
+			if (PackagesReferencedOutsideOfCooker.Num() > 0)
+			{
+				UE::Cook::DumpPackageReferencers(PackagesReferencedOutsideOfCooker);
+			}
+		}
+	}
+}
diff --git a/Engine/Source/Editor/UnrealEd/Private/Cooker/CookGarbageCollect.h b/Engine/Source/Editor/UnrealEd/Private/Cooker/CookGarbageCollect.h
--- a/Engine/Source/Editor/UnrealEd/Private/Cooker/CookGarbageCollect.h
+++ b/Engine/Source/Editor/UnrealEd/Private/Cooker/CookGarbageCollect.h
@@ -2,9 +2,14 @@
 
 #pragma once
 
+#include "Containers/Array.h"
+#include "Containers/RingBuffer.h"
+#include "Containers/UnrealString.h"
 #include "CoreTypes.h"
+#include "Templates/UnrealTemplate.h"
 
 class UCookOnTheFlyServer;
+class UObject;
 
 namespace UE::Cook
 {
@@ -51,4 +56,35 @@ private:
 	bool bCurrentGCIsFull = false;
 };
 
+/** Stores data over time for SoftGC that is used to throttle how frequently we trigger periodic SoftGC. */
+class FSoftGCHistory
+{
+public:
+	void AddDurationMeasurement(float DurationSeconds);
+	bool IsTriggeringWithinBudget(UCookOnTheFlyServer& COTFS, double CurrentTimeSeconds, FString* OutDiagnostics);
+
+private:
+	TRingBuffer<float> DurationHistory;
+	int MaxHistoryLength = 5;
+	float AverageDurationSeconds = 0.f;
+};
+
+/** Scoped type to call Set/Clear SoftGCPackageToObjectListBuffer. */
+struct FScopeFindCookReferences
+{
+	FScopeFindCookReferences(UCookOnTheFlyServer& InCOTFS);
+	~FScopeFindCookReferences();
+
+	UCookOnTheFlyServer& COTFS;
+	TGuardValue<bool> SoftGCGuard;
+	bool bNeedsConstructBuffer;
+};
+
+/**
+ * For every package in memory, add a list of all of its public UObjects into the map
+ * used in garbage collection: UPackage::SoftGCPackageToObjectList. This will cause all of
+ * its public objects to be referenced if the UPackage is referenced.
+ */
+void ConstructSoftGCPackageToObjectList(TArray<UObject*>& PackageToObjectListBuffer);
+
 } // namespace UE::Cook
diff --git a/Engine/Source/Editor/UnrealEd/Private/Cooker/CookGenerationHelper.cpp b/Engine/Source/Editor/UnrealEd/Private/Cooker/CookGenerationHelper.cpp
--- a/Engine/Source/Editor/UnrealEd/Private/Cooker/CookGenerationHelper.cpp
+++ b/Engine/Source/Editor/UnrealEd/Private/Cooker/CookGenerationHelper.cpp
@@ -1848,7 +1848,7 @@ void FGenerationHelper::PostGarbageCollect(const TRefCountPtr<FGenerationHelper>
 				
 				{
 					// Compute UCookOnTheFlyServer's references so they are gathered by OBJ REFS below 
-					UCookOnTheFlyServer::FScopeFindCookReferences(Owner.GetPackageDatas().GetCookOnTheFlyServer());
+					FScopeFindCookReferences(Owner.GetPackageDatas().GetCookOnTheFlyServer());
 
 					StaticExec(nullptr, *FString::Printf(TEXT("OBJ REFS NAME=%s"),
 						*Info.PackageData->GetPackageName().ToString()));
diff --git a/Engine/Source/Editor/UnrealEd/Private/Cooker/CookSavePackage.cpp b/Engine/Source/Editor/UnrealEd/Private/Cooker/CookSavePackage.cpp
new file mode 100644
--- /dev/null
+++ b/Engine/Source/Editor/UnrealEd/Private/Cooker/CookSavePackage.cpp
@@ -0,0 +1,702 @@
+// Copyright Epic Games, Inc. All Rights Reserved.
+
+#include "Cooker/CookSavePackage.h"
+
+#include "AssetRegistry/AssetData.h"
+#include "AssetRegistry/IAssetRegistry.h"
+#include "Commandlets/AssetRegistryGenerator.h"
+#include "Cooker/CookDeterminismManager.h"
+#include "Cooker/CookDiagnostics.h"
+#include "Cooker/CookGenerationHelper.h"
+#include "Cooker/CookOnTheFlyServerInterface.h"
+#include "Cooker/CookPackageData.h"
+#include "Cooker/CookPlatformManager.h"
+#include "Cooker/CookProfiling.h"
+#include "Cooker/CookRequestCluster.h"
+#include "Cooker/CookRequests.h"
+#include "Cooker/CookTypes.h"
+#include "Cooker/PackageTracker.h"
+#include "CookerSettings.h"
+#include "CookOnTheSide/CookLog.h"
+#include "CookOnTheSide/CookOnTheFlyServer.h"
+#include "Containers/Set.h"
+#include "CoreGlobals.h"
+#include "Editor.h"
+#include "Engine/AssetManager.h"
+#include "Engine/World.h"
+#include "HAL/LowLevelMemTracker.h"
+#include "HAL/PlatformMisc.h"
+#include "HAL/PlatformTime.h"
+#include "Interfaces/ITargetPlatform.h"
+#include "Interfaces/ITargetPlatformManagerModule.h"
+#include "Misc/PackageAccessTracking.h"
+#include "Misc/Paths.h"
+#include "Misc/RedirectCollector.h"
+#include "Logging/LogMacros.h"
+#include "Serialization/ArchiveCookData.h"
+#include "TargetDomain/TargetDomainUtils.h"
+#include "Templates/UnrealTemplate.h"
+#include "Trace/Trace.h"
+#include "UObject/CookEnums.h"
+#include "UObject/ObjectMacros.h"
+#include "UObject/ObjectSaveOverride.h"
+#include "UObject/SavePackage.h"
+#include "UObject/SoftObjectPath.h"
+#include "UObject/UObjectGlobals.h"
+
+COREUOBJECT_API extern bool GOutputCookingWarnings;
+
+#if OUTPUT_COOKTIMING
+
+UE_TRACE_EVENT_BEGIN(UE_CUSTOM_COOKTIMER_LOG, SaveCookedPackage, NoSync)
+	UE_TRACE_EVENT_FIELD(UE::Trace::WideString, PackageName)
+UE_TRACE_EVENT_END()
+
+#endif //OUTPUT_COOKTIMING
+
+LLM_DEFINE_TAG(Cooker_SavePackage);
+
+void UCookOnTheFlyServer::SaveCookedPackage(UE::Cook::FSaveCookedPackageContext& Context)
+{
+	using namespace UE::Cook;
+
+	UE_SCOPED_HIERARCHICAL_CUSTOM_COOKTIMER_AND_DURATION(SaveCookedPackage, DetailedCookStats::TickCookOnTheSideSaveCookedPackageTimeSec)
+		UE_ADD_CUSTOM_COOKTIMER_META(SaveCookedPackage, PackageName, *WriteToString<256>(Context.PackageData.GetFileName()));
+	double SaveStartTime = FPlatformTime::Seconds();
+
+	UPackage* Package = Context.Package;
+	FOnScopeExit ScopedPackageFlags([Package, OriginalPackageFlags = Package->GetPackageFlags()]()
+		{
+			Package->SetPackageFlagsTo(OriginalPackageFlags);
+		});
+
+	Context.SetupPackage();
+
+	ITargetPlatformManagerModule& TPM = GetTargetPlatformManagerRef();
+	TGuardValue<bool> ScopedOutputCookerWarnings(GOutputCookingWarnings, IsCookFlagSet(ECookInitializationFlags::OutputVerboseCookerWarnings));
+	// SavePackage can CollectGarbage, so we need to store the currently-unqueued PackageData in a separate variable that we register for garbage collection
+	TGuardValue<UE::Cook::FPackageData*> ScopedSavingPackageData(SavingPackageData, &Context.PackageData);
+	TGuardValue<bool> ScopedIsSavingPackage(bIsSavingPackage, true);
+	// For legacy reasons we set GIsCookerLoadingPackage == true during save. Some classes use it to conditionally execute cook operations in both save and load
+	TGuardValue<bool> ScopedIsCookerLoadingPackage(GIsCookerLoadingPackage, true);
+
+	for (int32 PlatformIndex = 0; PlatformIndex < Context.PlatformsForPackage.Num(); ++PlatformIndex)
+	{
+		const ITargetPlatform* TargetPlatform = Context.PlatformsForPackage[PlatformIndex];
+		Context.SetupPlatform(TargetPlatform, PlatformIndex);
+		if (Context.bPlatformSetupSuccessful)
+		{
+			UE_SCOPED_HIERARCHICAL_COOKTIMER(GEditorSavePackage);
+			UE_TRACK_REFERENCING_PLATFORM_SCOPED(TargetPlatform);
+
+			TMap<UObject*, FObjectSaveOverride> SaveOverrides;
+			FArchiveCookData CookData(*TargetPlatform, *Context.ArchiveCookContext);
+			FSavePackageArgs SaveArgs;
+			SaveArgs.TopLevelFlags = Context.FlagsToCook;
+			SaveArgs.bForceByteSwapping = Context.bEndianSwap;
+			SaveArgs.bWarnOfLongFilename = false;
+			SaveArgs.SaveFlags = Context.SaveFlags;
+			SaveArgs.ArchiveCookData = &CookData;
+			SaveArgs.bSlowTask = false;
+			SaveArgs.SavePackageContext = Context.SavePackageContext;
+			SaveArgs.InOutSaveOverrides = &SaveOverrides;
+
+			Context.PackageWriter->UpdateSaveArguments(SaveArgs);
+			FSavePackageResultStruct AuthoritativeResult = ESavePackageResult::Error;
+			bool IsFirstPass = true;
+			for (;;)
+			{
+				try
+				{
+					LLM_SCOPE_BYTAG(Cooker_SavePackage);
+					FScopedActivePackage ScopedActivePackage(*this, Context.PackageData.GetPackageName(), NAME_None);
+					if (bSkipSave)
+					{
+						Context.SavePackageResult = ESavePackageResult::Success;
+					}
+					else
+					{
+						Context.SavePackageResult = GEditor->Save(Package, Context.World, *Context.PlatFilename, SaveArgs);
+					}
+				}
+				catch (std::exception&)
+				{
+					UE_LOG(LogCook, Warning, TEXT("Tried to save package %s for target platform %s but threw an exception"),
+						*Package->GetName(), *TargetPlatform->PlatformName());
+					Context.SavePackageResult = ESavePackageResult::Error;
+				}
+
+				bool IsAnotherSaveNeeded = Context.PackageWriter->IsAnotherSaveNeeded(Context.SavePackageResult, SaveArgs);
+				if (IsFirstPass)
+				{
+					AuthoritativeResult = MoveTemp(Context.SavePackageResult);
+					IsFirstPass = false;
+				}
+				if (IsAnotherSaveNeeded)
+				{
+					// We must not try a second save of a package while the first save is still in flight.
+					// The optimal solution is to wait for ONLY the package that needs a second save, but we don't
+					// have the bookkeeping data to do that, so we have to wait for all async package writes to complete.
+					UPackage::WaitForAsyncFileWrites();
+				}
+				else
+				{
+					break;
+				}
+			}
+			Context.SavePackageResult = MoveTemp(AuthoritativeResult);
+
+			// If package was actually saved check with asset manager to make sure it wasn't excluded for being a
+			// development or never cook package. But skip sending the warnings from this check if it was editor-only.
+			if (Context.SavePackageResult == ESavePackageResult::Success)
+			{
+				UE_SCOPED_HIERARCHICAL_COOKTIMER(VerifyCanCookPackage);
+				if (!UAssetManager::Get().VerifyCanCookPackage(this, Package->GetFName()))
+				{
+					Context.SavePackageResult = ESavePackageResult::Error;
+				}
+			}
+
+			++this->StatSavedPackageCount;
+		}
+
+		Context.FinishPlatform();
+	}
+
+	// Need to restore flags before calling FinishPackage because it might need to save again
+	ScopedPackageFlags.ExitEarly();
+	Context.FinishPackage();
+
+	constexpr double SavePackageMinDurationLogTimeSeconds = 600.;
+	float SaveDurationSeconds = static_cast<float>(FPlatformTime::Seconds() - SaveStartTime);
+	UE_CLOG(SaveDurationSeconds >= SavePackageMinDurationLogTimeSeconds,
+		LogCook, Display, TEXT("SavePackagePerformance: Package %s took %.0fs to save."),
+		*WriteToString<256>(Package->GetFName()), SaveDurationSeconds);
+}
+
+namespace UE::Cook
+{
+
+FSaveCookedPackageContext::FSaveCookedPackageContext(UCookOnTheFlyServer& InCOTFS, UE::Cook::FPackageData& InPackageData,
+	TArrayView<const ITargetPlatform*> InPlatformsForPackage, UE::Cook::FTickStackData& InStackData)
+	: COTFS(InCOTFS)
+	, PackageData(InPackageData)
+	, PlatformsForPackage(InPlatformsForPackage)
+	, StackData(InStackData)
+	, Package(PackageData.GetPackage())
+	, PackageName(Package ? Package->GetName() : FString())
+	, Filename(PackageData.GetFileName().ToString())
+{
+	PlatformDependencies.SetNum(InPlatformsForPackage.Num());
+}
+
+void FSaveCookedPackageContext::SetupPackage()
+{
+	check(Package && Package->IsFullyLoaded()); // PackageData should not be in the save state if Package is not fully loaded
+	check(Package->GetPathName().Equals(PackageName)); // We should only be saving outermost packages, so the path name should be the same as the package name
+	check(!Filename.IsEmpty()); // PackageData guarantees FileName is non-empty; if not found it should never make it into save state
+	if (Package->HasAnyPackageFlags(PKG_ReloadingForCooker))
+	{
+		UE_LOG(LogCook, Warning, TEXT("Package %s marked as reloading for cook was requested to save"), *PackageName);
+		UE_LOG(LogCook, Fatal, TEXT("Package %s marked as reloading for cook was requested to save"), *PackageName);
+	}
+
+	SaveFlags = SAVE_KeepGUID | SAVE_Async
+		| (COTFS.IsCookFlagSet(ECookInitializationFlags::Unversioned) ? SAVE_Unversioned : 0);
+	SaveFlags |= COTFS.IsCookFlagSet(ECookInitializationFlags::CookEditorOptional) ? SAVE_Optional : SAVE_None;
+
+	if (COTFS.CookByTheBookOptions->bCookSoftPackageReferences)
+	{
+		SaveFlags |= SAVE_CookSoftPackageReferences;
+	}
+
+	// Use SandboxFile to do path conversion to properly handle sandbox paths (outside of standard paths in particular).
+	Filename = COTFS.ConvertToFullSandboxPath(*Filename, true);
+}
+
+void FSaveCookedPackageContext::SetupPlatform(const ITargetPlatform* InTargetPlatform, int32 InPlatformIndex)
+{
+	PlatformIndex = InPlatformIndex;
+	TargetPlatform = InTargetPlatform;
+	PlatFilename = Filename.Replace(TEXT("[Platform]"), *TargetPlatform->PlatformName());
+	bPlatformSetupSuccessful = false;
+
+	ArchiveCookContext.Emplace(Package,
+		COTFS.IsDirectorCookByTheBook() ? UE::Cook::ECookType::ByTheBook : UE::Cook::ECookType::OnTheFly,
+		COTFS.IsCookingDLC() ? UE::Cook::ECookingDLC::Yes : UE::Cook::ECookingDLC::No,
+		TargetPlatform);
+
+	// don't save Editor resources from the Engine if the target doesn't have editoronly data
+	if (COTFS.IsCookFlagSet(ECookInitializationFlags::SkipEditorContent) &&
+		(PackageName.StartsWith(TEXT("/Engine/Editor")) || PackageName.StartsWith(TEXT("/Engine/VREditor"))) &&
+		!TargetPlatform->HasEditorOnlyData())
+	{
+		SavePackageResult = ESavePackageResult::ContainsEditorOnlyData;
+		const TCHAR* RejectedReason = TEXT("EngineEditorContent");
+		if (GCookProgressDisplay & (int32)ECookProgressDisplayMode::Instigators)
+		{
+			UE_LOG(LogCook, Display, TEXT("Cooking %s, Instigator: { %s } -> Rejected %s"), *PackageName,
+				*(PackageData.GetInstigator().ToString()), RejectedReason);
+		}
+		else
+		{
+			UE_LOG(LogCook, Display, TEXT("Cooking %s -> Rejected %s"), *PackageName, RejectedReason);
+		}
+		return;
+	}
+	// Check whether or not game-specific behaviour should prevent this package from being cooked for the target platform
+	else if (!UAssetManager::Get().ShouldCookForPlatform(Package, TargetPlatform))
+	{
+		SavePackageResult = ESavePackageResult::ContainsEditorOnlyData;
+		const TCHAR* RejectedReason = TEXT("NotAssetManagerShouldCookForPlatform");
+		if (GCookProgressDisplay & (int32)ECookProgressDisplayMode::Instigators)
+		{
+			UE_LOG(LogCook, Display, TEXT("Cooking %s, Instigator: { %s } -> Rejected %s"), *PackageName,
+				*(PackageData.GetInstigator().ToString()), RejectedReason);
+		}
+		else
+		{
+			UE_LOG(LogCook, Display, TEXT("Cooking %s -> Rejected %s"), *PackageName, RejectedReason);
+		}
+		return;
+	}
+	// check if this package is unsupported for the target platform (typically plugin content)
+	else
+	{
+		if (TSet<FName>* NeverCookPackages = COTFS.PackageTracker->PlatformSpecificNeverCookPackages.Find(TargetPlatform))
+		{
+			FGenerationHelper* GenerationHelper =
+				PackageData.IsGenerated() ? PackageData.GetParentGenerationHelper() : nullptr;
+
+			if (NeverCookPackages->Find(Package->GetFName()) ||
+				(GenerationHelper && NeverCookPackages->Find(GenerationHelper->GetOwner().GetPackageName())))
+			{
+				SavePackageResult = ESavePackageResult::ContainsEditorOnlyData;
+				const TCHAR* RejectedReason = TEXT("PlatformSpecificNeverCook");
+				if (GCookProgressDisplay & (int32)ECookProgressDisplayMode::Instigators)
+				{
+					UE_LOG(LogCook, Display, TEXT("Cooking %s, Instigator: { %s } -> Rejected %s"), *PackageName,
+						*(PackageData.GetInstigator().ToString()), RejectedReason);
+				}
+				else
+				{
+					UE_LOG(LogCook, Display, TEXT("Cooking %s -> Rejected %s"), *PackageName, RejectedReason);
+				}
+				return;
+			}
+		}
+	}
+
+	const FString FullFilename = FPaths::ConvertRelativePathToFull(PlatFilename);
+	if (FullFilename.Len() >= FPlatformMisc::GetMaxPathLength())
+	{
+		LogCookerMessage(FString::Printf(TEXT("Couldn't save package, filename is too long (%d >= %d): %s"),
+			FullFilename.Len(), FPlatformMisc::GetMaxPathLength(), *FullFilename), EMessageSeverity::Error);
+		SavePackageResult = ESavePackageResult::Error;
+		return;
+	}
+
+	if (!bHasDelayLoaded)
+	{
+		// look for a world object in the package (if there is one, there's a map)
+		World = UWorld::FindWorldInPackage(Package);
+		if (World)
+		{
+			FlagsToCook = RF_NoFlags;
+		}
+		bContainsMap = Package->ContainsMap();
+		bHasDelayLoaded = true;
+	}
+
+	UE_CLOG((GCookProgressDisplay & (int32)ECookProgressDisplayMode::Instigators) && PlatformIndex == 0,
+		LogCook, Display, TEXT("Cooking %s, Instigator: { %s }"),
+		*PackageName, *(PackageData.GetInstigator().ToString()));
+	UE_CLOG(GCookProgressDisplay & (int32)ECookProgressDisplayMode::PackageNames, LogCook, Display,
+		TEXT("Cooking %s"), *PackageName);
+
+	bEndianSwap = (!TargetPlatform->IsLittleEndian()) ^ (!PLATFORM_LITTLE_ENDIAN);
+
+	if (!TargetPlatform->HasEditorOnlyData())
+	{
+		Package->SetPackageFlags(PKG_FilterEditorOnly);
+	}
+	else
+	{
+		Package->ClearPackageFlags(PKG_FilterEditorOnly);
+	}
+
+	CookContext = &COTFS.FindOrCreateSaveContext(TargetPlatform);
+	SavePackageContext = &CookContext->SaveContext;
+	PackageWriter = CookContext->PackageWriter;
+	ICookedPackageWriter::FBeginPackageInfo Info;
+	Info.PackageName = Package->GetFName();
+	Info.LooseFilePath = PlatFilename;
+	PackageWriter->BeginPackage(Info);
+	if (CookContext->DeterminismManager)
+	{
+		CookContext->DeterminismManager->BeginPackage(Package, TargetPlatform, PackageWriter);
+	}
+	// Set platform-specific save flags
+	FPackagePlatformData& PlatformData = PackageData.FindOrAddPlatformData(TargetPlatform);
+	uint32 PlatformSaveFlagsMask = SAVE_AllowTimeout;
+	SaveFlags = (SaveFlags & ~PlatformSaveFlagsMask);
+	if (!PlatformData.IsSaveTimedOut())
+	{
+		// If we timedout before, do not allow another timeout, otherwise do allow it
+		SaveFlags |= SAVE_AllowTimeout;
+	}
+
+	// Indicate Setup was successful
+	bPlatformSetupSuccessful = true;
+	SavePackageResult = ESavePackageResult::Success;
+}
+
+bool IsRetryErrorCode(ESavePackageResult Result)
+{
+	return Result == ESavePackageResult::Timeout;
+}
+
+void FSaveCookedPackageContext::FinishPlatform()
+{
+	TRACE_CPUPROFILER_EVENT_SCOPE(FSaveCookedPackageContext::FinishPlatform);
+	check(PlatformIndex >= 0 && PlatformDependencies.IsValidIndex(PlatformIndex));
+
+	bool bSuccessful = SavePackageResult.IsSuccessful();
+	ECookResult CookResult = bSuccessful ? ECookResult::Succeeded : ECookResult::Failed;
+
+	// Calculate up-to-date assetregistry data for Generator and Generated packages
+	TOptional<TArray<FAssetDependency>> OverridePackageDependencies;
+	TOptional<UE::TargetDomain::FGeneratedPackageResultStruct> GeneratedResult;
+	TOptional<FAssetPackageData> AssetPackageDataBuffer;
+	TOptional<FAssetPackageData> OverrideAssetPackageData;
+	const FAssetPackageData* AssetPackageData = nullptr;
+	if (TRefCountPtr<FGenerationHelper> GenerationHelper = PackageData.GetGenerationHelper();
+		GenerationHelper)
+	{
+		OverridePackageDependencies.Emplace();
+		GenerationHelper->FinishGeneratorPlatformSave(PackageData, PlatformIndex == 0, *OverridePackageDependencies);
+	}
+	else if (TRefCountPtr<FGenerationHelper> ParentGenerationHelper = PackageData.GetParentGenerationHelper();
+		ParentGenerationHelper)
+	{
+		GeneratedResult.Emplace();
+		ParentGenerationHelper->FinishGeneratedPlatformSave(PackageData, *GeneratedResult);
+		OverridePackageDependencies.Emplace(GeneratedResult->PackageDependencies);
+		OverrideAssetPackageData.Emplace(GeneratedResult->AssetPackageData);
+		AssetPackageData = &GeneratedResult->AssetPackageData;
+	}
+	if (!AssetPackageData)
+	{
+		AssetPackageDataBuffer = COTFS.AssetRegistry->GetAssetPackageDataCopy(Package->GetFName());
+		AssetPackageData = AssetPackageDataBuffer.GetPtrOrNull();
+	}
+
+	// Commit the saved bytes and the iterative cook data to the PackageWriter
+	if (bPlatformSetupSuccessful)
+	{
+		CalculatePlatformAgnosticRuntimeDependencies();
+		CalculatePlatformRuntimeDependencies();
+
+		ICookedPackageWriter::FCommitPackageInfo Info;
+		// Note GetCommitAttachments mutates the SaveResult; it moves CookDependencies out of it
+		Info.Attachments = GetCommitAttachments(GeneratedResult.GetPtrOrNull(), OverridePackageDependencies.GetPtrOrNull());
+		Info.Status = PackageResultToCommitStatus(SavePackageResult.Result);
+		Info.PackageName = Package->GetFName();
+		Info.PackageHash = AssetPackageData ? AssetPackageData->GetPackageSavedHash() : FIoHash();
+		Info.WriteOptions = GetCommitWriteOptions();
+
+		PackageWriter->CommitPackage(MoveTemp(Info));
+		if (CookContext->DeterminismManager)
+		{
+			CookContext->DeterminismManager->EndPackage();
+		}
+	}
+
+	// Update asset registry
+	if (COTFS.IsDirectorCookByTheBook())
+	{
+		IAssetRegistryReporter& Reporter = *(COTFS.PlatformManager->GetPlatformData(TargetPlatform)->RegistryReporter);
+		TOptional<TArray<FAssetData>> AssetDatasFromSave;
+		if (bSuccessful)
+		{
+			AssetDatasFromSave.Emplace(MoveTemp(SavePackageResult.SavedAssets));
+		}
+		Reporter.UpdateAssetRegistryData(Package->GetFName(), Package, CookResult, &SavePackageResult,
+			MoveTemp(AssetDatasFromSave), MoveTemp(OverrideAssetPackageData), MoveTemp(OverridePackageDependencies),
+			COTFS);
+	}
+
+	// If not retrying, mark the package as cooked, either successfully or with failure
+	bool bIsRetryErrorCode = IsRetryErrorCode(SavePackageResult.Result);
+	if (!bIsRetryErrorCode)
+	{
+		PackageData.SetPlatformCooked(TargetPlatform, CookResult);
+	}
+
+	// Update flags used to determine garbage collection.
+	if (bSuccessful)
+	{
+		if (bContainsMap)
+		{
+			StackData.ResultFlags |= UCookOnTheFlyServer::COSR_CookedMap;
+		}
+		else
+		{
+			++COTFS.CookedPackageCountSinceLastGC;
+			StackData.ResultFlags |= UCookOnTheFlyServer::COSR_CookedPackage;
+		}
+	}
+
+	// Accumulate results for SaveCookedPackage_Finish
+	if (SavePackageResult.Result == ESavePackageResult::Timeout)
+	{
+		PackageData.FindOrAddPlatformData(TargetPlatform).SetSaveTimedOut(true);
+		bHasTimeOut = true;
+	}
+	bAnySaveSucceeded |= bSuccessful;
+
+	bHasRetryErrorCode |= bIsRetryErrorCode;
+
+	ArchiveCookContext.Reset();
+	PlatformIndex = -1;
+}
+
+void FSaveCookedPackageContext::FinishPackage()
+{
+	// If any save succeeded, add all dependencies from all platforms to the cook for the platform
+	if (bAnySaveSucceeded && !COTFS.CookByTheBookOptions->bSkipSoftReferences)
+	{
+		FName PackageFName = Package->GetFName();
+		if (PlatformsForPackage.Num() == 1)
+		{
+			TConstArrayView<const ITargetPlatform*> ReachablePlatforms(&TargetPlatform, 1);
+			for (const TPair<FPackageData*, EInstigator>& DependencyPair : PlatformDependencies[0])
+			{
+				COTFS.QueueDiscoveredPackage(*DependencyPair.Key,
+					UE::Cook::FInstigator(DependencyPair.Value, PackageFName), FDiscoveredPlatformSet(ReachablePlatforms));
+			}
+		}
+		else
+		{
+			// Merge the platform dependencies for each package into a single QueueDiscoveredPackage call, if possible
+			// It will not be possible if the different platforms have different instigators, so track a list of
+			// platforms for each package for each instigator type.
+			TMap<FPackageData*, TMap<EInstigator, TArray<const ITargetPlatform*>>> PackagePlatformsForInstigator;
+			for (int32 LocalIndex = 0; LocalIndex < PlatformsForPackage.Num(); ++LocalIndex)
+			{
+				const TMap<FPackageData*, EInstigator>& CurrentDependencies = PlatformDependencies[LocalIndex];
+				const ITargetPlatform* CurrentPlatform = PlatformsForPackage[LocalIndex];
+				for (const TPair<FPackageData*, EInstigator>& PackagePlatformPair : CurrentDependencies)
+				{
+					TMap<EInstigator, TArray<const ITargetPlatform*>>& TargetMap =
+						PackagePlatformsForInstigator.FindOrAdd(PackagePlatformPair.Key);
+					TargetMap.FindOrAdd(PackagePlatformPair.Value).Add(CurrentPlatform);
+				}
+			}
+			for (const TPair<FPackageData*, TMap<EInstigator, TArray<const ITargetPlatform*>>>& PackagePair
+				: PackagePlatformsForInstigator)
+			{
+				for (const TPair<EInstigator, TArray<const ITargetPlatform*>>& InstigatorPair : PackagePair.Value)
+				{
+					COTFS.QueueDiscoveredPackage(*PackagePair.Key,
+						UE::Cook::FInstigator(InstigatorPair.Key, PackageFName),
+						FDiscoveredPlatformSet(InstigatorPair.Value));
+				}
+			}
+		}
+	}
+
+	if (COTFS.IsDebugRecordUnsolicited())
+	{
+		COTFS.ProcessUnsolicitedPackages();
+		TMap<FPackageData*, EInstigator> AllPlatformDependenciesBuffer;
+		TMap<FPackageData*, EInstigator>* AllPlatformDependencies = &PlatformDependencies[0];
+		if (PlatformsForPackage.Num() > 1)
+		{
+			AllPlatformDependencies = &AllPlatformDependenciesBuffer;
+			for (int32 LocalIndex = 0; LocalIndex < PlatformsForPackage.Num(); ++LocalIndex)
+			{
+				AllPlatformDependenciesBuffer.Append(PlatformDependencies[LocalIndex]);
+			}
+		}
+
+		FDiagnostics::AnalyzeHiddenDependencies(COTFS, PackageData, PackageData.DetachUnsolicited(),
+			*AllPlatformDependencies, PlatformsForPackage, COTFS.bOnlyEditorOnlyDebug, COTFS.bHiddenDependenciesDebug);
+	}
+
+	if (!bHasRetryErrorCode)
+	{
+		if (COTFS.IsCookOnTheFlyMode() && PackageData.GetUrgency() != EUrgency::Blocking &&
+			(!COTFS.CookOnTheFlyRequestManager || COTFS.CookOnTheFlyRequestManager->ShouldUseLegacyScheduling()))
+		{
+			// this is an unsolicited package
+			if (FPaths::FileExists(Filename))
+			{
+				COTFS.PackageTracker->UnsolicitedCookedPackages.AddCookedPackage(
+					FFilePlatformRequest(PackageData.GetFileName(), EInstigator::Unspecified, PlatformsForPackage));
+
+#if DEBUG_COOKONTHEFLY
+				UE_LOG(LogCook, Display, TEXT("UnsolicitedCookedPackages: %s"), *Filename);
+#endif
+			}
+		}
+	}
+}
+
+void FSaveCookedPackageContext::CalculatePlatformAgnosticRuntimeDependencies()
+{
+	if (bPlatformAgnosticDependenciesCalculated)
+	{
+		return;
+	}
+	bPlatformAgnosticDependenciesCalculated = true;
+
+	FName PackageFName = PackageData.GetPackageName();
+	for (FName LocalizedPackageName : FRequestCluster::GetLocalizationReferences(PackageFName, COTFS))
+	{
+		UE::Cook::FPackageData* LocalizedPackageData = COTFS.PackageDatas->TryAddPackageDataByPackageName(LocalizedPackageName);
+		if (LocalizedPackageData)
+		{
+			AddDependency(PlatformAgnosticDependencies, LocalizedPackageData, false /* bHard */);
+		}
+	}
+	// Also add any references from the package that are required by the AssetManager
+	for (FName AMPackageName : FRequestCluster::GetAssetManagerReferences(PackageFName))
+	{
+		UE::Cook::FPackageData* AMPackageData = COTFS.PackageDatas->TryAddPackageDataByPackageName(AMPackageName);
+		if (AMPackageData)
+		{
+			AddDependency(PlatformAgnosticDependencies, AMPackageData, false /* bHard */);
+		}
+	}
+
+	// When using legacy WhatGetsCookedRules, add all the SoftObjectPaths discovered during the package's load, plus any added on
+	// during save, to the cook for all platforms
+	TSet<FName> SoftObjectPackages;
+	GRedirectCollector.ProcessSoftObjectPathPackageList(PackageFName, false /* bGetEditorOnly */, SoftObjectPackages);
+	for (FName SoftObjectPackage : SoftObjectPackages)
+	{
+		TMap<FSoftObjectPath, FSoftObjectPath> RedirectedPaths;
+
+		// If this is a redirector, extract destination from asset registry
+		if (COTFS.ContainsRedirector(SoftObjectPackage, RedirectedPaths))
+		{
+			for (TPair<FSoftObjectPath, FSoftObjectPath>& RedirectedPath : RedirectedPaths)
+			{
+				GRedirectCollector.AddAssetPathRedirection(RedirectedPath.Key, RedirectedPath.Value);
+			}
+		}
+
+		UE::Cook::FPackageData* SoftObjectPackageData = COTFS.PackageDatas->TryAddPackageDataByPackageName(SoftObjectPackage);
+		if (SoftObjectPackageData)
+		{
+			if (!COTFS.bSkipOnlyEditorOnly)
+			{
+				AddDependency(PlatformAgnosticDependencies, SoftObjectPackageData, false /* bHard */);
+			}
+
+			if (COTFS.IsDebugRecordUnsolicited())
+			{
+				PackageData.CreateOrGetUnsolicited().Add(SoftObjectPackageData, EInstigator::SoftDependency);
+			}
+		}
+	}
+}
+
+void FSaveCookedPackageContext::CalculatePlatformRuntimeDependencies()
+{
+	PlatformDependencies.IsValidIndex(PlatformIndex);
+	TMap<FPackageData*, EInstigator>& CurrentDependencies = PlatformDependencies[PlatformIndex];
+	CurrentDependencies.Reset();
+	// Add PlatformAgnostic dependencies to the PlatformDependencies
+	CurrentDependencies.Append(PlatformAgnosticDependencies);
+
+	// Add imports and softobjectpaths from the save to the SaveDependencies for the current Platform and for all Platforms
+	if (COTFS.bSkipOnlyEditorOnly)
+	{
+		FName PackageFName = Package->GetFName();
+		TConstArrayView<const ITargetPlatform*> ReachablePlatforms(&TargetPlatform, 1);
+		for (const TArray<FName>* DependencyNames : { &SavePackageResult.ImportPackages, &SavePackageResult.SoftPackageReferences })
+		{
+			bool bHard = DependencyNames == &SavePackageResult.ImportPackages;
+			for (FName DependencyName : *DependencyNames)
+			{
+				UE::Cook::FPackageData* DependencyData = COTFS.PackageDatas->TryAddPackageDataByPackageName(DependencyName);
+				if (DependencyData)
+				{
+					AddDependency(CurrentDependencies, DependencyData, bHard);
+				}
+			}
+		}
+	}
+}
+
+TArray<FName> FSaveCookedPackageContext::GetPlatformRuntimeDependencies() const
+{
+	TArray<FName> PlatformDependencyNames;
+	PlatformDependencyNames.Reserve(PlatformDependencies[PlatformIndex].Num());
+	for (const TPair<FPackageData*, EInstigator>& DependencyPair : PlatformDependencies[PlatformIndex])
+	{
+		PlatformDependencyNames.Add(DependencyPair.Key->GetPackageName());
+	}
+	return PlatformDependencyNames;
+}
+
+TArray<IPackageWriter::FCommitAttachmentInfo> FSaveCookedPackageContext::GetCommitAttachments(
+	UE::TargetDomain::FGeneratedPackageResultStruct* GeneratedResult,
+	TArray<FAssetDependency>* OverrideDependencies)
+{
+	TArray<IPackageWriter::FCommitAttachmentInfo> Result;
+	if (COTFS.bHybridIterativeEnabled)
+	{
+		UE_SCOPED_HIERARCHICAL_COOKTIMER(TargetDomainDependencies);
+		TArray<FName> PlatformRuntimeDependencies = GetPlatformRuntimeDependencies();
+		if (OverrideDependencies)
+		{
+			PlatformRuntimeDependencies.Reserve(PlatformRuntimeDependencies.Num() + OverrideDependencies->Num());
+			for (const FAssetDependency& Dependency : *OverrideDependencies)
+			{
+				PlatformRuntimeDependencies.Add(Dependency.AssetId.PackageName);
+			}
+		}
+		UE::TargetDomain::CollectAndStoreCookAttachments(Package, TargetPlatform, &SavePackageResult,
+			GeneratedResult, MoveTemp(PlatformRuntimeDependencies), Result);
+	}
+	if (CookContext->DeterminismManager)
+	{
+		CookContext->DeterminismManager->AppendCommitAttachments(Result);
+	}
+
+	return Result;
+}
+
+IPackageWriter::EWriteOptions FSaveCookedPackageContext::GetCommitWriteOptions() const
+{
+	IPackageWriter::EWriteOptions Result = IPackageWriter::EWriteOptions::None;
+	if (!COTFS.bSkipSave)
+	{
+		Result |= IPackageWriter::EWriteOptions::Write;
+
+		if (COTFS.IsDirectorCookByTheBook())
+		{
+			Result |= IPackageWriter::EWriteOptions::ComputeHash;
+		}
+	}
+	return Result;
+}
+
+void FSaveCookedPackageContext::AddDependency(TMap<FPackageData*, EInstigator>& InDependencies,
+	FPackageData* PackageData, bool bHard)
+{
+	EInstigator& Existing = InDependencies.FindOrAdd(PackageData, EInstigator::Unspecified);
+	if (bHard)
+	{
+		Existing = EInstigator::SaveTimeHardDependency;
+	}
+	else if (Existing == EInstigator::Unspecified)
+	{
+		Existing = EInstigator::SoftDependency;
+	}
+}
+
+} // namespace UE::Cook
diff --git a/Engine/Source/Editor/UnrealEd/Private/Cooker/CookSavePackage.h b/Engine/Source/Editor/UnrealEd/Private/Cooker/CookSavePackage.h
new file mode 100644
--- /dev/null
+++ b/Engine/Source/Editor/UnrealEd/Private/Cooker/CookSavePackage.h
@@ -0,0 +1,94 @@
+// Copyright Epic Games, Inc. All Rights Reserved.
+
+#pragma once
+
+#include "Containers/ArrayView.h"
+#include "Containers/Array.h"
+#include "Containers/Map.h"
+#include "Containers/UnrealString.h"
+#include "CoreTypes.h"
+#include "Misc/Optional.h"
+#include "Serialization/PackageWriter.h"
+#include "UObject/ArchiveCookContext.h"
+#include "UObject/ICookInfo.h"
+#include "UObject/NameTypes.h"
+#include "UObject/ObjectMacros.h"
+#include "UObject/Package.h"
+
+class FSavePackageContext;
+class ITargetPlatform;
+class UCookOnTheFlyServer;
+class UWorld;
+namespace UE::Cook { struct FCookSavePackageContext; }
+namespace UE::Cook { struct FPackageData; }
+namespace UE::Cook { struct FTickStackData; }
+namespace UE::TargetDomain { struct FGeneratedPackageResultStruct; }
+struct FAssetDependency;
+struct FSavePackageResultStruct;
+
+namespace UE::Cook
+{
+
+/** Local parameters and helper functions used by SaveCookedPackage */
+class FSaveCookedPackageContext
+{
+private: // Used only by UCookOnTheFlyServer, which has private access
+
+	FSaveCookedPackageContext(UCookOnTheFlyServer& InCOTFS, UE::Cook::FPackageData& InPackageData,
+		TArrayView<const ITargetPlatform*> InPlatformsForPackage, UE::Cook::FTickStackData& StackData);
+
+	// Hooks used by friends
+	void SetupPackage();
+	void SetupPlatform(const ITargetPlatform* InTargetPlatform, int32 InPlatformIndex);
+	void FinishPlatform();
+	void FinishPackage();
+
+	// private helper functions
+	void CalculatePlatformAgnosticRuntimeDependencies();
+	void CalculatePlatformRuntimeDependencies();
+	TArray<FName> GetPlatformRuntimeDependencies() const;
+	TArray<IPackageWriter::FCommitAttachmentInfo> GetCommitAttachments(
+		UE::TargetDomain::FGeneratedPackageResultStruct* GeneratedResult,
+		TArray<FAssetDependency>* OverrideDependencies);
+	IPackageWriter::EWriteOptions GetCommitWriteOptions() const;
+	static void AddDependency(TMap<FPackageData*, EInstigator>& InDependencies, FPackageData* PackageData, bool bHard);
+
+	// General Package Data
+	UCookOnTheFlyServer& COTFS;
+	FPackageData& PackageData;
+	TArrayView<const ITargetPlatform*> PlatformsForPackage;
+	TMap<FPackageData*, EInstigator> PlatformAgnosticDependencies;
+	TArray<TMap<FPackageData*, EInstigator>, TInlineAllocator<1>> PlatformDependencies;
+	FTickStackData& StackData;
+	UPackage* Package;
+	const FString PackageName;
+	FString Filename;
+	uint32 SaveFlags = 0;
+	bool bHasTimeOut = false;
+	bool bHasRetryErrorCode = false;
+	bool bPlatformAgnosticDependenciesCalculated = false;
+	bool bAnySaveSucceeded = false;
+
+	// General Package Data that is delay-loaded the first time we save a platform
+	UWorld* World = nullptr;
+	EObjectFlags FlagsToCook = RF_Public;
+	bool bHasDelayLoaded = false;
+	bool bContainsMap = false;
+
+	// Per-platform data, only valid in PerPlatform callbacks
+	const ITargetPlatform* TargetPlatform = nullptr;
+	FCookSavePackageContext* CookContext = nullptr;
+	// This holds cook data generated by Serialize() that isn't saved in the package
+	TOptional<FArchiveCookContext> ArchiveCookContext;
+	FSavePackageContext* SavePackageContext = nullptr;
+	ICookedPackageWriter* PackageWriter = nullptr;
+	FString PlatFilename;
+	FSavePackageResultStruct SavePackageResult;
+	int32 PlatformIndex = -1;
+	bool bPlatformSetupSuccessful = false;
+	bool bEndianSwap = false;
+
+	friend class ::UCookOnTheFlyServer;
+};
+
+} // namespace UE::Cook;
diff --git a/Engine/Source/Editor/UnrealEd/Private/Cooker/CookTypes.h b/Engine/Source/Editor/UnrealEd/Private/Cooker/CookTypes.h
--- a/Engine/Source/Editor/UnrealEd/Private/Cooker/CookTypes.h
+++ b/Engine/Source/Editor/UnrealEd/Private/Cooker/CookTypes.h
@@ -770,3 +770,4 @@ FMultiPackageReaderResults GetSaveExportsAndImports(UPackage* Package, UObject*
 
 IPackageWriter::ECommitStatus PackageResultToCommitStatus(ESavePackageResult Result);
 
+extern int32 GCookProgressDisplay;
diff --git a/Engine/Source/Editor/UnrealEd/Private/Cooker/SoftGC.cpp b/Engine/Source/Editor/UnrealEd/Private/Cooker/SoftGC.cpp
deleted file mode 100644
diff --git a/Engine/Source/Editor/UnrealEd/Private/Cooker/SoftGC.h b/Engine/Source/Editor/UnrealEd/Private/Cooker/SoftGC.h
deleted file mode 100644
