Disable UActorExportT3D and derivatives from being discovered by FindExporters if CVar disabled

* Addition of these classes were made to support a T3D use case that did not pan out in the end.  Copy/Paste blocking of Scene Graph achieved through CVars within the serialization code itself.
* Fixes non LevelExport export bugs from 5.5

diff --git a/Engine/Source/Editor/UnrealEd/Classes/Exporters/ActorExporterT3D.h b/Engine/Source/Editor/UnrealEd/Classes/Exporters/ActorExporterT3D.h
--- a/Engine/Source/Editor/UnrealEd/Classes/Exporters/ActorExporterT3D.h
+++ b/Engine/Source/Editor/UnrealEd/Classes/Exporters/ActorExporterT3D.h
@@ -22,10 +22,11 @@ public:
 
 	//~ Begin UExporter Interface
 	UNREALED_API virtual bool ExportText( const FExportObjectInnerContext* Context, UObject* Object, const TCHAR* Type, FOutputDevice& Ar, FFeedbackContext* Warn, uint32 PortFlags=0 ) override;
+	virtual bool SupportsObject(UObject* Object) const override;
 	//~ End UExporter Interface
 };
 
-UCLASS()
+UCLASS(MinimalAPI)
 class UGroupActorExporterT3D : public UActorExporterT3D
 {
 	GENERATED_BODY()
@@ -37,7 +38,7 @@ public:
 	//~ End UExporter Interface
 };
 
-UCLASS()
+UCLASS(MinimalAPI)
 class UPhysicsVolumeExporterT3D : public UActorExporterT3D
 {
 	GENERATED_BODY()
diff --git a/Engine/Source/Editor/UnrealEd/Private/EditorExporters.cpp b/Engine/Source/Editor/UnrealEd/Private/EditorExporters.cpp
--- a/Engine/Source/Editor/UnrealEd/Private/EditorExporters.cpp
+++ b/Engine/Source/Editor/UnrealEd/Private/EditorExporters.cpp
@@ -98,7 +98,7 @@ DEFINE_LOG_CATEGORY_STATIC(LogEditorExporters, Log, All);
 
 namespace Private
 {
-	bool bLevelExportT3D_UseActorExporters = true;
+	bool bLevelExportT3D_UseActorExporters = false;
 	static FAutoConsoleVariableRef CVarDumpHandlerTick(
 		TEXT("UnrealEd.T3DUseActorExporters"),
 		bLevelExportT3D_UseActorExporters,
@@ -538,6 +538,15 @@ bool UActorExporterT3D::ExportText(const FExportObjectInnerContext* Context, UOb
 	return true;
 }
 
+bool UActorExporterT3D::SupportsObject(UObject* Object) const
+{
+	if (Private::bLevelExportT3D_UseActorExporters)
+	{
+		return Super::SupportsObject(Object);
+	}
+	return false;
+}
+
 UGroupActorExporterT3D::UGroupActorExporterT3D(const FObjectInitializer& ObjectInitializer)
 {
 	SupportedClass = AGroupActor::StaticClass();
