Optimized editor loading region when region contains a lot of level instances (postpone rebuild of SActorEditorContext only when Ticked)

diff --git a/Engine/Source/Editor/UnrealEd/Private/SActorEditorContext.cpp b/Engine/Source/Editor/UnrealEd/Private/SActorEditorContext.cpp
--- a/Engine/Source/Editor/UnrealEd/Private/SActorEditorContext.cpp
+++ b/Engine/Source/Editor/UnrealEd/Private/SActorEditorContext.cpp
@@ -15,10 +15,11 @@
 void SActorEditorContext::Construct(const FArguments& InArgs)
 {
 	World = (InArgs._World);
+	bNeedsRebuild = false;
 	bIsContextExpanded = true;
 	check(GEditor);
 	GEditor->GetEditorWorldContext().AddRef(World);
-	UActorEditorContextSubsystem::Get()->OnActorEditorContextSubsystemChanged().AddSP(this, &SActorEditorContext::Rebuild);
+	UActorEditorContextSubsystem::Get()->OnActorEditorContextSubsystemChanged().AddSP(this, &SActorEditorContext::OnActorEditorContextSubsystemChanged);
 	FEditorDelegates::MapChange.AddSP(this, &SActorEditorContext::OnEditorMapChange);
 	Rebuild();
 }
@@ -36,6 +37,22 @@ SActorEditorContext::~SActorEditorContext()
 	}
 }
 
+void SActorEditorContext::OnActorEditorContextSubsystemChanged()
+{
+	bNeedsRebuild = true;
+}
+
+void SActorEditorContext::Tick(const FGeometry& AllottedGeometry, const double InCurrentTime, const float InDeltaTime)
+{
+	if (bNeedsRebuild)
+	{
+		bNeedsRebuild = false;
+		Rebuild();
+	}
+
+	SCompoundWidget::Tick(AllottedGeometry, InCurrentTime, InDeltaTime);
+}
+
 void SActorEditorContext::Rebuild()
 {
 	TArray<IActorEditorContextClient*> Clients = UActorEditorContextSubsystem::Get()->GetDisplayableClients();
diff --git a/Engine/Source/Editor/UnrealEd/Public/SActorEditorContext.h b/Engine/Source/Editor/UnrealEd/Public/SActorEditorContext.h
--- a/Engine/Source/Editor/UnrealEd/Public/SActorEditorContext.h
+++ b/Engine/Source/Editor/UnrealEd/Public/SActorEditorContext.h
@@ -13,13 +13,16 @@ public:
 
 	UNREALED_API void Construct(const FArguments& InArgs);
 	UNREALED_API ~SActorEditorContext();
+	UNREALED_API virtual void Tick(const FGeometry& AllottedGeometry, const double InCurrentTime, const float InDeltaTime) override;
 
 	static UNREALED_API bool IsVisible(UWorld* InWorld);
 private:
 
 	void OnEditorMapChange(uint32 MapChangeFlags = 0) { Rebuild(); }
+	void OnActorEditorContextSubsystemChanged();
 	void Rebuild();
 	
+	bool bNeedsRebuild;
 	bool bIsContextExpanded;
 	UWorld* World;
 };
