Updates WorldPartition Builders to wait for the AssetRegistry initial scanning before initializing WP.
Fixes invalid LevelInstanceActorDesc loading when AssetRegistry initialization is not complete.

diff --git a/Engine/Source/Editor/UnrealEd/Private/WorldPartition/WorldPartitionBuilder.cpp b/Engine/Source/Editor/UnrealEd/Private/WorldPartition/WorldPartitionBuilder.cpp
--- a/Engine/Source/Editor/UnrealEd/Private/WorldPartition/WorldPartitionBuilder.cpp
+++ b/Engine/Source/Editor/UnrealEd/Private/WorldPartition/WorldPartitionBuilder.cpp
@@ -2,6 +2,7 @@
 
 #include "WorldPartition/WorldPartitionBuilder.h"
 
+#include "AssetRegistry/AssetRegistryModule.h"
 #include "CoreMinimal.h"
 #include "Editor.h"
 #include "EditorWorldUtils.h"
@@ -100,6 +101,9 @@ bool UWorldPartitionBuilder::RunBuilder(UWorld* World)
 	bool bResult = true;
 	// Setup the world
 	{
+		IAssetRegistry& AssetRegistry = FModuleManager::GetModuleChecked<FAssetRegistryModule>("AssetRegistry").Get();
+		AssetRegistry.WaitForCompletion();
+
 		TUniquePtr<FScopedEditorWorld> EditorWorld;
 		if (World->bIsWorldInitialized)
 		{
