AssetEditorToolkit - Allow derived classes to generate their own MenuBar
Allow derived classes to prevent registering the AssetEditor's default commands.

The motivation for these changes is that we want LiveLinkHub to be an AssetEditorToolkit in order to leverage the functionality of WorkflowCentricApplication and ApplicationModes.
In the case of LiveLinkHub, we do not want the default File menus and commands since the hub lacks a bunch of the editor's functionality.

diff --git a/Engine/Source/Editor/UnrealEd/Private/Toolkits/AssetEditorToolkit.cpp b/Engine/Source/Editor/UnrealEd/Private/Toolkits/AssetEditorToolkit.cpp
--- a/Engine/Source/Editor/UnrealEd/Private/Toolkits/AssetEditorToolkit.cpp
+++ b/Engine/Source/Editor/UnrealEd/Private/Toolkits/AssetEditorToolkit.cpp
@@ -37,6 +37,7 @@
 #include "Framework/Commands/GenericCommands.h"
 #include "Interfaces/Interface_AsyncCompilation.h"
 #include "Widgets/Images/SImage.h"
+#include "StatusBarSubsystem.h"
 
 #define LOCTEXT_NAMESPACE "AssetEditorToolkit"
 
@@ -58,6 +59,12 @@
 	, OpenMethod(EAssetOpenMethod::Edit)
 {
 	WorkspaceMenuCategory = FWorkspaceItem::NewGroup(LOCTEXT("WorkspaceMenu_BaseAssetEditor", "Asset Editor"));
+
+	if (!FAssetEditorCommonCommands::IsRegistered())
+	{
+		// Ensure the asset editor common commands are registered in case UnrealEdMisc isn't used, and therefore not initializing the commands.
+		FAssetEditorCommonCommands::Register();
+	}
 }
 
 void FAssetEditorToolkit::InitAssetEditor( const EToolkitMode::Type Mode, const TSharedPtr< class IToolkitHost >& InitToolkitHost, const FName AppIdentifier, const TSharedRef<FTabManager::FLayout>& StandaloneDefaultLayout, const bool bCreateDefaultStandaloneMenu, const bool bCreateDefaultToolbar, UObject* ObjectToEdit, const bool bInIsToolbarFocusable, const bool bInUseSmallToolbarIcons, const TOptional<EAssetOpenMethod>& InOpenMethod )
@@ -254,53 +261,7 @@
 	check( ToolkitHost.IsValid() );
 	ToolkitManager.RegisterNewToolkit( SharedThis( this ) );
 
-	ToolkitCommands->MapAction(
-		FAssetEditorCommonCommands::Get().SaveAsset,
-		FExecuteAction::CreateSP( this, &FAssetEditorToolkit::SaveAsset_Execute ),
-		FCanExecuteAction::CreateSP( this, &FAssetEditorToolkit::CanSaveAsset_Internal ),
-		FIsActionChecked(),
-		FIsActionButtonVisible::CreateSP( this, &FAssetEditorToolkit::IsSaveAssetVisible ));
+	MapToolkitCommands();
 
-	ToolkitCommands->MapAction(
-		FAssetEditorCommonCommands::Get().SaveAssetAs,
-		FExecuteAction::CreateSP( this, &FAssetEditorToolkit::SaveAssetAs_Execute ),
-		FCanExecuteAction::CreateSP( this, &FAssetEditorToolkit::CanSaveAssetAs_Internal ),
-		FIsActionChecked(),
-		FIsActionButtonVisible::CreateSP( this, &FAssetEditorToolkit::IsSaveAssetAsVisible ));
-
-	ToolkitCommands->MapAction(
-		FGlobalEditorCommonCommands::Get().FindInContentBrowser,
-		FExecuteAction::CreateSP( this, &FAssetEditorToolkit::FindInContentBrowser_Execute ),
-		FCanExecuteAction::CreateSP( this, &FAssetEditorToolkit::CanFindInContentBrowser ),
-		FIsActionChecked(),
-		FIsActionButtonVisible::CreateSP(this, &FAssetEditorToolkit::IsFindInContentBrowserButtonVisible));
-		
-	ToolkitCommands->MapAction(
-		FGlobalEditorCommonCommands::Get().OpenDocumentation,
-		FExecuteAction::CreateSP(this, &FAssetEditorToolkit::BrowseDocumentation_Execute));
-
-	ToolkitCommands->MapAction(
-		FAssetEditorCommonCommands::Get().ReimportAsset,
-		FExecuteAction::CreateSP( this, &FAssetEditorToolkit::Reimport_Execute ),
-		FCanExecuteAction::CreateSP(this, &FAssetEditorToolkit::CanReimport_Internal));
-
-	FGlobalEditorCommonCommands::MapActions(ToolkitCommands);
-
-	if( IsWorldCentricAssetEditor() )
-	{
-		ToolkitCommands->MapAction(
-			FAssetEditorCommonCommands::Get().SwitchToStandaloneEditor,
-			FExecuteAction::CreateStatic( &FAssetEditorToolkit::SwitchToStandaloneEditor_Execute, TWeakPtr< FAssetEditorToolkit >( AsShared() ) ) );
-	}
-	else
-	{
-		if( GetPreviousWorldCentricToolkitHost().IsValid() )
-		{
-			ToolkitCommands->MapAction(
-				FAssetEditorCommonCommands::Get().SwitchToWorldCentricEditor,
-				FExecuteAction::CreateStatic( &FAssetEditorToolkit::SwitchToWorldCentricEditor_Execute, TWeakPtr< FAssetEditorToolkit >( AsShared() ) ) );
-		}
-	}
-
 	InitializeReadOnlyMenuProfiles();
 
@@ -851,6 +812,57 @@
 	EditorModeManager = MakeShared<FEditorModeTools>();
 }
 
+void FAssetEditorToolkit::MapToolkitCommands()
+{
+	ToolkitCommands->MapAction(
+		FAssetEditorCommonCommands::Get().SaveAsset,
+		FExecuteAction::CreateSP(this, &FAssetEditorToolkit::SaveAsset_Execute),
+		FCanExecuteAction::CreateSP(this, &FAssetEditorToolkit::CanSaveAsset_Internal),
+		FIsActionChecked(),
+		FIsActionButtonVisible::CreateSP(this, &FAssetEditorToolkit::IsSaveAssetVisible));
+
+	ToolkitCommands->MapAction(
+		FAssetEditorCommonCommands::Get().SaveAssetAs,
+		FExecuteAction::CreateSP(this, &FAssetEditorToolkit::SaveAssetAs_Execute),
+		FCanExecuteAction::CreateSP(this, &FAssetEditorToolkit::CanSaveAssetAs_Internal),
+		FIsActionChecked(),
+		FIsActionButtonVisible::CreateSP(this, &FAssetEditorToolkit::IsSaveAssetAsVisible));
+
+	ToolkitCommands->MapAction(
+		FGlobalEditorCommonCommands::Get().FindInContentBrowser,
+		FExecuteAction::CreateSP(this, &FAssetEditorToolkit::FindInContentBrowser_Execute),
+		FCanExecuteAction::CreateSP(this, &FAssetEditorToolkit::CanFindInContentBrowser),
+		FIsActionChecked(),
+		FIsActionButtonVisible::CreateSP(this, &FAssetEditorToolkit::IsFindInContentBrowserButtonVisible));
+
+	ToolkitCommands->MapAction(
+		FGlobalEditorCommonCommands::Get().OpenDocumentation,
+		FExecuteAction::CreateSP(this, &FAssetEditorToolkit::BrowseDocumentation_Execute));
+
+	ToolkitCommands->MapAction(
+		FAssetEditorCommonCommands::Get().ReimportAsset,
+		FExecuteAction::CreateSP(this, &FAssetEditorToolkit::Reimport_Execute),
+		FCanExecuteAction::CreateSP(this, &FAssetEditorToolkit::CanReimport_Internal));
+
+	FGlobalEditorCommonCommands::MapActions(ToolkitCommands);
+
+	if (IsWorldCentricAssetEditor())
+	{
+		ToolkitCommands->MapAction(
+			FAssetEditorCommonCommands::Get().SwitchToStandaloneEditor,
+			FExecuteAction::CreateStatic(&FAssetEditorToolkit::SwitchToStandaloneEditor_Execute, TWeakPtr< FAssetEditorToolkit >(AsShared())));
+	}
+	else
+	{
+		if (GetPreviousWorldCentricToolkitHost().IsValid())
+		{
+			ToolkitCommands->MapAction(
+				FAssetEditorCommonCommands::Get().SwitchToWorldCentricEditor,
+				FExecuteAction::CreateStatic(&FAssetEditorToolkit::SwitchToWorldCentricEditor_Execute, TWeakPtr< FAssetEditorToolkit >(AsShared())));
+		}
+	}
+}
+
 FAssetEditorModeManager* FAssetEditorToolkit::GetAssetEditorModeManager() const
 {
 	PRAGMA_DISABLE_DEPRECATION_WARNINGS
@@ -1551,6 +1563,11 @@
 	return bIsActuallyAnAsset;
 }
 
+UE::Editor::Toolbars::ECreateStatusBarOptions FAssetEditorToolkit::GetStatusBarCreationOptions() const
+{
+	return UE::Editor::Toolbars::ECreateStatusBarOptions::Default;
+}
+
 void FAssetEditorToolkit::AddMenuExtender(TSharedPtr<FExtender> Extender)
 {
 	if (TSharedPtr<SStandaloneAssetEditorToolkitHost> StandaloneHostPtr = StandaloneHost.Pin())
diff --git a/Engine/Source/Editor/UnrealEd/Private/Toolkits/SStandaloneAssetEditorToolkitHost.cpp b/Engine/Source/Editor/UnrealEd/Private/Toolkits/SStandaloneAssetEditorToolkitHost.cpp
--- a/Engine/Source/Editor/UnrealEd/Private/Toolkits/SStandaloneAssetEditorToolkitHost.cpp
+++ b/Engine/Source/Editor/UnrealEd/Private/Toolkits/SStandaloneAssetEditorToolkitHost.cpp
@@ -190,9 +190,15 @@ void SStandaloneAssetEditorToolkitHost::RestoreFromLayout( const TSharedRef<FTab
 
 	checkf(RestoredUI.IsValid(), TEXT("The layout must have a primary dock area"));
 
+	UE::Editor::Toolbars::ECreateStatusBarOptions StatusBarCreationOptions = UE::Editor::Toolbars::ECreateStatusBarOptions::Default;
+	if (HostedAssetEditorToolkit.IsValid())
+	{
+		StatusBarCreationOptions = HostedAssetEditorToolkit->GetStatusBarCreationOptions();
+	}
+
 	if (UStatusBarSubsystem* StatusBarSubsystem = GEditor ? GEditor->GetEditorSubsystem<UStatusBarSubsystem>() : nullptr)
 	{
-		StatusBarWidget = StatusBarSubsystem->MakeStatusBarWidget(StatusBarName, HostTab);
+		StatusBarWidget = StatusBarSubsystem->MakeStatusBarWidget(StatusBarName, HostTab, StatusBarCreationOptions);
 	}
 
 	this->ChildSlot
@@ -232,18 +238,35 @@ FName SStandaloneAssetEditorToolkitHost::GetMenuName() const
 	return *(FString(TEXT("AssetEditor.")) + MenuAppName.ToString() + TEXT(".MainMenu"));
 }
 
+TSharedRef<SWidget> SStandaloneAssetEditorToolkitHost::CreateMenuBar(FToolMenuContext& ToolMenuContext) const
+{
+	TSharedPtr<SWidget> ToolbarMenu;
+
+	const FName AssetEditorMenuName = GetMenuName();
+	if (HostedAssetEditorToolkit.IsValid())
+	{
+		// This might return nullptr if the implementation doesn't provide a custom toolbar menu.
+		ToolbarMenu = HostedAssetEditorToolkit->CreateMenuBar(MyTabManager, AssetEditorMenuName, ToolMenuContext);
+	}
+
+	if (!ToolbarMenu)
+	{
+		IMainFrameModule& MainFrameModule = FModuleManager::LoadModuleChecked<IMainFrameModule>( "MainFrame" );
+		ToolbarMenu = MainFrameModule.MakeMainMenu(MyTabManager, AssetEditorMenuName, ToolMenuContext);
+	}
+	
+	return ToolbarMenu.ToSharedRef();
+}
+
 void SStandaloneAssetEditorToolkitHost::GenerateMenus(bool bForceCreateMenu)
 {
 	if( bForceCreateMenu || DefaultMenuWidget != SNullWidget::NullWidget )
 	{
-		const FName AssetEditorMenuName = GetMenuName();
-
 		UAssetEditorToolkitMenuContext* ContextObject = NewObject<UAssetEditorToolkitMenuContext>();
 		ContextObject->Toolkit = HostedAssetEditorToolkit;
 		FToolMenuContext ToolMenuContext(HostedAssetEditorToolkit->GetToolkitCommands(), FExtender::Combine(MenuExtenders), ContextObject);
 		HostedAssetEditorToolkit->InitToolMenuContext(ToolMenuContext);
-		IMainFrameModule& MainFrameModule = FModuleManager::LoadModuleChecked<IMainFrameModule>( "MainFrame" );
-		DefaultMenuWidget = MainFrameModule.MakeMainMenu( MyTabManager, AssetEditorMenuName, ToolMenuContext );
+		DefaultMenuWidget = CreateMenuBar(ToolMenuContext);
 	}
 }
 
diff --git a/Engine/Source/Editor/UnrealEd/Private/Toolkits/SStandaloneAssetEditorToolkitHost.h b/Engine/Source/Editor/UnrealEd/Private/Toolkits/SStandaloneAssetEditorToolkitHost.h
--- a/Engine/Source/Editor/UnrealEd/Private/Toolkits/SStandaloneAssetEditorToolkitHost.h
+++ b/Engine/Source/Editor/UnrealEd/Private/Toolkits/SStandaloneAssetEditorToolkitHost.h
@@ -116,6 +116,9 @@ private:
 	void OnTabClosed(TSharedRef<SDockTab> TabClosed) const;
 	void ShutdownToolkitHost();
 
+	/** Create the main menu bar, either generated by the toolkit host or by the MainFrame module. */
+	TSharedRef<SWidget> CreateMenuBar(struct FToolMenuContext& ToolMenuContext) const;
+
 	FName GetMenuName() const;
 
 	/** Manages internal tab layout */
diff --git a/Engine/Source/Editor/UnrealEd/Public/Toolkits/AssetEditorToolkit.h b/Engine/Source/Editor/UnrealEd/Public/Toolkits/AssetEditorToolkit.h
--- a/Engine/Source/Editor/UnrealEd/Public/Toolkits/AssetEditorToolkit.h
+++ b/Engine/Source/Editor/UnrealEd/Public/Toolkits/AssetEditorToolkit.h
@@ -27,6 +27,11 @@ class SStandaloneAssetEditorToolkitHost;
 class FWorkspaceItem;
 class UToolMenu;
 
+namespace UE::Editor::Toolbars
+{
+	enum class ECreateStatusBarOptions : uint8;
+}
+
 DECLARE_DELEGATE_RetVal( bool, FRequestAssetEditorClose );
 DECLARE_DELEGATE_RetVal( void, FAssetEditorClosing);
 
@@ -233,6 +238,8 @@ public:
 	/** Get name used by tool menu */
 	UNREALED_API virtual FName GetToolMenuToolbarName(FName& OutParentName) const;
 	UNREALED_API virtual void InitToolMenuContext(FToolMenuContext& MenuContext);
+	/** Method that can be overriden by derived classes to provide a custom MenuBar to the Standalone Toolkit Host. */
+	UNREALED_API virtual TSharedPtr<SWidget> CreateMenuBar(const TSharedPtr<FTabManager>& InTabManager, const FName InMenuName, FToolMenuContext& InToolMenuContext) { return nullptr; }
 	UNREALED_API FName GetToolMenuToolbarName() const;
 	UNREALED_API FName GetToolMenuAppName() const;
 	UNREALED_API FName GetToolMenuName() const;
@@ -278,6 +285,9 @@ public:
 	/** True if this actually is editing an asset */
 	UNREALED_API bool IsActuallyAnAsset() const;
 
+	/** Called by StandaloneAssetEditorToolkitHost to determine what sections of the StatusBar should be created. */
+	UNREALED_API virtual UE::Editor::Toolbars::ECreateStatusBarOptions GetStatusBarCreationOptions() const;
+	
 	/**
 	 * Gets the text to display in a toolkit titlebar for an object
 	 * @param	InObject	The object we want a description of
@@ -327,7 +337,7 @@ protected:
 	UNREALED_API TArray<TObjectPtr<UObject>>& GetEditingObjectPtrs();
 
 	/** Generate the toolbar for common asset actions like Save*/
-	UNREALED_API UToolMenu* GenerateCommonActionsToolbar(FToolMenuContext& MenuContext);
+	virtual UNREALED_API UToolMenu* GenerateCommonActionsToolbar(FToolMenuContext& MenuContext);
 
 	/** Generate the toolbar for read only mode specific content */
 	UToolMenu* GenerateReadOnlyToolbar(FToolMenuContext& MenuContext);
@@ -442,6 +452,9 @@ protected:
 	/** Get the name of the profile registered with UToolMenus for read only menu customizations */
 	UNREALED_API virtual FName GetReadOnlyMenuProfileName() { return ReadOnlyMenuProfileName; }
 
+	/** Can be overriden to prevent registering the default commands.  */
+	UNREALED_API virtual void MapToolkitCommands();
+
 private:
 	// Callback for persisting the Asset Editor's layout.
 	void HandleTabManagerPersistLayout( const TSharedRef<FTabManager::FLayout>& LayoutToSave )
