Use the WorldPartitionRuntimeVirtualTextureBuilder to build the RVT for all maps
* This builder will correctly load all primitives that should be baked to the RVT
* Modified that builder to allow it to run on non-WP maps

diff --git a/Engine/Source/Editor/UnrealEd/Private/WorldPartition/WorldPartitionRuntimeVirtualTextureBuilder.cpp b/Engine/Source/Editor/UnrealEd/Private/WorldPartition/WorldPartitionRuntimeVirtualTextureBuilder.cpp
--- a/Engine/Source/Editor/UnrealEd/Private/WorldPartition/WorldPartitionRuntimeVirtualTextureBuilder.cpp
+++ b/Engine/Source/Editor/UnrealEd/Private/WorldPartition/WorldPartitionRuntimeVirtualTextureBuilder.cpp
@@ -44,16 +44,11 @@ bool UWorldPartitionRuntimeVirtualTextureBuilder::RunInternal(UWorld* World, con
 {
 	IVirtualTexturingEditorModule& VTModule = FModuleManager::Get().LoadModuleChecked<IVirtualTexturingEditorModule>("VirtualTexturingEditor");
 	
-	UWorldPartition* WorldPartition = World->GetWorldPartition();
-	if (!WorldPartition)
+	FWorldPartitionHelpers::FForEachActorWithLoadingResult ForEachActorWithLoadingResult;
+	if (UWorldPartition* WorldPartition = World->GetWorldPartition())
 	{
-		UE_LOG(LogWorldPartitionRuntimeVirtualTextureBuilder, Error, TEXT("Failed to retrieve WorldPartition."));
-		return false;
+		LoadRuntimeVirtualTextureActors(WorldPartition, ForEachActorWithLoadingResult);
 	}
-		
-	// Load required actors
-	FWorldPartitionHelpers::FForEachActorWithLoadingResult ForEachActorWithLoadingResult;
-	LoadRuntimeVirtualTextureActors(WorldPartition, ForEachActorWithLoadingResult);
 	
 	// Recreate render state after shader compilation complete
 	{
diff --git a/Engine/Source/Editor/UnrealEd/Public/WorldPartition/WorldPartitionRuntimeVirtualTextureBuilder.h b/Engine/Source/Editor/UnrealEd/Public/WorldPartition/WorldPartitionRuntimeVirtualTextureBuilder.h
--- a/Engine/Source/Editor/UnrealEd/Public/WorldPartition/WorldPartitionRuntimeVirtualTextureBuilder.h
+++ b/Engine/Source/Editor/UnrealEd/Public/WorldPartition/WorldPartitionRuntimeVirtualTextureBuilder.h
@@ -8,7 +8,7 @@
 
 class UWorldPartition;
 
-// Example Command Line: ProjectName MapName -run=WorldPartitionBuilderCommandlet -SCCProvider=Perforce -Builder=WorldPartitionRuntimeVirtualTextureBuilder
+// Example Command Line: ProjectName MapName -run=WorldPartitionBuilderCommandlet -SCCProvider=Perforce -Builder=WorldPartitionRuntimeVirtualTextureBuilder -AllowCommandletRendering
 UCLASS(MinimalAPI)
 class UWorldPartitionRuntimeVirtualTextureBuilder : public UWorldPartitionBuilder
 {
@@ -19,6 +19,7 @@ public:
 	virtual bool RequiresCommandletRendering() const override { return true; }
 	virtual ELoadingMode GetLoadingMode() const override { return ELoadingMode::Custom; }
 	UNREALED_API virtual bool RunInternal(UWorld* World, const FCellInfo& InCellInfo, FPackageSourceControlHelper& PackageHelper) override;
+	virtual bool CanProcessNonPartitionedWorlds() const override { return true; }
 	// UWorldPartitionBuilder interface end
 
 	// Helper method to load actors that contribute to the Runtime Virtual Texture
