Fixed Build > Build Landscape dirtying the main level. The landscape proxies' outer is already properly dirtied if necessary and there's nothing in the persistent level that requires to be dirty (it's also not happening when using the in-viewport actionable notification)

diff --git a/Engine/Source/Editor/UnrealEd/Private/EditorBuildUtils.cpp b/Engine/Source/Editor/UnrealEd/Private/EditorBuildUtils.cpp
--- a/Engine/Source/Editor/UnrealEd/Private/EditorBuildUtils.cpp
+++ b/Engine/Source/Editor/UnrealEd/Private/EditorBuildUtils.cpp
@@ -534,7 +534,15 @@ bool FEditorBuildUtils::EditorBuild( UWorld* InWorld, FName Id, const bool bAllo
 		if (bDoBuild)
 		{
 			GEditor->ResetTransaction(NSLOCTEXT("UnrealEd", "BuildAllLandscape", "Building Landscape"));
+
+			// We can't set the busy cursor for all windows, because lighting
+			// needs a cursor for the lighting options dialog.
+			const FScopedBusyCursor BusyCursor;
+
+			bShouldMapCheck = false;
 			EditorBuildAllLandscape(InWorld);
+
+			bDirtyPersistentLevel = false;
 		}
 	}
 	else if (Id == FBuildOptions::BuildAll || Id == FBuildOptions::BuildAllSubmit)
