Cooker: SoftGC logging spam: Remove diagnostic text for garbage collection when it was triggered by periodic SoftGC. Add a minimum period between periodic SoftGC since some text is still printed on every GC.
#rnx

diff --git a/Engine/Config/BaseEditor.ini b/Engine/Config/BaseEditor.ini
--- a/Engine/Config/BaseEditor.ini
+++ b/Engine/Config/BaseEditor.ini
@@ -382,6 +382,8 @@ SoftGCDenominator=10
 ; When SoftGC is insufficient and we run out of memory, a full GC will be executed instead and will not be considered
 ; part of the SoftGC's time budget.
 SoftGCTimeFractionBudget = 0.05
+; Minimum duration between the SoftGC triggered by SoftGCTimeFractionBudget, to prevent spam
+SoftGCMinimumPeriodSeconds = 30
 
 ; If true, iterative cooks will be allowed even if referenced ini values have changed.
 IterativeIgnoreIni=false
diff --git a/Engine/Source/Editor/UnrealEd/Classes/CookOnTheSide/CookOnTheFlyServer.h b/Engine/Source/Editor/UnrealEd/Classes/CookOnTheSide/CookOnTheFlyServer.h
--- a/Engine/Source/Editor/UnrealEd/Classes/CookOnTheSide/CookOnTheFlyServer.h
+++ b/Engine/Source/Editor/UnrealEd/Classes/CookOnTheSide/CookOnTheFlyServer.h
@@ -368,6 +368,7 @@ private:
 	int32 SoftGCStartNumerator = 5;
 	int32 SoftGCDenominator = 10;
 	float SoftGCTimeFractionBudget = 0.05;
+	float SoftGCMinimumPeriodSeconds = 30.f;
 	TUniquePtr<UE::Cook::FSoftGCHistory> SoftGCHistory;
 	bool bUseSoftGC = false;
 	bool bWarnedExceededMaxMemoryWithinGCCooldown = false;
@@ -591,9 +592,11 @@ public:
 		COSR_MarkedUpKeepPackages	= 0x00000040,
 		COSR_RequiresGC_OOM			= 0x00000080,
 		COSR_RequiresGC_PackageCount= 0x00000100,
-		COSR_RequiresGC_IdleTimer	= 0x00000200,
+		COSR_RequiresGC_Periodic	= 0x00000200,
 		COSR_YieldTick				= 0x00000400,
-		COSR_RequiresGC_Soft_OOM	= 0x00000800,
+		COSR_RequiresGC_Soft		= 0x00000800,
+		COSR_RequiresGC_IdleTimer   UE_DEPRECATED(5.6, "Use COSR_RequiresGC_Periodic") = 0x40000000,
+		COSR_RequiresGC_Soft_OOM	UE_DEPRECATED(5.6, "Use COSR_RequiresGC_Soft or COSR_RequiresGC_Periodic") = 0x80000000,
 	};
 
 	struct FCookByTheBookStartupOptions
diff --git a/Engine/Source/Editor/UnrealEd/Private/Commandlets/CookCommandlet.cpp b/Engine/Source/Editor/UnrealEd/Private/Commandlets/CookCommandlet.cpp
--- a/Engine/Source/Editor/UnrealEd/Private/Commandlets/CookCommandlet.cpp
+++ b/Engine/Source/Editor/UnrealEd/Private/Commandlets/CookCommandlet.cpp
@@ -44,6 +44,7 @@
 #include "Stats/StatsMisc.h"
 #include "Templates/UnrealTemplate.h"
 #include "UObject/Class.h"
+#include "UObject/GarbageCollection.h"
 #include "UObject/MetaData.h"
 #include "UObject/Package.h"
 #include "UObject/SavePackage.h"
@@ -654,55 +655,56 @@ void UCookCommandlet::ConditionalCollectGarbage(uint32 TickResults, UCookOnTheFl
 	}
 
 	double GCStartTime = FPlatformTime::Seconds();
+	bool bSoft = (TickResults & UCookOnTheFlyServer::COSR_RequiresGC_Soft) != 0;
+	bool bPeriodic = (TickResults & UCookOnTheFlyServer::COSR_RequiresGC_Periodic) != 0;
 	FString GCReason;
-	FString GCType = bPartialGC ? TEXT(" partial gc") : TEXT("");
+	FString GCType = bPartialGC && !bSoft ? TEXT(" partial gc") : TEXT("");
+
 	if ((TickResults & UCookOnTheFlyServer::COSR_RequiresGC_PackageCount) != 0)
 	{
 		GCReason = TEXT("Exceeded packages per GC");
 	}
-	else if ((TickResults & UCookOnTheFlyServer::COSR_RequiresGC_Soft_OOM) != 0)
-	{
-		GCType = TEXT("");
-		GCReason = TEXT("Soft GC");
-	}
 	else if ((TickResults & UCookOnTheFlyServer::COSR_RequiresGC_OOM) != 0)
 	{
 		// this can cause thrashing if the cooker loads the same stuff into memory next tick
 		GCReason = TEXT("Exceeded Max Memory");
 
-		int32 JobsToLogAt = GShaderCompilingManager->GetNumRemainingJobs();
-		double NextFlushMsgSeconds = FPlatformTime::Seconds();
-		UE_SCOPED_COOKTIMER_AND_DURATION(CookByTheBook_ShaderJobFlush, DetailedCookStats::ShaderFlushTimeSec);
-		UE_LOG(LogCookCommandlet, Display, TEXT("Detected max mem exceeded - forcing shader compilation flush"));
-		while (true)
+		if ((TickResults & UCookOnTheFlyServer::COSR_RequiresGC_Soft) != 0)
 		{
-			int32 NumRemainingJobs = GShaderCompilingManager->GetNumRemainingJobs();
-			if (NumRemainingJobs < 1000)
+			int32 JobsToLogAt = GShaderCompilingManager->GetNumRemainingJobs();
+			double NextFlushMsgSeconds = FPlatformTime::Seconds();
+			UE_SCOPED_COOKTIMER_AND_DURATION(CookByTheBook_ShaderJobFlush, DetailedCookStats::ShaderFlushTimeSec);
+			UE_LOG(LogCookCommandlet, Display, TEXT("Detected max mem exceeded - forcing shader compilation flush"));
+			while (true)
 			{
-				UE_LOG(LogCookCommandlet, Display, TEXT("Finished flushing shader jobs at %d"), NumRemainingJobs);
-				break;
-			}
+				int32 NumRemainingJobs = GShaderCompilingManager->GetNumRemainingJobs();
+				if (NumRemainingJobs < 1000)
+				{
+					UE_LOG(LogCookCommandlet, Display, TEXT("Finished flushing shader jobs at %d"), NumRemainingJobs);
+					break;
+				}
 
-			if (NumRemainingJobs < JobsToLogAt)
-			{
-				double Now = FPlatformTime::Seconds();
-				if (NextFlushMsgSeconds <= Now)
+				if (NumRemainingJobs < JobsToLogAt)
 				{
-					UE_LOG(LogCookCommandlet, Display, TEXT("Flushing shader jobs, remaining jobs %d"), NumRemainingJobs);
-					NextFlushMsgSeconds = Now + 10;
+					double Now = FPlatformTime::Seconds();
+					if (NextFlushMsgSeconds <= Now)
+					{
+						UE_LOG(LogCookCommandlet, Display, TEXT("Flushing shader jobs, remaining jobs %d"), NumRemainingJobs);
+						NextFlushMsgSeconds = Now + 10;
+					}
 				}
-			}
 
-			GShaderCompilingManager->ProcessAsyncResults(0.f, false);
+				GShaderCompilingManager->ProcessAsyncResults(0.f, false);
 
-			FPlatformProcess::Sleep(0.05);
+				FPlatformProcess::Sleep(0.05);
 
-			// GShaderCompilingManager->FinishAllCompilation();
+				// GShaderCompilingManager->FinishAllCompilation();
+			}
 		}
 	}
-	else if (TickResults & UCookOnTheFlyServer::COSR_RequiresGC_IdleTimer)
+	else if (bPeriodic)
 	{
-		GCReason = TEXT("Cooker has been idle for long time gc");
+		GCReason = TEXT("Periodic GC");
 	}
 	else
 	{
@@ -729,9 +731,23 @@ void UCookCommandlet::ConditionalCollectGarbage(uint32 TickResults, UCookOnTheFl
 	UE_LOG(LogCookCommandlet, Display, TEXT("GarbageCollection...%s (%s)"), *GCType, *GCReason);
 	{
 		TGuardValue<bool> SoftGCGuard(UPackage::bSupportCookerSoftGC, true);
+#if !NO_LOGGING
+		ELogVerbosity::Type SavedVerbosity = LogGarbage.GetVerbosity();
+		bool bSuppressGCLogs = bSoft && bPeriodic && SavedVerbosity > ELogVerbosity::Warning;
+		if (bSuppressGCLogs)
+		{
+			LogGarbage.SetVerbosity(ELogVerbosity::Warning);
+		}
+#endif
 		COTFS.OnCookerStartCollectGarbage(TickResults);
 		CollectGarbage(RF_NoFlags);
 		COTFS.OnCookerEndCollectGarbage(TickResults);
+#if !NO_LOGGING
+		if (bSuppressGCLogs)
+		{
+			LogGarbage.SetVerbosity(SavedVerbosity);
+		}
+#endif
 		if (COTFS.NeedsDiagnosticSecondGC())
 		{
 			UE_LOG(LogCookCommandlet, Display, TEXT("Second GarbageCollect requested by cooker..."));
diff --git a/Engine/Source/Editor/UnrealEd/Private/CookOnTheFlyServer.cpp b/Engine/Source/Editor/UnrealEd/Private/CookOnTheFlyServer.cpp
--- a/Engine/Source/Editor/UnrealEd/Private/CookOnTheFlyServer.cpp
+++ b/Engine/Source/Editor/UnrealEd/Private/CookOnTheFlyServer.cpp
@@ -1901,7 +1901,8 @@ FString UCookOnTheFlyServer::GetCookSettingsForMemoryLogText() const
 		(bUseSoftGC && SoftGCStartNumerator > 0) ? TEXT("enabled") : TEXT("disabled"),
 		(bUseSoftGC && SoftGCStartNumerator > 0) ? *FString::Printf(TEXT(" (%d/%d)"), SoftGCStartNumerator, SoftGCDenominator) : TEXT(""),
 		(bUseSoftGC && SoftGCTimeFractionBudget > 0) ? TEXT("enabled") : TEXT("disabled"),
-		(bUseSoftGC && SoftGCTimeFractionBudget > 0) ? *FString::Printf(TEXT(" (%.3g)"), SoftGCTimeFractionBudget) : TEXT("")
+		(bUseSoftGC && SoftGCTimeFractionBudget > 0) ? *FString::Printf(TEXT(" (%.3g budget, %.0fs min period)"),
+			SoftGCTimeFractionBudget, SoftGCMinimumPeriodSeconds) : TEXT("")
 	);
 }
 
@@ -5168,7 +5169,7 @@ void UCookOnTheFlyServer::TickRecompileShaderRequestsPrivate(UE::Cook::FTickStac
 	if (bProcessedRequests)
 	{
 		// Ask for GC to run again when we processed some shaders requests to ensure material get evicted and we don't keep their package open
-		StackData.ResultFlags |= COSR_RequiresGC | COSR_RequiresGC_IdleTimer | COSR_YieldTick;
+		StackData.ResultFlags |= COSR_RequiresGC | COSR_RequiresGC_Periodic | COSR_YieldTick;
 	}
 }
 
@@ -5551,6 +5552,7 @@ void FInitializeConfigSettings::LoadLocal(const FString& InOutputDirectoryOverri
 	GConfig->GetInt(TEXT("CookSettings"), TEXT("SoftGCStartNumerator"), SoftGCStartNumerator, GEditorIni);
 	GConfig->GetInt(TEXT("CookSettings"), TEXT("SoftGCDenominator"), SoftGCDenominator, GEditorIni);
 	GConfig->GetFloat(TEXT("CookSettings"), TEXT("SoftGCTimeFractionBudget"), SoftGCTimeFractionBudget, GEditorIni);
+	GConfig->GetFloat(TEXT("CookSettings"), TEXT("SoftGCMinimumPeriodSeconds"), SoftGCMinimumPeriodSeconds, GEditorIni);
 
 	MemoryExpectedFreedToSpreadRatio = 0.10f;
 	GConfig->GetFloat(TEXT("CookSettings"), TEXT("MemoryExpectedFreedToSpreadRatio"),
diff --git a/Engine/Source/Editor/UnrealEd/Private/Cooker/CookGarbageCollect.cpp b/Engine/Source/Editor/UnrealEd/Private/Cooker/CookGarbageCollect.cpp
--- a/Engine/Source/Editor/UnrealEd/Private/Cooker/CookGarbageCollect.cpp
+++ b/Engine/Source/Editor/UnrealEd/Private/Cooker/CookGarbageCollect.cpp
@@ -113,7 +113,7 @@ void FCookGCDiagnosticContext::OnCookerStartCollectGarbage(UCookOnTheFlyServer&
 	if (bRequestFullGC)
 	{
 		COTFS.bGarbageCollectTypeSoft = false;
-		ResultFlagsFromTick = ResultFlagsFromTick & ~UCookOnTheFlyServer::COSR_RequiresGC_Soft_OOM;
+		ResultFlagsFromTick = ResultFlagsFromTick & ~UCookOnTheFlyServer::COSR_RequiresGC_Soft;
 	}
 	bCurrentGCIsFull = !COTFS.bGarbageCollectTypeSoft;
 }
@@ -203,6 +203,14 @@ bool FSoftGCHistory::IsTriggeringWithinBudget(UCookOnTheFlyServer& COTFS, double
 	{
 		return false;
 	}
+	float TimeSinceLastGCSeconds = CurrentTimeSeconds - COTFS.LastSoftGCTime;
+	if (TimeSinceLastGCSeconds < COTFS.SoftGCMinimumPeriodSeconds)
+	{
+		// Don't allow triggering SoftGC too frequently, even if it is within budget. This prevents spam
+		// from log messages that get printed every time garbage is collected.
+		return false;
+	}
+
 	if (DurationHistory.Num() == 0)
 	{
 		if (OutDiagnostics)
@@ -223,7 +231,6 @@ bool FSoftGCHistory::IsTriggeringWithinBudget(UCookOnTheFlyServer& COTFS, double
 		}
 		return true;
 	}
-	float TimeSinceLastGCSeconds = CurrentTimeSeconds - COTFS.LastSoftGCTime;
 	float CurrentTimeBudget = TimeSinceLastGCSeconds *
 		COTFS.SoftGCTimeFractionBudget / (1 - COTFS.SoftGCTimeFractionBudget);
 	if (CurrentTimeBudget >= AverageDurationSeconds)
@@ -390,7 +397,7 @@ void UCookOnTheFlyServer::PollGarbageCollection(UE::Cook::FTickStackData& StackD
 			CurrentTime - IdleStatusStartTime > GetIdleTimeToGC() &&
 			IdleStatusStartTime > GetLastGCTime())
 		{
-			StackData.ResultFlags |= COSR_RequiresGC | COSR_RequiresGC_IdleTimer | COSR_YieldTick;
+			StackData.ResultFlags |= COSR_RequiresGC | COSR_RequiresGC_Periodic | COSR_YieldTick;
 			return;
 		}
 	}
@@ -469,6 +476,7 @@ bool UCookOnTheFlyServer::PumpHasExceededMaxMemory(uint32& OutResultFlags)
 		}
 	}
 
+	bool bPeriodicTriggered = false;
 	bool bPressureTriggered = false;
 	if (MemoryTriggerGCAtPressureLevel != FPlatformMemoryStats::EMemoryPressureStatus::Unknown)
 	{
@@ -553,6 +561,7 @@ bool UCookOnTheFlyServer::PumpHasExceededMaxMemory(uint32& OutResultFlags)
 					*TriggerDiagnostics);
 				bTriggerGC = true;
 				bIsSoftGC = true;
+				bPeriodicTriggered = true;
 			}
 		}
 	}
@@ -562,9 +571,9 @@ bool UCookOnTheFlyServer::PumpHasExceededMaxMemory(uint32& OutResultFlags)
 		return false;
 	}
 
-	// Don't allow a second GC (soft or normal) within the GC cooldown period after a full GC
+	// Don't allow a second OOM GC (soft or normal) within the GC cooldown period after a full GC, because this can cause thrashing
 	constexpr float GCCooldown = 60.f;
-	if (LastFullGCTime + GCCooldown > CurrentTime)
+	if (!bPeriodicTriggered && LastFullGCTime + GCCooldown > CurrentTime)
 	{
 		if (!bIsSoftGC && !bWarnedExceededMaxMemoryWithinGCCooldown)
 		{
@@ -583,17 +592,17 @@ bool UCookOnTheFlyServer::PumpHasExceededMaxMemory(uint32& OutResultFlags)
 	UE_LOG(LogCook, Display, TEXT("Garbage collection triggered (%s). Triggered by conditions:%s"),
 		TypeMessage, TriggerMessages.ToString());
 	OutResultFlags |= COSR_RequiresGC | COSR_YieldTick;
-	OutResultFlags |= COSR_RequiresGC_OOM;
+	OutResultFlags |= bPeriodicTriggered ? COSR_RequiresGC_Periodic : COSR_RequiresGC_OOM;
 	if (bIsSoftGC)
 	{
-		OutResultFlags |= COSR_RequiresGC_Soft_OOM;
+		OutResultFlags |= COSR_RequiresGC_Soft;
 	}
 	return true;
 }
 
 void UCookOnTheFlyServer::SetGarbageCollectType(uint32 ResultFlagsFromTick)
 {
-	bGarbageCollectTypeSoft = (ResultFlagsFromTick & COSR_RequiresGC_Soft_OOM);
+	bGarbageCollectTypeSoft = (ResultFlagsFromTick & COSR_RequiresGC_Soft);
 }
 
 void UCookOnTheFlyServer::ClearGarbageCollectType()
@@ -968,7 +977,7 @@ void UCookOnTheFlyServer::EvaluateGarbageCollectionResults(bool bWasDueToOOM, bo
 	};
 	bWarnedExceededMaxMemoryWithinGCCooldown = false;
 	LastGCTime = FPlatformTime::Seconds();
-	bool bWasSoftGC = ResultFlags & COSR_RequiresGC_Soft_OOM;
+	bool bWasSoftGC = ResultFlags & COSR_RequiresGC_Soft;
 	if (bWasSoftGC)
 	{
 		LastSoftGCTime = LastGCTime;
@@ -1007,7 +1016,7 @@ void UCookOnTheFlyServer::EvaluateGarbageCollectionResults(bool bWasDueToOOM, bo
 	{
 		return;
 	}
-	if (IsCookOnTheFlyMode() && !bWasDueToOOM)
+	if (!bWasDueToOOM)
 	{
 		return;
 	}
diff --git a/Engine/Source/Editor/UnrealEd/Private/Cooker/CookTypes.cpp b/Engine/Source/Editor/UnrealEd/Private/Cooker/CookTypes.cpp
--- a/Engine/Source/Editor/UnrealEd/Private/Cooker/CookTypes.cpp
+++ b/Engine/Source/Editor/UnrealEd/Private/Cooker/CookTypes.cpp
@@ -740,6 +740,7 @@ FCbWriter& operator<<(FCbWriter& Writer, const UE::Cook::FInitializeConfigSettin
 	Writer << "SoftGCStartNumerator" << Value.SoftGCStartNumerator;
 	Writer << "SoftGCDenominator" << Value.SoftGCDenominator;
 	Writer << "SoftGCTimeFractionBudget" << Value.SoftGCTimeFractionBudget;
+	Writer << "SoftGCMinimumPeriodSeconds" << Value.SoftGCMinimumPeriodSeconds;
 	Writer << "MinFreeUObjectIndicesBeforeGC" << Value.MinFreeUObjectIndicesBeforeGC;
 	Writer << "MaxNumPackagesBeforePartialGC" << Value.MaxNumPackagesBeforePartialGC;
 	Writer << "ConfigSettingDenyList" << Value.ConfigSettingDenyList;
@@ -777,6 +778,7 @@ bool LoadFromCompactBinary(FCbFieldView Field, UE::Cook::FInitializeConfigSettin
 	bOk = LoadFromCompactBinary(Field["SoftGCStartNumerator"], OutValue.SoftGCStartNumerator) & bOk;
 	bOk = LoadFromCompactBinary(Field["SoftGCDenominator"], OutValue.SoftGCDenominator) & bOk;
 	bOk = LoadFromCompactBinary(Field["SoftGCTimeFractionBudget"], OutValue.SoftGCTimeFractionBudget) & bOk;
+	bOk = LoadFromCompactBinary(Field["SoftGCMinimumPeriodSeconds"], OutValue.SoftGCMinimumPeriodSeconds) & bOk;
 	bOk = LoadFromCompactBinary(Field["MinFreeUObjectIndicesBeforeGC"], OutValue.MinFreeUObjectIndicesBeforeGC) & bOk;
 	bOk = LoadFromCompactBinary(Field["MaxNumPackagesBeforePartialGC"], OutValue.MaxNumPackagesBeforePartialGC) & bOk;
 	bOk = LoadFromCompactBinary(Field["ConfigSettingDenyList"], OutValue.ConfigSettingDenyList) & bOk;
@@ -807,6 +809,7 @@ void FInitializeConfigSettings::MoveOrCopy(SourceType&& Source, TargetType&& Tar
 	Target.SoftGCStartNumerator = Source.SoftGCStartNumerator;
 	Target.SoftGCDenominator = Source.SoftGCDenominator;
 	Target.SoftGCTimeFractionBudget = Source.SoftGCTimeFractionBudget;
+	Target.SoftGCMinimumPeriodSeconds = Source.SoftGCMinimumPeriodSeconds;
 	Target.MinFreeUObjectIndicesBeforeGC = Source.MinFreeUObjectIndicesBeforeGC;
 	Target.MaxNumPackagesBeforePartialGC = Source.MaxNumPackagesBeforePartialGC;
 	Target.ConfigSettingDenyList = MoveTempIfPossible(Source.ConfigSettingDenyList);
diff --git a/Engine/Source/Editor/UnrealEd/Private/Cooker/CookTypes.h b/Engine/Source/Editor/UnrealEd/Private/Cooker/CookTypes.h
--- a/Engine/Source/Editor/UnrealEd/Private/Cooker/CookTypes.h
+++ b/Engine/Source/Editor/UnrealEd/Private/Cooker/CookTypes.h
@@ -441,6 +441,7 @@ namespace UE::Cook
 		int32 SoftGCStartNumerator = 0;
 		int32 SoftGCDenominator = 1;
 		float SoftGCTimeFractionBudget = 0.f;
+		float SoftGCMinimumPeriodSeconds = 0.f;
 		TArray<FString> ConfigSettingDenyList;
 		/** max number of objects of a specific type which are allowed to async cache at once */
 		TMap<FName, int32> MaxAsyncCacheForType;
