Added more modes to cook.cvarcontrol (2 = stomp on preview cvars, 3 = stomp on all cvars)

diff --git a/Engine/Source/Developer/TargetPlatform/Private/TargetPlatformManagerModule.cpp b/Engine/Source/Developer/TargetPlatform/Private/TargetPlatformManagerModule.cpp
--- a/Engine/Source/Developer/TargetPlatform/Private/TargetPlatformManagerModule.cpp
+++ b/Engine/Source/Developer/TargetPlatform/Private/TargetPlatformManagerModule.cpp
@@ -30,7 +30,6 @@
 #include "DesktopPlatformModule.h"
 #include "Interfaces/ITurnkeySupportModule.h"
 #include "DataDrivenShaderPlatformInfo.h"
-#include "RHIShaderFormatDefinitions.inl"
 
 DEFINE_LOG_CATEGORY_STATIC(LogTargetPlatformManager, Log, All);
 
@@ -40,15 +39,6 @@ DEFINE_LOG_CATEGORY_STATIC(LogTargetPlatformManager, Log, All);
 #endif
 
 
-static int32 GCookTimeCVarControl = 0;
-static FAutoConsoleVariableRef CVarCookTimeCVarControl(
-	TEXT("cook.cvarcontrol"),
-	GCookTimeCVarControl,
-	TEXT("Controls how cvars for other platforms are managed during cooking:\n"
-		"0: Completely disabled\n"
-		"1: Redirects CVars for the cooking platform to the cooking DP. CURRENTLY ALSO ENABLES redirecting a ShaderPlatform to the cooking platform"
-		"(more to come)"));
-
 
 #if AUTOSDKS_ENABLED
 namespace UE::AutoSDK
@@ -286,36 +276,6 @@ public:
 			GetShaderFormats();
 		}
 
-		if (GCookTimeCVarControl == 1)
-		{
-			// hook up some global mappings
-			// @todo clear out existing mappings for each Invalidate?
-			for (const ITargetPlatform* Platform : GetActiveTargetPlatforms())
-			{
-				FName PlatformName = *Platform->IniPlatformName();
-				FName DPName = *Platform->CookingDeviceProfileName();
-				// register that when we cook for this platform, we will want to use the given DP when looking up CVar values
-				// if it matches the platform, we don't set it, because that's the default, so don't actually do anything special
-				if (DPName != PlatformName)
-				{
-					ConsoleVariablePlatformMapping::RegisterPlatformToDeviceProfileMapping(PlatformName, DPName);
-				}
-
-				if (Platform != GetRunningTargetPlatform())
-				{
-					// register that when we cook for this platform's shaderplatforms, we will want to use the given platform when looking up cvars
-					TArray<FName> ShaderFormats;
-					Platform->GetAllTargetedShaderFormats(ShaderFormats);
-					for (FName SF : ShaderFormats)
-					{
-						EShaderPlatform SP = ShaderFormatNameToShaderPlatform(SF);
-						ConsoleVariablePlatformMapping::RegisterShaderPlatformToPlatformMapping((int)SP, PlatformName);
-					}
-				}
-			}
-		}
-
-
 		OnTargetPlatformsInvalidated.Broadcast();
 	}
 
diff --git a/Engine/Source/Editor/UnrealEd/Private/CookOnTheFlyServer.cpp b/Engine/Source/Editor/UnrealEd/Private/CookOnTheFlyServer.cpp
--- a/Engine/Source/Editor/UnrealEd/Private/CookOnTheFlyServer.cpp
+++ b/Engine/Source/Editor/UnrealEd/Private/CookOnTheFlyServer.cpp
@@ -225,6 +225,18 @@ static FAutoConsoleVariableRef CVarCookDisplayWarnBusyTime(
 	TEXT("Controls the time before the cooker will issue a warning that there is a deadlock in a busy queue.\n"),
 	ECVF_Default);
 
+static int32 GCookTimeCVarControl = 0;
+static FAutoConsoleVariableRef CVarCookTimeCVarControl(
+	TEXT("cook.cvarcontrol"),
+	GCookTimeCVarControl,
+	TEXT("Controls how cvars for other platforms are managed during cooking:\n"
+		"0: Completely disabled\n"
+		"1: Redirects CVars for the cooking platform to the cooking DP\n"
+		"2: Performs mode 1, and will also update the active value of any CVar flagged with ECVF_Preview to match the platform/DP value\n"
+		"2: Same as mode 2, except it will update ALL CVars, not just ECVF_Preview CVars\n"
+		"NOTE: CURRENTLY any non-zero mode will also enable redirecting a ShaderPlatform to the cooking platform"),
+	ECVF_Default);
+
 ////////////////////////////////////////////////////////////////
 /// Cook on the fly server
 ///////////////////////////////////////////////////////////////
@@ -5480,6 +5492,9 @@ void UCookOnTheFlyServer::Initialize( ECookMode::Type DesiredCookMode, ECookInit
 #endif
 }
 
+
+
+
 void UCookOnTheFlyServer::InitializeAtFirstSession()
 {
 	UE::EditorDomain::UtilsCookInitialize();
@@ -5613,6 +5628,70 @@ static EMPCookGeneratorSplit ParseMPCookGeneratorSplitFromString(const FString&
 
 }
 
+#include "RHIShaderFormatDefinitions.inl"
+
+namespace UE::Cook::CVarControl
+{
+	void UpdateCVars(FBeginCookContext& BeginContext)
+	{
+		static FName CookTimeCVarTag("CookTimeCVars");
+
+		if (GCookTimeCVarControl == 2 || GCookTimeCVarControl == 3)
+		{
+			checkf(BeginContext.TargetPlatforms.Num() == 1, TEXT("When using cook.CVarControl in mode 2 or 3, only a single TargetPlatform may be cooked at once."));
+		}
+
+
+		// in case we had cooked before on this run, reset everything
+		IConsoleManager::Get().UnsetAllConsoleVariablesWithTag(CookTimeCVarTag);
+
+		ITargetPlatformManagerModule* TPM = GetTargetPlatformManager();
+
+		// hook up some global mappings
+		// @todo clear out existing mappings for each Invalidate?
+		for (const ITargetPlatform* Platform : BeginContext.TargetPlatforms)
+		{
+			FName PlatformName = *Platform->IniPlatformName();
+			FName DPName = *Platform->CookingDeviceProfileName();
+
+			if (GCookTimeCVarControl != 0)
+			{
+				// register that when we cook for this platform, we will want to use the given DP when looking up CVar values
+				// if it matches the platform, we don't set it, because that's the default, so don't actually do anything special
+				if (DPName != PlatformName)
+				{
+					ConsoleVariablePlatformMapping::RegisterPlatformToDeviceProfileMapping(PlatformName, DPName);
+				}
+			}
+			if (GCookTimeCVarControl == 2)
+			{
+				IConsoleManager::Get().StompPlatformCVars(PlatformName, DPName.ToString(), CookTimeCVarTag, ECVF_SetByCode, ECVF_Preview, ECVF_Cheat);
+			}
+			else if (GCookTimeCVarControl == 3)
+			{
+				IConsoleManager::Get().StompPlatformCVars(PlatformName, DPName.ToString(), CookTimeCVarTag, ECVF_SetByCode, ECVF_Default, ECVF_Default);
+			}
+
+			// for now, only do this if we are performing any cvar control modes
+			if (GCookTimeCVarControl != 0)
+			{
+				if (Platform != TPM->GetRunningTargetPlatform())
+				{
+					// register that when we cook for this platform's shaderplatforms, we will want to use the given platform when looking up cvars
+					TArray<FName> ShaderFormats;
+					Platform->GetAllTargetedShaderFormats(ShaderFormats);
+					for (FName SF : ShaderFormats)
+					{
+						EShaderPlatform SP = ShaderFormatNameToShaderPlatform(SF);
+						ConsoleVariablePlatformMapping::RegisterShaderPlatformToPlatformMapping((int)SP, PlatformName);
+					}
+				}
+			}
+		}
+
+	}
+}
+
 void UCookOnTheFlyServer::SetInitializeConfigSettings(UE::Cook::FInitializeConfigSettings&& Settings)
 {
 	Settings.MoveToLocal(*this);
@@ -9760,6 +9839,7 @@ void UCookOnTheFlyServer::SetBeginCookConfigSettings(FBeginCookContext& BeginCon
 	PackageDatas->SetBeginCookConfigSettings(Settings.CookShowInstigator);
 	UE::Cook::FGenerationHelper::SetBeginCookConfigSettings();
 	SetNeverCookPackageConfigSettings(BeginContext, Settings);
+	UE::Cook::CVarControl::UpdateCVars(BeginContext);
 }
 
 namespace UE::Cook
diff --git a/Engine/Source/Runtime/Core/Private/HAL/ConsoleManager.cpp b/Engine/Source/Runtime/Core/Private/HAL/ConsoleManager.cpp
--- a/Engine/Source/Runtime/Core/Private/HAL/ConsoleManager.cpp
+++ b/Engine/Source/Runtime/Core/Private/HAL/ConsoleManager.cpp
@@ -122,6 +122,7 @@ public:
 	virtual void LoadAllPlatformCVars(FName PlatformName, const FString& DeviceProfileName = FString()) override;
 	virtual void ClearAllPlatformCVars(FName PlatformName = NAME_None, const FString& DeviceProfileName = FString()) override;
 	virtual void PreviewPlatformCVars(FName PlatformName, const FString& DeviceProfileName, FName PreviewModeTag) override;
+	virtual void StompPlatformCVars(FName PlatformName, const FString& DeviceProfileName, FName Tag, EConsoleVariableFlags SetBy, EConsoleVariableFlags RequiredFlags, EConsoleVariableFlags DisallowedFlags) override;
 #endif
 
 	// Internal functionality to this file -----------------------------------
@@ -3305,10 +3306,16 @@ void FConsoleManager::LoadAllPlatformCVars(FName PlatformName, const FString& De
 
 void FConsoleManager::PreviewPlatformCVars(FName PlatformName, const FString& DeviceProfileName, FName PreviewModeTag)
 {
-	UE_LOG(LogConsoleManager, Display, TEXT("Previewing Platform '%s', DeviceProfile '%s', ModeTag '%s'"), *PlatformName.ToString(), *DeviceProfileName, *PreviewModeTag.ToString());
-	
+	// we want Preview but not Cheat
+	StompPlatformCVars(PlatformName, DeviceProfileName, PreviewModeTag, ECVF_SetByPreview, ECVF_Preview, ECVF_Cheat);
+}
+
+void FConsoleManager::StompPlatformCVars(FName PlatformName, const FString& DeviceProfileName, FName Tag, EConsoleVariableFlags SetBy, EConsoleVariableFlags RequiredFlags, EConsoleVariableFlags DisallowedFlags)
+{
+	UE_LOG(LogConsoleManager, Display, TEXT("Previewing/Stomping CVars for Platform '%s', DeviceProfile '%s', ModeTag '%s'"), *PlatformName.ToString(), *DeviceProfileName, *Tag.ToString());
+
 	LoadAllPlatformCVars(PlatformName, DeviceProfileName.Len() ? DeviceProfileName : PlatformName.ToString());
-	
+
 	FName PlatformKey = MakePlatformKey(PlatformName, DeviceProfileName);
 
 	for (auto Pair : ConsoleObjects)
@@ -3319,22 +3326,22 @@ void FConsoleManager::PreviewPlatformCVars(FName PlatformName, const FString& De
 		}
 		if (IConsoleVariable* CVar = Pair.Value->AsVariable())
 		{
-			// we want Preview but not Cheat
-			if ((CVar->GetFlags() & (ECVF_Preview | ECVF_Cheat)) == ECVF_Preview)
+			// require some flags, but disallow others
+			if ((CVar->GetFlags() & (RequiredFlags | DisallowedFlags)) == RequiredFlags)
 			{
-				EConsoleVariableFlags Flags = ECVF_SetByPreview;
+				EConsoleVariableFlags Flags = SetBy;
 				if (CVar->GetFlags() & ECVF_ScalabilityGroup)
 				{
 					// we want to set SG cvars so they can be queried, but not send updates so that we don't use host platform's cvars
 					Flags = (EConsoleVariableFlags)(Flags | ECVF_Set_SetOnly_Unsafe);
 				}
-				
+
 				// if we have a value for the platform, then set it in the real CVar
 				if (CVar->HasPlatformValueVariable(PlatformKey, GSpecialDPNameForPremadePlatformKey))
 				{
 					TSharedPtr<IConsoleVariable> PlatformCVar = CVar->GetPlatformValueVariable(PlatformKey, GSpecialDPNameForPremadePlatformKey);
-					CVar->Set(*PlatformCVar->GetString(), Flags, PreviewModeTag);
-					
+					CVar->Set(*PlatformCVar->GetString(), Flags, Tag);
+
 					UE_LOG(LogConsoleManager, Display, TEXT("  |-> Setting %s = %s"), *Pair.Key, *PlatformCVar->GetString());
 				}
 			}
diff --git a/Engine/Source/Runtime/Core/Public/HAL/IConsoleManager.h b/Engine/Source/Runtime/Core/Public/HAL/IConsoleManager.h
--- a/Engine/Source/Runtime/Core/Public/HAL/IConsoleManager.h
+++ b/Engine/Source/Runtime/Core/Public/HAL/IConsoleManager.h
@@ -1196,6 +1196,12 @@ struct IConsoleManager
 	 */
 	virtual void PreviewPlatformCVars(FName PlatformName, const FString& DeviceProfileName, FName PreviewModeTag) = 0;
 
+	/**
+	 * Applies the cvars from the platform/DeviceProfile pair (can be blank to use the platform's named DP) and sets at the given SetBy priority, with the PreviewModeTag
+	 * to be unset later. RequiredFlags must be present ont he CVar to be set, unless DisallowedFlags is present, then it will not set it
+	 */
+	virtual void StompPlatformCVars(FName PlatformName, const FString& DeviceProfileName, FName Tag, EConsoleVariableFlags SetBy, EConsoleVariableFlags RequiredFlags, EConsoleVariableFlags DisallowedFlags) = 0;
+
 	/**
 	 * Empties the cache for the given Platform/DP for all CVars. If using NAME_None for Platofrm, this will wipe every cached platform/DP value. If DeviceProfileName
 	 * is empty, it will use the PlatfomName as the DeviceProfile name
