Fix IWYU issue in SlateIconFinder

diff --git a/Engine/Source/Editor/UnrealEd/Public/ClassIconFinder.h b/Engine/Source/Editor/UnrealEd/Public/ClassIconFinder.h
--- a/Engine/Source/Editor/UnrealEd/Public/ClassIconFinder.h
+++ b/Engine/Source/Editor/UnrealEd/Public/ClassIconFinder.h
@@ -7,6 +7,7 @@
 #include "HAL/Platform.h"
 #include "Styling/SlateIconFinder.h"
 #include "Textures/SlateIcon.h"
+#include "UObject/Class.h"
 #include "UObject/NameTypes.h"
 #include "UObject/WeakObjectPtrTemplates.h"
 
