WorldPartitionRuntimeVirtualTextureBuilder: Added -AutoSubmit support

diff --git a/Engine/Source/Editor/UnrealEd/Private/WorldPartition/WorldPartitionRuntimeVirtualTextureBuilder.cpp b/Engine/Source/Editor/UnrealEd/Private/WorldPartition/WorldPartitionRuntimeVirtualTextureBuilder.cpp
--- a/Engine/Source/Editor/UnrealEd/Private/WorldPartition/WorldPartitionRuntimeVirtualTextureBuilder.cpp
+++ b/Engine/Source/Editor/UnrealEd/Private/WorldPartition/WorldPartitionRuntimeVirtualTextureBuilder.cpp
@@ -143,5 +143,14 @@ bool UWorldPartitionRuntimeVirtualTextureBuilder::RunInternal(UWorld* World, con
 	// Wait for VT Textures to be ready before saving
 	FAssetCompilingManager::Get().FinishAllCompilation();
 
-	return SavePackages(ModifiedPackages.Array(), PackageHelper, false);
+	if (!SavePackages(ModifiedPackages.Array(), PackageHelper, false))
+	{
+		return false;
+	}
+
+	TArray<FString> FilesToSubmit;
+	FilesToSubmit.Append(SourceControlHelpers::PackageFilenames(ModifiedPackages.Array()));
+
+	const FString ChangeDescription = FString::Printf(TEXT("Built RVT for world '%s'"), *World->GetName());
+	return OnFilesModified(FilesToSubmit, ChangeDescription);
 }
diff --git a/Engine/Source/Editor/UnrealEd/Public/WorldPartition/WorldPartitionRuntimeVirtualTextureBuilder.h b/Engine/Source/Editor/UnrealEd/Public/WorldPartition/WorldPartitionRuntimeVirtualTextureBuilder.h
--- a/Engine/Source/Editor/UnrealEd/Public/WorldPartition/WorldPartitionRuntimeVirtualTextureBuilder.h
+++ b/Engine/Source/Editor/UnrealEd/Public/WorldPartition/WorldPartitionRuntimeVirtualTextureBuilder.h
@@ -8,7 +8,7 @@
 
 class UWorldPartition;
 
-// Example Command Line: ProjectName MapName -run=WorldPartitionBuilderCommandlet -SCCProvider=Perforce -Builder=WorldPartitionRuntimeVirtualTextureBuilder -AllowCommandletRendering
+// Example Command Line: ProjectName MapName -run=WorldPartitionBuilderCommandlet -SCCProvider=Perforce -Builder=WorldPartitionRuntimeVirtualTextureBuilder -AllowCommandletRendering [-AutoSubmit]
 UCLASS(MinimalAPI)
 class UWorldPartitionRuntimeVirtualTextureBuilder : public UWorldPartitionBuilder
 {
