Add a GetAllOpenEditors utility function to Asset Editor Subsystem

diff --git a/Engine/Source/Editor/UnrealEd/Private/Subsystems/AssetEditorSubsystem.cpp b/Engine/Source/Editor/UnrealEd/Private/Subsystems/AssetEditorSubsystem.cpp
--- a/Engine/Source/Editor/UnrealEd/Private/Subsystems/AssetEditorSubsystem.cpp
+++ b/Engine/Source/Editor/UnrealEd/Private/Subsystems/AssetEditorSubsystem.cpp
@@ -365,6 +365,14 @@ TArray<UObject*> UAssetEditorSubsystem::GetAllEditedAssets()
 }
 
 
+TArray<IAssetEditorInstance*> UAssetEditorSubsystem::GetAllOpenEditors() const
+{
+	TArray<IAssetEditorInstance*> AllEditors;
+	OpenedEditors.GenerateKeyArray(AllEditors);
+	return AllEditors;
+}
+
+
 void UAssetEditorSubsystem::NotifyEditorOpeningPreWidgets(const TArray< UObject* >& Assets, IAssetEditorInstance* InInstance)
 {
 	EditorOpeningPreWidgetsEvent.Broadcast(Assets, InInstance);
diff --git a/Engine/Source/Editor/UnrealEd/Public/Subsystems/AssetEditorSubsystem.h b/Engine/Source/Editor/UnrealEd/Public/Subsystems/AssetEditorSubsystem.h
--- a/Engine/Source/Editor/UnrealEd/Public/Subsystems/AssetEditorSubsystem.h
+++ b/Engine/Source/Editor/UnrealEd/Public/Subsystems/AssetEditorSubsystem.h
@@ -178,6 +178,9 @@ public:
 	/** Get all assets currently being tracked with open editors */
 	UNREALED_API TArray<UObject*> GetAllEditedAssets();
 
+	/** Returns all editors that are currently alive */
+	UNREALED_API TArray<IAssetEditorInstance*> GetAllOpenEditors() const;
+
 	/** Notify the asset editor manager that an asset was opened */
 	UNREALED_API void NotifyAssetOpened(UObject* Asset, IAssetEditorInstance* Instance);
 	UNREALED_API void NotifyAssetsOpened(const TArray< UObject* >& Assets, IAssetEditorInstance* Instance);
