Improved a project setting to include/exclude localized variants when deleting the source assets or folders containing them. It previously would only allow including/exluding them when the operation would be renaming or moving the assets.

Project Setting: Editor/Assets/Internationalization/Manage Localized Variants Alongside Source Asset

[RN]

diff --git a/Engine/Source/Developer/AssetTools/Private/AssetRenameManager.h b/Engine/Source/Developer/AssetTools/Private/AssetRenameManager.h
--- a/Engine/Source/Developer/AssetTools/Private/AssetRenameManager.h
+++ b/Engine/Source/Developer/AssetTools/Private/AssetRenameManager.h
@@ -58,9 +58,6 @@ public:
 	/** Filters packages list depending on if it actually has soft object paths pointing to the specific object being renamed */
 	bool CheckPackageForSoftObjectReferences(UPackage* Package, const TMap<FSoftObjectPath, FSoftObjectPath>& AssetRedirectorMap, TMap<FSoftObjectPath, TArray<UObject*>>& OutReferencingObjects) const;
 
-	/** If this returns false, then ExpandAssetsRenameDataWithTheirVariants would just return a copy of its input */
-	bool CanExpandAssetsRenameDataWithTheirVariants(const TArray<FAssetRenameData>& InAssetsAndNames) const;
-
 private:
 	/** Callback used by DiscoverintAssetsDialog to call RenameAssetsAndVariants */
 	void RenameAssetsAndVariantsCallback(TArray<FAssetRenameData> InAssetsRenameData, bool bAutoCheckout, bool bWithDialog) const;
diff --git a/Engine/Source/Developer/AssetTools/Private/AssetViewUtils.cpp b/Engine/Source/Developer/AssetTools/Private/AssetViewUtils.cpp
--- a/Engine/Source/Developer/AssetTools/Private/AssetViewUtils.cpp
+++ b/Engine/Source/Developer/AssetTools/Private/AssetViewUtils.cpp
@@ -435,6 +435,13 @@ bool AssetViewUtils::DeleteFolders(const TArray<FString>& PathsToDelete)
 	TArray<FAssetData> AssetDataList;
 	GetAssetsInPaths(PathsToDelete, AssetDataList);
 
+	// Add localized variants associated with the current assets to delete if the user wants to include them
+	if (!ObjectTools::AddExtraLocalizedVariantsToDelete(AssetDataList, AssetDataList))
+	{
+		// Either the user requested to cancel the current operation or an error occurred that requires interrupting the current operation...
+		return false;
+	}
+
 	const int32 NumAssetsInPaths = AssetDataList.Num();
 	bool bAllowFolderDelete = false;
 	if ( NumAssetsInPaths == 0 )
diff --git a/Engine/Source/Developer/SourceControl/Public/SourceControlPreferences.h b/Engine/Source/Developer/SourceControl/Public/SourceControlPreferences.h
--- a/Engine/Source/Developer/SourceControl/Public/SourceControlPreferences.h
+++ b/Engine/Source/Developer/SourceControl/Public/SourceControlPreferences.h
@@ -57,6 +57,6 @@ public:
 	TMap<FName, FString> SpecificCollectionChangelistTags;
 
 	/** Revision Control is a requirement to move/rename/delete localizable assets */
-	UPROPERTY(config, EditAnywhere, Category = Internationalization, meta = (ToolTip = "If the project uses Localization and its localized content might not be on disk (but still in your Revision Control), then enabling this feature will require a Revision Control Provider to be configured to move/rename/delete files to be sure no localized variants become orphaned. (Works only with Perforce for now)"))
+	UPROPERTY(config, EditAnywhere, DisplayName = "Requires Revision Control To Manage Localizable Assets", Category = Internationalization, meta = (ToolTip = "If the project uses Localization and its localized content might not be on disk (but still in your Revision Control), then enabling this feature will require a Revision Control Provider to be configured to move/rename/delete files to be sure no localized variants become orphaned. (Works only with Perforce for now)"))
 	bool bRequiresRevisionControlToRenameLocalizableAssets = false;
 };
diff --git a/Engine/Source/Editor/UnrealEd/Private/ObjectTools.cpp b/Engine/Source/Editor/UnrealEd/Private/ObjectTools.cpp
--- a/Engine/Source/Editor/UnrealEd/Private/ObjectTools.cpp
+++ b/Engine/Source/Editor/UnrealEd/Private/ObjectTools.cpp
@@ -68,6 +68,7 @@
 #include "PackageTools.h"
 #include "Internationalization/TextPackageNamespaceUtil.h"
 #include "Framework/Application/SlateApplication.h"
+#include "ILocalizedAssetTools.h"
 
 #include "BusyCursor.h"
 #include "Dialogs/DlgMoveAssets.h"
@@ -108,6 +109,7 @@
 #include "UncontrolledChangelistsModule.h"
 #include "AssetCompilingManager.h"
 #include "ObjectEditorUtils.h"
+#include "Settings/EditorProjectSettings.h"
 #include "Settings/EditorStyleSettings.h"
 #include "ProfilingDebugging/AssetMetadataTrace.h"
 #include "Engine/DataTable.h"
@@ -2698,11 +2700,27 @@ namespace ObjectTools
 
 	int32 DeleteAssets( const TArray<FAssetData>& AssetsToDelete, bool bShowConfirmation )
 	{
+		TArray<FAssetData> AllAssetsToDelete;
+		const UEditorProjectAssetSettings* Settings = GetDefault<UEditorProjectAssetSettings>();
+		bool bAddExtraLocalizedVariants = (Settings == nullptr || Settings->bRenameLocalizedVariantsAlongsideSourceAsset);
+		if (bAddExtraLocalizedVariants && bShowConfirmation)
+		{
+			if (!AddExtraLocalizedVariantsToDelete(AssetsToDelete, AllAssetsToDelete))
+			{
+				// Either the user requested to cancel the current operation or an error occurred that requires interrupting the current operation...
+				return 0;
+			}
+		}
+		else
+		{
+			AllAssetsToDelete = AssetsToDelete;
+		}
+
 		TArray<TWeakObjectPtr<UPackage>> PackageFilesToDelete;
 		TArray<UObject*> ObjectsToDelete;
-		for ( int i = 0; i < AssetsToDelete.Num(); i++ )
+		for ( int i = 0; i < AllAssetsToDelete.Num(); i++ )
 		{
-			const FAssetData& AssetData = AssetsToDelete[i];
+			const FAssetData& AssetData = AllAssetsToDelete[i];
 			UObject *ObjectToDelete = AssetData.GetAsset({ ULevel::LoadAllExternalObjectsTag });
 			// Assets can be loaded even when their underlying type/class no longer exists...
 			if ( ObjectToDelete!=nullptr )
@@ -2811,6 +2829,131 @@ namespace ObjectTools
 		}
 	}
 
+	bool AddExtraLocalizedVariantsToDelete(const TArray<FAssetData>& InAssetsToDelete, TArray<FAssetData>& OutAssetsWithVariantsToDelete)
+	{
+		IAssetTools& AssetTools = FModuleManager::LoadModuleChecked<FAssetToolsModule>("AssetTools").Get();
+		TSharedPtr<ILocalizedAssetTools> LocalizedAssetTools = AssetTools.GetLocalizedAssetTools();
+
+		TArray<FName> PackageNames;
+		Algo::Transform(InAssetsToDelete, PackageNames, [](const FAssetData& InAssetData) { return InAssetData.PackageName; });
+
+		TMap<FName, TArray<FName>> VariantsBySourcesOnDisk;
+		TMap<FName, TArray<FName>> VariantsBySourcesOnlyInRevisionControl;
+		TArray<FName> PackagesNotFound;
+		bool bNeedToCheckInRevisionControl = USourceControlPreferences::RequiresRevisionControlToRenameLocalizableAssets();
+		ELocalizedAssetsResult Result = LocalizedAssetTools->GetLocalizedVariants(PackageNames, VariantsBySourcesOnDisk, bNeedToCheckInRevisionControl, VariantsBySourcesOnlyInRevisionControl, &PackagesNotFound);
+		bool bRevisionControlWasNeeded = (Result == ELocalizedAssetsResult::RevisionControlNotAvailable);
+
+		// We have all the assets and variants that we need
+		TArray<FName> PackagesAndVariants;
+		for (const TPair<FName, TArray<FName>>& LocalizedVariantsWithSourceAsset : VariantsBySourcesOnDisk)
+		{
+			PackagesAndVariants.Add(LocalizedVariantsWithSourceAsset.Key); // Add Source Asset
+			for (const FName& LocalizedVariant : LocalizedVariantsWithSourceAsset.Value)
+			{
+				PackagesAndVariants.Add(LocalizedVariant); // Add all Localized Variants
+			}
+		}
+		for (const FName& OtherAsset : PackagesNotFound)
+		{
+			PackagesAndVariants.Add(OtherAsset); // Add the other assets
+		}
+
+		// Error processing
+		if (!VariantsBySourcesOnlyInRevisionControl.IsEmpty())
+		{
+			// Files in Revision Control are needed
+			TArray<FText> FilesNeeded;
+			for (const TPair<FName, TArray<FName>>& LocalizedVariantsWithSourceAssetInSCC : VariantsBySourcesOnlyInRevisionControl)
+			{
+				for (const FName& LocalizedVariantInSCC : LocalizedVariantsWithSourceAssetInSCC.Value)
+				{
+					FilesNeeded.Add(FText::AsCultureInvariant(LocalizedVariantInSCC.ToString()));
+				}
+			}
+
+			UE_LOG(LogObjectTools, Error, TEXT("A file that needs to be renamed was detected in Revision Control but not on your disk."));
+			LocalizedAssetTools->OpenFilesInRevisionControlRequiredDialog(FilesNeeded);
+
+			if (&OutAssetsWithVariantsToDelete != &InAssetsToDelete)
+			{
+				OutAssetsWithVariantsToDelete = InAssetsToDelete;
+			}
+			return false;
+		}
+		else if (bRevisionControlWasNeeded)
+		{
+			// Revision Control is needed. Extraction must fail.
+			UE_LOG(LogObjectTools, Error, TEXT("%s"), *LocalizedAssetTools->GetRevisionControlIsNotAvailableWarningText().ToString());
+			LocalizedAssetTools->OpenRevisionControlRequiredDialog();
+
+			if (&OutAssetsWithVariantsToDelete != &InAssetsToDelete)
+			{
+				OutAssetsWithVariantsToDelete = InAssetsToDelete;
+			}
+			return false;
+		}
+
+		TArray<FText> NewAssetsAdded;
+		for (const FName& Package : PackagesAndVariants)
+		{
+			if (!PackageNames.Contains(Package))
+			{
+				FString PackageNameStr = Package.ToString();
+				UE_LOG(LogObjectTools, Display, TEXT("Added asset to extraction (while checking for localized variants): %s"), *PackageNameStr);
+				NewAssetsAdded.Add(FText::AsCultureInvariant(PackageNameStr));
+			}
+		}
+
+		if (NewAssetsAdded.IsEmpty())
+		{
+			if (&OutAssetsWithVariantsToDelete != &InAssetsToDelete)
+			{
+				OutAssetsWithVariantsToDelete = InAssetsToDelete;
+			}
+			return true;
+		}
+
+		ELocalizedVariantsInclusion InclusionResult = LocalizedAssetTools->OpenIncludeLocalizedVariantsListDialog(NewAssetsAdded);
+		switch (InclusionResult)
+		{
+		case ELocalizedVariantsInclusion::Exclude:
+			if (&OutAssetsWithVariantsToDelete != &InAssetsToDelete)
+			{
+				OutAssetsWithVariantsToDelete = InAssetsToDelete;
+			}
+			return true;
+		case ELocalizedVariantsInclusion::Cancel:
+			return false;
+		case ELocalizedVariantsInclusion::Include:
+		default:
+			// The inclusion is already done at this point. Only the exclude/cancel options need to do some revert work.
+			break;
+		}
+
+		// Get the Asset Registry
+		IAssetRegistry& AssetRegistry = IAssetRegistry::GetChecked();
+
+		// Don't add localized variants if the asset registry is not ready
+		if (AssetRegistry.IsLoadingAssets())
+		{
+			FNotificationInfo Info(NSLOCTEXT("UnrealEd", "Warning_CantDeleteRebuildingAssetRegistry", "Unable To Delete While Discovering Assets"));
+			Info.ExpireDuration = 3.0f;
+			FSlateNotificationManager::Get().AddNotification(Info);
+			return false;
+		}
+
+		// Iterate over each package path to convert to FAssetData
+		for (const FName& PackagePath : PackagesAndVariants)
+		{
+			TArray<FAssetData> FoundAssets;
+			AssetRegistry.GetAssetsByPackageName(PackagePath, FoundAssets);
+			OutAssetsWithVariantsToDelete.Append(FoundAssets);
+		}
+
+		return true;
+	}
+
 	bool ContainsWorldInUse(const TArray< UObject* >& ObjectsToDelete)
 	{
 		TArray<const UWorld*> WorldsToDelete;
diff --git a/Engine/Source/Editor/UnrealEd/Public/ObjectTools.h b/Engine/Source/Editor/UnrealEd/Public/ObjectTools.h
--- a/Engine/Source/Editor/UnrealEd/Public/ObjectTools.h
+++ b/Engine/Source/Editor/UnrealEd/Public/ObjectTools.h
@@ -384,6 +384,16 @@ namespace ObjectTools
 	 */
 	UNREALED_API void AddExtraObjectsToDelete(TArray< UObject* >& ObjectsToDelete);
 
+	/**
+	 * Gathers additional localized variants to delete
+	 *
+	 * @param	InAssetsToDelete				Initial list of assets to delete
+	 * @param	OutAssetsWithVariantsToDelete	Final list of assets to delete (initial + variants if the user requested it)
+	 *
+	 * @return Whether or not the function succeeded. If it returns false, the delete process should be interrupted.
+	 */
+	UNREALED_API bool AddExtraLocalizedVariantsToDelete(const TArray<FAssetData>& InAssetsToDelete, TArray<FAssetData>& OutAssetsWithVariantsToDelete);
+
 	/**
 	 * Utility function to compose a string list of referencing objects
 	 *
diff --git a/Engine/Source/Editor/UnrealEd/Public/Settings/EditorProjectSettings.h b/Engine/Source/Editor/UnrealEd/Public/Settings/EditorProjectSettings.h
--- a/Engine/Source/Editor/UnrealEd/Public/Settings/EditorProjectSettings.h
+++ b/Engine/Source/Editor/UnrealEd/Public/Settings/EditorProjectSettings.h
@@ -287,9 +287,9 @@ public:
 	bool bPromptToDeleteUnreferencedRedirectors = true;
 	
 	/**
-	 * Rename localized variants alongside their source asset
+	 * Manage (rename/move/delete) localized variants alongside their source asset
 	 */
-	UPROPERTY(EditAnywhere, config, Category = Internationalization)
+	UPROPERTY(EditAnywhere, DisplayName = "Manage Localized Variants Alongside Source Asset", config, Category = Internationalization, meta = (ToolTip = "If the project uses Localization, localized variants of assets will be moved/renamed/deleted alongside their source asset."))
 	bool bRenameLocalizedVariantsAlongsideSourceAsset = true;
 };
 
