IncrementalCook: TransitiveBuildDependencies: Mark packages that are referenced as transitive build dependencies from other packages, stored on the cookdirector. In an upcoming change the cooker will store ops in the oplog for packages that are referenced  even if they are not runtime referenced, so that the next incremental cook does not invalidate all of the packages that depend on them.
#rnx

diff --git a/Engine/Source/Editor/UnrealEd/Private/CookOnTheFlyServer.cpp b/Engine/Source/Editor/UnrealEd/Private/CookOnTheFlyServer.cpp
--- a/Engine/Source/Editor/UnrealEd/Private/CookOnTheFlyServer.cpp
+++ b/Engine/Source/Editor/UnrealEd/Private/CookOnTheFlyServer.cpp
@@ -3293,11 +3293,21 @@ void UCookOnTheFlyServer::QueueDiscoveredPackage(UE::Cook::FPackageData& Package
 		DiscoveredPlatforms = ReachablePlatforms.GetPlatforms(*this, &Instigator,
 			TConstArrayView<const ITargetPlatform*>(), &BufferPlatforms);
 	}
-	if (Instigator.Category != EInstigator::ForceExplorableSaveTimeSoftDependency &&
-		PackageData.HasReachablePlatforms(DiscoveredPlatforms))
+	if (Instigator.Category == EInstigator::BuildDependency)
 	{
-		// Not a new discovery; ignore
-		return;
+		if (PackageData.IsBuildDependency(DiscoveredPlatforms))
+		{
+			// Not a new discovery; ignore
+			return;
+		}
+	}
+	else if (Instigator.Category != EInstigator::ForceExplorableSaveTimeSoftDependency)
+	{
+		if (PackageData.HasReachablePlatforms(DiscoveredPlatforms))
+		{
+			// Not a new discovery; ignore
+			return;
+		}
 	}
 
 	if (bHiddenDependenciesDebug)
@@ -3329,7 +3339,17 @@ void UCookOnTheFlyServer::QueueDiscoveredPackageOnDirector(UE::Cook::FPackageDat
 	if (!CookByTheBookOptions->bSkipHardReferences ||
 		(Instigator.Category == EInstigator::GeneratedPackage))
 	{
-		PackageData.QueueAsDiscovered(MoveTemp(Instigator), MoveTemp(ReachablePlatforms), Urgency);
+		if (Instigator.Category == EInstigator::BuildDependency)
+		{
+			TArray<const ITargetPlatform*, TInlineAllocator<ExpectedMaxNumPlatforms>> BufferPlatforms;
+			TConstArrayView<const ITargetPlatform*> DiscoveredPlatforms = ReachablePlatforms.GetPlatforms(
+				*this, &Instigator, TConstArrayView<const ITargetPlatform*>(), &BufferPlatforms);
+			PackageData.AddBuildDependencyPlatforms(DiscoveredPlatforms, MoveTemp(Instigator));
+		}
+		else
+		{
+			PackageData.QueueAsDiscovered(MoveTemp(Instigator), MoveTemp(ReachablePlatforms), Urgency);
+		}
 	}
 }
 
@@ -6116,6 +6136,7 @@ TArray<UE::Cook::FInstigator> UCookOnTheFlyServer::GetInstigatorChain(FName Pack
 			case EInstigator::SaveTimeSoftDependency: bGetNext = true; break;
 			case EInstigator::ForceExplorableSaveTimeSoftDependency: bGetNext = true; break;
 			case EInstigator::GeneratedPackage: bGetNext = true; break;
+			case EInstigator::BuildDependency: bGetNext = true; break;
 			default: break;
 		}
 		if (!bGetNext)
@@ -9003,6 +9024,30 @@ void UCookOnTheFlyServer::CookByTheBookFinishedInternal()
 	check(PackageDatas->GetSaveQueue().IsEmpty());
 	check(PackageDatas->GetSaveStalledSet().IsEmpty());
 
+	if (bHybridIterativeEnabled)
+	{
+		// Report the packages that are build dependencies and that were not saved
+		// TODO: cook these packages instead of just reporting them
+		TArray<FPackageData*> UncookedBuildDependencies;
+		PackageDatas->LockAndEnumeratePackageDatas([&UncookedBuildDependencies](FPackageData* PackageData)
+			{
+				for (const TPair<const ITargetPlatform*, FPackagePlatformData>& PlatformPair :
+					PackageData->GetPlatformDatas())
+				{
+					const FPackagePlatformData& Data = PlatformPair.Value;
+					if (Data.IsBuildDependency())
+					{
+						if (!Data.IsCookAttempted())
+						{
+							UncookedBuildDependencies.Add(PackageData);
+							break;
+						}
+					}
+				}
+			});
+		UE_LOG(LogCook, Display, TEXT("UncookedBuildDependencies: %d packages."), UncookedBuildDependencies.Num());
+	}
+
 	TArray<FPackageData*> DanglingGenerationHelpers;
 	PackageDatas->LockAndEnumeratePackageDatas([&DanglingGenerationHelpers](FPackageData* PackageData)
 		{
@@ -9714,6 +9759,7 @@ void UCookOnTheFlyServer::ResetCook(TConstArrayView<TPair<const ITargetPlatform*
 			{
 				bool bResetResults = Pair.Value;
 				PlatformData->ResetReachable();
+				PlatformData->SetIsBuildDependency(false);
 				if (bResetResults)
 				{
 					PackageData->ClearCookResults(TargetPlatform);
@@ -12212,6 +12258,7 @@ void UCookOnTheFlyServer::OnDiscoveredPackageDebug(FName PackageName, const UE::
 	case EInstigator::StartupPackage: [[fallthrough]];
 	case EInstigator::GeneratedPackage: [[fallthrough]];
 	case EInstigator::ForceExplorableSaveTimeSoftDependency:
+	case EInstigator::BuildDependency:
 		// Not a Hidden dependency
 		return;
 	default:
diff --git a/Engine/Source/Editor/UnrealEd/Private/Cooker/CookPackageData.cpp b/Engine/Source/Editor/UnrealEd/Private/Cooker/CookPackageData.cpp
--- a/Engine/Source/Editor/UnrealEd/Private/Cooker/CookPackageData.cpp
+++ b/Engine/Source/Editor/UnrealEd/Private/Cooker/CookPackageData.cpp
@@ -56,9 +56,9 @@ static FAutoConsoleVariableRef CVarPollAsyncPeriod(
 //////////////////////////////////////////////////////////////////////////
 // FPackageData
 FPackagePlatformData::FPackagePlatformData()
-	: bReachable(0), bVisitedByCluster(0), bSaveTimedOut(0), bCookable(1), bExplorable(1), bExplorableOverride(0)
-	, bIterativelyUnmodified(0), bIterativelySkipped(0), bRegisteredForCachedObjectsInOuter(0),
-	CookResults((uint8)ECookResult::NotAttempted)
+	: bReachable(0), bIsBuildDependency(0), bVisitedByCluster(0), bSaveTimedOut(0), bCookable(1), bExplorable(1)
+	, bExplorableOverride(0) , bIterativelyUnmodified(0), bIterativelySkipped(0), bRegisteredForCachedObjectsInOuter(0)
+	, CookResults((uint8)ECookResult::NotAttempted)
 {
 }
 
@@ -92,7 +92,8 @@ bool FPackagePlatformData::NeedsCooking(const ITargetPlatform* PlatformItBelongs
 
 FPackageData::FPackageData(FPackageDatas& PackageDatas, const FName& InPackageName, const FName& InFileName)
 	: ParentGenerationHelper(nullptr), PackageName(InPackageName), FileName(InFileName), PackageDatas(PackageDatas)
-	, Instigator(EInstigator::NotYetRequested), Urgency(static_cast<uint32>(EUrgency::Normal)), bIsCookLast(0)
+	, Instigator(EInstigator::NotYetRequested), BuildInstigator(EInstigator::NotYetRequested)
+	, Urgency(static_cast<uint32>(EUrgency::Normal)), bIsCookLast(0)
 	, bIsVisited(0)
 	, bHasSaveCache(0), bPrepareSaveFailed(0), bPrepareSaveRequiresGC(0)
 	, MonitorCookResult((uint8)ECookResult::NotAttempted)
@@ -239,6 +240,48 @@ void FPackageData::AddReachablePlatformsInternal(FPackageData& PackageData,
 	}
 }
 
+void FPackageData::AddBuildDependencyPlatforms(TConstArrayView<const ITargetPlatform*> Platforms,
+	FInstigator&& InInstigator)
+{
+	if (Platforms.IsEmpty())
+	{
+		return;
+	}
+	if (BuildInstigator.Category == EInstigator::NotYetRequested)
+	{
+		BuildInstigator = MoveTemp(InInstigator);
+	}
+
+	for (const ITargetPlatform* Platform : Platforms)
+	{
+		FPackagePlatformData& PlatformData = FindOrAddPlatformData(Platform);
+		PlatformData.SetIsBuildDependency(true);
+	}
+}
+
+bool FPackageData::IsBuildDependency(const TArrayView<const ITargetPlatform* const>& Platforms) const
+{
+	if (Platforms.Num() == 0)
+	{
+		return true;
+	}
+	if (PlatformDatas.Num() == 0)
+	{
+		return false;
+	}
+
+	for (const ITargetPlatform* QueryPlatform : Platforms)
+	{
+		const FPackagePlatformData* PlatformData = PlatformDatas.Find(QueryPlatform);
+		if (!PlatformData || !PlatformData->IsBuildDependency())
+		{
+			return false;
+		}
+	}
+	return true;
+}
+
+
 void FPackageData::QueueAsDiscovered(FInstigator&& InInstigator, FDiscoveredPlatformSet&& ReachablePlatforms,
 	EUrgency InUrgency)
 {
@@ -412,6 +455,14 @@ void FPackageData::ResetReachable()
 	}
 }
 
+void FPackageData::ResetIsBuildDependency()
+{
+	for (TPair<const ITargetPlatform*, FPackagePlatformData>& Pair : PlatformDatas)
+	{
+		Pair.Value.SetIsBuildDependency(false);
+	}
+}
+
 void FPackageData::ClearCookResults()
 {
 	bool bModifiedCookAttempted = false;
@@ -2556,6 +2607,7 @@ void FPackageDatas::ClearCookedPlatforms()
 	LockAndEnumeratePackageDatas([](FPackageData* PackageData)
 	{
 		PackageData->ResetReachable();
+		PackageData->ResetIsBuildDependency();
 		PackageData->ClearCookResults();
 	});
 }
diff --git a/Engine/Source/Editor/UnrealEd/Private/Cooker/CookPackageData.h b/Engine/Source/Editor/UnrealEd/Private/Cooker/CookPackageData.h
--- a/Engine/Source/Editor/UnrealEd/Private/Cooker/CookPackageData.h
+++ b/Engine/Source/Editor/UnrealEd/Private/Cooker/CookPackageData.h
@@ -164,6 +164,10 @@ struct FPackagePlatformData
 	bool IsReachable() const { return bReachable != 0; }
 	void SetReachable(bool bValue) { bReachable = (uint32)bValue; }
 
+	/** Marked as a transitive build dependency from a reachable package, so we need to save its dependency data.*/
+	bool IsBuildDependency() const { return bIsBuildDependency != 0; }
+	void SetIsBuildDependency(bool bValue) { bIsBuildDependency = (uint32)bValue; }
+
 	/** The package has passed through a RequestCluster and its transitive dependencies were added. */
 	bool IsVisitedByCluster() const { return bVisitedByCluster != 0; }
 	void SetVisitedByCluster(bool bValue) { bVisitedByCluster = (uint32)bValue; }
@@ -235,6 +239,7 @@ struct FPackagePlatformData
 
 private:
 	uint32 bReachable : 1;
+	uint32 bIsBuildDependency : 1;
 	uint32 bVisitedByCluster : 1;
 	uint32 bSaveTimedOut : 1;
 	uint32 bCookable : 1;
@@ -370,6 +375,19 @@ public:
 	/** Accessor for RequestClusters to add reachable platforms directly without modifying dependent data. */
 	void AddReachablePlatforms(FRequestCluster& RequestCluster, TConstArrayView<const ITargetPlatform*> Platforms,
 		FInstigator&& InInstigator);
+	/**
+	 * Set that the packagedata is reachable as a build dependency for the given platform. This setting is ignored
+	 * if the package is already reachable for the given platform, and is cleared if it ever becomes reachable.
+	 */
+	void AddBuildDependencyPlatforms(TConstArrayView<const ITargetPlatform*> Platforms, FInstigator&& InInstigator);
+	/** Clear the IsBuildDependency flag for all platforms on the packagedata. */
+	void ResetIsBuildDependency();
+
+	/**
+	 * Return true if and only if the package is a BuildDependency for every element of Platforms.
+	 * Returns true if Platforms is empty.
+	 */
+	bool IsBuildDependency(const TArrayView<const ITargetPlatform* const>& Platforms) const;
 
 	/** Add the given reachable platforms to this PackageData and send it back to Request state for exploration. */
 	void QueueAsDiscovered(FInstigator&& InInstigator, FDiscoveredPlatformSet&& ReachablePlatforms, EUrgency InUrgency);
@@ -662,13 +680,20 @@ public:
 	 * Return the instigator for this package. The Instigator is the first code location or
 	 * referencing package that causes the package to enter the requested state.
 	 */
-	const FInstigator& GetInstigator() const { return Instigator; }
-	bool HasInstigator() const { return Instigator.Category != EInstigator::NotYetRequested; }
+	const FInstigator& GetInstigator() const;
+	bool HasInstigator() const;
 	/** Setting the instigator is mostly private - it should only be done during clustering. */
 	void SetInstigator(FRequestCluster& Cluster, FInstigator&& InInstigator);
 	void SetInstigator(FCookWorkerClient& Client, FInstigator&& InInstigator);
 	void SetInstigator(FGenerationHelper& InGenerationHelper, FInstigator&& InInstigator);
 
+	/**
+	 * Return the BuildInstigator for this package. The BuildInstigator is the first other package
+	 * that marked this package as a BuildDependency.
+	 */
+	const FInstigator& GetBuildInstigator() const;
+	bool HasBuildInstigator() const;
+
 	/** Get whether COTFS is keeping this package referenced referenced during GC. */
 	bool IsKeepReferencedDuringGC() const { return static_cast<bool>(bKeepReferencedDuringGC); }
 	/** Set whether COTFS is keeping this package referenced referenced during GC. */
@@ -801,6 +826,7 @@ private:
 	int32 CookedPlatformDataNextIndex = -1;
 	int32 NumRetriesBeginCacheOnObject = 0;
 	FInstigator Instigator;
+	FInstigator BuildInstigator;
 
 	FWorkerId WorkerAssignment = FWorkerId::Invalid();
 	FWorkerId WorkerAssignmentConstraint = FWorkerId::Invalid();
@@ -1595,4 +1621,25 @@ inline void FPackageData::SetDoesGeneratedRequireGenerator(ICookPackageSplitter:
 	DoesGeneratedRequireGeneratorValue = static_cast<uint32>(Value);
 }
 
+inline const FInstigator& FPackageData::GetInstigator() const
+{
+	return Instigator;
+}
+
+inline bool FPackageData::HasInstigator() const
+{
+	return Instigator.Category != EInstigator::NotYetRequested;
+}
+
+inline const FInstigator& FPackageData::GetBuildInstigator() const
+{
+	return BuildInstigator;
+}
+
+inline bool FPackageData::HasBuildInstigator() const
+{
+	return BuildInstigator.Category != EInstigator::NotYetRequested;
+}
+
+
 } // namespace UE::Cook
\ No newline at end of file
diff --git a/Engine/Source/Editor/UnrealEd/Private/Cooker/CookRequestCluster.cpp b/Engine/Source/Editor/UnrealEd/Private/Cooker/CookRequestCluster.cpp
--- a/Engine/Source/Editor/UnrealEd/Private/Cooker/CookRequestCluster.cpp
+++ b/Engine/Source/Editor/UnrealEd/Private/Cooker/CookRequestCluster.cpp
@@ -1775,7 +1775,7 @@ void FRequestCluster::FGraphSearch::FExploreEdgesContext::AddPlatformDependencyR
 
 void FRequestCluster::FGraphSearch::FExploreEdgesContext::ProcessPlatformAttachments(int32 PlatformIndex,
 	const ITargetPlatform* TargetPlatform, FFetchPlatformData& FetchPlatformData,
-	FPackagePlatformData& PackagePlatformData, UE::TargetDomain::FCookAttachments& PlatformAttachments,
+	FPackagePlatformData& PackagePlatformData, UE::TargetDomain::FIncrementalCookAttachments& PlatformAttachments,
 	bool bExploreDependencies)
 {
 	bool bFoundBuildDefinitions = false;
@@ -2087,10 +2087,11 @@ void FRequestCluster::FQueryVertexBatch::Send()
 		}
 		if (bAllHaveAlreadyCompletedFetch)
 		{
-			// We are contractually obligated to kick the vertex. Normally we would call FCookAttachments::Fetch with it
-			// and would then kick the vertex in our callback. Also, it might still be in the AsyncQueueResults for one
-			// of the platforms so it will be kicked by TickExplore pulling it out of the AsyncQueueResults. But if all
-			// requested platforms already previously pulled it out of AsyncQueueResults, then we need to kick it again.
+			// We are contractually obligated to kick the vertex. Normally we would call
+			// FIncrementalCookAttachments::Fetch with it and would then kick the vertex in our callback. Also, it
+			// might still be in the AsyncQueueResults for one of the platforms so it will be kicked by TickExplore
+			// pulling it out of the AsyncQueueResults. But if all requested platforms already previously pulled it
+			// out of AsyncQueueResults, then we need to kick it again.
 			ThreadSafeOnlyVars.KickVertex(Vertex);
 		}
 	}
@@ -2121,12 +2122,12 @@ void FRequestCluster::FQueryVertexBatch::Send()
 															// always use empty
 			)
 		{
-			TUniqueFunction<void(FName PackageName, UE::TargetDomain::FCookAttachments&& Result)> Callback =
-				[this, PlatformIndex](FName PackageName, UE::TargetDomain::FCookAttachments&& Attachments)
+			TUniqueFunction<void(FName PackageName, UE::TargetDomain::FIncrementalCookAttachments&& Result)> Callback =
+				[this, PlatformIndex](FName PackageName, UE::TargetDomain::FIncrementalCookAttachments&& Attachments)
 			{
 				RecordCacheResults(PackageName, PlatformIndex, MoveTemp(Attachments));
 			};
-			UE::TargetDomain::FCookAttachments::Fetch(PlatformData.PackageNames, FetchPlatformData.Platform,
+			UE::TargetDomain::FIncrementalCookAttachments::Fetch(PlatformData.PackageNames, FetchPlatformData.Platform,
 				FetchPlatformData.Writer, MoveTemp(Callback));
 		}
 		else
@@ -2141,7 +2142,7 @@ void FRequestCluster::FQueryVertexBatch::Send()
 			for (int32 PackageNameIndex = 0; PackageNameIndex < NumPackageNames; ++PackageNameIndex)
 			{
 				FName PackageName = PackageNamesData[PackageNameIndex];
-				UE::TargetDomain::FCookAttachments Attachments;
+				UE::TargetDomain::FIncrementalCookAttachments Attachments;
 				RecordCacheResults(PackageName, PlatformIndex, MoveTemp(Attachments));
 			}
 		}
@@ -2149,7 +2150,7 @@ void FRequestCluster::FQueryVertexBatch::Send()
 }
 
 void FRequestCluster::FQueryVertexBatch::RecordCacheResults(FName PackageName, int32 PlatformIndex,
-	UE::TargetDomain::FCookAttachments&& CookAttachments)
+	UE::TargetDomain::FIncrementalCookAttachments&& CookAttachments)
 {
 	FVertexData* Vertex = Vertices.FindChecked(PackageName);
 	FQueryPlatformData& PlatformData = Vertex->PlatformData[PlatformIndex];
diff --git a/Engine/Source/Editor/UnrealEd/Private/Cooker/CookRequestCluster.h b/Engine/Source/Editor/UnrealEd/Private/Cooker/CookRequestCluster.h
--- a/Engine/Source/Editor/UnrealEd/Private/Cooker/CookRequestCluster.h
+++ b/Engine/Source/Editor/UnrealEd/Private/Cooker/CookRequestCluster.h
@@ -121,7 +121,7 @@ private:
 		 * Thread synchronization: this field is write-once from the async thread and is not readable until
 		 * bSchedulerThreadFetchCompleted.
 		 */
-		UE::TargetDomain::FCookAttachments CookAttachments;
+		UE::TargetDomain::FIncrementalCookAttachments CookAttachments;
 		bool bSchedulerThreadFetchCompleted = false;
 		bool bExploreRequested = false;
 		bool bExploreCompleted = false;
@@ -167,7 +167,7 @@ private:
 		void Send();
 
 		void RecordCacheResults(FName PackageName, int32 PlatformIndex,
-			UE::TargetDomain::FCookAttachments&& CookAttachments);
+			UE::TargetDomain::FIncrementalCookAttachments&& CookAttachments);
 
 		struct FPlatformData
 		{
@@ -281,7 +281,7 @@ private:
 			void AddPlatformDependencyRange(TConstArrayView<FName> Range, int32 PlatformIndex, EInstigator InstigatorType);
 			void ProcessPlatformAttachments(int32 PlatformIndex, const ITargetPlatform* TargetPlatform,
 				FFetchPlatformData& FetchPlatformData, FPackagePlatformData& PackagePlatformData,
-				UE::TargetDomain::FCookAttachments& PlatformAttachments, bool bExploreDependencies);
+				UE::TargetDomain::FIncrementalCookAttachments& PlatformAttachments, bool bExploreDependencies);
 
 			void SetIsIterativelyUnmodified(int32 PlatformIndex, bool bIterativelyUnmodified,
 				FPackagePlatformData& PackagePlatformData);
diff --git a/Engine/Source/Editor/UnrealEd/Private/Cooker/CookSavePackage.cpp b/Engine/Source/Editor/UnrealEd/Private/Cooker/CookSavePackage.cpp
--- a/Engine/Source/Editor/UnrealEd/Private/Cooker/CookSavePackage.cpp
+++ b/Engine/Source/Editor/UnrealEd/Private/Cooker/CookSavePackage.cpp
@@ -5,6 +5,7 @@
 #include "AssetRegistry/AssetData.h"
 #include "AssetRegistry/IAssetRegistry.h"
 #include "Commandlets/AssetRegistryGenerator.h"
+#include "Cooker/CookDependency.h"
 #include "Cooker/CookDeterminismManager.h"
 #include "Cooker/CookDiagnostics.h"
 #include "Cooker/CookGenerationHelper.h"
@@ -396,10 +397,12 @@ void FSaveCookedPackageContext::FinishPlatform()
 	{
 		CalculatePlatformAgnosticRuntimeDependencies();
 		CalculatePlatformRuntimeDependencies();
+		// Note CalculateCookDependencies mutates this->SaveResult; it moves CookDependencies out of it
+		CalculateCookDependencies(GeneratedResult.GetPtrOrNull(), OverridePackageDependencies.GetPtrOrNull());
+		RecordPlatformBuildDependencies();
 
 		ICookedPackageWriter::FCommitPackageInfo Info;
-		// Note GetCommitAttachments mutates the SaveResult; it moves CookDependencies out of it
-		Info.Attachments = GetCommitAttachments(GeneratedResult.GetPtrOrNull(), OverridePackageDependencies.GetPtrOrNull());
+		Info.Attachments = GetCommitAttachments();
 		Info.Status = PackageResultToCommitStatus(SavePackageResult.Result);
 		Info.PackageName = Package->GetFName();
 		Info.PackageHash = AssetPackageData ? AssetPackageData->GetPackageSavedHash() : FIoHash();
@@ -470,28 +473,42 @@ void FSaveCookedPackageContext::FinishPackage()
 		if (PlatformsForPackage.Num() == 1)
 		{
 			TConstArrayView<const ITargetPlatform*> ReachablePlatforms(&TargetPlatform, 1);
-			for (const TPair<FPackageData*, EInstigator>& DependencyPair : PlatformDependencies[0])
+			for (const TPair<FPackageData*, EInstigator>& DependencyPair : PlatformDependencies[0].RuntimeDependencies)
 			{
 				COTFS.QueueDiscoveredPackage(*DependencyPair.Key,
-					UE::Cook::FInstigator(DependencyPair.Value, PackageFName), FDiscoveredPlatformSet(ReachablePlatforms));
+					FInstigator(DependencyPair.Value, PackageFName), FDiscoveredPlatformSet(ReachablePlatforms));
+			}
+			for (FPackageData* BuildDependency : PlatformDependencies[0].BuildDependencies)
+			{
+				COTFS.QueueDiscoveredPackage(*BuildDependency,
+					FInstigator(EInstigator::BuildDependency, PackageFName), FDiscoveredPlatformSet(ReachablePlatforms));
 			}
 		}
 		else
 		{
-			// Merge the platform dependencies for each package into a single QueueDiscoveredPackage call, if possible
-			// It will not be possible if the different platforms have different instigators, so track a list of
-			// platforms for each package for each instigator type.
 			TMap<FPackageData*, TMap<EInstigator, TArray<const ITargetPlatform*>>> PackagePlatformsForInstigator;
 			for (int32 LocalIndex = 0; LocalIndex < PlatformsForPackage.Num(); ++LocalIndex)
 			{
-				const TMap<FPackageData*, EInstigator>& CurrentDependencies = PlatformDependencies[LocalIndex];
+				const FPlatformDiscoveryData& DData = PlatformDependencies[LocalIndex];
 				const ITargetPlatform* CurrentPlatform = PlatformsForPackage[LocalIndex];
-				for (const TPair<FPackageData*, EInstigator>& PackagePlatformPair : CurrentDependencies)
+
+				// Merge the runtime dependencies for each package into a single QueueDiscoveredPackage call, if possible
+				// It will not be possible if the different platforms have different instigators, so track a list of
+				// platforms for each package for each instigator type.
+				for (const TPair<FPackageData*, EInstigator>& PackagePlatformPair : DData.RuntimeDependencies)
 				{
 					TMap<EInstigator, TArray<const ITargetPlatform*>>& TargetMap =
 						PackagePlatformsForInstigator.FindOrAdd(PackagePlatformPair.Key);
 					TargetMap.FindOrAdd(PackagePlatformPair.Value).Add(CurrentPlatform);
 				}
+
+				// Merge the build dependencies for each package into a single QueueDiscoveredPackage call.
+				for (FPackageData* BuildDependency : DData.BuildDependencies)
+				{
+					TMap<EInstigator, TArray<const ITargetPlatform*>>& TargetMap =
+						PackagePlatformsForInstigator.FindOrAdd(BuildDependency);
+					TargetMap.FindOrAdd(EInstigator::BuildDependency).Add(CurrentPlatform);
+				}
 			}
 			for (const TPair<FPackageData*, TMap<EInstigator, TArray<const ITargetPlatform*>>>& PackagePair
 				: PackagePlatformsForInstigator)
@@ -510,13 +527,13 @@ void FSaveCookedPackageContext::FinishPackage()
 	{
 		COTFS.ProcessUnsolicitedPackages();
 		TMap<FPackageData*, EInstigator> AllPlatformDependenciesBuffer;
-		TMap<FPackageData*, EInstigator>* AllPlatformDependencies = &PlatformDependencies[0];
+		TMap<FPackageData*, EInstigator>* AllPlatformDependencies = &PlatformDependencies[0].RuntimeDependencies;
 		if (PlatformsForPackage.Num() > 1)
 		{
 			AllPlatformDependencies = &AllPlatformDependenciesBuffer;
 			for (int32 LocalIndex = 0; LocalIndex < PlatformsForPackage.Num(); ++LocalIndex)
 			{
-				AllPlatformDependenciesBuffer.Append(PlatformDependencies[LocalIndex]);
+				AllPlatformDependenciesBuffer.Append(PlatformDependencies[LocalIndex].RuntimeDependencies);
 			}
 		}
 
@@ -605,8 +622,8 @@ void FSaveCookedPackageContext::CalculatePlatformAgnosticRuntimeDependencies()
 
 void FSaveCookedPackageContext::CalculatePlatformRuntimeDependencies()
 {
-	PlatformDependencies.IsValidIndex(PlatformIndex);
-	TMap<FPackageData*, EInstigator>& CurrentDependencies = PlatformDependencies[PlatformIndex];
+	check(PlatformDependencies.IsValidIndex(PlatformIndex));
+	TMap<FPackageData*, EInstigator>& CurrentDependencies = PlatformDependencies[PlatformIndex].RuntimeDependencies;
 	CurrentDependencies.Reset();
 	// Add PlatformAgnostic dependencies to the PlatformDependencies
 	CurrentDependencies.Append(PlatformAgnosticDependencies);
@@ -634,19 +651,17 @@ void FSaveCookedPackageContext::CalculatePlatformRuntimeDependencies()
 TArray<FName> FSaveCookedPackageContext::GetPlatformRuntimeDependencies() const
 {
 	TArray<FName> PlatformDependencyNames;
-	PlatformDependencyNames.Reserve(PlatformDependencies[PlatformIndex].Num());
-	for (const TPair<FPackageData*, EInstigator>& DependencyPair : PlatformDependencies[PlatformIndex])
+	PlatformDependencyNames.Reserve(PlatformDependencies[PlatformIndex].RuntimeDependencies.Num());
+	for (const TPair<FPackageData*, EInstigator>& DependencyPair : PlatformDependencies[PlatformIndex].RuntimeDependencies)
 	{
 		PlatformDependencyNames.Add(DependencyPair.Key->GetPackageName());
 	}
 	return PlatformDependencyNames;
 }
 
-TArray<IPackageWriter::FCommitAttachmentInfo> FSaveCookedPackageContext::GetCommitAttachments(
-	UE::TargetDomain::FGeneratedPackageResultStruct* GeneratedResult,
+void FSaveCookedPackageContext::CalculateCookDependencies(UE::TargetDomain::FGeneratedPackageResultStruct* GeneratedResult,
 	TArray<FAssetDependency>* OverrideDependencies)
 {
-	TArray<IPackageWriter::FCommitAttachmentInfo> Result;
 	if (COTFS.bHybridIterativeEnabled)
 	{
 		UE_SCOPED_HIERARCHICAL_COOKTIMER(TargetDomainDependencies);
@@ -659,8 +674,35 @@ TArray<IPackageWriter::FCommitAttachmentInfo> FSaveCookedPackageContext::GetComm
 				PlatformRuntimeDependencies.Add(Dependency.AssetId.PackageName);
 			}
 		}
-		UE::TargetDomain::CollectAndStoreCookAttachments(Package, TargetPlatform, &SavePackageResult,
-			GeneratedResult, MoveTemp(PlatformRuntimeDependencies), Result);
+		PlatformCookAttachments = UE::TargetDomain::FIncrementalCookAttachments::Collect(Package, TargetPlatform,
+			&SavePackageResult, GeneratedResult, MoveTemp(PlatformRuntimeDependencies));
+	}
+}
+
+void FSaveCookedPackageContext::RecordPlatformBuildDependencies()
+{
+	if (COTFS.bHybridIterativeEnabled)
+	{
+		TSet<FPackageData*>& CurrentDependencies = PlatformDependencies[PlatformIndex].BuildDependencies;
+		for (const FCookDependency& BuildDependency :
+			PlatformCookAttachments.Dependencies.GetTransitiveBuildDependencies())
+		{
+			FPackageData* BuildPackageData = COTFS.PackageDatas->TryAddPackageDataByPackageName(
+				BuildDependency.GetPackageName());
+			if (BuildPackageData)
+			{
+				CurrentDependencies.Add(BuildPackageData);
+			}
+		}
+	}
+}
+
+TArray<IPackageWriter::FCommitAttachmentInfo> FSaveCookedPackageContext::GetCommitAttachments()
+{
+	TArray<IPackageWriter::FCommitAttachmentInfo> Result;
+	if (COTFS.bHybridIterativeEnabled)
+	{
+		PlatformCookAttachments.AppendCommitAttachments(Result);
 	}
 	if (CookContext->DeterminismManager)
 	{
diff --git a/Engine/Source/Editor/UnrealEd/Private/Cooker/CookSavePackage.h b/Engine/Source/Editor/UnrealEd/Private/Cooker/CookSavePackage.h
--- a/Engine/Source/Editor/UnrealEd/Private/Cooker/CookSavePackage.h
+++ b/Engine/Source/Editor/UnrealEd/Private/Cooker/CookSavePackage.h
@@ -9,6 +9,7 @@
 #include "CoreTypes.h"
 #include "Misc/Optional.h"
 #include "Serialization/PackageWriter.h"
+#include "TargetDomain/TargetDomainUtils.h"
 #include "UObject/ArchiveCookContext.h"
 #include "UObject/ICookInfo.h"
 #include "UObject/NameTypes.h"
@@ -22,7 +23,6 @@ class UWorld;
 namespace UE::Cook { struct FCookSavePackageContext; }
 namespace UE::Cook { struct FPackageData; }
 namespace UE::Cook { struct FTickStackData; }
-namespace UE::TargetDomain { struct FGeneratedPackageResultStruct; }
 struct FAssetDependency;
 struct FSavePackageResultStruct;
 
@@ -47,9 +47,11 @@ private: // Used only by UCookOnTheFlyServer, which has private access
 	void CalculatePlatformAgnosticRuntimeDependencies();
 	void CalculatePlatformRuntimeDependencies();
 	TArray<FName> GetPlatformRuntimeDependencies() const;
-	TArray<IPackageWriter::FCommitAttachmentInfo> GetCommitAttachments(
-		UE::TargetDomain::FGeneratedPackageResultStruct* GeneratedResult,
+	void CalculateCookDependencies(UE::TargetDomain::FGeneratedPackageResultStruct* GeneratedResult,
 		TArray<FAssetDependency>* OverrideDependencies);
+	void RecordPlatformBuildDependencies();
+
+	TArray<IPackageWriter::FCommitAttachmentInfo> GetCommitAttachments();
 	IPackageWriter::EWriteOptions GetCommitWriteOptions() const;
 	static void AddDependency(TMap<FPackageData*, EInstigator>& InDependencies, FPackageData* PackageData, bool bHard);
 
@@ -58,7 +60,13 @@ private: // Used only by UCookOnTheFlyServer, which has private access
 	FPackageData& PackageData;
 	TArrayView<const ITargetPlatform*> PlatformsForPackage;
 	TMap<FPackageData*, EInstigator> PlatformAgnosticDependencies;
-	TArray<TMap<FPackageData*, EInstigator>, TInlineAllocator<1>> PlatformDependencies;
+	struct FPlatformDiscoveryData
+	{
+		TMap<FPackageData*, EInstigator> RuntimeDependencies;
+		TSet<FPackageData*> BuildDependencies;
+	};
+	TArray<FPlatformDiscoveryData, TInlineAllocator<1>> PlatformDependencies;
+	UE::TargetDomain::FIncrementalCookAttachments PlatformCookAttachments;
 	FTickStackData& StackData;
 	UPackage* Package;
 	const FString PackageName;
diff --git a/Engine/Source/Editor/UnrealEd/Private/Cooker/CookWorkerClient.cpp b/Engine/Source/Editor/UnrealEd/Private/Cooker/CookWorkerClient.cpp
--- a/Engine/Source/Editor/UnrealEd/Private/Cooker/CookWorkerClient.cpp
+++ b/Engine/Source/Editor/UnrealEd/Private/Cooker/CookWorkerClient.cpp
@@ -1088,6 +1088,7 @@ void FCookWorkerClient::HandleRetractionMessage(FMPCollectorClientMessageContext
 			COTFS.DemoteToIdle(*PackageData, ESendFlags::QueueAddAndRemove,
 				ESuppressCookReason::RetractedByCookDirector);
 			PackageData->ResetReachable();
+			PackageData->ResetIsBuildDependency();
 		}
 	}
 
diff --git a/Engine/Source/Editor/UnrealEd/Private/Cooker/CookWorkerServer.cpp b/Engine/Source/Editor/UnrealEd/Private/Cooker/CookWorkerServer.cpp
--- a/Engine/Source/Editor/UnrealEd/Private/Cooker/CookWorkerServer.cpp
+++ b/Engine/Source/Editor/UnrealEd/Private/Cooker/CookWorkerServer.cpp
@@ -1071,13 +1071,25 @@ void FCookWorkerServer::QueueDiscoveredPackage(FDiscoveredPackageReplication&& D
 			&BufferPlatforms);
 	}
 
-	if (Instigator.Category != EInstigator::ForceExplorableSaveTimeSoftDependency &&
-		PackageData.HasReachablePlatforms(DiscoveredPlatforms))
+	if (Instigator.Category == EInstigator::BuildDependency)
 	{
-		// The CookWorker thought there were some new reachable platforms, but the Director already knows about
-		// all of them; ignore the report
-		return;
+		if (PackageData.IsBuildDependency(DiscoveredPlatforms))
+		{
+			// The CookWorker thought there were some new builddependency platforms, but the Director already knows
+			// about all of them; ignore the report
+			return;
+		}
+	}
+	else if (Instigator.Category != EInstigator::ForceExplorableSaveTimeSoftDependency)
+	{
+		if (PackageData.HasReachablePlatforms(DiscoveredPlatforms))
+		{
+			// The CookWorker thought there were some new reachable platforms, but the Director already knows about
+			// all of them; ignore the report
+			return;
+		}
 	}
+
 	if (COTFS.bSkipOnlyEditorOnly &&
 		Instigator.Category == EInstigator::Unsolicited &&
 		Platforms.GetSource() == EDiscoveredPlatformSet::CopyFromInstigator &&
diff --git a/Engine/Source/Editor/UnrealEd/Private/EditorDomain/EditorDomainUtils.cpp b/Engine/Source/Editor/UnrealEd/Private/EditorDomain/EditorDomainUtils.cpp
--- a/Engine/Source/Editor/UnrealEd/Private/EditorDomain/EditorDomainUtils.cpp
+++ b/Engine/Source/Editor/UnrealEd/Private/EditorDomain/EditorDomainUtils.cpp
@@ -1979,9 +1979,11 @@ bool TrySavePackage(UPackage* Package)
 
 	if (bStorageResultValid)
 	{
-		TArray<IPackageWriter::FCommitAttachmentInfo, TInlineAllocator<2>> Attachments;
-		UE::TargetDomain::CollectAndStoreCookAttachments(Package, nullptr, nullptr, nullptr, TArray<FName>(),
-			Attachments);
+		UE::TargetDomain::FIncrementalCookAttachments CookAttachments =
+			UE::TargetDomain::FIncrementalCookAttachments::Collect(
+			Package, nullptr, nullptr, nullptr, TArray<FName>());
+		TArray<IPackageWriter::FCommitAttachmentInfo> Attachments;
+		CookAttachments.AppendCommitAttachments(Attachments);
 		if (!Attachments.IsEmpty())
 		{
 			UE::TargetDomain::CommitEditorDomainCookAttachments(Package->GetFName(), Attachments);
diff --git a/Engine/Source/Editor/UnrealEd/Private/TargetDomain/TargetDomainUtils.cpp b/Engine/Source/Editor/UnrealEd/Private/TargetDomain/TargetDomainUtils.cpp
--- a/Engine/Source/Editor/UnrealEd/Private/TargetDomain/TargetDomainUtils.cpp
+++ b/Engine/Source/Editor/UnrealEd/Private/TargetDomain/TargetDomainUtils.cpp
@@ -924,26 +924,51 @@ FCbWriter& operator<<(FCbWriter& Writer, const UE::TargetDomain::FBuildDefinitio
 namespace UE::TargetDomain
 {
 
-void FCookAttachments::Reset()
+void FIncrementalCookAttachments::Reset()
 {
 	Dependencies.Reset();
 	BuildDefinitions.Reset();
 }
 
-void FCookAttachments::Empty()
+void FIncrementalCookAttachments::Empty()
 {
 	Dependencies.Empty();
 	BuildDefinitions.Empty();
 }
 
-bool TryCollectAndStoreCookDependencies(UPackage* Package, const ITargetPlatform* TargetPlatform,
-	FSavePackageResultStruct* SaveResult, const FGeneratedPackageResultStruct* GeneratedResult,
-	TArray<FName>&& RuntimeDependencies, IPackageWriter::FCommitAttachmentInfo& OutResult)
+void FIncrementalCookAttachments::AppendCommitAttachments(TArray<IPackageWriter::FCommitAttachmentInfo>& OutAttachments)
 {
+	if (Dependencies.IsValid())
+	{
+		FCbWriter Writer;
+		Writer << Dependencies;
+		IPackageWriter::FCommitAttachmentInfo& Attachment = OutAttachments.Emplace_GetRef();
+		Attachment.Key = CookDependenciesAttachmentKey;
+		Attachment.Value = Writer.Save().AsObject();
+	}
+
+	if (!BuildDefinitions.Definitions.IsEmpty())
+	{
+		FCbWriter Writer;
+		Writer << BuildDefinitions;
+		IPackageWriter::FCommitAttachmentInfo& Attachment = OutAttachments.Emplace_GetRef();
+		Attachment.Key = BuildDefinitionsAttachmentKey;
+		Attachment.Value = Writer.Save().AsObject();
+	}
+}
+
+FIncrementalCookAttachments FIncrementalCookAttachments::Collect(UPackage* Package,
+	const ITargetPlatform* TargetPlatform, FSavePackageResultStruct* SaveResult,
+	const FGeneratedPackageResultStruct* GeneratedResult, TArray<FName>&& RuntimeDependencies)
+{
+	FIncrementalCookAttachments Result;
+
+	Result.CommitStatus = IPackageWriter::ECommitStatus::NotCommitted;
+
 	TArray<TPair<ELogVerbosity::Type, FString>> Messages;
-	FCookDependencies CookDependencies = FCookDependencies::Collect(Package, TargetPlatform, SaveResult,
+	Result.Dependencies = FCookDependencies::Collect(Package, TargetPlatform, SaveResult,
 		GeneratedResult, MoveTemp(RuntimeDependencies), &Messages);
-	if (!CookDependencies.IsValid())
+	if (!Result.Dependencies.IsValid())
 	{
 		TStringBuilder<256> LogText;
 		LogText << TEXT("Could not collect CookDependencies for package '") << Package->GetFName() << TEXT("'");
@@ -952,37 +977,16 @@ bool TryCollectAndStoreCookDependencies(UPackage* Package, const ITargetPlatform
 		// IterativeTODO: This error occurs due to dependencies on _Verse and on some transient packages.
 		// Raise Verbosity to Error once that is fixed.
 		UE_LOG(LogCook, Verbose, TEXT("%s"), *LogText);
-		OutResult.Value = FCbObject();
-		return false;
 	}
 
-	FCbWriter Writer;
-	Writer << CookDependencies;
-	OutResult.Key = CookDependenciesAttachmentKey;
-	OutResult.Value = Writer.Save().AsObject();
-	return true;
-}
-
-bool TryCollectAndStoreBuildDefinitionList(UPackage* Package, const ITargetPlatform* TargetPlatform,
-	IPackageWriter::FCommitAttachmentInfo& OutResult)
-{
-	FBuildDefinitionList Definitions = FBuildDefinitionList::Collect(Package, TargetPlatform);
-	if (Definitions.Definitions.IsEmpty())
-	{
-		OutResult.Value = FCbObject();
-		return false;
-	}
+	Result.BuildDefinitions = FBuildDefinitionList::Collect(Package, TargetPlatform);
 
-	FCbWriter Writer;
-	Writer << Definitions;
-	OutResult.Key = BuildDefinitionsAttachmentKey;
-	OutResult.Value = Writer.Save().AsObject();
-	return true;
+	return Result;
 }
 
-void FCookAttachments::Fetch(TArrayView<FName> PackageNames, const ITargetPlatform* TargetPlatform,
+void FIncrementalCookAttachments::Fetch(TArrayView<FName> PackageNames, const ITargetPlatform* TargetPlatform,
 	ICookedPackageWriter* PackageWriter,
-	TUniqueFunction<void(FName PackageName, FCookAttachments&& Result)>&& Callback)
+	TUniqueFunction<void(FName PackageName, FIncrementalCookAttachments&& Result)>&& Callback)
 {
 	for (FName PackageName : PackageNames)
 	{
@@ -1000,7 +1004,7 @@ void FCookAttachments::Fetch(TArrayView<FName> PackageNames, const ITargetPlatfo
 		{
 			if (!GEditorDomainOplog)
 			{
-				Callback(PackageName, FCookAttachments());
+				Callback(PackageName, FIncrementalCookAttachments());
 				continue;
 			}
 			DependenciesObj = GEditorDomainOplog->GetOplogAttachment(PackageName, CookDependenciesAttachmentKey);
@@ -1009,7 +1013,7 @@ void FCookAttachments::Fetch(TArrayView<FName> PackageNames, const ITargetPlatfo
 				: IPackageWriter::ECommitStatus::NotCommitted;
 		}
 
-		FCookAttachments Result;
+		FIncrementalCookAttachments Result;
 		if (LoadFromCompactBinary(DependenciesObj, Result.Dependencies))
 		{
 			Result.Dependencies.PackageName = PackageName;
diff --git a/Engine/Source/Editor/UnrealEd/Private/TargetDomain/TargetDomainUtils.h b/Engine/Source/Editor/UnrealEd/Private/TargetDomain/TargetDomainUtils.h
--- a/Engine/Source/Editor/UnrealEd/Private/TargetDomain/TargetDomainUtils.h
+++ b/Engine/Source/Editor/UnrealEd/Private/TargetDomain/TargetDomainUtils.h
@@ -26,7 +26,7 @@ template <typename FuncType> class TUniqueFunction;
 namespace UE::Cook { class FCookDependency; }
 namespace UE::TargetDomain { class FCookDependencies; }
 namespace UE::TargetDomain { struct FBuildDefinitionList; }
-namespace UE::TargetDomain { struct FCookAttachments; }
+namespace UE::TargetDomain { struct FIncrementalCookAttachments; }
 
 bool LoadFromCompactBinary(FCbObjectView ObjectView, UE::TargetDomain::FCookDependencies& CookAttachments);
 FCbWriter& operator<<(FCbWriter& Writer, const UE::TargetDomain::FCookDependencies& CookAttachments);
@@ -115,7 +115,7 @@ public:
 		TArray<TPair<ELogVerbosity::Type, FString>>* OutMessages);
 
 	// Fetch function to load the dependencies from a PackageStore is not yet implemented independently for
-	// this structure. Use FCookAttachments instead. 
+	// this structure. Use FIncrementalCookAttachments instead.
 
 private:
 	TArray<FName> BuildPackageDependencies;
@@ -132,7 +132,7 @@ private:
 
 	friend bool ::LoadFromCompactBinary(FCbObjectView ObjectView, FCookDependencies& CookAttachments);
 	friend FCbWriter& ::operator<<(FCbWriter& Writer, const FCookDependencies& CookAttachments);
-	friend FCookAttachments;
+	friend FIncrementalCookAttachments;
 };
 
 /**
@@ -178,11 +178,8 @@ private:
 	friend FCbWriter& ::operator<<(FCbWriter& Writer, const FBuildDefinitionList& Definitions);
 };
 
-/**
- * All of the attachments that we want to read during RequestCluster reference traversal. We read them all at once
- * to batch up the fetch of the attachments from the PackageStore.
- */
-struct FCookAttachments
+/** All of the metadata that is written/read to the oplog for the incremental cook of a package. */
+struct FIncrementalCookAttachments
 {
 	FCookDependencies Dependencies;
 	FBuildDefinitionList BuildDefinitions;
@@ -190,36 +187,16 @@ struct FCookAttachments
 
 	void Reset();
 	void Empty();
+	void AppendCommitAttachments(TArray<IPackageWriter::FCommitAttachmentInfo>& OutAttachments);
 
 	static void Fetch(TArrayView<FName> PackageNames, const ITargetPlatform* TargetPlatform,
 		ICookedPackageWriter* PackageWriter,
-		TUniqueFunction<void(FName PackageName, FCookAttachments&& Result)>&& Callback);
+		TUniqueFunction<void(FName PackageName, FIncrementalCookAttachments&& Result)>&& Callback);
+	static FIncrementalCookAttachments Collect(UPackage* Package,
+		const ITargetPlatform* TargetPlatform, FSavePackageResultStruct* SaveResult,
+		const FGeneratedPackageResultStruct* GeneratedResult, TArray<FName>&& RuntimeDependencies);
 };
 
-bool TryCollectAndStoreCookDependencies(UPackage* Package, const ITargetPlatform* TargetPlatform,
-	FSavePackageResultStruct* SaveResult, const FGeneratedPackageResultStruct* GeneratedResult,
-	TArray<FName>&& RuntimeDependencies, IPackageWriter::FCommitAttachmentInfo& OutResult);
-bool TryCollectAndStoreBuildDefinitionList(UPackage* Package, const ITargetPlatform* TargetPlatform,
-	IPackageWriter::FCommitAttachmentInfo& OutResult);
-
-template <typename ArrayType>
-void CollectAndStoreCookAttachments(UPackage* Package,
-	const ITargetPlatform* TargetPlatform, FSavePackageResultStruct* SaveResult,
-	const FGeneratedPackageResultStruct* GeneratedResult, TArray<FName>&& RuntimeDependencies,
-	ArrayType& Output)
-{
-	IPackageWriter::FCommitAttachmentInfo Result;
-	if (TryCollectAndStoreCookDependencies(Package, TargetPlatform, SaveResult, GeneratedResult,
-		MoveTemp(RuntimeDependencies), Result))
-	{
-		Output.Add(MoveTemp(Result));
-	}
-	if (TryCollectAndStoreBuildDefinitionList(Package, TargetPlatform, Result))
-	{
-		Output.Add(MoveTemp(Result));
-	}
-}
-
 /** Return whether iterative cook is enabled for the given packagename, based on used-class allowlist/blocklist. */
 bool IsIterativeEnabled(FName PackageName, bool bAllowAllClasses,
 	const FAssetPackageData* OverrideAssetPackageData = nullptr);
diff --git a/Engine/Source/Runtime/CoreUObject/Public/UObject/ICookInfo.h b/Engine/Source/Runtime/CoreUObject/Public/UObject/ICookInfo.h
--- a/Engine/Source/Runtime/CoreUObject/Public/UObject/ICookInfo.h
+++ b/Engine/Source/Runtime/CoreUObject/Public/UObject/ICookInfo.h
@@ -66,6 +66,7 @@ namespace UE::Cook
 	CallbackMacro(SaveTimeSoftDependency, false) \
 	CallbackMacro(ForceExplorableSaveTimeSoftDependency, false) \
 	CallbackMacro(GeneratedPackage, false) \
+	CallbackMacro(BuildDependency, false) \
 
 
 /** The different ways a package can be discovered by the cooker. */
