[Viewport Toolbar] Viewport DPI scale value is not properly set

This CL adds a refresh request for DPI scale value at the end of SEditorViewport::Construct().
Retrieving that value before that ended up caching a default value instead of the expected one, because the widget parent window is still not setup.

The introduced refresh request allows code to properly retrieve the actual DPI scale.

diff --git a/Engine/Source/Editor/UnrealEd/Private/SEditorViewport.cpp b/Engine/Source/Editor/UnrealEd/Private/SEditorViewport.cpp
--- a/Engine/Source/Editor/UnrealEd/Private/SEditorViewport.cpp
+++ b/Engine/Source/Editor/UnrealEd/Private/SEditorViewport.cpp
@@ -214,6 +214,15 @@ void SEditorViewport::Construct( const FArguments& InArgs )
 
 	
 	PopulateViewportOverlays(ViewportOverlay.ToSharedRef());
+
+	// Any code retrieving DPI scale before this point might have done that too soon, when the parent Window of the
+	// Viewport Widget is not valid yet. This causes the Viewport Client to cache a default DPI scale instead of the
+	// correct and expected value.
+	// An example of this happening would be FEditorViewportClient::GetPreviewScreenPercentage() which is called when
+	// creating Screen Percentage menus in Asset Editors (e.g. Level, Static Mesh, etc.)
+	// The following RequestUpdateDPIScale() call marks the DPI value for refresh.
+	// This ensures CachedDPIScale value can be properly retrieved at the right time, once the Widget is fully setup and its parent window can be accessed.
+	GetViewportClient()->RequestUpdateDPIScale();
 }
 
 FReply SEditorViewport::OnKeyDown( const FGeometry& MyGeometry, const FKeyEvent& InKeyEvent )
