[Sequencer] Fix selectability not working with marquee selections

diff --git a/Engine/Source/Editor/Sequencer/Private/SequencerEdMode.cpp b/Engine/Source/Editor/Sequencer/Private/SequencerEdMode.cpp
--- a/Engine/Source/Editor/Sequencer/Private/SequencerEdMode.cpp
+++ b/Engine/Source/Editor/Sequencer/Private/SequencerEdMode.cpp
@@ -288,6 +288,11 @@ TSharedPtr<ISequencer> FSequencerEdMode::GetFirstActiveSequencer() const
 
 bool FSequencerEdMode::StartTracking(FEditorViewportClient* InViewportClient, FViewport* InViewport)
 {
+	if (SelectabilityTool->IsSelectionLimited())
+	{
+		SelectabilityTool->StartTracking(InViewportClient, InViewport);
+	}
+
 	if (IsPressingMoveTimeSlider(InViewport))
 	{
 		TSharedPtr<FSequencer> ActiveSequencer;
@@ -320,8 +325,14 @@ bool FSequencerEdMode::StartTracking(FEditorViewportClient* InViewportClient, FV
 	}
 	return FEdMode::StartTracking(InViewportClient, InViewport);
 }
+
 bool FSequencerEdMode::EndTracking(FEditorViewportClient* InViewportClient, FViewport* InViewport)
 {
+	if (SelectabilityTool->IsSelectionLimited())
+	{
+		SelectabilityTool->EndTracking(InViewportClient, InViewport);
+	}
+
 	if (IsPressingMoveTimeSlider(InViewport))
 	{
 		return true;
diff --git a/Engine/Source/Editor/Sequencer/Private/SequencerSelectabilityTool.cpp b/Engine/Source/Editor/Sequencer/Private/SequencerSelectabilityTool.cpp
--- a/Engine/Source/Editor/Sequencer/Private/SequencerSelectabilityTool.cpp
+++ b/Engine/Source/Editor/Sequencer/Private/SequencerSelectabilityTool.cpp
@@ -37,7 +37,7 @@ bool FSequencerSelectabilityTool::FrustumSelect(const FConvexVolume& InFrustum,
 	// Need to check for a zero frustum since ComponentIsTouchingSelectionFrustum will return true, selecting everything, when this is the case
 	const bool bAreTopBottomMalformed = InFrustum.Planes[0].IsNearlyZero() && InFrustum.Planes[2].IsNearlyZero();
 	const bool bAreRightLeftMalformed = InFrustum.Planes[1].IsNearlyZero() && InFrustum.Planes[3].IsNearlyZero();
-	if (bAreTopBottomMalformed || bAreRightLeftMalformed)
+	if (bAreTopBottomMalformed || bAreRightLeftMalformed || InEditorViewportClient->IsInGameView())
 	{
 		return false;
 	}
diff --git a/Engine/Source/Editor/UnrealEd/Private/EditorViewportSelectability.cpp b/Engine/Source/Editor/UnrealEd/Private/EditorViewportSelectability.cpp
--- a/Engine/Source/Editor/UnrealEd/Private/EditorViewportSelectability.cpp
+++ b/Engine/Source/Editor/UnrealEd/Private/EditorViewportSelectability.cpp
@@ -5,9 +5,14 @@
 #include "Editor/EditorPerProjectUserSettings.h"
 #include "Elements/Framework/EngineElementsLibrary.h"
 #include "Elements/Framework/TypedElementRegistry.h"
+#include "Elements/Framework/TypedElementSelectionSet.h"
+#include "Elements/Interfaces/TypedElementObjectInterface.h"
 #include "Elements/Interfaces/TypedElementWorldInterface.h"
 #include "Engine/Canvas.h"
 #include "EngineUtils.h"
+#include "GameFramework/Actor.h"
+#include "GameFramework/Volume.h"
+#include "LevelEditorSubsystem.h"
 #include "LevelEditorViewport.h"
 #include "Math/Color.h"
 #include "ScopedTransaction.h"
@@ -317,7 +322,7 @@ bool FEditorViewportSelectability::SelectActorsByPredicate(UWorld* const InWorld
 
 bool FEditorViewportSelectability::IsActorSelectableClass(const AActor& InActor)
 {
-	const bool bInvalidClass = InActor.IsA<AWorldSettings>() || InActor.IsA<ABrush>();
+	const bool bInvalidClass = InActor.IsA<AWorldSettings>();
 	return !bInvalidClass;
 }
 
@@ -356,52 +361,6 @@ TTypedElement<ITypedElementWorldInterface> FEditorViewportSelectability::GetType
 	return TypedElementRegistry->GetElement<ITypedElementWorldInterface>(ActorElementHandle);
 }
 
-bool FEditorViewportSelectability::DoesActorIntersectBox(const AActor& InActor, const FBox& InBox, FEditorViewportClient* const InEditorViewportClient, const bool bInUseStrictSelection)
-{
-	if (InActor.IsHiddenEd() || !IsActorSelectableClass(InActor))
-	{
-		return false;
-	}
-
-	if (GCurrentLevelEditingViewportClient
-		&& InEditorViewportClient == GCurrentLevelEditingViewportClient
-		&& IsActorInLevelHiddenLayer(InActor, GCurrentLevelEditingViewportClient))
-	{
-		return false;
-	}
-
-	const TTypedElement<ITypedElementWorldInterface> WorldActorElement = GetTypedWorldElementFromActor(InActor);
-	if (!WorldActorElement)
-	{
-		return false;
-	}
-
-	return WorldActorElement.IsElementInBox(InBox, bInUseStrictSelection);
-}
-
-bool FEditorViewportSelectability::DoesActorIntersectFrustum(const AActor& InActor, const FConvexVolume& InFrustum, FEditorViewportClient* const InEditorViewportClient, const bool bInUseStrictSelection)
-{
-	if (InActor.IsHiddenEd() || !IsActorSelectableClass(InActor))
-	{
-		return false;
-	}
-
-	if (GCurrentLevelEditingViewportClient
-		&& InEditorViewportClient == GCurrentLevelEditingViewportClient
-		&& IsActorInLevelHiddenLayer(InActor, GCurrentLevelEditingViewportClient))
-	{
-		return false;
-	}
-
-	const TTypedElement<ITypedElementWorldInterface> WorldActorElement = GetTypedWorldElementFromActor(InActor);
-	if (!WorldActorElement)
-	{
-		return false;
-	}
-
-	return WorldActorElement.IsElementInConvexVolume(InFrustum, bInUseStrictSelection);
-}
-
 bool FEditorViewportSelectability::GetCursorForHovered(EMouseCursor::Type& OutCursor) const
 {
 	if (bSelectionLimited && MouseCursor.IsSet())
@@ -519,64 +478,257 @@ bool FEditorViewportSelectability::HandleClick(FEditorViewportClient* const InVi
 	return bNotSelectable;
 }
 
+void FEditorViewportSelectability::StartTracking(FEditorViewportClient* const InViewportClient, FViewport* const InViewport)
+{
+	FIntPoint MousePosition;
+	InViewport->GetMousePos(MousePosition);
+
+	DragStartPosition = FVector(MousePosition.X, MousePosition.Y, 0);
+	DragEndPosition = DragStartPosition;
+}
+
+void FEditorViewportSelectability::EndTracking(FEditorViewportClient* const InViewportClient, FViewport* const InViewport)
+{
+	FIntPoint MousePosition;
+	InViewport->GetMousePos(MousePosition);
+
+	DragEndPosition = FVector(MousePosition);
+
+	if (DragStartPosition.X > DragEndPosition.X)
+	{
+		Swap(DragStartPosition.X, DragEndPosition.X);
+	}
+	if (DragStartPosition.Y > DragEndPosition.Y)
+	{
+		Swap(DragStartPosition.Y, DragEndPosition.Y);
+	}
+
+	// Extend the endpoint of the rect to get the actual line
+	const int32 MinX = UE::LWC::FloatToIntCastChecked<int32>(FMath::Max<double>(0.0, DragStartPosition.X));
+	const int32 MinY = UE::LWC::FloatToIntCastChecked<int32>(FMath::Max<double>(0.0, DragStartPosition.Y));
+	const FIntPoint ViewportSize = InViewport->GetSizeXY();
+	const int32 MaxX = FMath::Min(ViewportSize.X, FMath::TruncToInt32(DragEndPosition.X + 1.0));
+	const int32 MaxY = FMath::Min(ViewportSize.Y, FMath::TruncToInt32(DragEndPosition.Y + 1.0));
+
+	const FIntPoint Min { MinX, MinY };
+	const FIntPoint Max { MaxX, MaxY };
+	DragSelectionRect = { Min, Max };
+}
+
 bool FEditorViewportSelectability::BoxSelectWorldActors(FBox& InBox, FEditorViewportClient* const InEditorViewportClient, const bool bInSelect)
 {
-	if (!InEditorViewportClient || InEditorViewportClient->IsInGameView())
+	if (!GEditor || !InEditorViewportClient || InEditorViewportClient->IsInGameView())
 	{
 		return false;
 	}
 
-	UWorld* const World = InEditorViewportClient->GetWorld();
-	if (!IsValid(World))
+	UTypedElementRegistry* const TypedElementRegistry = UTypedElementRegistry::GetInstance();
+	if (!ensure(TypedElementRegistry))
+	{
+		return false;
+	}
+
+	UTypedElementSelectionSet* const SelectionSet = GetLevelEditorSelectionSet();
+	if (!IsValid(SelectionSet))
 	{
 		return false;
 	}
 
 	const ULevelEditorViewportSettings* const LevelEditorViewportSettings = GetDefault<ULevelEditorViewportSettings>();
-	const bool bUseStrictSelection = IsValid(LevelEditorViewportSettings) ? LevelEditorViewportSettings->bStrictBoxSelection : false;
+	if (!IsValid(LevelEditorViewportSettings))
+	{
+		return false;
+	}
 
-	const auto Predicate = [this, &InBox, InEditorViewportClient, bUseStrictSelection](AActor* const InActor) -> bool
+	SelectionSet->Modify();
+
+	const bool bUseStrictSelection = LevelEditorViewportSettings->bStrictBoxSelection;
+
+	FWorldSelectionElementArgs SelectionArgs
 	{
-		const bool bSelectable = IsObjectSelectableInViewport(InActor);
-		const bool bIntersects = DoesActorIntersectBox(*InActor, InBox, InEditorViewportClient, bUseStrictSelection);
-		return bSelectable && bIntersects;
+		SelectionSet,
+		ETypedElementSelectionMethod::Primary,
+		FTypedElementSelectionOptions(),
+		&InEditorViewportClient->EngineShowFlags,
+		bUseStrictSelection,
+		false
 	};
 
+	TArray<FTypedElementHandle> ElementsToSelect;
+
+	auto AddToElementsToSelect = [this, SelectionSet, &ElementsToSelect](const FTypedElementHandle& InElement)
+	{
+		if (IsTypedElementSelectable(InElement))
+		{
+			ElementsToSelect.Add(SelectionSet->GetSelectionElement(InElement, ETypedElementSelectionMethod::Primary));
+		}
+	};
+
+	if (LevelEditorViewportSettings->bTransparentBoxSelection)
+	{
+		// Get a list of box-culled elements
+		for (FActorIterator It(InEditorViewportClient->GetWorld()); It; ++It)
+		{
+			AActor* const Actor = *It;
+			GetSelectionElements(Actor, [this, &InBox, &SelectionArgs, &AddToElementsToSelect]
+				(const TTypedElement<ITypedElementWorldInterface>& InWorldElement)
+				{
+					for (const FTypedElementHandle& Element : InWorldElement.GetSelectionElementsInBox(InBox, SelectionArgs))
+					{
+						AddToElementsToSelect(Element);
+					}
+				});
+		}
+	}
+	else
+	{
+		const FTypedElementListRef ElementList = TypedElementRegistry->CreateElementList();
+		InEditorViewportClient->Viewport->GetElementHandlesInRect(DragSelectionRect, ElementList);
+
+		if (bUseStrictSelection)
+		{
+			ElementList->ForEachElement<ITypedElementWorldInterface>([this, bUseStrictSelection, &InBox, &AddToElementsToSelect]
+				(const TTypedElement<ITypedElementWorldInterface>& InElement)
+				{
+					if (InElement.IsElementInBox(InBox, bUseStrictSelection))
+					{
+						AddToElementsToSelect(InElement);
+					}
+					return true;
+				});
+		}
+		else
+		{
+			ElementList->ForEachElementHandle([this, &AddToElementsToSelect]
+				(const FTypedElementHandle& InElement)
+				{
+					AddToElementsToSelect(InElement);
+					return true;
+				});
+		}
+	}
+
 	const bool bShiftDown = InEditorViewportClient->Viewport->KeyState(EKeys::LeftShift)
 		|| InEditorViewportClient->Viewport->KeyState(EKeys::RightShift);
 
-	SelectActorsByPredicate(World, bInSelect, !bShiftDown, Predicate);
+	if (!bShiftDown)
+	{
+		// If the user is selecting, but isn't hold down SHIFT, remove the previous selections
+		SelectionSet->SetSelection(MoveTemp(ElementsToSelect), FTypedElementSelectionOptions());
+	}
+	else
+	{
+		SelectionSet->SelectElements(MoveTemp(ElementsToSelect), FTypedElementSelectionOptions());
+	}
 
 	return true;
 }
 
 bool FEditorViewportSelectability::FrustumSelectWorldActors(const FConvexVolume& InFrustum, FEditorViewportClient* const InEditorViewportClient, const bool bInSelect)
 {
-	if (!InEditorViewportClient || InEditorViewportClient->IsInGameView())
+	if (!GEditor || !InEditorViewportClient || InEditorViewportClient->IsInGameView())
 	{
 		return false;
 	}
 
-	UWorld* const World = InEditorViewportClient->GetWorld();
-	if (!IsValid(World))
+	UTypedElementRegistry* const TypedElementRegistry = UTypedElementRegistry::GetInstance();
+	if (!ensure(TypedElementRegistry))
+	{
+		return false;
+	}
+
+	UTypedElementSelectionSet* const SelectionSet = GetLevelEditorSelectionSet();
+	if (!IsValid(SelectionSet))
 	{
 		return false;
 	}
 
 	const ULevelEditorViewportSettings* const LevelEditorViewportSettings = GetDefault<ULevelEditorViewportSettings>();
-	const bool bUseStrictSelection = IsValid(LevelEditorViewportSettings) ? LevelEditorViewportSettings->bStrictBoxSelection : false;
+	if (!IsValid(LevelEditorViewportSettings))
+	{
+		return false;
+	}
+
+	SelectionSet->Modify();
+
+	const bool bUseStrictSelection = LevelEditorViewportSettings->bStrictBoxSelection;
+
+	FWorldSelectionElementArgs SelectionArgs
+	{
+		SelectionSet,
+		ETypedElementSelectionMethod::Primary,
+		FTypedElementSelectionOptions(),
+		&InEditorViewportClient->EngineShowFlags,
+		bUseStrictSelection,
+		false
+	};
 
-	const auto Predicate = [this, &InFrustum, InEditorViewportClient, bUseStrictSelection](AActor* const InActor) -> bool
+	TArray<FTypedElementHandle> ElementsToSelect;
+
+	auto AddToElementsToSelect = [this, SelectionSet, &ElementsToSelect](const FTypedElementHandle& InElement)
 	{
-		const bool bSelectable = IsObjectSelectableInViewport(InActor);
-		const bool bIntersects = DoesActorIntersectFrustum(*InActor, InFrustum, InEditorViewportClient, bUseStrictSelection);
-		return bSelectable && bIntersects;
+		if (IsTypedElementSelectable(InElement))
+		{
+			ElementsToSelect.Add(SelectionSet->GetSelectionElement(InElement, ETypedElementSelectionMethod::Primary));
+		}
 	};
 
+	if (LevelEditorViewportSettings->bTransparentBoxSelection)
+	{
+		// Get a list of frustum-culled elements
+		for (FActorIterator It(InEditorViewportClient->GetWorld()); It; ++It)
+		{
+			AActor* Actor = *It;
+			GetSelectionElements(Actor, [this, &InFrustum, &SelectionArgs, &AddToElementsToSelect]
+				(const TTypedElement<ITypedElementWorldInterface>& InWorldElement)
+				{
+					for (const FTypedElementHandle& Element : InWorldElement.GetSelectionElementsInConvexVolume(InFrustum, SelectionArgs))
+					{
+						AddToElementsToSelect(Element);
+					}
+				});
+		}
+	}
+	else
+	{
+		const FTypedElementListRef ElementList = TypedElementRegistry->CreateElementList();
+		InEditorViewportClient->Viewport->GetElementHandlesInRect(DragSelectionRect, ElementList);
+
+		if (bUseStrictSelection)
+		{
+			ElementList->ForEachElement<ITypedElementWorldInterface>([this, bUseStrictSelection, &InFrustum, &AddToElementsToSelect]
+				(const TTypedElement<ITypedElementWorldInterface>& InElement)
+				{
+					if (InElement.IsElementInConvexVolume(InFrustum, bUseStrictSelection))
+					{
+						AddToElementsToSelect(InElement);
+					}
+					return true;
+				});
+		}
+		else
+		{
+			ElementList->ForEachElementHandle([this, &AddToElementsToSelect]
+				(const FTypedElementHandle& InElement)
+				{
+					AddToElementsToSelect(InElement);
+					return true;
+				});
+		}
+	}
+
 	const bool bShiftDown = InEditorViewportClient->Viewport->KeyState(EKeys::LeftShift)
 		|| InEditorViewportClient->Viewport->KeyState(EKeys::RightShift);
 
-	SelectActorsByPredicate(World, bInSelect, !bShiftDown, Predicate);
+	if (!bShiftDown)
+	{
+		// If the user is selecting, but isn't hold down SHIFT, remove the previous selections
+		SelectionSet->SetSelection(MoveTemp(ElementsToSelect), FTypedElementSelectionOptions());
+	}
+	else
+	{
+		SelectionSet->SelectElements(MoveTemp(ElementsToSelect), FTypedElementSelectionOptions());
+	}
 
 	return true;
 }
@@ -614,4 +766,96 @@ FText FEditorViewportSelectability::GetLimitedSelectionText(const TSharedPtr<FUI
 	return HelpText;
 }
 
+UTypedElementSelectionSet* FEditorViewportSelectability::GetLevelEditorSelectionSet()
+{
+	if (!ensure(GEditor))
+	{
+		return nullptr;
+	}
+
+	UTypedElementRegistry* const TypedElementRegistry = UTypedElementRegistry::GetInstance();
+	if (!ensure(TypedElementRegistry))
+	{
+		return nullptr;
+	}
+	
+	ULevelEditorSubsystem* const EditorSubsystem = GEditor->GetEditorSubsystem<ULevelEditorSubsystem>();
+	if (!ensure(IsValid(EditorSubsystem)))
+	{
+		return nullptr;
+	}
+
+	return EditorSubsystem->GetSelectionSet();
+}
+
+bool FEditorViewportSelectability::IsTypedElementSelectable(const FTypedElementHandle& InElementHandle) const
+{
+	UTypedElementRegistry* const TypedElementRegistry = UTypedElementRegistry::GetInstance();
+	if (!ensure(TypedElementRegistry))
+	{
+		return false;
+	}
+
+	ITypedElementObjectInterface* const ObjectInterface = TypedElementRegistry->GetElementInterface<ITypedElementObjectInterface>(InElementHandle);
+	if (!ObjectInterface)
+	{
+		return false;
+	}
+
+	UPrimitiveComponent* const PrimitiveComponent = ObjectInterface->GetObjectAs<UPrimitiveComponent>(InElementHandle);
+	if (!IsValid(PrimitiveComponent))
+	{
+		return false;
+	}
+
+	AActor* const Actor = PrimitiveComponent->GetOwner();
+	if (!IsValid(Actor))
+	{
+		return false;
+	}
+
+	if (Actor->IsHiddenEd() || !IsActorSelectableClass(*Actor))
+	{
+		return false;
+	}
+
+	if (GCurrentLevelEditingViewportClient && IsActorInLevelHiddenLayer(*Actor, GCurrentLevelEditingViewportClient))
+	{
+		return false;
+	}
+
+	return IsObjectSelectableInViewport(Actor);
+}
+
+void FEditorViewportSelectability::GetSelectionElements(AActor* const InActor, const TFunctionRef<void(const TTypedElement<ITypedElementWorldInterface>&)> InPredicate)
+{
+	if (!IsValid(InActor))
+	{
+		return;
+	}
+
+	if (InActor->IsA(AVolume::StaticClass()))
+	{
+		if (!GCurrentLevelEditingViewportClient
+			|| !GCurrentLevelEditingViewportClient->IsVolumeVisibleInViewport(*InActor))
+		{
+			return;
+		}
+	}
+
+	UTypedElementRegistry* const TypedElementRegistry = UTypedElementRegistry::GetInstance();
+	if (!ensure(TypedElementRegistry))
+	{
+		return;
+	}
+
+	if (const FTypedElementHandle ElementHandle = UEngineElementsLibrary::AcquireEditorActorElementHandle(InActor))
+	{
+		if (const TTypedElement<ITypedElementWorldInterface> WorldElement = TypedElementRegistry->GetElement<ITypedElementWorldInterface>(ElementHandle))
+		{
+			InPredicate(WorldElement);
+		}
+	}
+}
+
 #undef LOCTEXT_NAMESPACE
diff --git a/Engine/Source/Editor/UnrealEd/Public/EditorViewportSelectability.h b/Engine/Source/Editor/UnrealEd/Public/EditorViewportSelectability.h
--- a/Engine/Source/Editor/UnrealEd/Public/EditorViewportSelectability.h
+++ b/Engine/Source/Editor/UnrealEd/Public/EditorViewportSelectability.h
@@ -16,7 +16,9 @@ class FLevelEditorViewportClient;
 class FUICommandInfo;
 class HHitProxy;
 class ITypedElementWorldInterface;
+class UTypedElementSelectionSet;
 struct FConvexVolume;
+struct FTypedElementHandle;
 struct FViewportClick;
 
 DECLARE_DELEGATE_RetVal(UWorld*, FOnGetWorld);
@@ -76,6 +78,9 @@ public:
 
 	bool HandleClick(FEditorViewportClient* const InViewportClient, HHitProxy* const InHitProxy, const FViewportClick& InClick);
 
+	void StartTracking(FEditorViewportClient* const InViewportClient, FViewport* const InViewport);
+	void EndTracking(FEditorViewportClient* const InViewportClient, FViewport* const InViewport);
+
 	/** Selects or deselects all actors in a level world that are inside a defined box. */
 	bool BoxSelectWorldActors(FBox& InBox, FEditorViewportClient* const InEditorViewportClient, const bool bInSelect);
 	/** Selects or deselects all actors in a level world that are inside a defined convex volume. */
@@ -88,9 +93,6 @@ protected:
 
 	static TTypedElement<ITypedElementWorldInterface> GetTypedWorldElementFromActor(const AActor& InActor);
 
-    static bool DoesActorIntersectBox(const AActor& InActor, const FBox& InBox, FEditorViewportClient* const InEditorViewportClient, const bool bInUseStrictSelection);
-    static bool DoesActorIntersectFrustum(const AActor& InActor, const FConvexVolume& InFrustum, FEditorViewportClient* const InEditorViewportClient, const bool bInUseStrictSelection);
-
 	/**
 	 * Selects or deselects actors in a world. If no actors are specified, uses all the actors in the level.
 	 * 
@@ -109,6 +111,12 @@ protected:
 	/** Updates an actors hovered state and visuals. */
 	void UpdateHoveredActorPrimitives(AActor* const InActor);
 
+	static UTypedElementSelectionSet* GetLevelEditorSelectionSet();
+
+	bool IsTypedElementSelectable(const FTypedElementHandle& InElementHandle) const;
+
+	void GetSelectionElements(AActor* const InActor, const TFunctionRef<void(const TTypedElement<ITypedElementWorldInterface>&)> InPredicate);
+
 	bool bSelectionLimited = false;
 
 	/** Hovered primitives and their last overlay color before we apply the hover overlay */
@@ -121,4 +129,8 @@ protected:
 
 	/** Delegate used to check if an object is selectable in the viewport */
 	FOnIsObjectSelectableInViewport OnIsObjectSelectableInViewportDelegate;
+
+	FVector DragStartPosition;
+	FVector DragEndPosition;
+	FIntRect DragSelectionRect;
 };
