IncrementalCook: IterativeValidate: Add a cap to the warning messages written for assets that differ, to prevent the possibility of the log being too big and causing farm jobs to fail and human reading of the log to be difficult.
#rnx

diff --git a/Engine/Config/BaseEditor.ini b/Engine/Config/BaseEditor.ini
--- a/Engine/Config/BaseEditor.ini
+++ b/Engine/Config/BaseEditor.ini
@@ -502,8 +502,14 @@ SaveUnversioned=true
 +IterativeClassDenyList=/Script/Engine.World 								; Not yet documented
 +IterativeClassDenyList=/Script/PoseSearch.PoseSearchDatabase				; UE-224771
 
-; IterativeValidatePackageIgnoreList - package names that will be displayed but not treated as an error by IterativeValidate
-; +IterativeValidatePackageIgnoreList=/Game/PackageName
+; Settings for the IterativeValidate cooker mode, which reports false positive iterative skips in incremental cooks
+[IterativeValidate]
+; PackageIgnoreList - package names that will be displayed but not treated as an error by IterativeValidate
+; +PackageIgnoreList=/Game/PackageName
+; The maximum number of iterative validate failures to report before we start suppressing their log, to reduce spam in the job
+; This is a SoftMaximum, reports for new AssetClass types will still be logged, but for an AssetClass that has already been reported
+; once, no further occurrences of packages with that AssetClass will be logged. Negative values mean no max, 0 value means a max of 0.
+LoggingSoftMaximum=1000
 
 [CookPlatformDataCacheSettings]
 ; Add <ShortClassName>=<Integer> here to limit the number of instances that can have BeginCacheForCookedPlatformData called at once
diff --git a/Engine/Source/Editor/UnrealEd/Private/Cooker/IterativeValidatePackageWriter.cpp b/Engine/Source/Editor/UnrealEd/Private/Cooker/IterativeValidatePackageWriter.cpp
--- a/Engine/Source/Editor/UnrealEd/Private/Cooker/IterativeValidatePackageWriter.cpp
+++ b/Engine/Source/Editor/UnrealEd/Private/Cooker/IterativeValidatePackageWriter.cpp
@@ -261,6 +261,7 @@ void FIterativeValidateMPCollector::ServerReceiveMessage(UE::Cook::FMPCollectorS
 		else if (MessageSubtypeAsInteger == (uint8)EMessageSubtype::ClientToServer_ReplIsAnotherSaveNeeded)
 		{
 			ReadAndSyncPackageStatus(Message, PackageName);
+			Owner->MarkPackageCompletedOnDirector(PackageName, Context.GetWorkerId());
 		}
 		else
 		{
@@ -330,13 +331,25 @@ FIterativeValidatePackageWriter::FIterativeValidatePackageWriter(UCookOnTheFlySe
 		LogIterativeValidate.GetCategoryName(), TEXT(""), GPrintLogTimes).Len());
 
 	TArray<FString> IterativeValidatePackageIgnoreList;
-	GConfig->GetArray(TEXT("TargetDomain"), TEXT("IterativeValidatePackageIgnoreList"),
+	GConfig->GetArray(TEXT("IterativeValidate"), TEXT("PackageIgnoreList"),
 		IterativeValidatePackageIgnoreList, GEditorIni);
+	TArray<FString> LegacyIterativeValidatePackageIgnoreList;
+	GConfig->GetArray(TEXT("TargetDomain"), TEXT("IterativeValidatePackageIgnoreList"),
+		LegacyIterativeValidatePackageIgnoreList, GEditorIni);
+	if (!LegacyIterativeValidatePackageIgnoreList.IsEmpty())
+	{
+		UE_LOG(LogIterativeValidate, Warning,
+			TEXT("Editor.ini:[TargetDomain]:IterativeValidatePackageIgnoreList is deprecated; use Editor.ini:[IterativeValidate]:PackageIgnoreList instead."));
+		IterativeValidatePackageIgnoreList.Append(LegacyIterativeValidatePackageIgnoreList);
+	}
 	PackageIgnoreList.Reserve(IterativeValidatePackageIgnoreList.Num());
 	for (const FString& PackageName : IterativeValidatePackageIgnoreList)
 	{
 		PackageIgnoreList.Add(FName(FStringView(PackageName)));
 	}
+
+	LoggingSoftMaximum = -1;
+	GConfig->GetValue(TEXT("IterativeValidate"), TEXT("LoggingSoftMaximum"), LoggingSoftMaximum, GEditorIni);
 }
 
 void FIterativeValidatePackageWriter::BeginPackage(const FBeginPackageInfo& Info)
@@ -778,7 +791,8 @@ void FIterativeValidatePackageWriter::EndCook(const FCookInfo& Info)
 		{
 			TStringBuilder<1024> MessageWithDiagnostics;
 			MessageWithDiagnostics << Message;
-			TMap<FTopLevelAssetPath, int32> ClassFalsePositiveCounts = GetSummaryFalsePositiveCounts();
+			TMap<FTopLevelAssetPath, int32> ClassFalsePositiveCounts = GetSummaryStatusCounts(
+				EPackageStatus::DeclaredUnmodified_FoundModified_FalsePositive);
 			int32 NumPrinted = 0;
 			constexpr int32 MaxNumPrinted = 5;
 			for (const TPair<FTopLevelAssetPath, int32>& Pair : ClassFalsePositiveCounts)
@@ -892,6 +906,18 @@ void FIterativeValidatePackageWriter::UpdateSaveArguments(FSavePackageArgs& Save
 }
 
 bool FIterativeValidatePackageWriter::IsAnotherSaveNeeded(FSavePackageResultStruct& PreviousResult, FSavePackageArgs& SaveArgs)
+{
+	bool bResult = IsAnotherSaveNeededInternal(PreviousResult, SaveArgs);
+	if (!bResult && COTFS.GetCookMode() != ECookMode::CookWorker)
+	{
+		MarkPackageCompletedOnDirector(BeginInfo.PackageName, UE::Cook::FWorkerId::Local());
+	}
+	return bResult;
+
+}
+
+bool FIterativeValidatePackageWriter::IsAnotherSaveNeededInternal(FSavePackageResultStruct& PreviousResult,
+	FSavePackageArgs& SaveArgs)
 {
 	bPackageFirstPass = false;
 	checkf(!Inner->IsAnotherSaveNeeded(PreviousResult, SaveArgs),
@@ -1033,14 +1059,40 @@ bool FIterativeValidatePackageWriter::IsAnotherSaveNeeded(FSavePackageResultStru
 
 void FIterativeValidatePackageWriter::OnDiffWriterMessage(ELogVerbosity::Type Verbosity, FStringView Message)
 {
+	PackageMessageMap.FindOrAdd(BeginInfo.PackageName).Add(FMessage{ FString(Message), Verbosity });
+}
+
+void FIterativeValidatePackageWriter::MarkPackageCompletedOnDirector(FName PackageName,
+	UE::Cook::FWorkerId WorkerId)
+{
+	FPackageStatusInfo* Status = PackageStatusMap.Find(PackageName);
+	TArray<FMessage>* Messages = PackageMessageMap.Find(PackageName);
+
+	if (!Status || !Messages)
+	{
+		return;
+	}
+
+	int32& TotalCount = TotalStatusCounts.FindOrAdd(Status->Status, 0);
+	int32& ClassCount = ClassStatusCounts.FindOrAdd(Status->AssetClass).FindOrAdd(Status->Status, 0);
+	++TotalCount;
+	++ClassCount;
+	if (LoggingSoftMaximum >= 0 && TotalCount > LoggingSoftMaximum && ClassCount > 1)
+	{
+		// Suppress reported logs for this package
+		return;
+	}
+
 	switch (Phase)
 	{
 	case EPhase::AllInOnePhase:
-		FMsg::Logf(__FILE__, __LINE__, LogIterativeValidate.GetCategoryName(), Verbosity,
-			TEXT("%s"), *ResolveText(Message));
+		for (FMessage& Message : *Messages)
+		{
+			FMsg::Logf(__FILE__, __LINE__, LogIterativeValidate.GetCategoryName(), Message.Verbosity,
+				TEXT("%s"), *ResolveText(Message.Text));
+		}
 		break;
 	case EPhase::Phase1:
-		PackageMessageMap.FindOrAdd(BeginInfo.PackageName).Add(FMessage{ FString(Message), Verbosity });
 		break;
 	case EPhase::Phase2:
 		break;
@@ -1180,15 +1232,15 @@ FIterativeValidatePackageWriter::FStatusCounts FIterativeValidatePackageWriter::
 	return StatusCounts;
 }
 
-TMap<FTopLevelAssetPath, int32> FIterativeValidatePackageWriter::GetSummaryFalsePositiveCounts()
+TMap<FTopLevelAssetPath, int32> FIterativeValidatePackageWriter::GetSummaryStatusCounts(EPackageStatus PackageStatus)
 {
 	TMap<FTopLevelAssetPath, int32> Result;
-	for (const TPair<FName, FPackageStatusInfo>& Pair : PackageStatusMap)
+	for (const TPair<FTopLevelAssetPath, TMap<EPackageStatus, int32>>& Pair : ClassStatusCounts)
 	{
-		if (Pair.Value.Status == EPackageStatus::DeclaredUnmodified_FoundModified_FalsePositive
-			&& Pair.Value.AssetClass.IsValid())
+		const int32* Count = Pair.Value.Find(PackageStatus);
+		if (Count)
 		{
-			++(Result.FindOrAdd(Pair.Value.AssetClass, 0));
+			Result.Add(Pair.Key, *Count);
 		}
 	}
 
diff --git a/Engine/Source/Editor/UnrealEd/Private/Cooker/IterativeValidatePackageWriter.h b/Engine/Source/Editor/UnrealEd/Private/Cooker/IterativeValidatePackageWriter.h
--- a/Engine/Source/Editor/UnrealEd/Private/Cooker/IterativeValidatePackageWriter.h
+++ b/Engine/Source/Editor/UnrealEd/Private/Cooker/IterativeValidatePackageWriter.h
@@ -6,6 +6,7 @@
 
 class UCookOnTheFlyServer;
 class FIterativeValidateMPCollector; 
+namespace UE::Cook { struct FWorkerId; }
 
 /**
  * A CookedPackageWriter that diffs the cook results of iteratively-unmodified packages between their last cook
@@ -113,14 +114,19 @@ protected:
 	EPackageStatus GetPackageStatus(FName PackageName) const;
 	void SetPackageStatus(FName PackageName, EPackageStatus NewStatus);
 	FStatusCounts CountPackagesByStatus();
-	TMap<FTopLevelAssetPath, int32> GetSummaryFalsePositiveCounts();
+	TMap<FTopLevelAssetPath, int32> GetSummaryStatusCounts(EPackageStatus PackageStatus);
+	bool IsAnotherSaveNeededInternal(FSavePackageResultStruct& PreviousResult, FSavePackageArgs& SaveArgs);
+	void MarkPackageCompletedOnDirector(FName PackageName, UE::Cook::FWorkerId WorkerId);
 
 protected:
 	TMap<FName, FPackageStatusInfo> PackageStatusMap;
 	TMap<FName, TArray<FMessage>> PackageMessageMap;
+	TMap<EPackageStatus, int32> TotalStatusCounts; // Only populated on director
+	TMap<FTopLevelAssetPath, TMap<EPackageStatus, int32>> ClassStatusCounts; // Only populated on director
 	TSet<FName> PackageIgnoreList;
 	FString MetadataPath;
 	UCookOnTheFlyServer& COTFS;
+	int32 LoggingSoftMaximum = -1;
 	EPhase Phase = EPhase::AllInOnePhase;
 	ESaveAction SaveAction = ESaveAction::IgnoreResults;
 	bool bPackageFirstPass = false;
