[Editor] Copy/Paste fix missing remapping of sub-objects properties (allows later call to RemapProperty to fixup those properties)

#rnx

diff --git a/Engine/Source/Editor/UnrealEd/Private/EditorObject.cpp b/Engine/Source/Editor/UnrealEd/Private/EditorObject.cpp
--- a/Engine/Source/Editor/UnrealEd/Private/EditorObject.cpp
+++ b/Engine/Source/Editor/UnrealEd/Private/EditorObject.cpp
@@ -254,7 +254,7 @@ static const TCHAR* ImportProperties(
 	FFeedbackContext*			Warn,
 	int32						Depth,
 	FObjectInstancingGraph&		InstanceGraph,
-	const TMap<FSoftObjectPath, UObject*>* ObjectRemapper
+	TMap<FSoftObjectPath, UObject*>* ObjectRemapper
 	)
 {
 	check(ObjectStruct!=NULL);
@@ -476,6 +476,9 @@ static const TCHAR* ImportProperties(
 				return NULL;
 			}
 
+			FString ExportedObjectFullName;
+			const bool bExportPath = FParse::Value(Str, TEXT("ExportPath="), ExportedObjectFullName);
+			
 			// points to the parent class's template subobject/component, if we are overriding a subobject/component declared in our parent class
 			UObject* BaseTemplate = NULL;
 			bool bRedefiningSubobject = false;
@@ -713,6 +716,11 @@ static const TCHAR* ImportProperties(
 				}
 				FArchiveReplaceObjectRef<UObject> ReplaceAr(SubobjectOuter, ReplacementMap, EArchiveReplaceObjectFlags::IgnoreArchetypeRef);
 
+				if (bExportPath && ObjectRemapper)
+				{
+					ObjectRemapper->Add(MoveTemp(ExportedObjectFullName), ComponentTemplate);
+				}
+
 				if (bOverrideT3DSupported)
 				{
 					if (OverriddenOperation.IsSet())
@@ -924,7 +932,7 @@ const TCHAR* ImportObjectProperties(
 	int32					Depth,
 	int32					LineNumber,
 	FObjectInstancingGraph* InInstanceGraph,
-	const TMap<FSoftObjectPath, UObject*>* ObjectRemapper
+	TMap<FSoftObjectPath, UObject*>* ObjectRemapper
 	)
 {
 	FImportObjectParams Params;
diff --git a/Engine/Source/Editor/UnrealEd/Private/Factories/EditorFactories.cpp b/Engine/Source/Editor/UnrealEd/Private/Factories/EditorFactories.cpp
--- a/Engine/Source/Editor/UnrealEd/Private/Factories/EditorFactories.cpp
+++ b/Engine/Source/Editor/UnrealEd/Private/Factories/EditorFactories.cpp
@@ -1293,7 +1293,11 @@ UObject* ULevelFactory::FactoryCreateText
 				const FName OldPath = Actor->GetFolderPath();
 
 				Actor->PreEditChange(nullptr);
-				ImportObjectProperties( (uint8*)Actor, *PropText, Actor->GetClass(), Actor, Actor, Warn, 0, INDEX_NONE, NULL, &ExistingToNewMap );
+
+				// Remapping of sub-objects across actors is currently dependant on the order in which we import those actors 
+				// so for now make it unsupported for consistency by copying the ExistingToNewMap for every imported actor.
+				TMap<FSoftObjectPath, UObject*> CopyExistingToNewMap(ExistingToNewMap);
+				ImportObjectProperties( (uint8*)Actor, *PropText, Actor->GetClass(), Actor, Actor, Warn, 0, INDEX_NONE, NULL, &CopyExistingToNewMap);
 				bActorChanged = true;
 
 				// Path might have been set through OnLevelActorAdded (i.e. by UActorEditorContextSubsystem).
diff --git a/Engine/Source/Editor/UnrealEd/Public/Editor.h b/Engine/Source/Editor/UnrealEd/Public/Editor.h
--- a/Engine/Source/Editor/UnrealEd/Public/Editor.h
+++ b/Engine/Source/Editor/UnrealEd/Public/Editor.h
@@ -519,7 +519,7 @@ struct FImportObjectParams
 	FObjectInstancingGraph* InInstanceGraph = nullptr;
 
 	/** provides a mapping from an existing object, typically an actor, (which may no longer be loaded) to a new instance to which it should be remapped */
-	const TMap<FSoftObjectPath, UObject*>* ObjectRemapper = nullptr;
+	TMap<FSoftObjectPath, UObject*>* ObjectRemapper = nullptr;
 
 	/** True if we should call PreEditChange/PostEditChange on the object as it's imported.  Pass false here
 		if you're going to do that on your own. */
@@ -565,7 +565,7 @@ const TCHAR* ImportObjectProperties(
 	int32					Depth,
 	int32					LineNumber = INDEX_NONE,
 	FObjectInstancingGraph* InstanceGraph = nullptr,
-	const TMap<FSoftObjectPath, UObject*>* ObjectRemapper = nullptr
+	TMap<FSoftObjectPath, UObject*>* ObjectRemapper = nullptr
 	);
 
 //
