MPCook: Add log feedback when a socket fails to send but is not in the error state. This has occurred on a CookWorker once, and it softlocks the cook because the CookWorker stops communicating with the director but does not terminate.
#rnx

diff --git a/Engine/Source/Editor/UnrealEd/Private/Cooker/CompactBinaryTCP.cpp b/Engine/Source/Editor/UnrealEd/Private/Cooker/CompactBinaryTCP.cpp
--- a/Engine/Source/Editor/UnrealEd/Private/Cooker/CompactBinaryTCP.cpp
+++ b/Engine/Source/Editor/UnrealEd/Private/Cooker/CompactBinaryTCP.cpp
@@ -214,7 +214,7 @@ bool FCompactBinaryTCPImpl::TryCreatePayload(TArray<FMarshalledMessage>& Pending
 	}
 }
 
-EConnectionStatus FCompactBinaryTCPImpl::PollSendBytes(FSocket* Socket, const void* Data, 
+EConnectionStatus FCompactBinaryTCPImpl::PollSendBytes(FSocket* Socket, const void* Data,
 	uint64 DataSize, uint64& InOutBytesWritten)
 {
 	check(DataSize <= MAX_int32);
@@ -231,18 +231,24 @@ EConnectionStatus FCompactBinaryTCPImpl::PollSendBytes(FSocket* Socket, const vo
 			InOutBytesWritten += BytesSent;
 		}
 
-		// Return values of false indicate the socket send failed with an error, but the error
-		// might be temporary, such as EWOULDBLOCK. Even when return value is false, return incomplete
-		// unless socket's ConnectionState has been set to error
-		if (!bResult && Socket->GetConnectionState() == SCS_ConnectionError)
+		if (!bResult)
 		{
-			UE_LOG(LogSockets, Error, TEXT("Could not WritePacket to Socket. Socket is in the error state."));
-			return EConnectionStatus::Failed;
-		}
-		else
-		{
-			return EConnectionStatus::Incomplete;
+			// Return values of false indicate the socket send failed with an error, but the error
+			// might be temporary, such as EWOULDBLOCK. Even when return value is false, return incomplete
+			// unless socket's ConnectionState has been set to error
+			if (Socket->GetConnectionState() == SCS_ConnectionError)
+			{
+				UE_LOG(LogSockets, Error, TEXT("Could not WritePacket to Socket. Socket is in the error state."));
+				return EConnectionStatus::Failed;
+			}
+
+			// TODO: Decide which errors are fatal even if the socket is still connected. We have seen one case
+			// where the socket send fails forever, but we were not able to see what error it was returning.
+			ESocketErrors Error = ISocketSubsystem::Get(PLATFORM_SOCKETSUBSYSTEM)->GetLastErrorCode();
+			UE_LOG(LogSockets, Log, TEXT("FCompactBinaryTCPImpl::PollSendBytes: Socket->Send failed with error %d."), (int32)Error);
 		}
+
+		return EConnectionStatus::Incomplete;
 	}
 	InOutBytesWritten = 0;
 	return EConnectionStatus::Okay;
