- Color widget editor changes
- Added a mode to show alpha & color separate vs combined (default option, right click menu to change)
- Adding a key will select the interpolated color
- Adding a key with CTRL pressed will insert 1.0 for all channels
- Adding a key with ALT pressed will insert 0.0 for all channels

diff --git a/Engine/Source/Editor/UnrealEd/Private/SColorGradientEditor.cpp b/Engine/Source/Editor/UnrealEd/Private/SColorGradientEditor.cpp
--- a/Engine/Source/Editor/UnrealEd/Private/SColorGradientEditor.cpp
+++ b/Engine/Source/Editor/UnrealEd/Private/SColorGradientEditor.cpp
@@ -198,12 +198,12 @@ void FGradientStopMark::SetTime( float NewTime, FCurveOwnerInterface& CurveOwner
 void SColorGradientEditor::Construct( const FArguments& InArgs )
 {
 	IsEditingEnabled = InArgs._IsEditingEnabled;
-	LastModifiedColor = FLinearColor::White;
 	CurveOwner = NULL;
 	ViewMinInput = InArgs._ViewMinInput;
 	ViewMaxInput = InArgs._ViewMaxInput;
 	bClampStopsToViewRange = InArgs._ClampStopsToViewRange;
 	bDraggingAlphaValue = false;
+	bDrawColorAndAlphaSeparate = InArgs._DrawColorAndAlphaSeparate;
 	bDraggingStop = false;
 	DistanceDragged = 0.0f;
 	ContextMenuPosition = FVector2D::ZeroVector;
@@ -230,10 +230,11 @@ int32 SColorGradientEditor::OnPaint( const FPaintArgs& Args, const FGeometry& Al
 		int32 Start = 0;
 		int32 Finish = FMath::TruncToInt( GradientAreaGeometry.GetLocalSize().X );
 
-		TArray<FSlateGradientStop> Stops;
+		TArray<FSlateGradientStop> ColorStops;
+		TArray<FSlateGradientStop> AlphaStops;
 
 		// If no alpha keys are available, treat the curve as being completely opaque for drawing purposes
-		bool bHasAnyAlphaKeys = CurveOwner->HasAnyAlphaKeys(); 
+		const bool bHasAnyAlphaKeys = CurveOwner->HasAnyAlphaKeys(); 
 
 		// If any transparency (A < 1) is found, we'll draw a checkerboard to visualize the color with alpha
 		bool bHasTransparency = false;
@@ -272,7 +273,7 @@ int32 SColorGradientEditor::OnPaint( const FPaintArgs& Args, const FGeometry& Al
 		for( int32 CurrentStep = Start; CurrentStep < Finish; CurrentStep+=2 )
 		{
 			// Figure out the time from the current screen unit
-			float Time = ScaleInfo.LocalXToInput(CurrentStep);
+			const float Time = ScaleInfo.LocalXToInput(CurrentStep);
 
 			// Sample the curve
 			FLinearColor Color = CurveOwner->GetLinearColorValue( Time );
@@ -286,35 +287,71 @@ int32 SColorGradientEditor::OnPaint( const FPaintArgs& Args, const FGeometry& Al
 			{
 				bHasTransparency |= (Color.A < 1.0f);
 			}
-	
-			Stops.Add( FSlateGradientStop( FVector2D( CurrentStep, 0.0f ), Color ) );
+
+			if (bDrawColorAndAlphaSeparate)
+			{
+				ColorStops.Emplace(FVector2D(CurrentStep, 0.0f), FLinearColor(Color.R, Color.G, Color.B));
+				AlphaStops.Emplace(FVector2D(CurrentStep, 0.0f), FLinearColor(Color.A, Color.A, Color.A));
+			}
+			else
+			{
+				ColorStops.Emplace(FVector2D(CurrentStep, 0.0f), Color);
+			}
 		}
 
-		if( Stops.Num() > 0 )
+		if ( ColorStops.Num() > 0 )
 		{
-			if( bHasTransparency )
+			// Draw the color gradient
+			if (bDrawColorAndAlphaSeparate)
 			{
-				// Draw a checkerboard behind there is any transparency visible
-				FSlateDrawElement::MakeBox
-				( 
+				const FVector2D GradientLocalSize = GradientAreaGeometry.GetLocalSize();
+				const FGeometry ColorGradientAreaGeometry = GradientAreaGeometry.MakeChild(FVector2D(GradientLocalSize.X, GradientLocalSize.Y * 0.70f), FSlateLayoutTransform(FVector2D(0.0f, 0.0f)));
+				const FGeometry AlphaGradientAreaGeometry = GradientAreaGeometry.MakeChild(FVector2D(GradientLocalSize.X, GradientLocalSize.Y * 0.30f), FSlateLayoutTransform(FVector2D(0.0f, GradientLocalSize.Y * 0.70f)));
+
+				FSlateDrawElement::MakeGradient
+				(
 					OutDrawElements,
 					LayerId,
-					GradientAreaGeometry.ToPaintGeometry(),
-					FAppStyle::GetBrush("Checkerboard"),
-					DrawEffects 
+					ColorGradientAreaGeometry.ToPaintGeometry(),
+					ColorStops,
+					Orient_Vertical,
+					DrawEffects
+				);
+				FSlateDrawElement::MakeGradient
+				(
+					OutDrawElements,
+					LayerId,
+					AlphaGradientAreaGeometry.ToPaintGeometry(),
+					AlphaStops,
+					Orient_Vertical,
+					DrawEffects
 				);
 			}
+			else
+			{
+				if (bHasTransparency)
+				{
+					// Draw a checkerboard behind there is any transparency visible
+					FSlateDrawElement::MakeBox
+					(
+						OutDrawElements,
+						LayerId,
+						GradientAreaGeometry.ToPaintGeometry(),
+						FAppStyle::GetBrush("Checkerboard"),
+						DrawEffects
+					);
+				}
 
-			// Draw the color gradient
-			FSlateDrawElement::MakeGradient
-			( 
-				OutDrawElements, 
-				LayerId,
-				GradientAreaGeometry.ToPaintGeometry(),
-				Stops,
-				Orient_Vertical,
-				DrawEffects
-			);	
+				FSlateDrawElement::MakeGradient
+				(
+					OutDrawElements,
+					LayerId,
+					GradientAreaGeometry.ToPaintGeometry(),
+					ColorStops,
+					Orient_Vertical,
+					DrawEffects
+				);
+			}
 		}
 
 		// Get actual editable stop marks
@@ -521,11 +558,21 @@ FReply SColorGradientEditor::OnMouseButtonUp( const FGeometry& MyGeometry, const
 				FGeometry ColorMarkAreaMouseGeometry = FColorGradientEditorGeometry::GetColorMarkAreaMouseGeometry( MyGeometry );
 				FGeometry AlphaMarkAreaMouseGeometry = FColorGradientEditorGeometry::GetAlphaMarkAreaMouseGeometry( MyGeometry );
 
+				TOptional<FLinearColor> StopColor;
+				if (MouseEvent.IsControlDown())
+				{
+					StopColor = FLinearColor::White;
+				}
+				else if (MouseEvent.IsAltDown())
+				{
+					StopColor = FLinearColor::Transparent;
+				}
+
 				if( ColorMarkAreaMouseGeometry.IsUnderLocation( MouseEvent.GetScreenSpacePosition() ) )
 				{
 					// Add a new color mark
 					bool bColorStop = true;
-					SelectedStop = AddStop( MouseEvent.GetScreenSpacePosition(), MyGeometry, bColorStop );
+					SelectedStop = AddStop( MouseEvent.GetScreenSpacePosition(), MyGeometry, bColorStop, StopColor);
 
 					return FReply::Handled().CaptureMouse( SharedThis(this) );
 
@@ -534,7 +581,7 @@ FReply SColorGradientEditor::OnMouseButtonUp( const FGeometry& MyGeometry, const
 				{
 					// Add a new alpha mark
 					bool bColorStop = false;
-					SelectedStop = AddStop( MouseEvent.GetScreenSpacePosition(), MyGeometry, bColorStop );
+					SelectedStop = AddStop( MouseEvent.GetScreenSpacePosition(), MyGeometry, bColorStop, StopColor);
 
 					return FReply::Handled().CaptureMouse( SharedThis(this) );
 				}
@@ -546,9 +593,20 @@ FReply SColorGradientEditor::OnMouseButtonUp( const FGeometry& MyGeometry, const
 		else if( MouseEvent.GetEffectingButton() == EKeys::RightMouseButton && !bDraggingStop)
 		{
 			// Didnt move the mouse too far, open a context menu
-			if( DistanceDragged < DragThresholdDist && SelectedStop.IsValid( *CurveOwner ) )
+			if( DistanceDragged < DragThresholdDist )
 			{
-				OpenGradientStopContextMenu( MouseEvent );
+				if (SelectedStop.IsValid(*CurveOwner))
+				{
+					OpenGradientStopContextMenu(MouseEvent);
+				}
+				else
+				{
+					const FGeometry GradientAreaPaintGeometry = FColorGradientEditorGeometry::GetGradientAreaPaintGeometry(MyGeometry);
+					if (GradientAreaPaintGeometry.IsUnderLocation(MouseEvent.GetScreenSpacePosition()))
+					{
+						OpenGradientOptionsMenu(MouseEvent);
+					}
+				}
 			}
 
 			DistanceDragged = 0;
@@ -588,6 +646,29 @@ void SColorGradientEditor::SetCurveOwner( FCurveOwnerInterface* InCurveOwner )
 }
 
 
+void SColorGradientEditor::OpenGradientOptionsMenu(const FPointerEvent& MouseEvent)
+{
+	FMenuBuilder MenuBuilder(true, NULL);
+
+	MenuBuilder.AddMenuEntry(
+		LOCTEXT("DrawColorAndAlphaSeparate", "Draw Color And Alpha Separate"),
+		FText(),
+		FSlateIcon(),
+		FUIAction(
+			FExecuteAction::CreateLambda([this]() { bDrawColorAndAlphaSeparate = !bDrawColorAndAlphaSeparate; }),
+			FCanExecuteAction(),
+			FIsActionChecked::CreateLambda([this]() { return bDrawColorAndAlphaSeparate; })
+		),
+		NAME_None,
+		EUserInterfaceActionType::ToggleButton
+	);
+
+	ContextMenuPosition = MouseEvent.GetScreenSpacePosition();
+
+	FWidgetPath WidgetPath = MouseEvent.GetEventPath() != nullptr ? *MouseEvent.GetEventPath() : FWidgetPath();
+	FSlateApplication::Get().PushMenu(AsShared(), WidgetPath, MenuBuilder.MakeWidget(), ContextMenuPosition, FPopupTransitionEffect::ContextMenu);
+}
+
 void SColorGradientEditor::OpenGradientStopContextMenu(const FPointerEvent& MouseEvent)
 {
 	const FVector2D& Location = MouseEvent.GetScreenSpacePosition();
@@ -742,11 +823,6 @@ void SColorGradientEditor::OnSelectedStopColorChanged( FLinearColor InNewColor )
 	SelectedStop.SetColor( InNewColor, *CurveOwner );
 	TArray<FRichCurveEditInfo> ChangedCurves{ CurveOwner->GetCurves()[0], CurveOwner->GetCurves()[1], CurveOwner->GetCurves()[2] };
 	CurveOwner->OnCurveChanged(ChangedCurves);
-
-	// Set the the last edited color.  The next time a new stop is added we'll use this value
-	LastModifiedColor.R = InNewColor.R;
-	LastModifiedColor.G = InNewColor.G;
-	LastModifiedColor.B = InNewColor.B;
 }
 
 void SColorGradientEditor::OnCancelSelectedStopColorChange( FLinearColor PreviousColor )
@@ -810,9 +886,6 @@ void SColorGradientEditor::OnAlphaValueCommitted( float NewValue, ETextCommit::T
 		TArray<FRichCurveEditInfo> ChangedCurves{ CurveOwner->GetCurves()[3] };
 		CurveOwner->OnCurveChanged(ChangedCurves);
 	}
-
-	// Set the alpha of the last edited color.  The next time a new alpha stop is added we'll use this value
-	LastModifiedColor.A = NewValue;
 }
 
 void SColorGradientEditor::OnRemoveSelectedGradientStop()
@@ -1001,7 +1074,7 @@ void SColorGradientEditor::DeleteStop( const FGradientStopMark& InMark )
 	CurveOwner->OnCurveChanged(CurveOwner->GetCurves());
 }
 
-FGradientStopMark SColorGradientEditor::AddStop( const FVector2D& Position, const FGeometry& MyGeometry, bool bColorStop )
+FGradientStopMark SColorGradientEditor::AddStop( const FVector2D& Position, const FGeometry& MyGeometry, bool bColorStop, TOptional<FLinearColor> OptionalColor)
 {
 	FScopedTransaction AddStopTrans( LOCTEXT("AddGradientStop", "Add Gradient Stop") );
 
@@ -1019,15 +1092,20 @@ FGradientStopMark SColorGradientEditor::AddStop( const FVector2D& Position, cons
 		FRealCurve* RedCurve = Curves[0].CurveToEdit;
 		FRealCurve* GreenCurve = Curves[1].CurveToEdit;
 		FRealCurve* BlueCurve = Curves[2].CurveToEdit;
-		
-		NewStop.RedKeyHandle = RedCurve->AddKey( NewStopTime, LastModifiedColor.R );
-		NewStop.GreenKeyHandle = GreenCurve->AddKey( NewStopTime, LastModifiedColor.G );
-		NewStop.BlueKeyHandle = BlueCurve->AddKey( NewStopTime, LastModifiedColor.B );
+
+		const float RValue = OptionalColor.IsSet() ? OptionalColor.GetValue().A : RedCurve->Eval(NewStopTime, 1.0f);
+		const float GValue = OptionalColor.IsSet() ? OptionalColor.GetValue().G : GreenCurve->Eval(NewStopTime, 1.0f);
+		const float BValue = OptionalColor.IsSet() ? OptionalColor.GetValue().B : BlueCurve->Eval(NewStopTime, 1.0f);
+
+		NewStop.RedKeyHandle = RedCurve->AddKey( NewStopTime, RValue );
+		NewStop.GreenKeyHandle = GreenCurve->AddKey( NewStopTime, GValue );
+		NewStop.BlueKeyHandle = BlueCurve->AddKey( NewStopTime, BValue );
 	}
 	else
 	{
 		FRealCurve* AlphaCurve = Curves[3].CurveToEdit;
-		NewStop.AlphaKeyHandle = AlphaCurve->AddKey( NewStopTime, LastModifiedColor.A );
+		const float AlphaValue = OptionalColor.IsSet() ? OptionalColor.GetValue().A : AlphaCurve->Eval(NewStopTime, 1.0f);
+		NewStop.AlphaKeyHandle = AlphaCurve->AddKey( NewStopTime, AlphaValue );
 	}
 
 	CurveOwner->OnCurveChanged(CurveOwner->GetCurves());
diff --git a/Engine/Source/Editor/UnrealEd/Public/SColorGradientEditor.h b/Engine/Source/Editor/UnrealEd/Public/SColorGradientEditor.h
--- a/Engine/Source/Editor/UnrealEd/Public/SColorGradientEditor.h
+++ b/Engine/Source/Editor/UnrealEd/Public/SColorGradientEditor.h
@@ -50,11 +50,13 @@ public:
 		, _ViewMaxInput(0.0f)
 		, _IsEditingEnabled( true )
 		, _ClampStopsToViewRange( false )
+		, _DrawColorAndAlphaSeparate( true )
 	{}
 		SLATE_ATTRIBUTE( float, ViewMinInput )
 		SLATE_ATTRIBUTE( float, ViewMaxInput )
 		SLATE_ATTRIBUTE( bool, IsEditingEnabled )
 		SLATE_ARGUMENT( bool, ClampStopsToViewRange )
+		SLATE_ARGUMENT(bool, DrawColorAndAlphaSeparate )
 	SLATE_END_ARGS()
 
 
@@ -79,6 +81,13 @@ public:
 	UNREALED_API void SetCurveOwner( FCurveOwnerInterface* InCurveOwner );
 
 private:
+	/**
+	 * Opens a context menu with options for how we display the gradient
+	 *
+	 * @param MouseEvent	The mouse event that raised this handler (containing the location and widget path of the click)
+	 */
+	UNREALED_API void OpenGradientOptionsMenu(const FPointerEvent& MouseEvent);
+
 	/**
 	 * Opens a context menu with options for the selected gradient stop
 	 *
@@ -181,7 +190,7 @@ private:
 	 * @param MyGeometry	The geometry of the stop area
 	 * @param bColorStop	If true a color stop is added, otherwise an alpha stop is added
 	 */
-	UNREALED_API FGradientStopMark AddStop( const FVector2D& Position, const FGeometry& MyGeometry, bool bColorStop );
+	UNREALED_API FGradientStopMark AddStop( const FVector2D& Position, const FGeometry& MyGeometry, bool bColorStop, TOptional<FLinearColor> Color );
 
 	/**
 	 * Moves a gradient stop to a new time
@@ -194,8 +203,6 @@ private:
 private:
 	/** The currently selected stop */
 	FGradientStopMark SelectedStop;
-	/** The color that was last modified. New stops are added with this color */
-	FLinearColor LastModifiedColor;
 	/** interface to the curves being edited */
 	FCurveOwnerInterface* CurveOwner;
 	/** Current min input value that is visible */
@@ -218,4 +225,6 @@ private:
 	bool bDraggingAlphaValue;
 	/** True if a gradient stop is being dragged */
 	bool bDraggingStop;
+	/** Do we draw a gradient for color and alpha separate or combined? */
+	bool bDrawColorAndAlphaSeparate;
 };
