Cooker: WhiteSpace: Move CookPackageData.h functions into out-of-class inline functions to match coding standard requirement of braces on newlines.
#rnx

diff --git a/Engine/Source/Editor/UnrealEd/Private/Cooker/CookPackageData.h b/Engine/Source/Editor/UnrealEd/Private/Cooker/CookPackageData.h
--- a/Engine/Source/Editor/UnrealEd/Private/Cooker/CookPackageData.h
+++ b/Engine/Source/Editor/UnrealEd/Private/Cooker/CookPackageData.h
@@ -161,38 +161,38 @@
 	FPackagePlatformData();
 
 	/** Requested from StartCookByTheBook or CookOnTheFly, or is a transitive dependency of those requests. */
-	bool IsReachable() const { return bReachable != 0; }
-	void SetReachable(bool bValue) { bReachable = (uint32)bValue; }
+	bool IsReachable() const;
+	void SetReachable(bool bValue);
 
 	/** Marked as a transitive build dependency from a reachable package, so we need to save its dependency data.*/
-	bool IsBuildDependency() const { return bIsBuildDependency != 0; }
-	void SetIsBuildDependency(bool bValue) { bIsBuildDependency = (uint32)bValue; }
+	bool IsBuildDependency() const;
+	void SetIsBuildDependency(bool bValue);
 
 	/** The package has passed through a RequestCluster and its transitive dependencies were added. */
-	bool IsVisitedByCluster() const { return bVisitedByCluster != 0; }
-	void SetVisitedByCluster(bool bValue) { bVisitedByCluster = (uint32)bValue; }
+	bool IsVisitedByCluster() const;
+	void SetVisitedByCluster(bool bValue);
 
 	/** UPackage::Save was called on the package, but timedout and needs to be retried. */
-	bool IsSaveTimedOut() const { return bSaveTimedOut != 0; }
-	void SetSaveTimedOut(bool bValue) { bSaveTimedOut = (uint32)bValue; }
+	bool IsSaveTimedOut() const;
+	void SetSaveTimedOut(bool bValue);
 
 	/** Whether the package is allowed to cook. Initially true, may be set false during Cluster search. */
-	bool IsCookable() const { return bCookable != 0; }
-	void SetCookable(bool bValue) { bCookable = (uint32)bValue; }
+	bool IsCookable() const;
+	void SetCookable(bool bValue);
 
 	/**
 	 * Whether the package is searchable for transitive dependencies during Cluster evaluation.
 	 * This can be true even if the package is not cookable. Initially true, may be set false during Cluster search
 	 */
-	bool IsExplorable() const { return bExplorable != 0; }
-	void SetExplorable(bool bValue) { bExplorable = (uint32)bValue; }
+	bool IsExplorable() const;
+	void SetExplorable(bool bValue);
 
 	/**
 	 * Whether the package, which might be set as not explorable from IsRequestCookable, is set to explorable
 	 * based on conditions discovered during the cook.
 	 */
-	bool IsExplorableOverride() const { return bExplorableOverride != 0; }
-	void SetExplorableOverride(bool bValue) { bExplorableOverride = (uint32)bValue; }
+	bool IsExplorableOverride() const;
+	void SetExplorableOverride(bool bValue);
 
 	/** All flags modified by reachability calculations are returned to default. */
 	void ResetReachable();
@@ -206,36 +206,30 @@
 	void MarkCookableForWorker(FCookWorkerClient& CookWorkerClient);
 
 	/** The package was found to be unmodified in the current iterative cook. */
-	bool IsIterativelyUnmodified() const { return bIterativelyUnmodified != 0; }
-	void SetIterativelyUnmodified(bool bValue) { bIterativelyUnmodified = (uint32)bValue; }
+	bool IsIterativelyUnmodified() const;
+	void SetIterativelyUnmodified(bool bValue);
 
 	/**
 	 * The package was skipped in the current iterative cook. This might not be equal to IterativelyUnmodified,
 	 * depending on packagewriter).
 	 */
-	bool IsIterativelySkipped() const { return bIterativelySkipped != 0; }
-	void SetIterativelySkipped(bool bValue) { bIterativelySkipped = (uint32)bValue; }
+	bool IsIterativelySkipped() const;
+	void SetIterativelySkipped(bool bValue);
 
-	ECookResult GetCookResults() const { return (ECookResult)CookResults; }
-	bool IsCookAttempted() const { return CookResults != (uint32)ECookResult::NotAttempted; }
-	bool IsCookSucceeded() const { return CookResults == (uint32)ECookResult::Succeeded; }
-	void SetCookResults(ECookResult Value)
-	{
-		// ECookResult::Invalid is only used in replication and is not allowed in FPackagePlatformData
-		check(Value != ECookResult::Invalid);
-		CookResults = (uint32)Value;
-		SetReportedToDirector(false);
-	}
+	ECookResult GetCookResults() const;
+	bool IsCookAttempted() const;
+	bool IsCookSucceeded() const;
+	void SetCookResults(ECookResult Value);
 
 	/** Return if we need to call SavePackage on it (reachable, cookable, not yet cooked) */
 	bool NeedsCooking(const ITargetPlatform* PlatformItBelongsTo) const;
 
-	bool IsRegisteredForCachedObjectsInOuter() const { return bRegisteredForCachedObjectsInOuter != 0; }
-	void SetRegisteredForCachedObjectsInOuter(bool bValue) { bRegisteredForCachedObjectsInOuter = bValue; }
+	bool IsRegisteredForCachedObjectsInOuter() const;
+	void SetRegisteredForCachedObjectsInOuter(bool bValue);
 
 	/** Only read/written on CookWorkers */
-	bool IsReportedToDirector() { return bReportedToDirector != 0; }
-	void SetReportedToDirector(bool bValue) { bReportedToDirector = (uint32)bValue; }
+	bool IsReportedToDirector();
+	void SetReportedToDirector(bool bValue);
 
 private:
 	uint32 bReachable : 1;
@@ -325,7 +319,7 @@
 	FPackageData& operator=(const FPackageData& other) = delete;
 	FPackageData& operator=(FPackageData&& other) = delete;
 
-	FPackageDatas& GetPackageDatas() const { return PackageDatas; }
+	FPackageDatas& GetPackageDatas() const;
 
 	/** Return a copy of Package->GetName(). It is derived from the FileName if necessary, and is never modified. */
 	const FName& GetPackageName() const;
@@ -524,14 +518,14 @@
 
 	/** Get/Set whether the package has been marked for cooking after all other packages, via -cooklast */
 	void SetIsCookLast(bool bValue);
-	bool GetIsCookLast() const { return bIsCookLast != 0;}
+	bool GetIsCookLast() const;
 
 	/**
 	 * Get/Set a visited flag used when searching graphs of PackageData. User of the graph is responsible for 
 	 * setting the bIsVisited flag back to empty when graph operations are done.
 	 */
-	bool GetIsVisited() const { return bIsVisited != 0; }
-	void SetIsVisited(bool bValue) { bIsVisited = static_cast<uint32>(bValue); }
+	bool GetIsVisited() const;
+	void SetIsVisited(bool bValue);
 
 	/**
 	 * The list of objects inside the package.  Only non-empty during saving; it is populated on demand by
@@ -564,8 +558,8 @@
 	static int32 GetMaxNumRetriesBeginCacheOnObjects();
 
 	/** Get/Set the flag for whether CachedObjectsInOuter is populated. Always false except during save state. */
-	bool GetHasSaveCache() const { return static_cast<bool>(bHasSaveCache); }
-	void SetHasSaveCache(bool Value) { bHasSaveCache = Value != 0; }
+	bool GetHasSaveCache() const;
+	void SetHasSaveCache(bool Value);
 
 	/**
 	 * Get/Set the SubState within EPackageState::Save. Outside of EPackageStateProperty::Saving it is always
@@ -583,11 +577,11 @@
 	void UpdateSaveAfterGarbageCollect(bool& bOutDemote);
 
 	/** Get/Set the flag for whether PrepareSave has been called and returned an error. */
-	bool HasPrepareSaveFailed() const { return static_cast<bool>(bPrepareSaveFailed); }
-	void SetHasPrepareSaveFailed(bool bValue) { bPrepareSaveFailed = bValue != 0; }
+	bool HasPrepareSaveFailed() const;
+	void SetHasPrepareSaveFailed(bool bValue);
 
-	bool IsPrepareSaveRequiresGC() const { return bPrepareSaveRequiresGC; }
-	void SetIsPrepareSaveRequiresGC(bool bValue) { bPrepareSaveRequiresGC = bValue != 0; }
+	bool IsPrepareSaveRequiresGC() const;
+	void SetIsPrepareSaveRequiresGC(bool bValue);
 
 	/** Validate that the BeginCacheForCookedPlatformData-dependent fields are empty, when entering save. */
 	void CheckCookedPlatformDataEmpty() const;
@@ -598,8 +592,8 @@
 	void ClearCookedPlatformData();
 
 	/** Get/Set the Monitor's flag that counts whether this PackageData has finished cooking and with what result. */
-	ECookResult GetMonitorCookResult() const { return (ECookResult)MonitorCookResult; }
-	void SetMonitorCookResult(ECookResult Value) { MonitorCookResult = (uint8) Value; }
+	ECookResult GetMonitorCookResult() const;
+	void SetMonitorCookResult(ECookResult Value);
 
 	/** Remove all data about the given platform from all fields in this PackageData. */
 	void OnRemoveSessionPlatform(const ITargetPlatform* Platform);
@@ -695,22 +689,22 @@
 	bool HasBuildInstigator() const;
 
 	/** Get whether COTFS is keeping this package referenced referenced during GC. */
-	bool IsKeepReferencedDuringGC() const { return static_cast<bool>(bKeepReferencedDuringGC); }
+	bool IsKeepReferencedDuringGC() const;
 	/** Set whether COTFS is keeping this package referenced referenced during GC. */
-	void SetKeepReferencedDuringGC(bool Value) { bKeepReferencedDuringGC = Value != 0; }
+	void SetKeepReferencedDuringGC(bool Value);
 
 	/** Return whether the package was cooked during this session */
-	bool GetWasCookedThisSession() const { return static_cast<bool>(bWasCookedThisSession); }
+	bool GetWasCookedThisSession() const;
 
 	/** For MultiProcessCooks, Get the id of the worker this Package is assigned to; InvalidId means owned by local. */
-	FWorkerId GetWorkerAssignment() const { return WorkerAssignment; }
+	FWorkerId GetWorkerAssignment() const;
 	/**
 	 * Set the id of the worker this Package is assigned to. If value changes from Valid to Invalid and SendFlags
 	 * includes QueueRemove, also calls NotifyRemovedFromWorker.
 	 */
 	void SetWorkerAssignment(FWorkerId InWorkerAssignment, ESendFlags SendFlags = ESendFlags::QueueAddAndRemove);
 	/** Get the workerid that is the only worker allowed to cook this package; InvalidId means no constraint. */
-	FWorkerId GetWorkerAssignmentConstraint() const { return WorkerAssignmentConstraint; }
+	FWorkerId GetWorkerAssignmentConstraint() const;
 	/** Set the workerid that is the only worker allowed to cook this package; default is InvalidId; */
 	void SetWorkerAssignmentConstraint(FWorkerId InWorkerAssignment)
 	{
@@ -761,24 +755,12 @@
 	 * when a property changes from true to false.
 	 */
 	typedef void (FPackageData::*FEdgeFunction)();
-	inline void UpdateDownEdge(bool bOld, bool bNew, const FEdgeFunction& EdgeFunction)
-	{
-		if ((bOld != bNew) & bOld)
-		{
-			(this->*EdgeFunction)();
-		}
-	}
+	inline void UpdateDownEdge(bool bOld, bool bNew, const FEdgeFunction& EdgeFunction);
 	/**
 	 * Helper function to call the given EdgeFunction (e.g. OnEnterInProgress)\
 	 * when a property changes from false to true.
 	 */
-	inline void UpdateUpEdge(bool bOld, bool bNew, const FEdgeFunction& EdgeFunction)
-	{
-		if ((bOld != bNew) & bNew)
-		{
-			(this->*EdgeFunction)();
-		}
-	}
+	inline void UpdateUpEdge(bool bOld, bool bNew, const FEdgeFunction& EdgeFunction);
 
 	/* Entry/Exit gates for PackageData states, used to enforce state contracts and free unneeded memory. */
 	void OnEnterIdle();
@@ -1029,16 +1011,16 @@
 	uint32 RemoveRequestExceptFromCluster(FPackageData* PackageData, FRequestCluster* ExceptFromCluster);
 	void UpdateUrgency(FPackageData* PackageData, EUrgency bOldUrgency, EUrgency NewUrgency);
 
-	FPackageDataSet& GetRestartedRequests() { return RestartedRequests; }
+	FPackageDataSet& GetRestartedRequests();
 	/**
 	 * Unlike other containers on PackageData, GetDiscoveryQueue is not an ownership container.
 	 * PackageDatas in the DiscoveryQueue are PackageDatas we need to look at - in the right order during
 	 * PumpRequests - they can be in any state and are owned by another container (or are in the idle state).
 	 */
-	TRingBuffer<FDiscoveryQueueElement>& GetDiscoveryQueue() { return DiscoveryQueue; }
-	TRingBuffer<FRequestCluster>& GetRequestClusters() { return RequestClusters; }
-	FPackageDataSet& GetReadyRequestsUrgent() { return UrgentRequests; }
-	FPackageDataSet& GetReadyRequestsNormal() { return NormalRequests; }
+	TRingBuffer<FDiscoveryQueueElement>& GetDiscoveryQueue();
+	TRingBuffer<FRequestCluster>& GetRequestClusters();
+	FPackageDataSet& GetReadyRequestsUrgent();
+	FPackageDataSet& GetReadyRequestsNormal();
 	/**
 	 * Add an FPackageData by name that will be notified when all packages that were present in the DiscoveryQueue
 	 * before the fence was created have been assigned or demoted to idle.
@@ -1331,17 +1313,8 @@
 
 	/** Enumerate PendingPlatformDatas: the list of pending calls to BeginCacheForCookedPlatformData. */
 	template <typename FunctionType>
-	void ForEachPendingCookedPlatformData(const FunctionType& Function)
-	{
-		for (FPendingCookedPlatformDataContainer& Container : PendingCookedPlatformDataLists)
-		{
-			for (FPendingCookedPlatformData& Data : Container)
-			{
-				Function(Data);
-			}
-		}
-	}
-	int32 GetPendingCookedPlatformDataNum() const { return PendingCookedPlatformDataNum; }
+	void ForEachPendingCookedPlatformData(const FunctionType& Function);
+	int32 GetPendingCookedPlatformDataNum() const;
 	void AddPendingCookedPlatformData(FPendingCookedPlatformData&& Data);
 
 	/**
@@ -1361,10 +1334,7 @@
 	template <typename CallbackType>
 	void LockAndEnumeratePackageDatas(CallbackType&& Callback);
 
-	TMap<UObject*, FCachedCookedPlatformDataState>& GetCachedCookedPlatformDataObjects()
-	{
-		return CachedCookedPlatformDataObjects;
-	}
+	TMap<UObject*, FCachedCookedPlatformDataState>& GetCachedCookedPlatformDataObjects();
 	void CachedCookedPlatformDataObjectsPostGarbageCollect(const TSet<UObject*>& SaveQueueObjectsThatStillExist);
 
 private:
@@ -1428,86 +1398,143 @@
 ///////////////////////////////////////////////////////
 
 
-/** Return and increment the RoofToLankRank that should be assigned to the next PackageData needing one. */
-inline uint32 FPackageDatas::GetNextLeafToRootRank()
+inline bool FPackagePlatformData::IsReachable() const
 {
-	return NextLeafToRootRank++;
+	return bReachable != 0;
 }
 
-inline void FPackageDatas::ResetLeafToRootRank()
+inline void FPackagePlatformData::SetReachable(bool bValue)
 {
-	NextLeafToRootRank = 0;
+	bReachable = (uint32)bValue;
 }
 
-template <typename CallbackType>
-inline void FPackageDatas::LockAndEnumeratePackageDatas(CallbackType&& Callback)
+inline bool FPackagePlatformData::IsBuildDependency() const
 {
-	FReadScopeLock ExistenceReadLock(ExistenceLock);
-	EnumeratePackageDatasWithinLock(Forward<CallbackType>(Callback));
+	return bIsBuildDependency != 0;
 }
 
-template <typename CallbackType>
-inline void FPackageDatas::EnumeratePackageDatasWithinLock(CallbackType&& Callback)
+inline void FPackagePlatformData::SetIsBuildDependency(bool bValue)
 {
-	Allocator.EnumerateAllocations(Forward<CallbackType>(Callback));
+	bIsBuildDependency = (uint32)bValue;
 }
 
-template<typename CallbackType>
-inline void FPackageDatas::UpdateThreadsafePackageData(FName PackageName, CallbackType&& Callback)
+inline bool FPackagePlatformData::IsVisitedByCluster() const
 {
-	FWriteScopeLock ExistenceWriteLock(ExistenceLock);
-	FThreadsafePackageData& Value = ThreadsafePackageDatas.FindOrAdd(PackageName);
-	bool bNew = false;
-	if (!Value.bInitialized)
-	{
-		Value.bInitialized = true;
-		bNew = true;
-	}
-	Callback(Value, bNew);
+	return bVisitedByCluster != 0;
 }
 
-inline TOptional<FThreadsafePackageData> FPackageDatas::FindThreadsafePackageData(FName PackageName)
+inline void FPackagePlatformData::SetVisitedByCluster(bool bValue)
 {
-	FReadScopeLock ExistenceReadLock(ExistenceLock);
-	FThreadsafePackageData* Value = ThreadsafePackageDatas.Find(PackageName);
-	return Value ? TOptional<FThreadsafePackageData>(*Value) : TOptional<FThreadsafePackageData>();
+	bVisitedByCluster = (uint32)bValue;
 }
 
-inline FPackageDataMonitor& FPackageDatas::GetMonitor()
+inline bool FPackagePlatformData::IsSaveTimedOut() const
 {
-	return Monitor;
+	return bSaveTimedOut != 0;
 }
 
-inline UCookOnTheFlyServer& FPackageDatas::GetCookOnTheFlyServer()
+inline void FPackagePlatformData::SetSaveTimedOut(bool bValue)
 {
-	return CookOnTheFlyServer;
+	bSaveTimedOut = (uint32)bValue;
 }
 
-inline FRequestQueue& FPackageDatas::GetRequestQueue()
+inline bool FPackagePlatformData::IsCookable() const
 {
-	return RequestQueue;
+	return bCookable != 0;
 }
 
-inline TFastPointerSet<FPackageData*>& FPackageDatas::GetAssignedToWorkerSet()
+inline void FPackagePlatformData::SetCookable(bool bValue)
 {
-	return AssignedToWorkerSet;
+	bCookable = (uint32)bValue;
 }
 
-inline FLoadQueue& FPackageDatas::GetLoadQueue()
+inline bool FPackagePlatformData::IsExplorable() const
 {
-	return LoadQueue;
+	return bExplorable != 0;
 }
 
-inline TFastPointerSet<FPackageData*>& FPackageDatas::GetSaveStalledSet()
+inline void FPackagePlatformData::SetExplorable(bool bValue)
 {
-	return SaveStalledSet;
+	bExplorable = (uint32)bValue;
 }
 
-inline FPackageDataQueue& FPackageDatas::GetSaveQueue()
+inline bool FPackagePlatformData::IsExplorableOverride() const
 {
-	return SaveQueue;
+	return bExplorableOverride != 0;
 }
 
+inline void FPackagePlatformData::SetExplorableOverride(bool bValue)
+{
+	bExplorableOverride = (uint32)bValue;
+}
 
+inline bool FPackagePlatformData::IsIterativelyUnmodified() const
+{
+	return bIterativelyUnmodified != 0;
+}
+
+inline void FPackagePlatformData::SetIterativelyUnmodified(bool bValue)
+{
+	bIterativelyUnmodified = (uint32)bValue;
+}
+
+inline bool FPackagePlatformData::IsIterativelySkipped() const
+{
+	return bIterativelySkipped != 0;
+}
+
+inline void FPackagePlatformData::SetIterativelySkipped(bool bValue)
+{
+	bIterativelySkipped = (uint32)bValue;
+}
+
+inline ECookResult FPackagePlatformData::GetCookResults() const
+{
+	return (ECookResult)CookResults;
+}
+
+inline bool FPackagePlatformData::IsCookAttempted() const
+{
+	return CookResults != (uint32)ECookResult::NotAttempted;
+}
+
+inline bool FPackagePlatformData::IsCookSucceeded() const
+{
+	return CookResults == (uint32)ECookResult::Succeeded;
+}
+
+inline void FPackagePlatformData::SetCookResults(ECookResult Value)
+{
+	// ECookResult::Invalid is only used in replication and is not allowed in FPackagePlatformData
+	check(Value != ECookResult::Invalid);
+	CookResults = (uint32)Value;
+	SetReportedToDirector(false);
+}
+
+inline bool FPackagePlatformData::IsRegisteredForCachedObjectsInOuter() const
+{
+	return bRegisteredForCachedObjectsInOuter != 0;
+}
+
+inline void FPackagePlatformData::SetRegisteredForCachedObjectsInOuter(bool bValue)
+{
+	bRegisteredForCachedObjectsInOuter = bValue;
+}
+
+inline bool FPackagePlatformData::IsReportedToDirector()
+{
+	return bReportedToDirector != 0;
+}
+
+inline void FPackagePlatformData::SetReportedToDirector(bool bValue)
+{
+	bReportedToDirector = (uint32)bValue;
+}
+
+inline FPackageDatas& FPackageData::GetPackageDatas() const
+{
+	return PackageDatas;
+}
+
 inline const FName& FPackageData::GetPackageName() const
 {
@@ -1548,19 +1575,6 @@
 }
 
 template <typename ArrayType>
-inline void FPackageData::GetCachedObjectsInOuterPlatforms(ArrayType& OutPlatforms) const
-{
-	OutPlatforms.Reset(PlatformDatas.Num());
-	for (const TPair<const ITargetPlatform*, FPackagePlatformData>& Pair : PlatformDatas)
-	{
-		if (Pair.Value.IsRegisteredForCachedObjectsInOuter())
-		{
-			OutPlatforms.Add(Pair.Key);
-		}
-	}
-}
-
-template <typename ArrayType>
 inline void FPackageData::GetReachablePlatforms(ArrayType& OutPlatforms) const
 {
 	OutPlatforms.Reset(PlatformDatas.Num());
@@ -1573,21 +1587,6 @@
 	}
 }
 
-inline ESaveSubState FPackageData::GetSaveSubState() const
-{
-	return static_cast<ESaveSubState>(SaveSubState);
-}
-
-inline ESuppressCookReason FPackageData::GetSuppressCookReason() const
-{
-	return static_cast<ESuppressCookReason>(SuppressCookReason);
-}
-
-inline void FPackageData::SetSuppressCookReason(ESuppressCookReason Reason)
-{
-	SuppressCookReason = static_cast<uint32>(Reason);
-}
-
 inline EUrgency FPackageData::GetUrgency() const
 {
 	return static_cast<EUrgency>(Urgency);
@@ -1601,4 +1600,87 @@
 	}
 }
 
+inline ESuppressCookReason FPackageData::GetSuppressCookReason() const
+{
+	return static_cast<ESuppressCookReason>(SuppressCookReason);
+}
+
+inline void FPackageData::SetSuppressCookReason(ESuppressCookReason Reason)
+{
+	SuppressCookReason = static_cast<uint32>(Reason);
+}
+
+inline bool FPackageData::GetIsCookLast() const
+{
+	return bIsCookLast != 0;
+}
+
+inline bool FPackageData::GetIsVisited() const
+{
+	return bIsVisited != 0;
+}
+
+inline void FPackageData::SetIsVisited(bool bValue)
+{
+	bIsVisited = static_cast<uint32>(bValue);
+}
+
+template <typename ArrayType>
+inline void FPackageData::GetCachedObjectsInOuterPlatforms(ArrayType& OutPlatforms) const
+{
+	OutPlatforms.Reset(PlatformDatas.Num());
+	for (const TPair<const ITargetPlatform*, FPackagePlatformData>& Pair : PlatformDatas)
+	{
+		if (Pair.Value.IsRegisteredForCachedObjectsInOuter())
+		{
+			OutPlatforms.Add(Pair.Key);
+		}
+	}
+}
+
+inline bool FPackageData::GetHasSaveCache() const
+{
+	return static_cast<bool>(bHasSaveCache);
+}
+
+inline void FPackageData::SetHasSaveCache(bool Value)
+{
+	bHasSaveCache = Value != 0;
+}
+
+inline ESaveSubState FPackageData::GetSaveSubState() const
+{
+	return static_cast<ESaveSubState>(SaveSubState);
+}
+
+inline bool FPackageData::HasPrepareSaveFailed() const
+{
+	return static_cast<bool>(bPrepareSaveFailed);
+}
+
+inline void FPackageData::SetHasPrepareSaveFailed(bool bValue)
+{
+	bPrepareSaveFailed = bValue != 0;
+}
+
+inline bool FPackageData::IsPrepareSaveRequiresGC() const
+{
+	return bPrepareSaveRequiresGC;
+}
+
+inline void FPackageData::SetIsPrepareSaveRequiresGC(bool bValue)
+{
+	bPrepareSaveRequiresGC = bValue != 0;
+}
+
+inline ECookResult FPackageData::GetMonitorCookResult() const
+{
+	return (ECookResult)MonitorCookResult;
+}
+
+inline void FPackageData::SetMonitorCookResult(ECookResult Value)
+{
+	MonitorCookResult = (uint8)Value;
+}
+
 inline bool FPackageData::IsGenerated() const
@@ -1641,5 +1723,173 @@
 	return BuildInstigator.Category != EInstigator::NotYetRequested;
 }
 
+inline bool FPackageData::IsKeepReferencedDuringGC() const
+{
+	return static_cast<bool>(bKeepReferencedDuringGC);
+}
 
+inline void FPackageData::SetKeepReferencedDuringGC(bool Value)
+{
+	bKeepReferencedDuringGC = Value != 0;
+}
+
+inline bool FPackageData::GetWasCookedThisSession() const
+{
+	return static_cast<bool>(bWasCookedThisSession);
+}
+
+inline FWorkerId FPackageData::GetWorkerAssignment() const
+{
+	return WorkerAssignment;
+}
+
+inline FWorkerId FPackageData::GetWorkerAssignmentConstraint() const
+{
+	return WorkerAssignmentConstraint;
+}
+
+typedef void (FPackageData::* FEdgeFunction)();
+inline void FPackageData::UpdateDownEdge(bool bOld, bool bNew, const FEdgeFunction& EdgeFunction)
+{
+	if ((bOld != bNew) & bOld)
+	{
+		(this->*EdgeFunction)();
+	}
+}
+
+inline void FPackageData::UpdateUpEdge(bool bOld, bool bNew, const FEdgeFunction& EdgeFunction)
+{
+	if ((bOld != bNew) & bNew)
+	{
+		(this->*EdgeFunction)();
+	}
+}
+
+inline FPackageDataSet& FRequestQueue::GetRestartedRequests()
+{
+	return RestartedRequests;
+}
+
+inline TRingBuffer<FDiscoveryQueueElement>& FRequestQueue::GetDiscoveryQueue()
+{
+	return DiscoveryQueue;
+}
+
+inline TRingBuffer<FRequestCluster>& FRequestQueue::GetRequestClusters()
+{
+	return RequestClusters;
+}
+
+inline FPackageDataSet& FRequestQueue::GetReadyRequestsUrgent()
+{
+	return UrgentRequests;
+}
+
+inline FPackageDataSet& FRequestQueue::GetReadyRequestsNormal()
+{
+	return NormalRequests;
+}
+
+
+inline FPackageDataMonitor& FPackageDatas::GetMonitor()
+{
+	return Monitor;
+}
+
+inline UCookOnTheFlyServer& FPackageDatas::GetCookOnTheFlyServer()
+{
+	return CookOnTheFlyServer;
+}
+
+inline uint32 FPackageDatas::GetNextLeafToRootRank()
+{
+	return NextLeafToRootRank++;
+}
+
+inline void FPackageDatas::ResetLeafToRootRank()
+{
+	NextLeafToRootRank = 0;
+}
+
+inline FRequestQueue& FPackageDatas::GetRequestQueue()
+{
+	return RequestQueue;
+}
+
+inline TFastPointerSet<FPackageData*>& FPackageDatas::GetAssignedToWorkerSet()
+{
+	return AssignedToWorkerSet;
+}
+
+inline FLoadQueue& FPackageDatas::GetLoadQueue()
+{
+	return LoadQueue;
+}
+
+inline TFastPointerSet<FPackageData*>& FPackageDatas::GetSaveStalledSet()
+{
+	return SaveStalledSet;
+}
+
+inline FPackageDataQueue& FPackageDatas::GetSaveQueue()
+{
+	return SaveQueue;
+}
+
+template<typename CallbackType>
+inline void FPackageDatas::UpdateThreadsafePackageData(FName PackageName, CallbackType&& Callback)
+{
+	FWriteScopeLock ExistenceWriteLock(ExistenceLock);
+	FThreadsafePackageData& Value = ThreadsafePackageDatas.FindOrAdd(PackageName);
+	bool bNew = false;
+	if (!Value.bInitialized)
+	{
+		Value.bInitialized = true;
+		bNew = true;
+	}
+	Callback(Value, bNew);
+}
+
+inline TOptional<FThreadsafePackageData> FPackageDatas::FindThreadsafePackageData(FName PackageName)
+{
+	FReadScopeLock ExistenceReadLock(ExistenceLock);
+	FThreadsafePackageData* Value = ThreadsafePackageDatas.Find(PackageName);
+	return Value ? TOptional<FThreadsafePackageData>(*Value) : TOptional<FThreadsafePackageData>();
+}
+
+template <typename FunctionType>
+inline void FPackageDatas::ForEachPendingCookedPlatformData(const FunctionType& Function)
+{
+	for (FPendingCookedPlatformDataContainer& Container : PendingCookedPlatformDataLists)
+	{
+		for (FPendingCookedPlatformData& Data : Container)
+		{
+			Function(Data);
+		}
+	}
+}
+
+inline int32 FPackageDatas::GetPendingCookedPlatformDataNum() const
+{
+	return PendingCookedPlatformDataNum;
+}
+
+template <typename CallbackType>
+inline void FPackageDatas::LockAndEnumeratePackageDatas(CallbackType&& Callback)
+{
+	FReadScopeLock ExistenceReadLock(ExistenceLock);
+	EnumeratePackageDatasWithinLock(Forward<CallbackType>(Callback));
+}
+
+inline TMap<UObject*, FCachedCookedPlatformDataState>& FPackageDatas::GetCachedCookedPlatformDataObjects()
+{
+	return CachedCookedPlatformDataObjects;
+}
+
+template <typename CallbackType>
+inline void FPackageDatas::EnumeratePackageDatasWithinLock(CallbackType&& Callback)
+{
+	Allocator.EnumerateAllocations(Forward<CallbackType>(Callback));
+}
+
 } // namespace UE::Cook
