[WorldPartitionRenameDuplicateBuilder] World Partition Rename Duplicate Builder commandlet fixes
* Handle invalid references - skip and log
* Properly replace references to duplicated objects

diff --git a/Engine/Source/Editor/UnrealEd/Private/WorldPartition/WorldPartitionRenameDuplicateBuilder.cpp b/Engine/Source/Editor/UnrealEd/Private/WorldPartition/WorldPartitionRenameDuplicateBuilder.cpp
--- a/Engine/Source/Editor/UnrealEd/Private/WorldPartition/WorldPartitionRenameDuplicateBuilder.cpp
+++ b/Engine/Source/Editor/UnrealEd/Private/WorldPartition/WorldPartitionRenameDuplicateBuilder.cpp
@@ -41,6 +41,7 @@ public:
 		ArIgnoreOuterRef = true;
 		ArIsObjectReferenceCollector = true;
 		ArShouldSkipBulkData = true;
+		ArIsModifyingWeakAndStrongReferences = true;
 
 		SubObjects.Add(Root);
 
@@ -155,7 +156,23 @@ bool UWorldPartitionRenameDuplicateBuilder::RunInternal(UWorld* World, const FCe
 	TArray<TPair<FGuid, TArray<FGuid>>> ActorsWithRefs;
 	for (FActorDescContainerInstanceCollection::TIterator<> Iterator(WorldPartition); Iterator; ++Iterator)
 	{
-		ActorsWithRefs.Emplace(Iterator->GetGuid(), Iterator->GetReferences());
+		FWorldPartitionActorDescInstance* ActorDescInstance = WorldPartition->GetActorDescInstance(Iterator->GetGuid());
+
+		TArray<FGuid> ValidReferences;
+		const TArray<FGuid>& References = Iterator->GetReferences();
+		for (const FGuid& Reference : References)
+		{
+			if (WorldPartition->GetActorDescInstance(Reference))
+			{
+				ValidReferences.Add(Reference);
+			}
+			else
+			{
+				UE_LOG(LogWorldPartitionRenameDuplicateBuilder, Warning, TEXT("Invalid reference %s of actor %s skipped"), *Reference.ToString(EGuidFormats::DigitsWithHyphensInBraces), *ActorDescInstance->GetActorPackage().ToString());
+			}
+		}
+
+		ActorsWithRefs.Emplace(Iterator->GetGuid(), ValidReferences);
 	}
 
 	TArray<TArray<FGuid>> ActorClusters = GenerateObjectsClusters(ActorsWithRefs);
