Add RemoveAllApplicationModes method to clear the list of ApplicationModes registered with WorkflowCentricApplication

#rnx

diff --git a/Engine/Plugins/Animation/LiveLinkHub/Source/LiveLinkHub/Private/LiveLinkHub.cpp b/Engine/Plugins/Animation/LiveLinkHub/Source/LiveLinkHub/Private/LiveLinkHub.cpp
--- a/Engine/Plugins/Animation/LiveLinkHub/Source/LiveLinkHub/Private/LiveLinkHub.cpp
+++ b/Engine/Plugins/Animation/LiveLinkHub/Source/LiveLinkHub/Private/LiveLinkHub.cpp
@@ -131,6 +131,8 @@ void FLiveLinkHub::InitializePostEngineInit()
 
 FLiveLinkHub::~FLiveLinkHub()
 {
+	RemoveAllApplicationModes();
+
 	IModularFeatures::Get().UnregisterModularFeature(IConsoleCommandExecutor::ModularFeatureName(), CommandExecutor.Get());
 	CommandExecutor.Reset();
 
diff --git a/Engine/Source/Editor/UnrealEd/Private/WorkflowCentricApplication.cpp b/Engine/Source/Editor/UnrealEd/Private/WorkflowCentricApplication.cpp
--- a/Engine/Source/Editor/UnrealEd/Private/WorkflowCentricApplication.cpp
+++ b/Engine/Source/Editor/UnrealEd/Private/WorkflowCentricApplication.cpp
@@ -154,6 +154,11 @@ void FWorkflowCentricApplication::AddApplicationMode(FName ModeName, TSharedRef<
 	ApplicationModeList.Add(ModeName, Mode);
 }
 
+void FWorkflowCentricApplication::RemoveAllApplicationModes()
+{
+	ApplicationModeList.Reset();
+}
+
 /////////////////////////////////////////////////////
 
 #undef LOCTEXT_NAMESPACE
diff --git a/Engine/Source/Editor/UnrealEd/Public/WorkflowOrientedApp/WorkflowCentricApplication.h b/Engine/Source/Editor/UnrealEd/Public/WorkflowOrientedApp/WorkflowCentricApplication.h
--- a/Engine/Source/Editor/UnrealEd/Public/WorkflowOrientedApp/WorkflowCentricApplication.h
+++ b/Engine/Source/Editor/UnrealEd/Public/WorkflowOrientedApp/WorkflowCentricApplication.h
@@ -56,6 +56,9 @@ public:
 protected:
 
 	UNREALED_API virtual void AddApplicationMode(FName ModeName, TSharedRef<FApplicationMode> Mode);
+
+	// Remove all application modes that were added through AddApplicationMode.
+	UNREALED_API virtual void RemoveAllApplicationModes();
 protected:
 	TSharedPtr<FApplicationMode> CurrentAppModePtr;
 
