Fix missing Substrate material properties in Debug Raytracing menu when Substrate is enabled.

diff --git a/Engine/Source/Editor/UnrealEd/Private/RayTracingDebugVisualizationMenuCommands.cpp b/Engine/Source/Editor/UnrealEd/Private/RayTracingDebugVisualizationMenuCommands.cpp
--- a/Engine/Source/Editor/UnrealEd/Private/RayTracingDebugVisualizationMenuCommands.cpp
+++ b/Engine/Source/Editor/UnrealEd/Private/RayTracingDebugVisualizationMenuCommands.cpp
@@ -71,6 +71,10 @@ void FRayTracingDebugVisualizationMenuCommands::BuildCommandMap()
 			Record.Type = FVisualizationType::Performance;
 			break;
 
+		case FRayTracingVisualizationData::FModeType::Material:
+			Record.Type = FVisualizationType::Material;
+			break;
+
 		case FRayTracingVisualizationData::FModeType::Other:
 			Record.Type = FVisualizationType::Other;
 			break;
@@ -102,6 +106,12 @@ void FRayTracingDebugVisualizationMenuCommands::BuildVisualisationSubMenu(FMenuB
 			Menu.EndSection();
 		}
 
+		{
+			Menu.BeginSection("LevelViewportRayTracingVisualizationMaterial", LOCTEXT("RayTracingVisualizationMaterial", "Material"));
+			Commands.AddCommandTypeToMenu(Menu, FVisualizationType::Material);
+			Menu.EndSection();
+		}
+
 		if (bShowOther)
 		{
 			Menu.BeginSection("LevelViewportRayTracingVisualizationOther", LOCTEXT("RayTracingVisualizationOther", "Other"));
diff --git a/Engine/Source/Editor/UnrealEd/Public/RayTracingDebugVisualizationMenuCommands.h b/Engine/Source/Editor/UnrealEd/Public/RayTracingDebugVisualizationMenuCommands.h
--- a/Engine/Source/Editor/UnrealEd/Public/RayTracingDebugVisualizationMenuCommands.h
+++ b/Engine/Source/Editor/UnrealEd/Public/RayTracingDebugVisualizationMenuCommands.h
@@ -19,6 +19,7 @@ public:
 		Overview,
 		Standard,
 		Performance,
+		Material,
 		Other,
 	};
 
diff --git a/Engine/Source/Runtime/Engine/Private/RayTracingVisualizationData.cpp b/Engine/Source/Runtime/Engine/Private/RayTracingVisualizationData.cpp
--- a/Engine/Source/Runtime/Engine/Private/RayTracingVisualizationData.cpp
+++ b/Engine/Source/Runtime/Engine/Private/RayTracingVisualizationData.cpp
@@ -81,7 +81,7 @@ void FRayTracingVisualizationData::Initialize()
 
 			if (Substrate::IsSubstrateEnabled())
 			{
-				AddVisualizationMode(TEXT("SubstrateMaterialProperties"), LOCTEXT("SubstrateMaterialProperties", "Substrate Material Properties"), FModeType::Other, RAY_TRACING_DEBUG_VIZ_SUBSTRATE_DATA, true);
+				AddVisualizationMode(TEXT("SubstrateMaterialProperties"), LOCTEXT("SubstrateMaterialProperties", "Substrate Material Properties"), FModeType::Material, RAY_TRACING_DEBUG_VIZ_SUBSTRATE_DATA, true);
 			}
 		}
 
diff --git a/Engine/Source/Runtime/Engine/Public/RayTracingVisualizationData.h b/Engine/Source/Runtime/Engine/Public/RayTracingVisualizationData.h
--- a/Engine/Source/Runtime/Engine/Public/RayTracingVisualizationData.h
+++ b/Engine/Source/Runtime/Engine/Public/RayTracingVisualizationData.h
@@ -13,6 +13,7 @@ public:
 		Overview,
 		Standard,
 		Performance,
+		Material,
 		Other,
 	};
 
