[Viewport Toolbar] Wrong tooltips for snapping entries

diff --git a/Engine/Source/Editor/UnrealEd/Private/ViewportToolbar/UnrealEdViewportToolbar.cpp b/Engine/Source/Editor/UnrealEd/Private/ViewportToolbar/UnrealEdViewportToolbar.cpp
--- a/Engine/Source/Editor/UnrealEd/Private/ViewportToolbar/UnrealEdViewportToolbar.cpp
+++ b/Engine/Source/Editor/UnrealEd/Private/ViewportToolbar/UnrealEdViewportToolbar.cpp
@@ -333,7 +333,7 @@ FToolMenuEntry CreateLocationSnapCheckboxMenu()
 		Entry = UnrealEd::CreateCheckboxSubmenu(
 			"GridSnapping",
 			LocationSnapLabel,
-			FEditorViewportCommands::Get().SurfaceSnapping->GetDescription(),
+			FEditorViewportCommands::Get().LocationGridSnap->GetDescription(),
 			FToolMenuExecuteAction::CreateLambda(
 				[](const FToolMenuContext& InContext)
 				{
@@ -405,7 +405,7 @@ FToolMenuEntry CreateRotationSnapCheckboxMenu()
 	FToolMenuEntry Entry = UnrealEd::CreateCheckboxSubmenu(
 		"RotationSnapping",
 		RotationSnapLabel,
-		FEditorViewportCommands::Get().SurfaceSnapping->GetDescription(),
+		FEditorViewportCommands::Get().RotationGridSnap->GetDescription(),
 		FToolMenuExecuteAction::CreateLambda(
 			[](const FToolMenuContext& InContext)
 			{
@@ -509,7 +509,7 @@ FToolMenuEntry CreateScaleSnapCheckboxMenu()
 	FToolMenuEntry Entry = UnrealEd::CreateCheckboxSubmenu(
 		"ScaleSnapping",
 		ScaleSnapLabel,
-		FEditorViewportCommands::Get().SurfaceSnapping->GetDescription(),
+		FEditorViewportCommands::Get().ScaleGridSnap->GetDescription(),
 		FToolMenuExecuteAction::CreateLambda(
 			[](const FToolMenuContext& InContext)
 			{
