Optimization for  FAssetRegistryGenerator::FinalizeChunkIDs : do not call FAssetRegistryState::PruneAssetData.
Instead create a new object FAssetRegistryState and call InitializeFromExistingAndPrune. Then move the new object to the existing variable State. This is way faster : from 1m down to 4s.

#rnx

diff --git a/Engine/Source/Editor/UnrealEd/Private/Commandlets/AssetRegistryGenerator.cpp b/Engine/Source/Editor/UnrealEd/Private/Commandlets/AssetRegistryGenerator.cpp
--- a/Engine/Source/Editor/UnrealEd/Private/Commandlets/AssetRegistryGenerator.cpp
+++ b/Engine/Source/Editor/UnrealEd/Private/Commandlets/AssetRegistryGenerator.cpp
@@ -1634,9 +1634,17 @@ void FAssetRegistryGenerator::FinalizeChunkIDs(const TSet<FName>& InCookedPackag
 	AllPackages.Append(DevelopmentOnlyPackages);
 
 	// Prune our asset registry to cooked + dev only list
-	FAssetRegistrySerializationOptions DevelopmentSaveOptions;
-	AssetRegistry.InitializeSerializationOptions(DevelopmentSaveOptions, TargetPlatform->IniPlatformName(), UE::AssetRegistry::ESerializationTarget::ForDevelopment);
-	State.PruneAssetData(AllPackages, TSet<FName>(), DevelopmentSaveOptions);
+	{
+		FAssetRegistrySerializationOptions DevelopmentSaveOptions;
+		AssetRegistry.InitializeSerializationOptions(DevelopmentSaveOptions, TargetPlatform->IniPlatformName(), UE::AssetRegistry::ESerializationTarget::ForDevelopment);
+		DevelopmentSaveOptions.bKeepDevelopmentAssetRegistryTags = true;
+
+		// Create a new FAssetRegistryState and call InitializeFromExistingAndPrune, then move the result to State.
+		// This is faster than calling State.PruneAssetData. PruneAssetData removes iteratively all the packages one by one and it is slow.
+		FAssetRegistryState PrunedAssetRegistry;
+		PrunedAssetRegistry.InitializeFromExistingAndPrune(State, AllPackages, TSet<FName>(), TSet<int32>(), DevelopmentSaveOptions);
+		State = MoveTemp(PrunedAssetRegistry);
+	}
 
 	// Development only packages should have been marked via UpdateAssetRegistryData with a negative size indicating why they were not cooked 
 	for (FName DevelopmentOnlyPackage : DevelopmentOnlyPackages)
