World Bookmark - Store the ActorEditorContext using a ActorEditorContextStateCollection object instead of an actor
* This allow context properties to be visible in the outliner
* Allows us to remove the use of a dummy actor to store the context, which came with its own set of issues

diff --git a/Engine/Source/Editor/UnrealEd/Private/EditorState/ActorEditorContextEditorState.cpp b/Engine/Source/Editor/UnrealEd/Private/EditorState/ActorEditorContextEditorState.cpp
--- a/Engine/Source/Editor/UnrealEd/Private/EditorState/ActorEditorContextEditorState.cpp
+++ b/Engine/Source/Editor/UnrealEd/Private/EditorState/ActorEditorContextEditorState.cpp
@@ -2,10 +2,12 @@
 
 #include "EditorState/ActorEditorContextEditorState.h"
 
+#include "ActorEditorContextState.h"
 #include "Engine/Level.h"
 #include "Engine/World.h"
 #include "GameFramework/Actor.h"
 #include "Subsystems/ActorEditorContextSubsystem.h"
+#include "IActorEditorContextClient.h"
 #include "UObject/Package.h"
 
 #include UE_INLINE_GENERATED_CPP_BY_NAME(ActorEditorContextEditorState)
@@ -18,23 +20,6 @@ UActorEditorContextEditorState::UActorEditorContextEditorState(const FObjectInit
 {
 }
 
-void UActorEditorContextEditorState::PostInitProperties()
-{
-	Super::PostInitProperties();
-
-	if (!HasAnyFlags(RF_ClassDefaultObject | RF_NeedLoad))
-	{
-		ActorEditorContext = NewObject<AActor>(this);
-	}
-}
-
-void UActorEditorContextEditorState::PostDuplicate(bool bDuplicateForPIE)
-{
-	Super::PostDuplicate(bDuplicateForPIE);
-
-	ActorEditorContext = NewObject<AActor>(this);
-}
-
 FText UActorEditorContextEditorState::GetCategoryText() const
 {
 	return FText(LOCTEXT("ActorEditorContextEditorStateCategoryText", "Actor Editor Context"));
@@ -42,23 +27,18 @@ FText UActorEditorContextEditorState::GetCategoryText() const
 
 UEditorState::FOperationResult UActorEditorContextEditorState::CaptureState()
 {
-	UWorld* CurrentWorld = GetStateWorld();
-
-	// Store context to saved actor with EActorEditorContextAction::ApplyContext
-	ActorEditorContext->SetFlags(RF_Transient);
-	ActorEditorContext->Rename(nullptr, GetTransientPackage(), REN_DoNotDirty | REN_DontCreateRedirectors | REN_NonTransactional);
-	ActorEditorContext->MarkAsGarbage();
-
-	UPackage* TempExternalActorPackage = CreatePackage(*FString::Printf(TEXT("/Temp/ActorEditorContextEditorState/%s"), *FGuid::NewGuid().ToString()));
-	TempExternalActorPackage->SetFlags(RF_Transient);
-
-	ActorEditorContext = NewObject<AActor>(CurrentWorld->PersistentLevel);
-	ActorEditorContext->SetExternalPackage(TempExternalActorPackage);
+	if (ActorEditorContextStateCollection == nullptr)
+	{
+		ActorEditorContextStateCollection = NewObject<UActorEditorContextStateCollection>(this);
+	}
 
-	UActorEditorContextSubsystem::Get()->ApplyContext(ActorEditorContext);
+	UActorEditorContextSubsystem::Get()->CaptureContext(ActorEditorContextStateCollection);
 
-	ActorEditorContext->SetExternalPackage(nullptr);
-	ActorEditorContext->Rename(nullptr, this, REN_DoNotDirty | REN_DontCreateRedirectors | REN_NonTransactional);
+	// Context is empty, no need to keep object around
+	if (ActorEditorContextStateCollection->IsEmpty())
+	{
+		ActorEditorContextStateCollection = nullptr;
+	}
 
 	return FOperationResult(FOperationResult::Success);
 }
@@ -72,15 +52,14 @@ UEditorState::FOperationResult UActorEditorContextEditorState::RestoreState() co
 		return FOperationResult(FOperationResult::Skipped, LOCTEXT("RestoreStateSkipped_ApplyActorEditorContextOnLoad", "User manually disabled application of the actor editor context"));
 	}
 
-	// Actor Editor Context - Apply context from saved actor with EActorEditorContextAction::InitializeContextFromActor
-	ActorEditorContext->Rename(nullptr, CurrentWorld->PersistentLevel, REN_DoNotDirty | REN_DontCreateRedirectors | REN_NonTransactional);
-	ActorEditorContext->SetExternalPackage(GetPackage());
-
-	UActorEditorContextSubsystem::Get()->ResetContext();
-	UActorEditorContextSubsystem::Get()->InitializeContextFromActor(ActorEditorContext);
-
-	ActorEditorContext->SetExternalPackage(nullptr);
-	ActorEditorContext->Rename(nullptr, const_cast<UActorEditorContextEditorState*>(this), REN_DoNotDirty | REN_DontCreateRedirectors | REN_NonTransactional);
+	if (ActorEditorContextStateCollection != nullptr)
+	{
+		UActorEditorContextSubsystem::Get()->RestoreContext(ActorEditorContextStateCollection);
+	}
+	else
+	{
+		UActorEditorContextSubsystem::Get()->ResetContext();
+	}
 
 	return FOperationResult(FOperationResult::Success);
 }
diff --git a/Engine/Source/Editor/UnrealEd/Public/EditorState/ActorEditorContextEditorState.h b/Engine/Source/Editor/UnrealEd/Public/EditorState/ActorEditorContextEditorState.h
--- a/Engine/Source/Editor/UnrealEd/Public/EditorState/ActorEditorContextEditorState.h
+++ b/Engine/Source/Editor/UnrealEd/Public/EditorState/ActorEditorContextEditorState.h
@@ -5,7 +5,7 @@
 #include "EditorState/WorldEditorState.h"
 #include "ActorEditorContextEditorState.generated.h"
 
-class AActor;
+class UActorEditorContextStateCollection;
 
 UCLASS(MinimalAPI)
 class UActorEditorContextEditorState : public UWorldDependantEditorState
@@ -21,15 +21,10 @@ protected:
 	virtual FOperationResult RestoreState() const override;
 	//~ End UEditorState Interface
 
-	//~ Begin UObject Interface
-	virtual void PostInitProperties() override;
-	virtual void PostDuplicate(bool bDuplicateForPIE) override;
-	//~ End UObject Interface
+	// Editor context state object, will be used to restore context when loading the bookmark.
+	UPROPERTY(VisibleAnywhere, Category = ActorEditorContext, Instanced, meta=(ShowOnlyInnerProperties, DisplayName="Context"))
+	TObjectPtr<UActorEditorContextStateCollection> ActorEditorContextStateCollection;
 
-	// Actor used to store the editor context, will be used to restore context when loading the bookmark.
-	UPROPERTY(DuplicateTransient)
-	TObjectPtr<AActor> ActorEditorContext;
-
-	UPROPERTY(EditAnywhere, Category = ActorEditorContext)
+	UPROPERTY(EditAnywhere, Category = ActorEditorContext, meta=(DisplayName="Apply Context on Load"))
 	bool bApplyActorEditorContextOnLoad;
 };
