MPCook: Add more feedback for display message "CookWorkerServer %d received FPackageResultsMessage for package %s which is not a pending package"; this message is possibly related to missing chunk errors during staging for the packages it reports.
#rnx

diff --git a/Engine/Source/Editor/UnrealEd/Private/Cooker/CookTypes.cpp b/Engine/Source/Editor/UnrealEd/Private/Cooker/CookTypes.cpp
--- a/Engine/Source/Editor/UnrealEd/Private/Cooker/CookTypes.cpp
+++ b/Engine/Source/Editor/UnrealEd/Private/Cooker/CookTypes.cpp
@@ -1119,3 +1119,24 @@ IPackageWriter::ECommitStatus PackageResultToCommitStatus(ESavePackageResult Res
 		return IPackageWriter::ECommitStatus::Error;
 	}
 }
+
+const TCHAR* LexToString(UE::Cook::ECookResult CookResult)
+{
+	using namespace UE::Cook;
+	switch (CookResult)
+	{
+	case ECookResult::NotAttempted:
+		return TEXT("NotAttempted");
+	case ECookResult::Succeeded:
+		return TEXT("Succeeded");
+	case ECookResult::Failed:
+		return TEXT("Failed");
+	case ECookResult::NeverCookPlaceholder:
+		return TEXT("NeverCookPlaceholder");
+	case ECookResult::Invalid:
+		[[fallthrough]];
+	default:
+		return TEXT("Invalid");
+		break;
+	}
+}
diff --git a/Engine/Source/Editor/UnrealEd/Private/Cooker/CookTypes.h b/Engine/Source/Editor/UnrealEd/Private/Cooker/CookTypes.h
--- a/Engine/Source/Editor/UnrealEd/Private/Cooker/CookTypes.h
+++ b/Engine/Source/Editor/UnrealEd/Private/Cooker/CookTypes.h
@@ -772,3 +772,5 @@ FMultiPackageReaderResults GetSaveExportsAndImports(UPackage* Package, UObject*
 IPackageWriter::ECommitStatus PackageResultToCommitStatus(ESavePackageResult Result);
 
 extern int32 GCookProgressDisplay;
+
+const TCHAR* LexToString(UE::Cook::ECookResult CookResult);
diff --git a/Engine/Source/Editor/UnrealEd/Private/Cooker/CookWorkerServer.cpp b/Engine/Source/Editor/UnrealEd/Private/Cooker/CookWorkerServer.cpp
--- a/Engine/Source/Editor/UnrealEd/Private/Cooker/CookWorkerServer.cpp
+++ b/Engine/Source/Editor/UnrealEd/Private/Cooker/CookWorkerServer.cpp
@@ -906,6 +906,31 @@ void FCookWorkerServer::SendMessageInLock(const IMPCollectorMessage& Message)
 	}
 }
 
+FString WriteCookStatus(FPackageData& PackageData, TConstArrayView<const ITargetPlatform*> SessionPlatforms)
+{
+	FString Result;
+	Result.Reserve(256);
+	auto BoolToString = [](bool bValue)
+		{
+			return bValue ? TEXT("true") : TEXT("false");
+		};
+	for (const ITargetPlatform* TargetPlatform : SessionPlatforms)
+	{
+		FPackagePlatformData* PlatformData = PackageData.FindPlatformData(TargetPlatform);
+		Result.Appendf(TEXT("[ %s: { Reachable: %s, Cookable: %s, CookResult: %s }, "),
+			*TargetPlatform->PlatformName(), BoolToString(PlatformData && PlatformData->IsReachable()),
+			BoolToString(PlatformData && PlatformData->IsCookable()),
+			(PlatformData ? ::LexToString(PlatformData->GetCookResults()) : TEXT("<NotCooked>")));
+	}
+	if (Result.Len() >= 2)
+	{
+		Result.LeftChopInline(2); // Remove the trailing ", "
+		Result.Append(TEXT(" ]"));
+	}
+
+	return Result;
+}
+
 void FCookWorkerServer::RecordResults(FPackageResultsMessage& Message)
 {
 	check(TickState.TickThread == ECookDirectorThread::SchedulerThread);
@@ -924,8 +949,11 @@ void FCookWorkerServer::RecordResults(FPackageResultsMessage& Message)
 		if (PendingPackages.Remove(PackageData) != 1)
 		{
 			UE_LOG(LogCook, Display,
-				TEXT("CookWorkerServer %d received FPackageResultsMessage for package %s which is not a pending package. Ignoring it."),
-				ProfileId, *Result.GetPackageName().ToString());
+				TEXT("CookWorkerServer %d received FPackageResultsMessage for package %s which is not a pending package. Ignoring it.")
+				TEXT("\n\tState = %s, WorkerId = %s, CookResults = { %s }"),
+				ProfileId, *Result.GetPackageName().ToString(),
+				LexToString(PackageData->GetState()), *Director.GetDisplayName(PackageData->GetWorkerAssignment()),
+				*WriteCookStatus(*PackageData, COTFS.GetSessionPlatforms()));
 			continue;
 		}
 		bRetiredAnyPackages = true;
