MPCook: Fix assertion when invalid WorkerId is passed to GetDisplayName.

diff --git a/Engine/Source/Editor/UnrealEd/Private/Cooker/CookDirector.cpp b/Engine/Source/Editor/UnrealEd/Private/Cooker/CookDirector.cpp
--- a/Engine/Source/Editor/UnrealEd/Private/Cooker/CookDirector.cpp
+++ b/Engine/Source/Editor/UnrealEd/Private/Cooker/CookDirector.cpp
@@ -788,7 +788,11 @@ void FCookDirector::DisplayRemainingPackages() const
 FString FCookDirector::GetDisplayName(const FWorkerId& WorkerId, int32 PreferredWidth) const
 {
 	FString Result;
-	if (WorkerId.IsLocal())
+	if (WorkerId.IsInvalid())
+	{
+		Result = TEXTVIEW("None");
+	}
+	else if (WorkerId.IsLocal())
 	{
 		Result = TEXTVIEW("Local");
 	}
