[Viewport Toolbar] Add Screen Percentage icon

- Adds proper icon
- Code cleanup so that icon works, and redundant code is removed

diff --git a/Engine/Source/Editor/EditorStyle/Private/StarshipStyle.cpp b/Engine/Source/Editor/EditorStyle/Private/StarshipStyle.cpp
--- a/Engine/Source/Editor/EditorStyle/Private/StarshipStyle.cpp
+++ b/Engine/Source/Editor/EditorStyle/Private/StarshipStyle.cpp
@@ -3065,6 +3065,7 @@ void FStarshipEditorStyle::FStyle::SetupGeneralIcons()
 	Set("AssetEditor.CompileStatus.Overlay.Warning", new IMAGE_BRUSH_SVG("Starship/Blueprints/CompileStatus_Warning_Badge", Icon20x20, FStyleColors::Warning));
 	Set("AssetEditor.CompileStatus.Overlay.Good", new IMAGE_BRUSH_SVG("Starship/Blueprints/CompileStatus_Good_Badge", Icon20x20, FStyleColors::AccentGreen));
 	Set("AssetEditor.CompileStatus.Overlay.Error", new IMAGE_BRUSH_SVG("Starship/Blueprints/CompileStatus_Fail_Badge", Icon20x20, FStyleColors::Error));
+	Set("AssetEditor.ScreenPercentage", new IMAGE_BRUSH_SVG("Starship/Common/ScreenPercentage", Icon16x16));
 	
 	Set("Debug", new IMAGE_BRUSH_SVG( "Starship/Common/Debug", Icon16x16 ) );
 	Set("Modules", new IMAGE_BRUSH_SVG( "Starship/Common/Modules", Icon16x16 ) );
diff --git a/Engine/Source/Editor/LevelEditor/Private/SLevelViewportToolBar.cpp b/Engine/Source/Editor/LevelEditor/Private/SLevelViewportToolBar.cpp
--- a/Engine/Source/Editor/LevelEditor/Private/SLevelViewportToolBar.cpp
+++ b/Engine/Source/Editor/LevelEditor/Private/SLevelViewportToolBar.cpp
@@ -503,11 +503,7 @@ void SLevelViewportToolBar::FillOptionsMenu(UToolMenu* Menu)
 				Section.AddEntry(UE::LevelEditor::CreateFarViewPlaneMenu(Viewport));
 			}
 
-			Section.AddSubMenu(
-				"ScreenPercentageSubMenu",
-				LOCTEXT("ScreenPercentageSubMenu", "Screen Percentage"),
-				LOCTEXT("ScreenPercentageSubMenu_ToolTip", "Customize the viewport's screen percentage"),
-				FNewToolMenuDelegate::CreateStatic(&UE::UnrealEd::ConstructScreenPercentageMenu));
+			Section.AddEntry(UE::UnrealEd::CreateScreenPercentageSubmenu());
 		}
 
 		{
diff --git a/Engine/Source/Editor/LevelEditor/Private/ViewportToolbar/LevelEditorSubmenus.cpp b/Engine/Source/Editor/LevelEditor/Private/ViewportToolbar/LevelEditorSubmenus.cpp
--- a/Engine/Source/Editor/LevelEditor/Private/ViewportToolbar/LevelEditorSubmenus.cpp
+++ b/Engine/Source/Editor/LevelEditor/Private/ViewportToolbar/LevelEditorSubmenus.cpp
@@ -1790,25 +1790,7 @@ FToolMenuEntry CreatePerformanceAndScalabilitySubmenu()
 
 					PerformanceAndScalabilitySection.AddEntry(CreateMaterialQualityLevelSubmenu());
 
-					PerformanceAndScalabilitySection.AddSubMenu(
-						"ScreenPercentageSubMenu",
-						LOCTEXT("ScreenPercentageSubMenu", "Screen Percentage"),
-						LOCTEXT("ScreenPercentageSubMenu_ToolTip", "Customize the viewport's screen percentage"),
-						FNewToolMenuDelegate::CreateLambda(
-							[](UToolMenu* ScreenPercentageSubMenu)
-							{
-								ScreenPercentageSubMenu->AddDynamicSection(
-									NAME_None,
-									FNewToolMenuDelegate::CreateLambda(
-										[](UToolMenu* InMenu) -> void
-										{
-											UE::UnrealEd::ConstructScreenPercentageMenu(InMenu);
-										}
-									)
-								);
-							}
-						)
-					);
+					PerformanceAndScalabilitySection.AddEntry(UE::UnrealEd::CreateScreenPercentageSubmenu());
 				}
 			}
 		)
diff --git a/Engine/Source/Editor/UnrealEd/Private/SCommonEditorViewportToolbarBase.cpp b/Engine/Source/Editor/UnrealEd/Private/SCommonEditorViewportToolbarBase.cpp
--- a/Engine/Source/Editor/UnrealEd/Private/SCommonEditorViewportToolbarBase.cpp
+++ b/Engine/Source/Editor/UnrealEd/Private/SCommonEditorViewportToolbarBase.cpp
@@ -54,18 +54,7 @@ void AddScreenPercentageMenu(FMenuBuilder& InMenuBuilder, const FEditorViewportC
 			UToolMenus::Get()->RegisterMenu(OldScreenPercentageMenuName, NAME_None, EMultiBoxType::Menu, false);
 
 		FToolMenuSection& UnnamedSection = Menu->FindOrAddSection(NAME_None);
-
-		UnnamedSection.AddSubMenu(
-			"ScreenPercentage",
-			LOCTEXT("ScreenPercentageSubMenu", "Screen Percentage"),
-			LOCTEXT("ScreenPercentageSubMenu_ToolTip", "Customize the viewport's screen percentage"),
-			FNewToolMenuDelegate::CreateLambda(
-				[](UToolMenu* Submenu) -> void
-				{
-					UE::UnrealEd::ConstructScreenPercentageMenu(Submenu);
-				}
-			)
-		);
+		UnnamedSection.AddEntry(CreateScreenPercentageSubmenu());
 	}
 
 	FToolMenuContext MenuContext;
diff --git a/Engine/Source/Editor/UnrealEd/Private/ViewportToolbar/UnrealEdViewportToolbar.cpp b/Engine/Source/Editor/UnrealEd/Private/ViewportToolbar/UnrealEdViewportToolbar.cpp
--- a/Engine/Source/Editor/UnrealEd/Private/ViewportToolbar/UnrealEdViewportToolbar.cpp
+++ b/Engine/Source/Editor/UnrealEd/Private/ViewportToolbar/UnrealEdViewportToolbar.cpp
@@ -3107,6 +3107,25 @@ void ConstructScreenPercentageMenu(UToolMenu* InMenu)
 	}
 }
 
+FToolMenuEntry CreateScreenPercentageSubmenu()
+{
+	FToolMenuEntry Entry = FToolMenuEntry::InitSubMenu(
+		"ScreenPercentage",
+		LOCTEXT("ScreenPercentageSubMenu", "Screen Percentage"),
+		LOCTEXT("ScreenPercentageSubMenu_ToolTip", "Customize the viewport's screen percentage"),
+		FNewToolMenuDelegate::CreateLambda(
+			[](UToolMenu* InMenu) -> void
+			{
+				ConstructScreenPercentageMenu(InMenu);
+			}
+		),
+		false,
+		FSlateIcon(FAppStyle::GetAppStyleSetName(), "AssetEditor.ScreenPercentage")
+	);
+
+	return Entry;
+}
+
 FToolMenuEntry CreateScalabilitySubmenu()
 {
 	return FToolMenuEntry::InitDynamicEntry(
@@ -3185,12 +3204,7 @@ FToolMenuEntry CreatePerformanceAndScalabilitySubmenu()
 					UnnamedSection.AddEntry(CreateRemoveRealtimeOverrideEntry(Context->Viewport));
 				}
 
-				UnnamedSection.AddSubMenu(
-					"ScreenPercentage",
-					LOCTEXT("ScreenPercentageSubMenu", "Screen Percentage"),
-					LOCTEXT("ScreenPercentageSubMenu_ToolTip", "Customize the viewport's screen percentage"),
-					FNewToolMenuDelegate::CreateStatic(&ConstructScreenPercentageMenu)
-				);
+				Submenu->AddMenuEntry("ScreenPercentageSubmenu", CreateScreenPercentageSubmenu());
 			}
 		)
 	);
diff --git a/Engine/Source/Editor/UnrealEd/Public/ViewportToolbar/UnrealEdViewportToolbar.h b/Engine/Source/Editor/UnrealEd/Public/ViewportToolbar/UnrealEdViewportToolbar.h
--- a/Engine/Source/Editor/UnrealEd/Public/ViewportToolbar/UnrealEdViewportToolbar.h
+++ b/Engine/Source/Editor/UnrealEd/Public/ViewportToolbar/UnrealEdViewportToolbar.h
@@ -196,6 +196,7 @@ UNREALED_API FOnViewportClientCamSpeedScalarChanged& OnViewportClientCamSpeedSca
 UNREALED_API void CreateCameraSpeedMenu(UToolMenu* InMenu);
 
 UNREALED_API void ConstructScreenPercentageMenu(UToolMenu* InMenu);
+UNREALED_API FToolMenuEntry CreateScreenPercentageSubmenu();
 
 UNREALED_API FToolMenuEntry CreateScalabilitySubmenu();
 
