[UnrealEd] "Browse to Asset" shortcut works with Levels

diff --git a/Engine/Source/Editor/UnrealEd/Private/EditorEngine.cpp b/Engine/Source/Editor/UnrealEd/Private/EditorEngine.cpp
--- a/Engine/Source/Editor/UnrealEd/Private/EditorEngine.cpp
+++ b/Engine/Source/Editor/UnrealEd/Private/EditorEngine.cpp
@@ -3351,6 +3351,20 @@ void UEditorEngine::GetAssetsToSyncToContentBrowser(TArray<FAssetData>& Assets,
 
 				return true;
 			});
+
+			// Only attempt to select the world if no elements are selected (a light might be selected, but it won't produce AssetData)
+			if (!SelectionSet->HasSelectedElements() && Assets.IsEmpty())
+			{
+				// If nothing is selected, sync to the current level (if it's not temp/untitled)
+				if (const UWorld* World = GEditor->GetEditorWorldContext().World())
+				{
+					if (const ULevel* Level = World->GetCurrentLevel();
+						Level && !FPackageName::IsTempPackage(Level->GetPackage()->GetName()))
+					{
+						Assets.Add(FAssetData(Level));
+					}
+				}
+			}
 		}
 	}
 }
