Make snapping checkboxes work in the overflow menu of the old viewport toolbars of asset editors

diff --git a/Engine/Source/Editor/UnrealEd/Private/STransformViewportToolbar.cpp b/Engine/Source/Editor/UnrealEd/Private/STransformViewportToolbar.cpp
--- a/Engine/Source/Editor/UnrealEd/Private/STransformViewportToolbar.cpp
+++ b/Engine/Source/Editor/UnrealEd/Private/STransformViewportToolbar.cpp
@@ -264,19 +264,39 @@ TSharedRef< SWidget > STransformViewportToolBar::MakeTransformToolBar( const TSh
 			HAlign_Fill,
 
 			// explictly specify what this widget should look like as a menu item
-			FNewMenuDelegate::CreateLambda( [this, Command]( FMenuBuilder& InMenuBuilder )
-			{
-				// TODO - debug why can't just use the Command / mapping isn't working 
-				InMenuBuilder.AddMenuEntry(Command);
+			FNewMenuDelegate::CreateLambda(
+				[this, Command](FMenuBuilder& InMenuBuilder)
+				{
+					FUIAction Action;
+					Action.ExecuteAction = FExecuteAction::CreateLambda(
+						[this]()
+						{
+							constexpr ECheckBoxState UnusedCheckedState = ECheckBoxState::Unchecked;
+							this->HandleToggleLocationGridSnap(UnusedCheckedState);
+						}
+					);
+					Action.GetActionCheckState =
+						FGetActionCheckState::CreateRaw(this, &STransformViewportToolBar::IsLocationGridSnapChecked);
+
+					const FName UnusedExtensionHook = NAME_None;
+					InMenuBuilder.AddMenuEntry(
+						Command->GetLabel(),
+						Command->GetDescription(),
+						Command->GetIcon(),
+						Action,
+						UnusedExtensionHook,
+						EUserInterfaceActionType::ToggleButton
+					);
 
-				InMenuBuilder.AddWrapperSubMenu(
-					LOCTEXT("GridSnapMenuSettings", "Grid Snap Settings"),
-					LOCTEXT("GridSnapMenuSettings_ToolTip", "Set the Position Grid Snap value"),
-					FOnGetContent::CreateSP(this, &STransformViewportToolBar::FillLocationGridSnapMenu),
-					FSlateIcon(Command->GetIcon())
-				);
-			}
-		));
+					InMenuBuilder.AddWrapperSubMenu(
+						LOCTEXT("GridSnapMenuSettings", "Grid Snap Settings"),
+						LOCTEXT("GridSnapMenuSettings_ToolTip", "Set the Position Grid Snap value"),
+						FOnGetContent::CreateSP(this, &STransformViewportToolBar::FillLocationGridSnapMenu),
+						FSlateIcon(Command->GetIcon())
+					);
+				}
+			)
+		);
 	}
 	ToolbarBuilder.EndSection();
 
@@ -303,18 +323,39 @@ TSharedRef< SWidget > STransformViewportToolBar::MakeTransformToolBar( const TSh
 			HAlign_Fill,
 
 			// explictly specify what this widget should look like as a menu item
-			FNewMenuDelegate::CreateLambda( [this, Command]( FMenuBuilder& InMenuBuilder )
-			{
-				InMenuBuilder.AddMenuEntry(Command);
+			FNewMenuDelegate::CreateLambda(
+				[this, Command](FMenuBuilder& InMenuBuilder)
+				{
+					FUIAction Action;
+					Action.ExecuteAction = FExecuteAction::CreateLambda(
+						[this]()
+						{
+							constexpr ECheckBoxState UnusedCheckedState = ECheckBoxState::Unchecked;
+							this->HandleToggleRotationGridSnap(UnusedCheckedState);
+						}
+					);
+					Action.GetActionCheckState =
+						FGetActionCheckState::CreateRaw(this, &STransformViewportToolBar::IsRotationGridSnapChecked);
+
+					const FName UnusedExtensionHook = NAME_None;
+					InMenuBuilder.AddMenuEntry(
+						Command->GetLabel(),
+						Command->GetDescription(),
+						Command->GetIcon(),
+						Action,
+						UnusedExtensionHook,
+						EUserInterfaceActionType::ToggleButton
+					);
 
-				InMenuBuilder.AddWrapperSubMenu(
-					LOCTEXT("RotationGridSnapMenuSettings", "Rotation Snap Settings"),
-					LOCTEXT("RotationGridSnapMenuSettings_ToolTip", "Adjust the Grid Settings for Rotation Snap"),
-					FOnGetContent::CreateSP(this, &STransformViewportToolBar::FillRotationGridSnapMenu),
-					FSlateIcon(Command->GetIcon())
-				);
-			}
-		));
+					InMenuBuilder.AddWrapperSubMenu(
+						LOCTEXT("RotationGridSnapMenuSettings", "Rotation Snap Settings"),
+						LOCTEXT("RotationGridSnapMenuSettings_ToolTip", "Adjust the Grid Settings for Rotation Snap"),
+						FOnGetContent::CreateSP(this, &STransformViewportToolBar::FillRotationGridSnapMenu),
+						FSlateIcon(Command->GetIcon())
+					);
+				}
+			)
+		);
 	}
 	ToolbarBuilder.EndSection();
 
@@ -349,7 +390,26 @@ TSharedRef< SWidget > STransformViewportToolBar::MakeTransformToolBar( const TSh
 			{
 				if (IsLayer2DSnapVisible() == EVisibility::Visible)
 				{
-					InMenuBuilder.AddMenuEntry(Command);
+					FUIAction Action;
+					Action.ExecuteAction = FExecuteAction::CreateLambda(
+						[this]()
+						{
+							constexpr ECheckBoxState UnusedCheckedState = ECheckBoxState::Unchecked;
+							this->HandleToggleLayer2DSnap(UnusedCheckedState);
+						}
+					);
+					Action.GetActionCheckState =
+						FGetActionCheckState::CreateRaw(this, &STransformViewportToolBar::IsLayer2DSnapChecked);
+
+					const FName UnusedExtensionHook = NAME_None;
+					InMenuBuilder.AddMenuEntry(
+						Command->GetLabel(),
+						Command->GetDescription(),
+						Command->GetIcon(),
+						Action,
+						UnusedExtensionHook,
+						EUserInterfaceActionType::ToggleButton
+					);
 
 					InMenuBuilder.AddWrapperSubMenu(
 						LOCTEXT("Layer2DSnapMenuSettings", "Layer 2D Snap Settings"),
@@ -388,18 +448,39 @@ TSharedRef< SWidget > STransformViewportToolBar::MakeTransformToolBar( const TSh
 			HAlign_Fill,
 
 			// explictly specify what this widget should look like as a menu item
-			FNewMenuDelegate::CreateLambda( [this, Command]( FMenuBuilder& InMenuBuilder )
-			{
-				InMenuBuilder.AddMenuEntry(Command);
+			FNewMenuDelegate::CreateLambda(
+				[this, Command](FMenuBuilder& InMenuBuilder)
+				{
+					FUIAction Action;
+					Action.ExecuteAction = FExecuteAction::CreateLambda(
+						[this]()
+						{
+							constexpr ECheckBoxState UnusedCheckedState = ECheckBoxState::Unchecked;
+							this->HandleToggleScaleGridSnap(UnusedCheckedState);
+						}
+					);
+					Action.GetActionCheckState =
+						FGetActionCheckState::CreateRaw(this, &STransformViewportToolBar::IsScaleGridSnapChecked);
+
+					const FName UnusedExtensionHook = NAME_None;
+					InMenuBuilder.AddMenuEntry(
+						Command->GetLabel(),
+						Command->GetDescription(),
+						Command->GetIcon(),
+						Action,
+						UnusedExtensionHook,
+						EUserInterfaceActionType::ToggleButton
+					);
 
-				InMenuBuilder.AddWrapperSubMenu(
-					LOCTEXT("ScaleGridSnapMenuSettings", "Scale Snap Settings"),
-					LOCTEXT("ScaleGridSnapMenuSettings_ToolTip", "Adjust the Grid Settings for Scale Snap"),
-					FOnGetContent::CreateSP(this, &STransformViewportToolBar::FillScaleGridSnapMenu),
-					FSlateIcon(Command->GetIcon())
-				);
-			}
-		));
+					InMenuBuilder.AddWrapperSubMenu(
+						LOCTEXT("ScaleGridSnapMenuSettings", "Scale Snap Settings"),
+						LOCTEXT("ScaleGridSnapMenuSettings_ToolTip", "Adjust the Grid Settings for Scale Snap"),
+						FOnGetContent::CreateSP(this, &STransformViewportToolBar::FillScaleGridSnapMenu),
+						FSlateIcon(Command->GetIcon())
+					);
+				}
+			)
+		);
 	}
 	ToolbarBuilder.EndSection();
 
