Hide the Default Button in the Scalability Menu when preview is inactive.

diff --git a/Engine/Source/Editor/UnrealEd/Private/SScalabilitySettings.cpp b/Engine/Source/Editor/UnrealEd/Private/SScalabilitySettings.cpp
--- a/Engine/Source/Editor/UnrealEd/Private/SScalabilitySettings.cpp
+++ b/Engine/Source/Editor/UnrealEd/Private/SScalabilitySettings.cpp
@@ -175,6 +175,9 @@ TSharedRef<SWidget> SScalabilitySettings::MakeDefaultButtonWidget()
 	return SNew(SButton)
 		.OnClicked(this, &SScalabilitySettings::OnDefaultClicked)
 		.ToolTipText(LOCTEXT("DefaultButtonTooltip", "Resets the Quality Levels to what the Device Profile sets"))
+		.Visibility_Lambda([]() {
+			return (GEditor && GEditor->IsFeatureLevelPreviewActive()) ? EVisibility::Visible : EVisibility::Collapsed;
+		})
 		.Content()
 		[
 			SNew(STextBlock)
