Component Visualizers: Fixed the Show Subcomponents setting not being taken into account

diff --git a/Engine/Source/Editor/UnrealEd/Private/EditorSelectUtils.cpp b/Engine/Source/Editor/UnrealEd/Private/EditorSelectUtils.cpp
--- a/Engine/Source/Editor/UnrealEd/Private/EditorSelectUtils.cpp
+++ b/Engine/Source/Editor/UnrealEd/Private/EditorSelectUtils.cpp
@@ -326,11 +326,16 @@ void UUnrealEdEngine::OnEditorElementSelectionChanged(const UTypedElementSelecti
 
 					ComponentVisualizerForSelection.IsEnabledDelegate.Emplace([bIsSelectedComponent = (Comp == SelectedComponent), WeakViz = TWeakPtr<FComponentVisualizer>(Visualizer), Comp]()
 					{
-						if (bIsSelectedComponent || (GetDefault<UEditorPerProjectUserSettings>()->bShowSelectionSubcomponents == true))
+						if (bIsSelectedComponent == true)
 						{
 							return true;
 						}
 
+						if (GetDefault<UEditorPerProjectUserSettings>()->bShowSelectionSubcomponents == false)
+						{
+							return false;
+						}
+
 						bool bShouldShowWithSelection = true;
 						if (TSharedPtr<FComponentVisualizer> PinnedViz = WeakViz.Pin())
 						{
