Updated GatherText to be able to run embedded within another editor process

* Added UGatherTextCommandlet::Execute as an entry point for running UGatherTextCommandlet::Main that allows additional options when running embedded within another editor process.
* Fixed UGatherTextFromAssetsCommandlet running a GC that can unload RF_Standalone objects when running embedded in a standard editor.
* Fixed UGatherTextFromAssetsCommandlet unloading the current world when running embedded in a standard editor.
* Instrumented UGatherTextCommandlet and UGatherTextFromAssetsCommandlet with FScopedSlowTask to provide feedback when running embedded in a standard editor.

This can improve the performance of running GatherText from an existing editor (vs starting a dedicated commandlet process), however it has the downside of potentially loading additional assets into your existing editor process that will remain loaded after the gather finishes.

diff --git a/Engine/Source/Editor/UnrealEd/Classes/Commandlets/GatherTextCommandlet.h b/Engine/Source/Editor/UnrealEd/Classes/Commandlets/GatherTextCommandlet.h
--- a/Engine/Source/Editor/UnrealEd/Classes/Commandlets/GatherTextCommandlet.h
+++ b/Engine/Source/Editor/UnrealEd/Classes/Commandlets/GatherTextCommandlet.h
@@ -15,8 +15,8 @@ namespace EOutputJson
 /**
  *	UGatherTextCommandlet: One commandlet to rule them all. This commandlet loads a config file and then calls other localization commandlets. Allows localization system to be easily extendable and flexible. 
  */
-UCLASS()
-class UGatherTextCommandlet : public UGatherTextCommandletBase
+UCLASS(MinimalAPI)
+class UGatherTextCommandlet final : public UGatherTextCommandletBase
 {
     GENERATED_UCLASS_BODY()
 public:
@@ -24,6 +24,9 @@ public:
 	virtual int32 Main(const FString& Params) override;
 	//~ End UCommandlet Interface
 
+	/** Internal implementation of Main that can be used to provide additional options when running embedded within another process */
+	UNREALED_API int32 Execute(const FString& Params, const TSharedPtr<const FGatherTextCommandletEmbeddedContext>& InEmbeddedContext);
+
 	int32 ProcessGatherConfig(const FString& GatherTextConfigPath, const TSharedPtr<FLocalizationSCC>& CommandletSourceControlInfo, const TArray<FString>& Tokens, const TArray<FString>& Switches, const TMap<FString, FString>& ParamVals);
 
 	static const FString UsageText;
diff --git a/Engine/Source/Editor/UnrealEd/Classes/Commandlets/GatherTextCommandletBase.h b/Engine/Source/Editor/UnrealEd/Classes/Commandlets/GatherTextCommandletBase.h
--- a/Engine/Source/Editor/UnrealEd/Classes/Commandlets/GatherTextCommandletBase.h
+++ b/Engine/Source/Editor/UnrealEd/Classes/Commandlets/GatherTextCommandletBase.h
@@ -84,6 +84,35 @@ private:
 	TArray<FFuzzyPath> FuzzyPaths;
 };
 
+/**
+ * Additional options and hooks that can be specified when running the GatherText commandlet embedded within another process
+ * @see UGatherTextCommandlet::Execute
+ */
+struct FGatherTextCommandletEmbeddedContext
+{
+	/** Optional override for the message of the overall localization slow task */
+	TOptional<FText> SlowTaskMessageOverride;
+
+	/** Callback used to perform additional tick tasks during the gather process */
+	TFunction<void()> TickCallback;
+
+	/** Callback used to allow user termination of the gather process */
+	TFunction<bool()> WasAbortRequestedCallback;
+
+	void RunTick() const
+	{
+		if (TickCallback)
+		{
+			TickCallback();
+		}
+	}
+
+	bool ShouldAbort() const
+	{
+		return WasAbortRequestedCallback && WasAbortRequestedCallback();
+	}
+};
+
 /**
  *	UGatherTextCommandletBase: Base class for localization commandlets. Just to force certain behaviors and provide helper functionality. 
  */
@@ -93,6 +122,7 @@ class UGatherTextCommandletBase : public UCommandlet
 	GENERATED_UCLASS_BODY()
 
 public:
+	UNREALED_API void SetEmbeddedContext(const TSharedPtr<const FGatherTextCommandletEmbeddedContext>& InEmbeddedContext);
 	UNREALED_API virtual void Initialize( const TSharedRef< FLocTextHelper >& InGatherManifestHelper, const TSharedPtr< FLocalizationSCC >& InSourceControlInfo );
 	UNREALED_API virtual void BeginDestroy() override;
 
@@ -126,6 +156,8 @@ protected:
 
 	static FName GetSplitPlatformNameFromPath_Static(const FString& InPath, const TMap<FName, FString>& InSplitPlatforms);
 
+	TSharedPtr<const FGatherTextCommandletEmbeddedContext> EmbeddedContext;
+
 	TSharedPtr< FLocTextHelper > GatherManifestHelper;
 
 	TSharedPtr< FLocalizationSCC > SourceControlInfo;
diff --git a/Engine/Source/Editor/UnrealEd/Classes/Commandlets/GatherTextFromAssetsCommandlet.h b/Engine/Source/Editor/UnrealEd/Classes/Commandlets/GatherTextFromAssetsCommandlet.h
--- a/Engine/Source/Editor/UnrealEd/Classes/Commandlets/GatherTextFromAssetsCommandlet.h
+++ b/Engine/Source/Editor/UnrealEd/Classes/Commandlets/GatherTextFromAssetsCommandlet.h
@@ -64,7 +64,7 @@ private:
 	void ApplyExcludeExactClassesFilter(const FARFilter& InFilter, TArray<FAssetData>& InOutAssetDataArray) const;
 	void FilterAssetsBasedOnIncludeExcludePaths(TArray<FAssetData>& InOutAssetDataArray) const;
 	
-	void DiscoverExternalActors(TArray<FAssetData>& InOutAssetDataArray);
+	bool DiscoverExternalActors(TArray<FAssetData>& InOutAssetDataArray);
 	void RemoveExistingExternalActors(TArray<FAssetData>& InOutAssetDataArray, TArray<FName>& OutPartitionedWorldPackageNames) const;
 
 	TSet<FName> GetPackageNamesToGather(const TArray<FAssetData>& InAssetDataArray) const;
@@ -72,7 +72,7 @@ private:
 	void ProcessAndRemoveCachedPackages(TMap<FName, TSet<FGuid>>& OutExternalActorsWithStaleOrMissingCaches);
 	void MergeInExternalActorsWithStaleOrMissingCaches(TMap<FName, TSet<FGuid>>& ExternalActorsWithStaleOrMissingCaches);
 
-	void LoadAndProcessUncachedPackages(TArray<FName>& OutPackagesWithStaleGatherCache);
+	bool LoadAndProcessUncachedPackages(TArray<FName>& OutPackagesWithStaleGatherCache);
 
 	void ReportStaleGatherCache(TArray<FName>& InPackagesWithStaleGatherCache) const;
 	/** Determines the loc cache state for a package. This determines whether the package should be fully loaded for gathering.*/
diff --git a/Engine/Source/Editor/UnrealEd/Private/Commandlets/GatherTextCommandlet.cpp b/Engine/Source/Editor/UnrealEd/Private/Commandlets/GatherTextCommandlet.cpp
--- a/Engine/Source/Editor/UnrealEd/Private/Commandlets/GatherTextCommandlet.cpp
+++ b/Engine/Source/Editor/UnrealEd/Private/Commandlets/GatherTextCommandlet.cpp
@@ -5,6 +5,7 @@
 #include "Misc/App.h"
 #include "Misc/Paths.h"
 #include "Misc/ConfigCacheIni.h"
+#include "Misc/ScopedSlowTask.h"
 #include "HAL/FileManager.h"
 #include "UObject/Class.h"
 #include "UObject/Package.h"
@@ -39,6 +40,13 @@ const FString UGatherTextCommandlet::UsageText
 
 int32 UGatherTextCommandlet::Main( const FString& Params )
 {
+	return Execute(Params, nullptr);
+}
+
+int32 UGatherTextCommandlet::Execute(const FString& Params, const TSharedPtr<const FGatherTextCommandletEmbeddedContext>& InEmbeddedContext)
+{
+	EmbeddedContext = InEmbeddedContext;
+
 	TArray<FString> Tokens;
 	TArray<FString> Switches;
 	TMap<FString, FString> ParamVals;
@@ -129,9 +137,24 @@ int32 UGatherTextCommandlet::Main( const FString& Params )
 		}
 	}
 
+	FScopedSlowTask SlowTask(GatherTextConfigPaths.Num(), (EmbeddedContext && EmbeddedContext->SlowTaskMessageOverride) ? *EmbeddedContext->SlowTaskMessageOverride : NSLOCTEXT("GatherTextCommandlet", "GatherTextTask", "Running Localization Tasks for Targets..."));
+
 	double AllCommandletExecutionStartTime = FPlatformTime::Seconds();
 	for (const FString& GatherTextConfigPath : GatherTextConfigPaths)
 	{
+		SlowTask.EnterProgressFrame();
+
+		if (SlowTask.ShouldCancel() || (EmbeddedContext && EmbeddedContext->ShouldAbort()))
+		{
+			UE_LOG(LogGatherTextCommandlet, Error, TEXT("GatherText aborted!"));
+			return -1;
+		}
+
+		if (EmbeddedContext)
+		{
+			EmbeddedContext->RunTick();
+		}
+
 		const int32 Result = ProcessGatherConfig(GatherTextConfigPath, CommandletSourceControlInfo, Tokens, Switches, ParamVals);
 		if (Result != 0)
 		{
@@ -251,10 +274,25 @@ int32 UGatherTextCommandlet::ProcessGatherConfig(const FString& GatherTextConfig
 		GeneratedParamsAndSwitches += FString::Printf(TEXT(" -%s"), *Switch);
 	}
 
+	FScopedSlowTask SlowTask(StepNames.Num(), NSLOCTEXT("GatherTextCommandlet", "GatherTextSubTask", "Running Localization Tasks..."));
+
 	const bool bRunningInPreview = Switches.Contains(TEXT("Preview"));
 	// Execute each step defined in the config file.
 	for (const FString& StepName : StepNames)
 	{
+		SlowTask.EnterProgressFrame();
+
+		if (SlowTask.ShouldCancel() || (EmbeddedContext && EmbeddedContext->ShouldAbort()))
+		{
+			UE_LOG(LogGatherTextCommandlet, Error, TEXT("GatherText aborted!"));
+			return -1;
+		}
+
+		if (EmbeddedContext)
+		{
+			EmbeddedContext->RunTick();
+		}
+
 		FString CommandletClassName = GConfig->GetStr( *StepName, TEXT("CommandletClass"), GatherTextConfigPath ) + TEXT("Commandlet");
 
 		UClass* CommandletClass = FindFirstObject<UClass>(*CommandletClassName, EFindFirstObjectOptions::None, ELogVerbosity::Warning, TEXT("UGatherTextCommandlet::ProcessGatherConfig"));
@@ -264,9 +302,9 @@ int32 UGatherTextCommandlet::ProcessGatherConfig(const FString& GatherTextConfig
 			continue;
 		}
 
-		UGatherTextCommandletBase* Commandlet = NewObject<UGatherTextCommandletBase>(GetTransientPackage(), CommandletClass);
+		TStrongObjectPtr<UGatherTextCommandletBase> Commandlet(NewObject<UGatherTextCommandletBase>(GetTransientPackage(), CommandletClass));
 		check(Commandlet);
-		FGCObjectScopeGuard CommandletGCGuard(Commandlet);
+		Commandlet->SetEmbeddedContext(EmbeddedContext);
 		Commandlet->Initialize( CommandletGatherManifestHelper, CommandletSourceControlInfo );
 		// As of now, all params and switches (with the exception of config) is passed on to child commandlets
 		// Instead of parsing in child commandlets, we'll just pass the params and switches along to determine if we need to run the child commandlet 
diff --git a/Engine/Source/Editor/UnrealEd/Private/Commandlets/GatherTextCommandletBase.cpp b/Engine/Source/Editor/UnrealEd/Private/Commandlets/GatherTextCommandletBase.cpp
--- a/Engine/Source/Editor/UnrealEd/Private/Commandlets/GatherTextCommandletBase.cpp
+++ b/Engine/Source/Editor/UnrealEd/Private/Commandlets/GatherTextCommandletBase.cpp
@@ -46,6 +46,11 @@ UGatherTextCommandletBase::UGatherTextCommandletBase(const FObjectInitializer& O
 	}
 }
 
+void UGatherTextCommandletBase::SetEmbeddedContext(const TSharedPtr<const FGatherTextCommandletEmbeddedContext>& InEmbeddedContext)
+{
+	EmbeddedContext = InEmbeddedContext;
+}
+
 void UGatherTextCommandletBase::Initialize( const TSharedRef< FLocTextHelper >& InGatherManifestHelper, const TSharedPtr< FLocalizationSCC >& InSourceControlInfo )
 {
 	GatherManifestHelper = InGatherManifestHelper;
diff --git a/Engine/Source/Editor/UnrealEd/Private/Commandlets/GatherTextFromAssetsCommandlet.cpp b/Engine/Source/Editor/UnrealEd/Private/Commandlets/GatherTextFromAssetsCommandlet.cpp
--- a/Engine/Source/Editor/UnrealEd/Private/Commandlets/GatherTextFromAssetsCommandlet.cpp
+++ b/Engine/Source/Editor/UnrealEd/Private/Commandlets/GatherTextFromAssetsCommandlet.cpp
@@ -8,6 +8,7 @@
 #include "Misc/ConfigCacheIni.h"
 #include "Misc/Paths.h"
 #include "Misc/OutputDeviceHelper.h"
+#include "Misc/ScopedSlowTask.h"
 #include "Misc/FeedbackContext.h"
 #include "UObject/GCObjectScopeGuard.h"
 #include "UObject/EditorObjectVersion.h"
@@ -421,7 +422,7 @@ void UGatherTextFromAssetsCommandlet::PurgeGarbage(const bool bPurgeReferencedPa
 		}
 	}
 
-	CollectGarbage(RF_NoFlags);
+	CollectGarbage(IsRunningCommandlet() ? RF_NoFlags : GARBAGE_COLLECTION_KEEPFLAGS);
 	ObjectsToKeepAlive.Reset();
 
 	// Fully process the shader compilation results when performing a full purge, as it's the only way to reclaim that memory
@@ -778,7 +779,7 @@ void UGatherTextFromAssetsCommandlet::RemoveExistingExternalActors(TArray<FAsset
 }
 
 /** Appends any external actors that also need to be gathered to the InOutAssetDataArray. */
-void UGatherTextFromAssetsCommandlet::DiscoverExternalActors(TArray<FAssetData>& InOutAssetDataArray)
+bool UGatherTextFromAssetsCommandlet::DiscoverExternalActors(TArray<FAssetData>& InOutAssetDataArray)
 {
 	UE_LOG(LogGatherTextFromAssetsCommandlet, Display, TEXT("Discovering external actors to gather..."));
 	const double DiscoveringExternalActorsStartTime = FPlatformTime::Seconds();
@@ -788,11 +789,21 @@ void UGatherTextFromAssetsCommandlet::DiscoverExternalActors(TArray<FAssetData>&
 
 	TArray<FName> ExternalActorPackageNames;
 	{
+		FScopedSlowTask SlowTask(PartitionedWorldPackageNames.Num(), NSLOCTEXT("GatherTextCommandlet", "DiscoverExternalActors", "Discovering External Actors..."));
+
 		FLoadPackageLogOutputRedirector LogOutputRedirector;
 
 		int32 NumPackagesProcessed = 0;
 		for (const FName PartitionedWorldPackageName : PartitionedWorldPackageNames)
 		{
+			SlowTask.EnterProgressFrame();
+
+			if (SlowTask.ShouldCancel() || (EmbeddedContext && EmbeddedContext->ShouldAbort()))
+			{
+				UE_LOG(LogGatherTextFromAssetsCommandlet, Error, TEXT("GatherText aborted!"));
+				return false;
+			}
+
 			FNameBuilder PackageNameStr(PartitionedWorldPackageName);
 
 			const int32 CurrentPackageNum = ++NumPackagesProcessed;
@@ -813,6 +824,10 @@ void UGatherTextFromAssetsCommandlet::DiscoverExternalActors(TArray<FAssetData>&
 
 			// Tick background tasks
 			UE::Private::GatherTextFromAssetsCommandlet::TickBackgroundTasks();
+			if (EmbeddedContext)
+			{
+				EmbeddedContext->RunTick();
+			}
 
 			if (UWorld* World = UWorld::FindWorldInPackage(Package))
 			{
@@ -830,6 +845,7 @@ void UGatherTextFromAssetsCommandlet::DiscoverExternalActors(TArray<FAssetData>&
 				IVS.CreateWorldPartition(true);
 
 				TOptional<FScopedEditorWorld> ScopeEditorWorld;
+				if (!World->IsInitialized())
 				{
 					FLoadPackageLogOutputRedirector::FScopedCapture ScopedCapture(&LogOutputRedirector, *PackageNameStr);
 					ScopeEditorWorld.Emplace(World, IVS); // Initializing FScopedEditorWorld can log warnings, so capture those like we do with loading errors
@@ -871,6 +887,7 @@ void UGatherTextFromAssetsCommandlet::DiscoverExternalActors(TArray<FAssetData>&
 	}
 
 	UE_LOG(LogGatherTextFromAssetsCommandlet, Display, TEXT("Discovering external actors took %.2f seconds."), FPlatformTime::Seconds() - DiscoveringExternalActorsStartTime);
+	return true;
 }
 
 /** Applies the passed in filter to the asset registry. If the filter is empty, the entire asset registry will be returned in InOutAssetDataArray. Else assets that pass the filter will be in InOutAssetDataArray.*/
@@ -1218,7 +1235,7 @@ void UGatherTextFromAssetsCommandlet::MergeInExternalActorsWithStaleOrMissingCac
 }
 
 /** Load the remaining pending packages for gather.*/
-void UGatherTextFromAssetsCommandlet::LoadAndProcessUncachedPackages(TArray<FName>& OutPackagesWithStaleGatherCache)
+bool UGatherTextFromAssetsCommandlet::LoadAndProcessUncachedPackages(TArray<FName>& OutPackagesWithStaleGatherCache)
 {
 	UE_SCOPED_TIMER(TEXT("UGatherTextFromAssetsCommandlet::LoadAndProcessUncachedPackages"), LogGatherTextFromAssetsCommandlet, Display);
 	FLoadPackageLogOutputRedirector LogOutputRedirector;
@@ -1227,6 +1244,8 @@ void UGatherTextFromAssetsCommandlet::LoadAndProcessUncachedPackages(TArray<FNam
 	int32 PackageCount = PackagesPendingGather.Num();
 	int32 NumPackagesFailedLoading = 0;
 	
+	FScopedSlowTask SlowTask(PackageCount, NSLOCTEXT("GatherTextCommandlet", "LoadAndProcessUncachedPackages", "Loading and Gathering Packages..."));
+
 	while (PackagesPendingGather.Num() > 0)
 	{
 		const FPackagePendingGather PackagePendingGather = PackagesPendingGather.Pop(EAllowShrinking::No);
@@ -1236,6 +1255,14 @@ void UGatherTextFromAssetsCommandlet::LoadAndProcessUncachedPackages(TArray<FNam
 		const float PercentageComplete = static_cast<float>(CurrentPackageNum) / static_cast<float>(PackageCount) * 100.0f;
 		UE_LOG(LogGatherTextFromAssetsCommandlet, Display, TEXT("[%6.2f%%] Loading package: '%s'..."), PercentageComplete, *PackageNameStr);
 
+		SlowTask.EnterProgressFrame();
+
+		if (SlowTask.ShouldCancel() || (EmbeddedContext && EmbeddedContext->ShouldAbort()))
+		{
+			UE_LOG(LogGatherTextFromAssetsCommandlet, Error, TEXT("GatherText aborted!"));
+			return false;
+		}
+
 		UPackage* Package = nullptr;
 		{
 			FLoadPackageLogOutputRedirector::FScopedCapture ScopedCapture(&LogOutputRedirector, *PackageNameStr);
@@ -1258,6 +1285,10 @@ void UGatherTextFromAssetsCommandlet::LoadAndProcessUncachedPackages(TArray<FNam
 
 		// Tick background tasks
 		UE::Private::GatherTextFromAssetsCommandlet::TickBackgroundTasks();
+		if (EmbeddedContext)
+		{
+			EmbeddedContext->RunTick();
+		}
 
 		// Because packages may not have been resaved after this flagging was implemented, we may have added packages to load that weren't flagged - potential false positives.
 		// The loading process should have reflagged said packages so that only true positives will have this flag.
@@ -1330,6 +1361,7 @@ void UGatherTextFromAssetsCommandlet::LoadAndProcessUncachedPackages(TArray<FNam
 					IVS.CreateWorldPartition(true);
 
 					TOptional<FScopedEditorWorld> ScopeEditorWorld;
+					if (!World->IsInitialized())
 					{
 						FLoadPackageLogOutputRedirector::FScopedCapture ScopedCapture(&LogOutputRedirector, *PackageNameStr);
 						ScopeEditorWorld.Emplace(World, IVS); // Initializing FScopedEditorWorld can log warnings, so capture those like we do with loading errors
@@ -1382,6 +1414,7 @@ void UGatherTextFromAssetsCommandlet::LoadAndProcessUncachedPackages(TArray<FNam
 		}
 	}
 	UE_LOG(LogGatherTextFromAssetsCommandlet, Display, TEXT("Loaded %d packages. %d failed."), NumPackagesProcessed, NumPackagesFailedLoading);
+	return true;
 }
 
 void UGatherTextFromAssetsCommandlet::ReportStaleGatherCache(TArray<FName>& InPackagesWithStaleGatherCache) const
@@ -1476,7 +1509,8 @@ int32 UGatherTextFromAssetsCommandlet::Main(const FString& Params)
 
 	// If the editor has loaded a persistent world then create an empty world prior to starting the asset gather
 	// This avoids any issues when loading and initializing worlds during the gather, as WP needs to re-initialize the world
-	if (GEditor)
+	// Note: We can skip this when running embedded within a normal editor (ie, not a commandlet) as editor worlds are already fully initialized
+	if (IsRunningCommandlet() && GEditor)
 	{
 		if (UWorld* EditorWorld = GEditor->GetEditorWorldContext().World();
 			EditorWorld && !FPackageName::IsTempPackage(FNameBuilder(EditorWorld->GetPackage()->GetFName()).ToView()))
@@ -1517,7 +1551,10 @@ int32 UGatherTextFromAssetsCommandlet::Main(const FString& Params)
 	}
 
 	// Discover the external actors for any worlds that are pending gather
-	DiscoverExternalActors(AssetDataArray);
+	if (!DiscoverExternalActors(AssetDataArray))
+	{
+		return -1;
+	}
 
 	// Collect the basic information about the packages that we're going to gather from
 	TSet<FName> PackageNamesToGather = GetPackageNamesToGather(AssetDataArray);
@@ -1564,7 +1601,10 @@ int32 UGatherTextFromAssetsCommandlet::Main(const FString& Params)
 	}
 
 	TArray<FName> PackagesWithStaleGatherCache;
-	LoadAndProcessUncachedPackages(PackagesWithStaleGatherCache);
+	if (!LoadAndProcessUncachedPackages(PackagesWithStaleGatherCache))
+	{
+		return -1;
+	}
 
 	UE_LOG(LogGatherTextFromAssetsCommandlet, Display, TEXT("Loading packages took %.2f seconds."), FPlatformTime::Seconds() - PackageLoadingStartTime);
 
