Ray Traced Translucency - fix preview mesh disappearing when previewing translucent materials

#tests tested in editor

diff --git a/Engine/Source/Editor/UnrealEd/Private/PreviewMaterial.cpp b/Engine/Source/Editor/UnrealEd/Private/PreviewMaterial.cpp
--- a/Engine/Source/Editor/UnrealEd/Private/PreviewMaterial.cpp
+++ b/Engine/Source/Editor/UnrealEd/Private/PreviewMaterial.cpp
@@ -224,6 +224,10 @@ public:
 			{
 				bShaderTypeMatches = true;
 			}
+			else if (FCString::Stristr(ShaderType->GetName(), TEXT("FLumenFrontLayerTranslucencyGBuffer")))
+			{
+				bShaderTypeMatches = true;
+			}
 
 			return bShaderTypeMatches;
 		}
