Allow saving into /Temp (which maps to Project\Saved) in the editor when running automation tests.
This allows saving a map and its external objects (actors, entities, etc.) into the saved folder to reload.

diff --git a/Engine/Source/Editor/UnrealEd/Private/FileHelpers.cpp b/Engine/Source/Editor/UnrealEd/Private/FileHelpers.cpp
--- a/Engine/Source/Editor/UnrealEd/Private/FileHelpers.cpp
+++ b/Engine/Source/Editor/UnrealEd/Private/FileHelpers.cpp
@@ -607,13 +607,7 @@ static bool SaveExternalPackages(UPackage* InPackage, TArray<UPackage*>& InOutPa
 
 	if (!bInAutosaving && (!FEditorFileUtils::ShouldSkipExternalObjectSave() || bInNewlyCreated))
 	{
-		for (UPackage* ExternalPackage : InPackage->GetExternalPackages())
-		{
-			if (!FPackageName::IsTempPackage(ExternalPackage->GetName()))
-			{
-				InOutPackagesToSave.Add(ExternalPackage);
-			}
-		}
+		InOutPackagesToSave = InPackage->GetExternalPackages();
 
 		if (InOutPackagesToSave.Num())
 		{
@@ -3046,14 +3040,17 @@ bool FEditorFileUtils::LoadMap(const FString& InFilename, bool LoadAsTemplate, b
 		FEditorDirectories::Get().SetLastDirectory(ELastDirectory::UNR, FPaths::GetPath(Filename)); // Save path as default for next time.
 	}
 
-	//ensure the name wasn't mangled during load before adding to the Recent File list
-	if (World->GetOutermost()->GetName() == LongMapPackageName)
+	if (FPackageName::IsValidLongPackageName(LongMapPackageName))
 	{
-		IMainFrameModule& MainFrameModule = FModuleManager::LoadModuleChecked<IMainFrameModule>( "MainFrame" );
-		FMainMRUFavoritesList* MRUFavoritesList = MainFrameModule.GetMRUFavoritesList();
-		if(MRUFavoritesList)
+		//ensure the name wasn't mangled during load before adding to the Recent File list
+		if (World->GetOutermost()->GetName() == LongMapPackageName)
 		{
-			MRUFavoritesList->AddMRUItem(LongMapPackageName);
+			IMainFrameModule& MainFrameModule = FModuleManager::LoadModuleChecked<IMainFrameModule>( "MainFrame" );
+			FMainMRUFavoritesList* MRUFavoritesList = MainFrameModule.GetMRUFavoritesList();
+			if(MRUFavoritesList)
+			{
+				MRUFavoritesList->AddMRUItem(LongMapPackageName);
+			}
 		}
 	}
 
@@ -3456,7 +3453,7 @@ static InternalSavePackageResult InternalSavePackage(UPackage* PackageToSave, bo
 	bool bAttemptSave = false;
 
 	// If the package already has a valid path to a non read-only location, use it to determine where the file should be saved
-	const bool bIncludeReadOnlyRoots = false;
+	const bool bIncludeReadOnlyRoots = GIsAutomationTesting;
 	const bool bIsValidPath = FPackageName::IsValidLongPackageName(PackageName, bIncludeReadOnlyRoots);
 	if( bIsValidPath )
 	{
@@ -3637,34 +3634,37 @@ static InternalSavePackageResult InternalSavePackage(UPackage* PackageToSave, bo
 			bWasSuccessful = SaveAsset(PackageToSave, PackageName, FinalPackageSavePath, SaveOutput);
 		}
 
-		ISourceControlProvider& SourceControlProvider = ISourceControlModule::Get().GetProvider();
-		FUncontrolledChangelistsModule& UncontrolledChangelistsModule = FUncontrolledChangelistsModule::Get();
-
-		if (ISourceControlModule::Get().IsEnabled())
+		if (!FPackageName::IsTempPackage(PackageName))
 		{
-			// Assume the package was correctly checked out from SCC
-			bOutPackageLocallyWritable = false;
+			ISourceControlProvider& SourceControlProvider = ISourceControlModule::Get().GetProvider();
+			FUncontrolledChangelistsModule& UncontrolledChangelistsModule = FUncontrolledChangelistsModule::Get();
+
+			if (ISourceControlModule::Get().IsEnabled())
+			{
+				// Assume the package was correctly checked out from SCC
+				bOutPackageLocallyWritable = false;
 
-			// Trusting the SCC status in the package file cache to minimize network activity during save.
-			const FSourceControlStatePtr SourceControlState = SourceControlProvider.GetState(PackageToSave, EStateCacheUsage::Use);
-			// If the package is in the depot, and not recognized as editable by source control, and not read-only, then we know the user has made the package locally writable!
-			const bool bSCCCanEdit = !SourceControlState.IsValid() || SourceControlState->CanCheckIn() || SourceControlState->IsIgnored() || SourceControlState->IsUnknown();
-			const bool bSCCIsCheckedOut = SourceControlState.IsValid() && SourceControlState->IsCheckedOut();
-			const bool bInDepot = SourceControlState.IsValid() && SourceControlState->IsSourceControlled();
-			if ( !bSCCCanEdit && bInDepot && !IFileManager::Get().IsReadOnly( *FinalPackageSavePath ) && SourceControlProvider.UsesLocalReadOnlyState() && !bSCCIsCheckedOut )
+				// Trusting the SCC status in the package file cache to minimize network activity during save.
+				const FSourceControlStatePtr SourceControlState = SourceControlProvider.GetState(PackageToSave, EStateCacheUsage::Use);
+				// If the package is in the depot, and not recognized as editable by source control, and not read-only, then we know the user has made the package locally writable!
+				const bool bSCCCanEdit = !SourceControlState.IsValid() || SourceControlState->CanCheckIn() || SourceControlState->IsIgnored() || SourceControlState->IsUnknown();
+				const bool bSCCIsCheckedOut = SourceControlState.IsValid() && SourceControlState->IsCheckedOut();
+				const bool bInDepot = SourceControlState.IsValid() && SourceControlState->IsSourceControlled();
+				if ( !bSCCCanEdit && bInDepot && !IFileManager::Get().IsReadOnly( *FinalPackageSavePath ) && SourceControlProvider.UsesLocalReadOnlyState() && !bSCCIsCheckedOut )
+				{
+					bOutPackageLocallyWritable = true;
+				}
+			}
+			else
 			{
-				bOutPackageLocallyWritable = true;
+				// If we are in offline mode, automatically add the modified package to an Uncontrolled Changelist
+				bOutPackageLocallyWritable = UncontrolledChangelistsModule.IsEnabled() && (!IFileManager::Get().IsReadOnly(*FinalPackageSavePath));
 			}
-		}
-		else
-		{
-			// If we are in offline mode, automatically add the modified package to an Uncontrolled Changelist
-			bOutPackageLocallyWritable = UncontrolledChangelistsModule.IsEnabled() && (!IFileManager::Get().IsReadOnly(*FinalPackageSavePath));
-		}
 
-		if (bWasSuccessful && bOutPackageLocallyWritable)
-		{
-			UncontrolledChangelistsModule.OnSaveWritable({ FinalPackageSavePath });
+			if (bWasSuccessful && bOutPackageLocallyWritable)
+			{
+				UncontrolledChangelistsModule.OnSaveWritable({ FinalPackageSavePath });
+			}
 		}
 
 		// Handle all failures the same way.
@@ -5682,4 +5682,4 @@ void UEditorLoadingAndSavingUtils::ReloadPackages(const TArray<UPackage*>& Packa
 	bOutAnyPackagesReloaded = UPackageTools::ReloadPackages(PackagesToReload, OutErrorMessage, InteractionMode);
 }
 
-#undef LOCTEXT_NAMESPACE
\ No newline at end of file
+#undef LOCTEXT_NAMESPACE
