Fix for splines viz disappearing on BP construction script rerun.
- Store a ComponentPropertyPath to retrieve the component after a reconstruction instead of storing the component pointer directly in the viz enabled delegate.

diff --git a/Engine/Source/Editor/UnrealEd/Private/EditorSelectUtils.cpp b/Engine/Source/Editor/UnrealEd/Private/EditorSelectUtils.cpp
--- a/Engine/Source/Editor/UnrealEd/Private/EditorSelectUtils.cpp
+++ b/Engine/Source/Editor/UnrealEd/Private/EditorSelectUtils.cpp
@@ -308,7 +308,9 @@ void UUnrealEdEngine::OnEditorElementSelectionChanged(const UTypedElementSelecti
 	{
 		// Iterate over components of that actor (and recurse through child components)
 		TInlineComponentArray<UActorComponent*> Components;
-		Actor->GetComponents(Components, true);
+
+		constexpr bool bIncludeFromChildActors = true;
+		Actor->GetComponents(Components, bIncludeFromChildActors);
 
 		for (int32 CompIdx = 0; CompIdx < Components.Num(); CompIdx++)
 		{
@@ -324,25 +326,29 @@ void UUnrealEdEngine::OnEditorElementSelectionChanged(const UTypedElementSelecti
 
 					FComponentVisualizerForSelection& ComponentVisualizerForSelection = VisualizersForSelection.Add_GetRef(MoveTemp(Temp));
 
-					ComponentVisualizerForSelection.IsEnabledDelegate.Emplace([bIsSelectedComponent = (Comp == SelectedComponent), WeakViz = TWeakPtr<FComponentVisualizer>(Visualizer), Comp]()
-					{
-						if (bIsSelectedComponent == true)
-						{
-							return true;
-						}
-
-						if (GetDefault<UEditorPerProjectUserSettings>()->bShowSelectionSubcomponents == false)
-						{
-							return false;
-						}
-
-						bool bShouldShowWithSelection = true;
-						if (TSharedPtr<FComponentVisualizer> PinnedViz = WeakViz.Pin())
+					ComponentVisualizerForSelection.IsEnabledDelegate.Emplace(
+						[bIsSelectedComponent = (Comp == SelectedComponent),
+						 WeakViz = TWeakPtr<FComponentVisualizer>(Visualizer),
+						 CompPath = CachedComponentVisualizer.ComponentPropertyPath]()
 						{
-							bShouldShowWithSelection = PinnedViz->ShouldShowForSelectedSubcomponents(Comp.Get());
+							if (bIsSelectedComponent == true)
+							{
+								return true;
+							}
+
+							if (GetDefault<UEditorPerProjectUserSettings>()->bShowSelectionSubcomponents == false)
+							{
+								return false;
+							}
+
+							bool bShouldShowWithSelection = true;
+							if (TSharedPtr<FComponentVisualizer> PinnedViz = WeakViz.Pin())
+							{
+								bShouldShowWithSelection = PinnedViz->ShouldShowForSelectedSubcomponents(CompPath.GetComponent());
+							}
+							return bShouldShowWithSelection;
 						}
-						return bShouldShowWithSelection;
-					});
+					);
 				}
 			}
 		}
