Unshelved from pending changelist '38208172':

PR #12578: Fix regression in GenerateBlueprintAPICommandlet due to misuse of pointers that breaks the commandlet

#12578

[RN] Fixed crashes when running GenerateBlueprintAPICommandlet

diff --git a/Engine/Source/Editor/UnrealEd/Private/Commandlets/GenerateBlueprintAPICommandlet.cpp b/Engine/Source/Editor/UnrealEd/Private/Commandlets/GenerateBlueprintAPICommandlet.cpp
--- a/Engine/Source/Editor/UnrealEd/Private/Commandlets/GenerateBlueprintAPICommandlet.cpp
+++ b/Engine/Source/Editor/UnrealEd/Private/Commandlets/GenerateBlueprintAPICommandlet.cpp
@@ -766,13 +766,13 @@ static void GenerateBlueprintAPIUtils::DumpPalette(uint32 Indent, UBlueprint* Bl
 //------------------------------------------------------------------------------
 static void GenerateBlueprintAPIUtils::DumpActionList(uint32 Indent, FGraphActionListBuilderBase& ActionList, FArchive* FileOutWriter)
 {
-	TArray< TSharedPtr<FEdGraphSchemaAction> const* > SortedActions;
+	TArray< TSharedPtr<FEdGraphSchemaAction> > SortedActions;
 	const TArray<FString>& CategoryFilter = GenerateBlueprintAPIUtils::CommandOptions.CategoryFilter;
 	bool HasCategoryFilter = CategoryFilter.Num() > 0;
 
 	for (int32 ActionIndex = 0; ActionIndex < ActionList.GetNumActions(); ++ActionIndex)
 	{
-		TSharedPtr<FEdGraphSchemaAction> Action = ActionList.GetSchemaAction(ActionIndex);
+		TSharedPtr<FEdGraphSchemaAction>& Action = ActionList.GetSchemaAction(ActionIndex);
 		if (!Action.IsValid())
 		{
 			continue;
@@ -780,7 +780,7 @@ static void GenerateBlueprintAPIUtils::DumpActionList(uint32 Indent, FGraphActio
 
 		if (!HasCategoryFilter)
 		{
-			SortedActions.Add(&Action);
+			SortedActions.Add(Action);
 		}
 		else
 		{
@@ -789,7 +789,7 @@ static void GenerateBlueprintAPIUtils::DumpActionList(uint32 Indent, FGraphActio
 			{
 				if (ActionCategory.Contains(filter, ESearchCase::CaseSensitive, ESearchDir::FromStart))
 				{
-					SortedActions.Add(&Action);
+					SortedActions.Add(Action);
 					break;
 				}
 			}
@@ -822,10 +822,10 @@ static void GenerateBlueprintAPIUtils::DumpActionList(uint32 Indent, FGraphActio
 	SortedActions.Sort(ActionSortFunctor());
 
 	FString LineEnding("\n");
-	for (TSharedPtr<FEdGraphSchemaAction> const* Action : SortedActions)
+	for (TSharedPtr<FEdGraphSchemaAction> const& Action : SortedActions)
 	{
 		FileOutWriter->Serialize(TCHAR_TO_ANSI(*LineEnding), LineEnding.Len());
-		DumpActionMenuItem(Indent + 2, *Action, ActionList, FileOutWriter);
+		DumpActionMenuItem(Indent + 2, Action, ActionList, FileOutWriter);
 		LineEnding = TEXT(",\n");
 	}
 
