MPCook: Fix staging error "File not found in manifest" after a race condition in Multiplatform, multiprocess cook, in this scenario:
1) A package is found reachable for one platform but not another, and is assigned to a worker
2) The worker completes the save and posts the result to the socket
3) The director finds the package reachable for the other platform, and removes it from the worker. It sends an AbortAssignment message to the worker.
4) The director decides to assign the package back to the worker, with the new list of platforms that includes both platforms. It sends an assignment message to the worker.
5) The worker receives the abort message, does nothing because it already finished, then receives the assign message, cooks the package again for the two requested platforms, and sends another result message with both platforms.
6) The first result message from the worker arrives on the director. We think its the second message, process its single platform, mark the package as cook complete for the first package, and mark it idle.
7) The second result message from the worker arrives with the results for all platforms, but we reject it because the package is now idle and not assigned to the worker.

Error: We were supposed to cook the package for the second platform but did not.

#rnx

diff --git a/Engine/Source/Editor/UnrealEd/Private/Cooker/CookDirector.cpp b/Engine/Source/Editor/UnrealEd/Private/Cooker/CookDirector.cpp
--- a/Engine/Source/Editor/UnrealEd/Private/Cooker/CookDirector.cpp
+++ b/Engine/Source/Editor/UnrealEd/Private/Cooker/CookDirector.cpp
@@ -244,7 +244,7 @@ FCookDirector::~FCookDirector()
 	TSet<FPackageData*> AbortedAssignments;
 	for (TPair<int32, TRefCountPtr<FCookWorkerServer>>& Pair : RemoteWorkers)
 	{
-		Pair.Value->AbortWorker(AbortedAssignments, ECookDirectorThread::SchedulerThread);
+		Pair.Value->AbortWorker(AbortedAssignments, ECookDirectorThread::SchedulerThread, HeartbeatNumber);
 	}
 	for (FPackageData* PackageData : AbortedAssignments)
 	{
@@ -354,7 +354,7 @@ void FCookDirector::AssignRequests(TArrayView<FPackageData*> Requests, TArray<FW
 	{
 		if (RemoteWorker->IsShuttingDown())
 		{
-			RemoteWorker->AbortAllAssignments(PackagesToReassignSet, ECookDirectorThread::SchedulerThread);
+			RemoteWorker->AbortAllAssignments(PackagesToReassignSet, ECookDirectorThread::SchedulerThread, HeartbeatNumber);
 		}
 	}
 	if (!PackagesToReassignSet.IsEmpty())
@@ -640,7 +640,8 @@ void FCookDirector::RemoveFromWorker(FPackageData& PackageData)
 		OwningWorker = *RemoteWorkerPtr;
 	}
 
-	OwningWorker->AbortAssignment(PackageData, ECookDirectorThread::SchedulerThread);
+	OwningWorker->AbortAssignment(PackageData, ECookDirectorThread::SchedulerThread, HeartbeatNumber);
+	bForceNextHeartbeat = true;
 }
 
 void FCookDirector::BroadcastGeneratorMessage(FGeneratorEventMessage&& Message)
@@ -685,7 +686,8 @@ void FCookDirector::TickFromSchedulerThread()
 	
 	bool bSendHeartbeat;
 	int32 LocalHeartbeatNumber;
-	TickHeartbeat(false /* bForceHeartbeat */, CurrentTime, bSendHeartbeat, LocalHeartbeatNumber);
+	TickHeartbeat(bForceNextHeartbeat, CurrentTime, bSendHeartbeat, LocalHeartbeatNumber);
+	bForceNextHeartbeat = false;
 
 	TArray<TRefCountPtr<FCookWorkerServer>, TInlineAllocator<16>> WorkersWithMessage;
 	bool bAllWorkersConnected = false;
@@ -987,7 +989,7 @@ void FCookDirector::ShutdownCookSession()
 	for (TRefCountPtr<FCookWorkerServer>& RemoteWorker : WorkersNeedingAbort)
 	{
 		TSet<FPackageData*> UnusedPendingPackages;
-		RemoteWorker->AbortWorker(UnusedPendingPackages, ECookDirectorThread::SchedulerThread);
+		RemoteWorker->AbortWorker(UnusedPendingPackages, ECookDirectorThread::SchedulerThread, HeartbeatNumber);
 	}
 
 	// Wait for all the shutdowns to complete
@@ -1621,10 +1623,10 @@ void FCookDirector::AbortWorker(FWorkerId WorkerId, ECookDirectorThread TickThre
 		}
 	}
 	TSet<FPackageData*> PackagesToReassignSet;
-	RemoteWorker->AbortAllAssignments(PackagesToReassignSet, TickThread);
+	RemoteWorker->AbortAllAssignments(PackagesToReassignSet, TickThread, HeartbeatNumber);
 	if (!RemoteWorker->IsShuttingDown())
 	{
-		RemoteWorker->AbortWorker(PackagesToReassignSet, TickThread);
+		RemoteWorker->AbortWorker(PackagesToReassignSet, TickThread, HeartbeatNumber);
 	}
 	TArray<FPackageData*> PackagesToReassign = PackagesToReassignSet.Array();
 
@@ -1986,7 +1988,7 @@ FCookDirector::FRetractionHandler::ERetractionState FCookDirector::FRetractionHa
 			// The worker(s) that sent the retraction message aborted all of the packages, so mark locally that they
 			// have been aborted
 			RemoteWorker->AbortAssignments(WorkerPackageDatas, ECookDirectorThread::SchedulerThread,
-				ENotifyRemote::LocalOnly);
+				Director.HeartbeatNumber, ENotifyRemote::LocalOnly);
 		}
 		PackageDatasToReassign.Append(MoveTemp(WorkerPackageDatas));
 	}
diff --git a/Engine/Source/Editor/UnrealEd/Private/Cooker/CookDirector.h b/Engine/Source/Editor/UnrealEd/Private/Cooker/CookDirector.h
--- a/Engine/Source/Editor/UnrealEd/Private/Cooker/CookDirector.h
+++ b/Engine/Source/Editor/UnrealEd/Private/Cooker/CookDirector.h
@@ -260,6 +260,7 @@ private:
 	bool bWorkersStalled = false;
 	bool bMultiprocessAvailable = false;
 	bool bReceivingMessages = false;
+	bool bForceNextHeartbeat = false;
 
 	// Data that is read-only while the CommunicationThread is active and is readable from any thread
 	FBeginCookContextForWorker BeginCookContext;
diff --git a/Engine/Source/Editor/UnrealEd/Private/Cooker/CookWorkerClient.cpp b/Engine/Source/Editor/UnrealEd/Private/Cooker/CookWorkerClient.cpp
--- a/Engine/Source/Editor/UnrealEd/Private/Cooker/CookWorkerClient.cpp
+++ b/Engine/Source/Editor/UnrealEd/Private/Cooker/CookWorkerClient.cpp
@@ -1040,7 +1040,9 @@ void FCookWorkerClient::HandleAbortPackagesMessage(FMPCollectorClientMessageCont
 	if (!bReadSuccessful)
 	{
 		LogInvalidMessage(TEXT("AbortPackagesMessage"));
+		return;
 	}
+
 	for (FName PackageName : Message.PackageNames)
 	{
 		FPackageData* PackageData = COTFS.PackageDatas->FindPackageDataByPackageName(PackageName);
diff --git a/Engine/Source/Editor/UnrealEd/Private/Cooker/CookWorkerServer.cpp b/Engine/Source/Editor/UnrealEd/Private/Cooker/CookWorkerServer.cpp
--- a/Engine/Source/Editor/UnrealEd/Private/Cooker/CookWorkerServer.cpp
+++ b/Engine/Source/Editor/UnrealEd/Private/Cooker/CookWorkerServer.cpp
@@ -252,13 +252,14 @@ void FCookWorkerServer::AppendAssignments(TArrayView<FPackageData*> Assignments,
 	PackagesToAssignInfoPackages.Append(InfoPackages);
 }
 
-void FCookWorkerServer::AbortAllAssignments(TSet<FPackageData*>& OutPendingPackages, ECookDirectorThread TickThread)
+void FCookWorkerServer::AbortAllAssignments(TSet<FPackageData*>& OutPendingPackages, ECookDirectorThread TickThread,
+	int32 CurrentHeartbeat)
 {
 	FCommunicationScopeLock ScopeLock(this, TickThread, ETickAction::Queue);
-	AbortAllAssignmentsInLock(OutPendingPackages);
+	AbortAllAssignmentsInLock(OutPendingPackages, CurrentHeartbeat);
 }
 
-void FCookWorkerServer::AbortAllAssignmentsInLock(TSet<FPackageData*>& OutPendingPackages)
+void FCookWorkerServer::AbortAllAssignmentsInLock(TSet<FPackageData*>& OutPendingPackages, int32 CurrentHeartbeat)
 {
 	if (PendingPackages.Num())
 	{
@@ -283,14 +284,14 @@ void FCookWorkerServer::AbortAllAssignmentsInLock(TSet<FPackageData*>& OutPendin
 }
 
 void FCookWorkerServer::AbortAssignment(FPackageData& PackageData, ECookDirectorThread TickThread,
-	ENotifyRemote NotifyRemote)
+	int32 CurrentHeartbeat, ENotifyRemote NotifyRemote)
 {
 	FPackageData* PackageDataPtr = &PackageData;
-	AbortAssignments(TConstArrayView<FPackageData*>(&PackageDataPtr, 1), TickThread, NotifyRemote);
+	AbortAssignments(TConstArrayView<FPackageData*>(&PackageDataPtr, 1), TickThread, CurrentHeartbeat, NotifyRemote);
 }
 
 void FCookWorkerServer::AbortAssignments(TConstArrayView<FPackageData*> PackageDatas, ECookDirectorThread TickThread,
-	ENotifyRemote NotifyRemote)
+	int32 CurrentHeartbeat, ENotifyRemote NotifyRemote)
 {
 	FCommunicationScopeLock ScopeLock(this, TickThread, ETickAction::Queue);
 
@@ -316,13 +317,15 @@ void FCookWorkerServer::AbortAssignments(TConstArrayView<FPackageData*> PackageD
 	{
 		SendMessageInLock(FAbortPackagesMessage(MoveTemp(PackageNamesToMessage)));
 	}
+	LastAbortHeartbeatNumber = CurrentHeartbeat;
 }
 
-void FCookWorkerServer::AbortWorker(TSet<FPackageData*>& OutPendingPackages, ECookDirectorThread TickThread)
+void FCookWorkerServer::AbortWorker(TSet<FPackageData*>& OutPendingPackages, ECookDirectorThread TickThread,
+	int32 CurrentHeartbeat)
 {
 	FCommunicationScopeLock ScopeLock(this, TickThread, ETickAction::Tick);
 
-	AbortAllAssignmentsInLock(OutPendingPackages);
+	AbortAllAssignmentsInLock(OutPendingPackages, CurrentHeartbeat);
 	switch (ConnectStatus)
 	{
 	case EConnectStatus::Uninitialized: // Fall through
@@ -663,6 +666,17 @@ void FCookWorkerServer::PumpSendMessages()
 
 void FCookWorkerServer::SendPendingMessages()
 {
+	// If we aborted any packages, do not allow any new assignment messages to be sent until we receive an acknowledge of
+	// the abort. This prevents us from incorrectly assuming a package results message that was sent before the abort is
+	// the package results message after reassignment of that package, with new requested platforms, that we sent after
+	// the abort.
+	// Because we contractually are not allowed to send QueuedMessagesToSendAfterPackagesToAssign until after we have sent
+	// the assignment message, do not allow those to be sent out either.
+	if (LastReceivedHeartbeatNumber <= LastAbortHeartbeatNumber)
+	{
+		return;
+	}
+
 	SendPendingPackages();
 	for (UE::CompactBinaryTCP::FMarshalledMessage& MarshalledMessage : QueuedMessagesToSendAfterPackagesToAssign)
 	{
@@ -931,6 +945,41 @@ FString WriteCookStatus(FPackageData& PackageData, TConstArrayView<const ITarget
 	return Result;
 }
 
+TArray<const ITargetPlatform*> GetCookedPlatformListFromPlatformResults(
+	TConstArrayView<const ITargetPlatform*> OrderedPlatforms,
+	TConstArrayView<FPackageRemoteResult::FPlatformResult> PlatformResults)
+{
+	TArray<const ITargetPlatform*> PlatformList;
+	int32 NumPlatforms = OrderedPlatforms.Num();
+	if (NumPlatforms != PlatformResults.Num())
+	{
+		return PlatformList;
+	}
+	for (int32 PlatformIndex = 0; PlatformIndex < NumPlatforms; ++PlatformIndex)
+	{
+		const FPackageRemoteResult::FPlatformResult& PlatformResult = PlatformResults[PlatformIndex];
+		if (PlatformResult.GetCookResults() != ECookResult::Invalid)
+		{
+			PlatformList.Add(OrderedPlatforms[PlatformIndex]);
+		}
+	}
+	return PlatformList;
+}
+
+FString PlatformListToString(TConstArrayView<const ITargetPlatform*> Platforms)
+{
+	TStringBuilder<256> Result;
+	for (const ITargetPlatform* TargetPlatform : Platforms)
+	{
+		Result << TargetPlatform->PlatformName() << TEXT(", ");
+	}
+	if (Result.Len() > 0)
+	{
+		Result.RemoveSuffix(2); // ", "
+	}
+	return FString(Result);
+}
+
 void FCookWorkerServer::RecordResults(FPackageResultsMessage& Message)
 {
 	check(TickState.TickThread == ECookDirectorThread::SchedulerThread);
@@ -1042,7 +1091,21 @@ void FCookWorkerServer::RecordResults(FPackageResultsMessage& Message)
 		PackageData->SetWorkerAssignment(FWorkerId::Invalid(), ESendFlags::QueueNone);
 		if (Result.GetSuppressCookReason() == ESuppressCookReason::NotSuppressed)
 		{
-			COTFS.PromoteToSaveComplete(*PackageData, ESendFlags::QueueAddAndRemove);
+			if (PackageData->GetPlatformsNeedingCookingNum() > 0)
+			{
+				TArray<const ITargetPlatform*> Remaining;
+				PackageData->GetPlatformsNeedingCooking(Remaining);
+				UE_LOG(LogCook, Display, TEXT("Package %s was completed by CookWorker %d for platforms { %s }, but it still needs to cook platforms { %s }.")
+					TEXT(" Sending it back to the request state."),
+					*PackageData->GetPackageName().ToString(), ProfileId,
+					*PlatformListToString(GetCookedPlatformListFromPlatformResults(OrderedSessionPlatforms, Result.GetPlatforms())),
+					*PlatformListToString(Remaining));
+				PackageData->SendToState(EPackageState::Request, ESendFlags::QueueAddAndRemove, EStateChangeReason::Discovered);
+			}
+			else
+			{
+				COTFS.PromoteToSaveComplete(*PackageData, ESendFlags::QueueAddAndRemove);
+			}
 		}
 		else if (Result.GetSuppressCookReason() == ESuppressCookReason::RetractedByCookDirector)
 		{
diff --git a/Engine/Source/Editor/UnrealEd/Private/Cooker/CookWorkerServer.h b/Engine/Source/Editor/UnrealEd/Private/Cooker/CookWorkerServer.h
--- a/Engine/Source/Editor/UnrealEd/Private/Cooker/CookWorkerServer.h
+++ b/Engine/Source/Editor/UnrealEd/Private/Cooker/CookWorkerServer.h
@@ -70,17 +70,18 @@ public:
 		ECookDirectorThread TickThread);
 	/** Remove assignment of the package from local state and optionally from the connected Client. */
 	void AbortAssignment(FPackageData& PackageData, ECookDirectorThread TickThread,
-		ENotifyRemote NotifyRemote = ENotifyRemote::NotifyRemote);
+		int32 CurrentHeartbeat, ENotifyRemote NotifyRemote = ENotifyRemote::NotifyRemote);
 	void AbortAssignments(TConstArrayView<FPackageData*> PackageData, ECookDirectorThread TickThread,
-		ENotifyRemote = ENotifyRemote::NotifyRemote);
+		int32 CurrentHeartbeat, ENotifyRemote NotifyRemote = ENotifyRemote::NotifyRemote);
 
 	/**
 	 * Remove assignment of all assigned packages from local state and from the connected Client.
 	 * Report all packages that were removed.
 	 */
-	void AbortAllAssignments(TSet<FPackageData*>& OutPendingPackages, ECookDirectorThread TickThread);
+	void AbortAllAssignments(TSet<FPackageData*>& OutPendingPackages, ECookDirectorThread TickThread,
+		int32 CurrentHeartbeat);
 	/** AbortAllAssignments and tell the Client to gracefully terminate. Report all packages that were unassigned. */
-	void AbortWorker(TSet<FPackageData*>& OutPendingPackages, ECookDirectorThread TickThread);
+	void AbortWorker(TSet<FPackageData*>& OutPendingPackages, ECookDirectorThread TickThread, int32 CurrentHeartbeat);
 	/** Take over the Socket for a CookWorker that has just connected. */
 	bool TryHandleConnectMessage(FWorkerConnectMessage& Message, FSocket* InSocket,
 		TArray<UE::CompactBinaryTCP::FMarshalledMessage>&& OtherPacketMessages, ECookDirectorThread TickThread);
@@ -188,7 +189,7 @@ private:
 	/** Helper for Tick, pump receive messages from a connected Client. */
 	void PumpReceiveMessages();
 	/** The main implementation of AbortAllAssignments, only callable from inside the lock. */
-	void AbortAllAssignmentsInLock(TSet<FPackageData*>& OutPendingPackages);
+	void AbortAllAssignmentsInLock(TSet<FPackageData*>& OutPendingPackages, int32 CurrentHeartbeat);
 	/** Send the message immediately to the Socket. If cannot complete immediately, it will be finished during Tick. */
 	void SendMessageInLock(const IMPCollectorMessage& Message);
 	/** Send this into the given state. Update any state-dependent variables. */
@@ -235,6 +236,7 @@ private:
 	int32 LastReceivedHeartbeatNumber = 0;
 	int32 PackagesAssignedFenceMarker = 0;
 	int32 PackagesRetiredFenceMarker = 0;
+	int32 LastAbortHeartbeatNumber = -1;
 	double ConnectStartTimeSeconds = 0.;
 	double ConnectTestStartTimeSeconds = 0.;
 	FWorkerId WorkerId = FWorkerId::Invalid();
@@ -326,7 +328,7 @@ struct FAbortPackagesMessage : public IMPCollectorMessage
 {
 public:
 	FAbortPackagesMessage() = default;
-	FAbortPackagesMessage(TArray<FName>&& InPackageNames);
+	explicit FAbortPackagesMessage(TArray<FName>&& InPackageNames);
 
 	virtual void Write(FCbWriter& Writer) const override;
 	virtual bool TryRead(FCbObjectView Object) override;
