[Viewport Toolbar] The blue line separating the viewport toolbar from the viewport can show under the black border of cinematic views

Hides the inherited "blue line", and creates a new one matching its behavior, but located after the toolbar.

diff --git a/Engine/Plugins/MovieScene/LevelSequenceEditor/Source/LevelSequenceEditor/Private/CinematicViewport/SCinematicLevelViewport.cpp b/Engine/Plugins/MovieScene/LevelSequenceEditor/Source/LevelSequenceEditor/Private/CinematicViewport/SCinematicLevelViewport.cpp
--- a/Engine/Plugins/MovieScene/LevelSequenceEditor/Source/LevelSequenceEditor/Private/CinematicViewport/SCinematicLevelViewport.cpp
+++ b/Engine/Plugins/MovieScene/LevelSequenceEditor/Source/LevelSequenceEditor/Private/CinematicViewport/SCinematicLevelViewport.cpp
@@ -359,6 +359,13 @@ public:
 				]
 			]
 			+ SVerticalBox::Slot()
+			.MaxHeight(1.0f)
+			[
+				CreateViewportIndicatorWidget(
+					TAttribute<EVisibility>::CreateSP(this, &SCinematicPreviewViewport::OnGetCinematicViewportIndicatorVisibility)
+					).ToSharedRef()
+			]
+			+ SVerticalBox::Slot()
 			.AutoHeight()
 			[
 				SNew(SBox)
@@ -393,6 +400,13 @@ public:
 		return ViewportContentVisibility == EVisibility::Visible ? EVisibility::HitTestInvisible : ViewportContentVisibility;
 	}
 
+	virtual EVisibility OnGetFocusedViewportIndicatorVisibility() const override { return EVisibility::Collapsed; }
+
+	EVisibility OnGetCinematicViewportIndicatorVisibility() const
+	{
+		return SLevelViewport::OnGetFocusedViewportIndicatorVisibility();
+	}
+
 private:
 	bool bShowToolbar;
 
diff --git a/Engine/Source/Editor/UnrealEd/Private/SEditorViewport.cpp b/Engine/Source/Editor/UnrealEd/Private/SEditorViewport.cpp
--- a/Engine/Source/Editor/UnrealEd/Private/SEditorViewport.cpp
+++ b/Engine/Source/Editor/UnrealEd/Private/SEditorViewport.cpp
@@ -108,13 +108,6 @@ void SEditorViewport::Construct( const FArguments& InArgs )
 	];
 	// clang-format on
 
-	// This makes a gradient that displays whether or not a viewport is active
-	FLinearColor ActiveBorderColor = FAppStyle::Get().GetSlateColor("EditorViewport.ActiveBorderColor").GetSpecifiedColor();
-	FLinearColor ActiveBorderColorTransparent = ActiveBorderColor;
-	ActiveBorderColorTransparent.A = 0.0f;
-
-	static TArray<FLinearColor> GradientStops{ ActiveBorderColorTransparent, ActiveBorderColor, ActiveBorderColorTransparent };
-
 	// clang-format off
 	ViewportOverlay->AddSlot()
 	.VAlign(VAlign_Top)
@@ -124,9 +117,9 @@ void SEditorViewport::Construct( const FArguments& InArgs )
 		.MaxDesiredHeight(1.0f)
 		.MinDesiredHeight(1.0f)
 		[
-			SNew(SComplexGradient)
-			.GradientColors(GradientStops)
-			.Orientation(EOrientation::Orient_Vertical)
+			CreateViewportIndicatorWidget(
+				TAttribute<EVisibility>::CreateSP(this, &SEditorViewport::OnGetFocusedViewportIndicatorVisibility)
+				).ToSharedRef()
 		]
 	];
 	// clang-format on
@@ -963,6 +956,25 @@ bool SEditorViewport::IsRotateToSurfaceNormalEnabled()
 	return Settings.bSnapRotation;
 }
 
+TSharedPtr<SWidget> SEditorViewport::CreateViewportIndicatorWidget(const TAttribute<EVisibility>& InVisibility)
+{
+	// This makes a gradient that displays whether a viewport is active
+	static FLinearColor ActiveBorderColor =
+		FAppStyle::Get().GetSlateColor("EditorViewport.ActiveBorderColor").GetSpecifiedColor();
+	static FLinearColor ActiveBorderColorTransparent(ActiveBorderColor.R, ActiveBorderColor.G, ActiveBorderColor.B, 0.0f);
+	static TArray<FLinearColor> GradientStops{ ActiveBorderColorTransparent, ActiveBorderColor, ActiveBorderColorTransparent };
+
+	return SNew(SBox)
+		.Visibility(InVisibility)
+		.MaxDesiredHeight(1.0f)
+		.MinDesiredHeight(1.0f)
+		[
+			SNew(SComplexGradient)
+			.GradientColors(GradientStops)
+			.Orientation(EOrientation::Orient_Vertical)
+		];
+}
+
 bool SEditorViewport::IsPreviewingScreenPercentage() const
 {
 	return Client->IsPreviewingScreenPercentage();
diff --git a/Engine/Source/Editor/UnrealEd/Public/SEditorViewport.h b/Engine/Source/Editor/UnrealEd/Public/SEditorViewport.h
--- a/Engine/Source/Editor/UnrealEd/Public/SEditorViewport.h
+++ b/Engine/Source/Editor/UnrealEd/Public/SEditorViewport.h
@@ -281,6 +281,12 @@ protected:
 	 */
 	static UNREALED_API bool IsRotateToSurfaceNormalEnabled();
 
+	/**
+	 * Creates a widget showing whether the viewport is active or not
+	 * @param InVisibility visibility attribute, can be used to dynamically hide the widget
+	 */
+	static UNREALED_API TSharedPtr<SWidget> CreateViewportIndicatorWidget(const TAttribute<EVisibility>& InVisibility);
+
 protected:
 	TSharedPtr<SOverlay> ViewportOverlay;
 
