[Viewport Toolbar] Copy main toolbar "Settings > Previewing > Draw Brush Polys" to Editor Preferences

- moves bShowBrushMarkerPolys to ULevelEditorViewportSettings
- deprecates former bShowBrushMarkerPolys found in UEditorEngine

diff --git a/Engine/Source/Editor/LevelEditor/Private/LevelEditorActions.cpp b/Engine/Source/Editor/LevelEditor/Private/LevelEditorActions.cpp
--- a/Engine/Source/Editor/LevelEditor/Private/LevelEditorActions.cpp
+++ b/Engine/Source/Editor/LevelEditor/Private/LevelEditorActions.cpp
@@ -2911,13 +2911,14 @@ bool FLevelEditorActionCallbacks::OnIsTransparentBoxSelectEnabled()
 
 void FLevelEditorActionCallbacks::OnDrawBrushMarkerPolys()
 {
-	GEditor->Exec( GetWorld(), *FString::Printf( TEXT("MODE SHOWBRUSHMARKERPOLYS=%d"), !GEditor->bShowBrushMarkerPolys ? 1 : 0 ) );
+	const bool bShowBrushMarkerPolys = GetDefault<ULevelEditorViewportSettings>()->bShowBrushMarkerPolys;
+	GEditor->Exec(GetWorld(), *FString::Printf(TEXT("MODE SHOWBRUSHMARKERPOLYS=%d"), !bShowBrushMarkerPolys ? 1 : 0));
 	GEditor->SaveConfig();
 }
 
 bool FLevelEditorActionCallbacks::OnIsDrawBrushMarkerPolysEnabled()
 {
-	return GEditor->bShowBrushMarkerPolys;
+	return GetDefault<ULevelEditorViewportSettings>()->bShowBrushMarkerPolys;
 }
 
 void FLevelEditorActionCallbacks::OnToggleOnlyLoadVisibleInPIE()
diff --git a/Engine/Source/Editor/UnrealEd/Classes/Editor/EditorEngine.h b/Engine/Source/Editor/UnrealEd/Classes/Editor/EditorEngine.h
--- a/Engine/Source/Editor/UnrealEd/Classes/Editor/EditorEngine.h
+++ b/Engine/Source/Editor/UnrealEd/Classes/Editor/EditorEngine.h
@@ -419,7 +419,8 @@ public:
 	FString GameCommandLine;
 
 	/** If true, show translucent marker polygons on the builder brush and volumes. */
-	UPROPERTY(EditAnywhere, config, Category=Advanced)
+	UE_DEPRECATED(5.6, "bShowBrushMarkerPoly has been moved to ULevelEditorViewportSettings::bShowBrushMarkerPoly. Setting can now be found under Editor Preferences->Level Editor->Viewport->Look And Feel")
+	UPROPERTY(config, meta=(DeprecatedProperty, DeprecationMessage = "bShowBrushMarkerPoly can now be found under Editor Preferences->Level Editor->Viewport->Look And Feel"))
 	uint32 bShowBrushMarkerPolys:1;
 
 	/** If true, socket snapping is enabled in the main level viewports. */
diff --git a/Engine/Source/Editor/UnrealEd/Classes/Settings/LevelEditorViewportSettings.h b/Engine/Source/Editor/UnrealEd/Classes/Settings/LevelEditorViewportSettings.h
--- a/Engine/Source/Editor/UnrealEd/Classes/Settings/LevelEditorViewportSettings.h
+++ b/Engine/Source/Editor/UnrealEd/Classes/Settings/LevelEditorViewportSettings.h
@@ -321,6 +321,10 @@ class ULevelEditorViewportSettings
 	UPROPERTY(EditAnywhere, config, Category = LookAndFeel, meta=(AdvancedDisplay))
 	uint32 bShowActorEditorContext : 1;
 
+	/** If true, show translucent marker polygons on the builder brush and volumes. */
+	UPROPERTY(EditAnywhere, config, Category=LookAndFeel)
+	uint32 bShowBrushMarkerPolys:1;
+
 	/** If true, the Edit widget of a transform will display the axis */
 	UPROPERTY(EditAnywhere, config, Category=LookAndFeel, meta=( DisplayName = "Enable Axis drawing for transform Edit Widget" ))
 	uint32 bAllowEditWidgetAxisDisplay:1;
diff --git a/Engine/Source/Editor/UnrealEd/Private/LevelEditorViewport.cpp b/Engine/Source/Editor/UnrealEd/Private/LevelEditorViewport.cpp
--- a/Engine/Source/Editor/UnrealEd/Private/LevelEditorViewport.cpp
+++ b/Engine/Source/Editor/UnrealEd/Private/LevelEditorViewport.cpp
@@ -5077,7 +5077,7 @@ void FLevelEditorViewportClient::DrawBrushDetails(const FSceneView* View, FPrimi
 		// Brush->Brush is checked to safe from brushes that were created without having their brush members attached.
 		// Check whether it satisfies the ShowBrushMarkersPoly condition of being selected or whether bDisplayShadedVolume is true
 		if (Brush->Brush && (FActorEditorUtils::IsABuilderBrush(Brush) || Brush->IsVolumeBrush()) && 
-			(Brush->bDisplayShadedVolume || (GEditor->bShowBrushMarkerPolys && ModeTools->GetSelectedActors()->IsSelected(Brush))))
+			(Brush->bDisplayShadedVolume || (GetDefault<ULevelEditorViewportSettings>()->bShowBrushMarkerPolys && ModeTools->GetSelectedActors()->IsSelected(Brush))))
 		{
 			// Build a mesh by basically drawing the triangles of each 
 			FDynamicMeshBuilder MeshBuilder(View->GetFeatureLevel());
diff --git a/Engine/Source/Editor/UnrealEd/Private/UnrealEdSrv.cpp b/Engine/Source/Editor/UnrealEd/Private/UnrealEdSrv.cpp
--- a/Engine/Source/Editor/UnrealEd/Private/UnrealEdSrv.cpp
+++ b/Engine/Source/Editor/UnrealEd/Private/UnrealEdSrv.cpp
@@ -3068,7 +3068,7 @@ bool UUnrealEdEngine::Exec_Mode( const TCHAR* Str, FOutputDevice& Ar )
 	if( FParse::Value(Str, TEXT("SHOWBRUSHMARKERPOLYS="), DWord1) )
 	{
 		FinishAllSnaps();
-		bShowBrushMarkerPolys = DWord1;
+		GetMutableDefault<ULevelEditorViewportSettings>()->bShowBrushMarkerPolys = DWord1;
 	}
 
 	if( FParse::Value(Str, TEXT("SELECTIONLOCK="), DWord1) )
