Register untracked soft object pointer during harvesting so they can be added to the redirector dependencies.

FPackageHarvester::operator<< : if a package was not register in the regular list of soft package references then save it in the list of untracked soft package references.

FHarvestedRealm : add new set of untracked package name. The harvesting stores untracked soft package references in this set.
FSavePackageResultStruct : add array to store untracked package names. This list is filled in GetFinalResult from the list in FHarvestRealm.

FCookDependencies : add array RedirectorsAffectingPackageDependencies. Redirectors for package names in this list are added to the build dependencies.
FCookDependencies::Collect : pass the list of untracked package name from the struct FSavePackageResultStruct to FCookDependencies.
FCookDependencies::TryCalculateCurrentKey : add redirectors of the packages in the list RedirectorsAffectingPackageDependencies to the build dependencies.

#rnx

diff --git a/Engine/Source/Editor/UnrealEd/Private/TargetDomain/TargetDomainUtils.cpp b/Engine/Source/Editor/UnrealEd/Private/TargetDomain/TargetDomainUtils.cpp
--- a/Engine/Source/Editor/UnrealEd/Private/TargetDomain/TargetDomainUtils.cpp
+++ b/Engine/Source/Editor/UnrealEd/Private/TargetDomain/TargetDomainUtils.cpp
@@ -210,6 +210,7 @@ FCookDependencies::ECurrentKeyResult FCookDependencies::TryCalculateCurrentKey(
 
 	FCoreRedirects::AppendHashOfRedirectsAffectingPackages(KeyBuilder, RuntimePackageDependencies);
 	FCoreRedirects::AppendHashOfRedirectsAffectingPackages(KeyBuilder, ScriptPackageDependencies);
+	FCoreRedirects::AppendHashOfRedirectsAffectingPackages(KeyBuilder, RedirectorsAffectingPackageDependencies);
 
 	if (!ConfigDependencies.IsEmpty())
 	{
@@ -270,6 +271,7 @@ void FCookDependencies::Reset()
 	ConfigDependencies.Reset();
 	RuntimePackageDependencies.Reset();
 	ScriptPackageDependencies.Reset();
+	RedirectorsAffectingPackageDependencies.Reset();
 	PackageName = FName();
 	StoredKey = FIoHash::Zero;
 	CurrentKey = FIoHash::Zero;
@@ -282,6 +284,7 @@ void FCookDependencies::Empty()
 	BuildPackageDependencies.Empty();
 	ConfigDependencies.Empty();
 	RuntimePackageDependencies.Empty();
+	RedirectorsAffectingPackageDependencies.Empty();
 }
 
 enum class EPackageMountPoint
@@ -573,6 +576,10 @@ FCookDependencies FCookDependencies::Collect(UPackage* Package, const ITargetPla
 			}
 			RemainingDependencies = NextIncludedDependencies;
 		}
+
+		Result.RedirectorsAffectingPackageDependencies = SaveResult->UntrackedSoftPackageReferences;
+		Result.RedirectorsAffectingPackageDependencies.Sort(FNameLexicalLess());
+		Result.RedirectorsAffectingPackageDependencies.SetNum(Algo::Unique(Result.RedirectorsAffectingPackageDependencies));
 	}
 
 	FName TransientPackageName = GetTransientPackage()->GetFName();
@@ -750,6 +757,13 @@ bool LoadFromCompactBinary(FCbObjectView ObjectView, UE::TargetDomain::FCookDepe
 				return false;
 			}
 		}
+		if (FieldView->GetName().Equals(UTF8TEXTVIEW("RedirectorsAffectingPackageDependencies")))
+		{
+			if (!LoadFromCompactBinary(FieldView++, Dependencies.RedirectorsAffectingPackageDependencies))
+			{
+				return false;
+			}
+		}
 		if (FieldView.GetName().Equals(UTF8TEXTVIEW("ScriptPackageDependencies")))
 		{
 			if (!LoadFromCompactBinary(FieldView++, Dependencies.ScriptPackageDependencies))
@@ -810,6 +824,10 @@ FCbWriter& operator<<(FCbWriter& Writer, const UE::TargetDomain::FCookDependenci
 	{
 		Writer << "RuntimePackageDependencies" << CookDependencies.RuntimePackageDependencies;
 	}
+	if (!CookDependencies.RedirectorsAffectingPackageDependencies.IsEmpty())
+	{
+		Writer << "RedirectorsAffectingPackageDependencies" << CookDependencies.RedirectorsAffectingPackageDependencies;
+	}
 	if (!CookDependencies.ScriptPackageDependencies.IsEmpty())
 	{
 		Writer << "ScriptPackageDependencies" << CookDependencies.ScriptPackageDependencies;
diff --git a/Engine/Source/Editor/UnrealEd/Private/TargetDomain/TargetDomainUtils.h b/Engine/Source/Editor/UnrealEd/Private/TargetDomain/TargetDomainUtils.h
--- a/Engine/Source/Editor/UnrealEd/Private/TargetDomain/TargetDomainUtils.h
+++ b/Engine/Source/Editor/UnrealEd/Private/TargetDomain/TargetDomainUtils.h
@@ -122,6 +122,7 @@ private:
 	TArray<FString> ConfigDependencies;
 	TArray<FName> RuntimePackageDependencies;
 	TArray<FName> ScriptPackageDependencies;
+	TArray<FName> RedirectorsAffectingPackageDependencies;
 	TArray<UE::Cook::FCookDependency> CookDependencies;
 	TArray<UE::Cook::FCookDependency> TransitiveBuildDependencies;
 	TArray<FString> ClassDependencies;
diff --git a/Engine/Source/Runtime/CoreUObject/Private/UObject/SavePackage/PackageHarvester.cpp b/Engine/Source/Runtime/CoreUObject/Private/UObject/SavePackage/PackageHarvester.cpp
--- a/Engine/Source/Runtime/CoreUObject/Private/UObject/SavePackage/PackageHarvester.cpp
+++ b/Engine/Source/Runtime/CoreUObject/Private/UObject/SavePackage/PackageHarvester.cpp
@@ -990,6 +990,11 @@ FArchive& FPackageHarvester::operator<<(FSoftObjectPath& Value)
 		ESoftObjectPathCollectType CollectType = ESoftObjectPathCollectType::AlwaysCollect;
 		ESoftObjectPathSerializeType SerializeType = ESoftObjectPathSerializeType::AlwaysSerialize;
 
+		FString Path = Value.ToString();
+		FName PackageName = FName(*FPackageName::ObjectPathToPackageName(Path));
+
+		bool bRecordedRuntimeDependency = false;
+
 		// Don't track as a soft reference to a package if this is a NeverCollect SoftObjectPath
 		ThreadContext.GetSerializationOptions(ReferencingPackageName, ReferencingPropertyName, CollectType, SerializeType, this);
 		if (CollectType != ESoftObjectPathCollectType::NeverCollect && CollectType != ESoftObjectPathCollectType::NonPackage)
@@ -1000,14 +1005,24 @@ FArchive& FPackageHarvester::operator<<(FSoftObjectPath& Value)
 			// this->FArchiveState::IsEditorOnlyPropertyOnTheStack
 			FHarvestScope EditorOnlyScope = EnterConditionalEditorOnlyScope(CollectType == ESoftObjectPathCollectType::EditorOnlyCollect);
 #endif
-			FString Path = Value.ToString();
-			FName PackageName = FName(*FPackageName::ObjectPathToPackageName(Path));
 			HarvestPackageHeaderName(PackageName);
-			ForEachExportHarvestingRealm([this, PackageName](ESaveRealm HarvestingRealm)
+			ForEachExportHarvestingRealm([this, PackageName, &bRecordedRuntimeDependency](ESaveRealm HarvestingRealm)
 				{
 					SaveContext.GetHarvestedRealm(HarvestingRealm).GetSoftPackageReferenceList().Add(PackageName);
+					bRecordedRuntimeDependency = true;
 				});
 		}
+
+		if (!bRecordedRuntimeDependency)
+		{
+			// Save the names of packages not recorded in the regular SoftPackageReferenceList in the UntrackedSoftPackageReferenceList.
+			// When doing cook saves, redirections on untracked soft packages are added to build dependencies. This way 
+			// when the redirection changes, a recook is triggered.
+			ForEachExportHarvestingRealm([this, PackageName](ESaveRealm HarvestingRealm)
+			{
+				SaveContext.GetHarvestedRealm(HarvestingRealm).GetUntrackedSoftPackageReferenceList().Add(PackageName);
+			});
+		}
 	}
 	return *this;
 }
diff --git a/Engine/Source/Runtime/CoreUObject/Private/UObject/SavePackage/SaveContext.cpp b/Engine/Source/Runtime/CoreUObject/Private/UObject/SavePackage/SaveContext.cpp
--- a/Engine/Source/Runtime/CoreUObject/Private/UObject/SavePackage/SaveContext.cpp
+++ b/Engine/Source/Runtime/CoreUObject/Private/UObject/SavePackage/SaveContext.cpp
@@ -245,6 +245,9 @@ FSavePackageResultStruct FSaveContext::GetFinalResult()
 	ResultData.CookDependencies = MoveTemp(ObjectSaveContext.CookBuildDependencies);
 #endif
 
+	const TSet<FName>& UntrackedPackageReferenceList = GetUntrackedSoftPackageReferenceList();
+	ResultData.UntrackedSoftPackageReferences = UntrackedPackageReferenceList.Array();
+
 	return ResultData;
 }
 
diff --git a/Engine/Source/Runtime/CoreUObject/Private/UObject/SavePackage/SaveContext.h b/Engine/Source/Runtime/CoreUObject/Private/UObject/SavePackage/SaveContext.h
--- a/Engine/Source/Runtime/CoreUObject/Private/UObject/SavePackage/SaveContext.h
+++ b/Engine/Source/Runtime/CoreUObject/Private/UObject/SavePackage/SaveContext.h
@@ -281,6 +281,16 @@ struct FHarvestedRealm
 		return SoftPackageReferenceList;
 	}
 
+	const TSet<FName>& GetUntrackedSoftPackageReferenceList() const
+	{
+		return UntrackedSoftPackageReferenceList;
+	}
+
+	TSet<FName>& GetUntrackedSoftPackageReferenceList()
+	{
+		return UntrackedSoftPackageReferenceList;
+	}
+
 	const TMap<TObjectPtr<UObject>, TArray<FName>>& GetSearchableNamesObjectMap() const
 	{
 		return SearchableNamesObjectMap;
@@ -460,6 +470,10 @@ private:
 	TSet<FSoftObjectPath> SoftObjectPathList;
 	// List of soft package reference found
 	TSet<FName> SoftPackageReferenceList;
+	// List of soft package reference not in the SoftPackageReferenceList because they are NeverCollect or EditorOnly.  
+	// They are saved here anyway in case there is a redirection. The redirection will be saved in the build dependencies so
+	// incremental cook can detect when the redirection changed.
+	TSet<FName> UntrackedSoftPackageReferenceList;
 	// Map of objects to their list of searchable names
 	TMap<TObjectPtr<UObject>, TArray<FName>> SearchableNamesObjectMap;
 	// Map of objects to their dependencies
@@ -942,6 +956,11 @@ public:
 		return GetHarvestedRealm(ESaveRealm::Game).GetSoftPackageReferenceList();
 	}
 
+	const TSet<FName>& GetUntrackedSoftPackageReferenceList() const
+	{
+		return GetHarvestedRealm().GetUntrackedSoftPackageReferenceList();
+	}
+
 	const TArray<FName>& GetPackageBuildDependencies(ESaveRealm SaveRealm = ESaveRealm::None)
 	{
 		SaveRealm = (SaveRealm == ESaveRealm::None) ? GetCurrentHarvestingRealm() : SaveRealm;
diff --git a/Engine/Source/Runtime/CoreUObject/Public/UObject/Package.h b/Engine/Source/Runtime/CoreUObject/Public/UObject/Package.h
--- a/Engine/Source/Runtime/CoreUObject/Public/UObject/Package.h
+++ b/Engine/Source/Runtime/CoreUObject/Public/UObject/Package.h
@@ -119,6 +119,7 @@ struct FSavePackageResultStruct
 
 	TArray<FName> ImportPackages;
 	TArray<FName> SoftPackageReferences;
+	TArray<FName> UntrackedSoftPackageReferences;
 
 #if WITH_EDITOR
 	TArray<UE::Cook::FCookDependency> CookDependencies;
