World Bookmark - CIS fix
#cisfix
#rnx

diff --git a/Engine/Source/Editor/UnrealEd/Public/EditorState/EditorStateCollection.h b/Engine/Source/Editor/UnrealEd/Public/EditorState/EditorStateCollection.h
--- a/Engine/Source/Editor/UnrealEd/Public/EditorState/EditorStateCollection.h
+++ b/Engine/Source/Editor/UnrealEd/Public/EditorState/EditorStateCollection.h
@@ -156,7 +156,8 @@ struct TStructOpsTypeTraits<FEditorStateCollection> : public TStructOpsTypeTrait
 {
 	enum
 	{
-		WithPostSerialize = true
+		WithPostSerialize = true,
+		WithCopy = false
 	};
 };
 
