Running some type of WorldPartition builders on non-partition worlds result in warnings

diff --git a/Engine/Source/Editor/UnrealEd/Private/WorldPartition/WorldPartitionBuilder.cpp b/Engine/Source/Editor/UnrealEd/Private/WorldPartition/WorldPartitionBuilder.cpp
--- a/Engine/Source/Editor/UnrealEd/Private/WorldPartition/WorldPartitionBuilder.cpp
+++ b/Engine/Source/Editor/UnrealEd/Private/WorldPartition/WorldPartitionBuilder.cpp
@@ -142,7 +142,7 @@ bool UWorldPartitionBuilder::RunBuilder(UWorld* World)
 		}
 		else
 		{
-			UE_LOG(LogWorldPartitionBuilder, Warning, TEXT("WorldPartition builder '%s' only works on partitioned worlds - skipping processing of '%s'."), *GetClass()->GetName(), *World->GetPathName());
+			UE_LOG(LogWorldPartitionBuilder, Log, TEXT("WorldPartition builder '%s' only works on partitioned worlds - skipping processing of '%s'."), *GetClass()->GetName(), *World->GetPathName());
 		}
 	}
 
