Include package name in UEditorEngine::Save trace information

[RN] UEditorEngine::Save now includes the path of the package being saved in its trace information

diff --git a/Engine/Source/Editor/UnrealEd/Private/EditorEngine.cpp b/Engine/Source/Editor/UnrealEd/Private/EditorEngine.cpp
--- a/Engine/Source/Editor/UnrealEd/Private/EditorEngine.cpp
+++ b/Engine/Source/Editor/UnrealEd/Private/EditorEngine.cpp
@@ -4355,10 +4355,16 @@ void UEditorEngine::CleanupPhysicsSceneThatWasInitializedForSave(UWorld* World,
 	}
 }
 
+UE_TRACE_EVENT_BEGIN(Cpu, EditorEngine_Save, NoSync)
+	UE_TRACE_EVENT_FIELD(UE::Trace::WideString, PackageName)
+UE_TRACE_EVENT_END()
+
 FSavePackageResultStruct UEditorEngine::Save(UPackage* InOuter, UObject* InAsset, const TCHAR* Filename,
 	const FSavePackageArgs& InSaveArgs)
 {
-	TRACE_CPUPROFILER_EVENT_SCOPE(UEditorEngine::Save);
+	UE_TRACE_LOG_SCOPED_T(Cpu, EditorEngine_Save, CpuChannel)
+		<< EditorEngine_Save.PackageName(InOuter ? *(InOuter->GetName()) : TEXT("None"));
+
 	UE_TRACK_REFERENCING_PACKAGE_SCOPED(InOuter, PackageAccessTrackingOps::NAME_Save); // Needs to be here in addition to UPackage::Save so that InitializePhysicsSceneForSaveIfNecessary and OnPreSaveWorld have appropriate referencing package info
 
 	FSavePackageArgs SaveArgs(InSaveArgs);
