IncrementalCook: Cook packages that were referenced as transitive build dependencies but were not runtime referenced, so we can store their CookDependencies for use by the incremental cook of their referencers.
#rnx

diff --git a/Engine/Source/Editor/UnrealEd/Classes/CookOnTheSide/CookOnTheFlyServer.h b/Engine/Source/Editor/UnrealEd/Classes/CookOnTheSide/CookOnTheFlyServer.h
--- a/Engine/Source/Editor/UnrealEd/Classes/CookOnTheSide/CookOnTheFlyServer.h
+++ b/Engine/Source/Editor/UnrealEd/Classes/CookOnTheSide/CookOnTheFlyServer.h
@@ -194,8 +194,10 @@ namespace UE::Cook
 	class ICookOnTheFlyNetworkServer;
 	class IWorkerRequests;
 	enum class ECachedCookedPlatformDataEvent : uint8;
+	enum class ECookPhase : uint8;
 	enum class EPackageState : uint8;
 	enum class EPollStatus : uint8;
+	enum class EReachability : uint8;
 	enum class EStateChangeReason : uint8;
 	enum class ESuppressCookReason : uint8;
 	enum class ESendFlags : uint8;
@@ -499,6 +501,7 @@ private:
 		SaveLimited,	// Process the SaveQueue, stopping when savequeuelength reaches the desired population level
 		Poll,			// Execute pollables which have exceeded their period
 		PollIdle,		// Execute pollables which have exceeded their idle period
+		KickBuildDependencies, // Find all packages that are build dependencies and request them for cook
 		WaitForAsync,	// Sleep for a time slice while we wait for async tasks to complete
 		YieldTick,		// Progress is blocked by an async result. Temporarily exit TickMainCookLoop.
 	};
@@ -541,6 +544,11 @@ private:
 	 * Report the number of requests that were completed (either skipped or successfully saved or failed to save)
 	 */
 	void PumpSaves(UE::Cook::FTickStackData& StackData, uint32 DesiredQueueLength, int32& OutNumPushed, bool& bOutBusy);
+	void PumpRuntimeSaves(UE::Cook::FTickStackData& StackData, uint32 DesiredQueueLength, int32& OutNumPushed, bool& bOutBusy);
+	void PumpBuildDependencySaves(UE::Cook::FTickStackData& StackData, uint32 DesiredQueueLength, int32& OutNumPushed, bool& bOutBusy);
+
+	/** Find all packages that are build dependencies and request them for cook. */
+	void KickBuildDependencies(UE::Cook::FTickStackData& StackData);
 
 	/**
 	 * Inspect the given package and queue it for saving if necessary.
@@ -932,6 +940,7 @@ protected:
 
 private:
 
+	UE::Cook::FInstigator GetInstigator(FName PackageName, UE::Cook::EReachability Reachability);
 	bool PumpHasExceededMaxMemory(uint32& OutResultFlags);
 
 	/**
@@ -1356,6 +1365,7 @@ private:
 
 	/** Cook (save) a package and process the results */
 	void SaveCookedPackage(UE::Cook::FSaveCookedPackageContext& Context);
+	void CommitUncookedPackage(UE::Cook::FSaveCookedPackageContext& Context);
 	/** Helper for package saves using ExternalActors: record ExternalActors for iterative builds. */
 	void RecordExternalActorDependencies(TConstArrayView<FName> ExternalActorDependencies);
 
@@ -1469,6 +1479,9 @@ private:
 	void BroadcastCookStarted();
 	void BroadcastCookFinished();
 
+	/** Returns the current Phase of the cook; this is separate axis of behavior from CookMode. @see ECookPhase. */
+	UE::Cook::ECookPhase GetCookPhase() const;
+
 	static UCookOnTheFlyServer* ActiveCOTFS;
 	uint32		StatLoadedPackageCount = 0;
 	uint32		StatSavedPackageCount = 0;
@@ -1510,6 +1523,11 @@ private:
 	bool bSaveBusy = false;
 	/** We need to track whether the compiler has been inactive for a long time before issuing a warning about it. */
 	bool bShaderCompilerWasActiveeOnPreviousBusyReport = true;
+	/**
+	 * After cooking all requested packages in CookByTheBook we need to search for any build dependencies and request them
+	 * for commit as build dependencies. This bool records whether we have done that for the current session.
+	 */
+	bool bKickedBuildDependencies = false;
 	/**
 	 * If preloading is enabled, we call TryPreload until it returns true before sending the package to LoadReady,
 	 * otherwise we skip TryPreload and it goes immediately.
diff --git a/Engine/Source/Editor/UnrealEd/Private/Commandlets/AssetRegistryGenerator.cpp b/Engine/Source/Editor/UnrealEd/Private/Commandlets/AssetRegistryGenerator.cpp
--- a/Engine/Source/Editor/UnrealEd/Private/Commandlets/AssetRegistryGenerator.cpp
+++ b/Engine/Source/Editor/UnrealEd/Private/Commandlets/AssetRegistryGenerator.cpp
@@ -1221,6 +1221,7 @@ static UE::Cook::ECookResult DiskSizeToCookResult(int64 DiskSize)
 	switch (DiskSize)
 	{
 	case -2: return UE::Cook::ECookResult::NeverCookPlaceholder;
+	case -3: return UE::Cook::ECookResult::NotAttempted;
 	default: return UE::Cook::ECookResult::Failed;
 	}
 }
@@ -1231,6 +1232,7 @@ static int64 CookResultToDiskSize(UE::Cook::ECookResult CookResult)
 	{
 	case UE::Cook::ECookResult::Succeeded: return 0;
 	case UE::Cook::ECookResult::NeverCookPlaceholder: return -2;
+	case UE::Cook::ECookResult::NotAttempted: return -3;
 	default: return -1;
 	}
 }
diff --git a/Engine/Source/Editor/UnrealEd/Private/CookOnTheFlyServer.cpp b/Engine/Source/Editor/UnrealEd/Private/CookOnTheFlyServer.cpp
--- a/Engine/Source/Editor/UnrealEd/Private/CookOnTheFlyServer.cpp
+++ b/Engine/Source/Editor/UnrealEd/Private/CookOnTheFlyServer.cpp
@@ -1237,6 +1237,12 @@ bool UCookOnTheFlyServer::IsCookWorkerMode() const
 	return ::IsCookWorkerMode(CurrentCookMode);
 }
 
+UE::Cook::ECookPhase UCookOnTheFlyServer::GetCookPhase() const
+{
+	using namespace UE::Cook;
+	return !bKickedBuildDependencies ? ECookPhase::Cook : ECookPhase::BuildDependencies;
+}
+
 bool UCookOnTheFlyServer::IsUsingLegacyCookOnTheFlyScheduling() const
 {
 	return CookOnTheFlyRequestManager && CookOnTheFlyRequestManager->ShouldUseLegacyScheduling();
@@ -1434,7 +1440,7 @@ void UCookOnTheFlyServer::RunCookList(ECookListOptions CookListOptions)
 		bool bIncludePackage;
 		if (EnumHasAnyFlags(CookListOptions, ECookListOptions::ShowRejected))
 		{
-			bIncludePackage = PackageData->HasInstigator();
+			bIncludePackage = PackageData->HasInstigator(EReachability::Runtime);
 			if (bIncludePackage)
 			{
 				// Skip printing out a message for external actors
@@ -1466,7 +1472,7 @@ void UCookOnTheFlyServer::RunCookList(ECookListOptions CookListOptions)
 			bRejected ? TEXT("Rejected: ") : TEXT(""),
 			*WriteToString<256>(PackageData->GetPackageName()),
 			bShowInstigators ? TEXT(", Instigator: ") : TEXT(""),
-			bShowInstigators ? *PackageData->GetInstigator().ToString() : TEXT(""));
+			bShowInstigators ? *PackageData->GetInstigator(EReachability::Runtime).ToString() : TEXT(""));
 	}
 }
 
@@ -1562,6 +1568,9 @@ void UCookOnTheFlyServer::TickMainCookLoop(UE::Cook::FTickStackData& StackData)
 		case ECookAction::PollIdle:
 			PumpPollables(StackData, true /* bIsIdle */);
 			break;
+		case ECookAction::KickBuildDependencies:
+			KickBuildDependencies(StackData);
+			break;
 		case ECookAction::WaitForAsync:
 			WaitForAsync(StackData);
 			break;
@@ -2473,6 +2482,11 @@ UCookOnTheFlyServer::ECookAction UCookOnTheFlyServer::DecideNextCookAction(UE::C
 		}
 	}
 
+	if (IsCookByTheBookMode() && GetCookPhase() != ECookPhase::BuildDependencies)
+	{
+		return ECookAction::KickBuildDependencies;
+	}
+
 	if (IsCookOnTheFlyMode() || IsCookWorkerMode())
 	{
 		// These modes are not done until a manual trigger, so continue polling idle
@@ -2586,7 +2600,9 @@ bool UCookOnTheFlyServer::TryCreateRequestCluster(UE::Cook::FPackageData& Packag
 {
 	check(PackageData.IsInProgress()); // This should only be called from Pump functions, and only on in-progress Packages
 	using namespace UE::Cook;
-	if (!PackageData.AreAllReachablePlatformsVisitedByCluster())
+	const EReachability DesiredReachability = GetCookPhase() == ECookPhase::Cook
+		? EReachability::Runtime : EReachability::Build;
+	if (!PackageData.AreAllReachablePlatformsVisitedByCluster(DesiredReachability))
 	{
 		PackageData.SendToState(EPackageState::Request, ESendFlags::QueueAdd, EStateChangeReason::Discovered);
 		return true;
@@ -2603,6 +2619,7 @@ void UCookOnTheFlyServer::PumpRequests(UE::Cook::FTickStackData& StackData, int3
 	FRequestQueue& RequestQueue = PackageDatas->GetRequestQueue();
 	FPackageDataSet& RestartedRequests = RequestQueue.GetRestartedRequests();
 	TRingBuffer<FDiscoveryQueueElement>& DiscoveryQueue = RequestQueue.GetDiscoveryQueue();
+	TRingBuffer<FPackageData*>& BuildDependencyDiscoveryQueue = RequestQueue.GetBuildDependencyDiscoveryQueue();
 	TRingBuffer<FRequestCluster>& RequestClusters = RequestQueue.GetRequestClusters();
 	const FCookerTimer& CookerTimer = StackData.Timer;
 
@@ -2629,7 +2646,7 @@ void UCookOnTheFlyServer::PumpRequests(UE::Cook::FTickStackData& StackData, int3
 				for (FPackageData*& PackageData : RequestsToLoad)
 				{
 					check(PackageData->GetState() == EPackageState::Request);
-					if (PackageData->GetPlatformsNeedingCookingNum() == 0)
+					if (PackageData->GetPlatformsNeedingCommitNum(GetCookPhase()) == 0)
 					{
 						ESuppressCookReason SuppressCookReason = PackageData->HasAnyCookedPlatform() ? ESuppressCookReason::AlreadyCooked :
 							ESuppressCookReason::OnlyEditorOnly;
@@ -2649,15 +2666,46 @@ void UCookOnTheFlyServer::PumpRequests(UE::Cook::FTickStackData& StackData, int3
 		}
 		else if (!RestartedRequests.IsEmpty())
 		{
-			RequestClusters.Add(FRequestCluster(*this, MoveTemp(RestartedRequests)));
+			EReachability ClusterReachability = GetCookPhase() == ECookPhase::Cook
+				? EReachability::Runtime
+				: EReachability::Build;
+			RequestClusters.Add(FRequestCluster(*this, MoveTemp(RestartedRequests), ClusterReachability));
 			RestartedRequests.Empty();
 		}
 		else if (!DiscoveryQueue.IsEmpty())
 		{
-			FRequestCluster Cluster(*this, DiscoveryQueue);
-			if (Cluster.NumPackageDatas() > 0)
+			if (GetCookPhase() == ECookPhase::Cook)
+			{
+				FRequestCluster Cluster(*this, DiscoveryQueue);
+				if (Cluster.NumPackageDatas() > 0)
+				{
+					RequestClusters.Add(MoveTemp(Cluster));
+				}
+			}
+			else
+			{
+				UE_LOG(LogCook, Warning, TEXT("Packages were added to the runtime discovery queue after starting BuildDependencies phase. Ignoring them.")
+					TEXT(" The first added package: %s"),
+					*DiscoveryQueue.First().PackageData->GetPackageName().ToString());
+				DiscoveryQueue.Empty();
+			}
+		}
+		else if (!BuildDependencyDiscoveryQueue.IsEmpty())
+		{
+			if (GetCookPhase() == ECookPhase::Cook)
+			{
+				UE_LOG(LogCook, Error, TEXT("Packages were added to the build dependency discovery queue before starting BuildDependencies phase. Ignoring them.")
+					TEXT(" The first added package: %s"),
+					*BuildDependencyDiscoveryQueue.First()->GetPackageName().ToString());
+				BuildDependencyDiscoveryQueue.Empty();
+			}
+			else
 			{
-				RequestClusters.Add(MoveTemp(Cluster));
+				FRequestCluster Cluster(*this, FRequestCluster::BuildDependencyQueue, BuildDependencyDiscoveryQueue);
+				if (Cluster.NumPackageDatas() > 0)
+				{
+					RequestClusters.Add(MoveTemp(Cluster));
+				}
 			}
 		}
 		else
@@ -2684,10 +2732,10 @@ void UCookOnTheFlyServer::PumpRequests(UE::Cook::FTickStackData& StackData, int3
 		{
 			continue;
 		}
-		if (PackageData->GetPlatformsNeedingCookingNum() == 0)
+		if (PackageData->GetPlatformsNeedingCommitNum(GetCookPhase()) == 0)
 		{
-			ESuppressCookReason SuppressCookReason = PackageData->HasAnyCookedPlatform() ? ESuppressCookReason::AlreadyCooked :
-				ESuppressCookReason::OnlyEditorOnly;
+			ESuppressCookReason SuppressCookReason = PackageData->HasAnyCommittedPlatforms()
+				? ESuppressCookReason::AlreadyCooked : ESuppressCookReason::OnlyEditorOnly;
 			DemoteToIdle(*PackageData, ESendFlags::QueueAdd, SuppressCookReason);
 			continue;
 		}
@@ -2803,13 +2851,13 @@ void UCookOnTheFlyServer::DemoteToIdle(UE::Cook::FPackageData& PackageData, UE::
 			}
 
 			// Reachability: Suppress the diagnostic that were found via cookload reference traversal but are not reachable on the target platforms
-			bPrintDiagnostic &= (PackageData.HasInstigator() || Reason != ESuppressCookReason::OnlyEditorOnly);
+			bPrintDiagnostic &= (PackageData.HasInstigator(EReachability::Runtime) || Reason != ESuppressCookReason::OnlyEditorOnly);
 
 			if (bPrintDiagnostic)
 			{
 				UE_CLOG((GCookProgressDisplay & (int32)ECookProgressDisplayMode::Instigators), LogCook, Display,
 					TEXT("Cooking %s, Instigator: { %s } -> Rejected %s"), *PackageNameStr,
-					*(PackageData.GetInstigator().ToString()), LexToString(Reason));
+					*(PackageData.GetInstigator(EReachability::Runtime).ToString()), LexToString(Reason));
 				UE_CLOG(GCookProgressDisplay & (int32)ECookProgressDisplayMode::PackageNames, LogCook, Display,
 					TEXT("Cooking %s -> Rejected %s"), *PackageNameStr, LexToString(Reason));
 			}
@@ -2983,7 +3031,9 @@ void UCookOnTheFlyServer::LoadPackageInQueue(UE::Cook::FPackageData& PackageData
 			FPackageData& OtherPackageData = PackageDatas->AddPackageDataByPackageNameChecked(LoadedPackage->GetFName());
 			UE_LOG(LogCook, Verbose, TEXT("Request for %s received going to save %s"), *PackageFileName.ToString(),
 				*OtherPackageData.GetFileName().ToString());
-			QueueDiscoveredPackage(OtherPackageData, FInstigator(PackageData.GetInstigator()), EDiscoveredPlatformSet::CopyFromInstigator);
+			QueueDiscoveredPackage(OtherPackageData,
+				FInstigator(EInstigator::ForceExplorableSaveTimeSoftDependency, PackageData.GetPackageName()),
+				EDiscoveredPlatformSet::CopyFromInstigator);
 
 			PackageData.SetPlatformsCooked(PlatformManager->GetSessionPlatforms(), ECookResult::Succeeded);
 			RejectPackageToLoad(PackageData, TEXT("is redirected to another filename"), ESuppressCookReason::Redirected);
@@ -3048,7 +3098,7 @@ void UCookOnTheFlyServer::LoadPackageInQueue(UE::Cook::FPackageData& PackageData
 		}
 	}
 
-	if (PackageData.GetPlatformsNeedingCookingNum() == 0)
+	if (PackageData.GetPlatformsNeedingCommitNum(GetCookPhase()) == 0)
 	{
 		// Already cooked. This can happen if we needed to load a package that was previously cooked and garbage collected because it is a loaddependency of a new request.
 		// Send the package back to idle, nothing further to do with it.
@@ -3081,7 +3131,7 @@ void UCookOnTheFlyServer::RejectPackageToLoad(UE::Cook::FPackageData& PackageDat
 	// make sure this package doesn't exist
 	for (const TPair<const ITargetPlatform*, UE::Cook::FPackagePlatformData>& Pair : PackageData.GetPlatformDatas())
 	{
-		if (Pair.Key == CookerLoadingPlatformKey || !Pair.Value.NeedsCooking(Pair.Key))
+		if (Pair.Key == CookerLoadingPlatformKey || !Pair.Value.NeedsCommit(Pair.Key, GetCookPhase()))
 		{
 			continue;
 		}
@@ -3282,6 +3332,8 @@ void UCookOnTheFlyServer::QueueDiscoveredPackage(UE::Cook::FPackageData& Package
 
 	TArray<const ITargetPlatform*, TInlineAllocator<ExpectedMaxNumPlatforms>> BufferPlatforms;
 	TConstArrayView<const ITargetPlatform*> DiscoveredPlatforms;
+	EReachability DiscoveredReachability =
+		Instigator.Category == EInstigator::BuildDependency ? EReachability::Build : EReachability::Runtime;
 	if (!bSkipOnlyEditorOnly)
 	{
 		BufferPlatforms = PlatformManager->GetSessionPlatforms();
@@ -3291,23 +3343,13 @@ void UCookOnTheFlyServer::QueueDiscoveredPackage(UE::Cook::FPackageData& Package
 	else
 	{
 		DiscoveredPlatforms = ReachablePlatforms.GetPlatforms(*this, &Instigator,
-			TConstArrayView<const ITargetPlatform*>(), &BufferPlatforms);
-	}
-	if (Instigator.Category == EInstigator::BuildDependency)
-	{
-		if (PackageData.IsBuildDependency(DiscoveredPlatforms))
-		{
-			// Not a new discovery; ignore
-			return;
-		}
+			TConstArrayView<const ITargetPlatform*>(), DiscoveredReachability, &BufferPlatforms);
 	}
-	else if (Instigator.Category != EInstigator::ForceExplorableSaveTimeSoftDependency)
+	if (Instigator.Category != EInstigator::ForceExplorableSaveTimeSoftDependency &&
+		PackageData.HasReachablePlatforms(DiscoveredReachability, DiscoveredPlatforms))
 	{
-		if (PackageData.HasReachablePlatforms(DiscoveredPlatforms))
-		{
-			// Not a new discovery; ignore
-			return;
-		}
+		// Not a new discovery; ignore
+		return;
 	}
 
 	if (bHiddenDependenciesDebug)
@@ -3339,17 +3381,7 @@ void UCookOnTheFlyServer::QueueDiscoveredPackageOnDirector(UE::Cook::FPackageDat
 	if (!CookByTheBookOptions->bSkipHardReferences ||
 		(Instigator.Category == EInstigator::GeneratedPackage))
 	{
-		if (Instigator.Category == EInstigator::BuildDependency)
-		{
-			TArray<const ITargetPlatform*, TInlineAllocator<ExpectedMaxNumPlatforms>> BufferPlatforms;
-			TConstArrayView<const ITargetPlatform*> DiscoveredPlatforms = ReachablePlatforms.GetPlatforms(
-				*this, &Instigator, TConstArrayView<const ITargetPlatform*>(), &BufferPlatforms);
-			PackageData.AddBuildDependencyPlatforms(DiscoveredPlatforms, MoveTemp(Instigator));
-		}
-		else
-		{
-			PackageData.QueueAsDiscovered(MoveTemp(Instigator), MoveTemp(ReachablePlatforms), Urgency);
-		}
+		PackageData.QueueAsDiscovered(MoveTemp(Instigator), MoveTemp(ReachablePlatforms), Urgency);
 	}
 }
 
@@ -3422,19 +3454,20 @@ UE::Cook::EPollStatus UCookOnTheFlyServer::QueueGeneratedPackages(UE::Cook::FGen
 	{
 		GenerationHelper.StartQueueGeneratedPackages(*this);
 		TArray<const ITargetPlatform*, TInlineAllocator<ExpectedMaxNumPlatforms>> ReachablePlatforms;
-		PackageData.GetReachablePlatforms(ReachablePlatforms);
+		PackageData.GetReachablePlatforms(EReachability::Runtime, ReachablePlatforms);
 		for (const FCookGenerationInfo& ChildInfo: GenerationHelper.GetPackagesToGenerate())
 		{
 			FPackageData* ChildPackageData = ChildInfo.PackageData;
 			// Set the Instigator now rather than delaying it until the discovery queue is processed.
-			ChildPackageData->SetInstigator(GenerationHelper, FInstigator(EInstigator::GeneratedPackage, OwnerName));
+			ChildPackageData->SetInstigator(GenerationHelper, EReachability::Runtime,
+				FInstigator(EInstigator::GeneratedPackage, OwnerName));
 			// The urgency of generated packages must be at least as high as the generator to satisfy the contract of 
 			// making the generator urgent. By default they are High urgency rather than Normal so that they are saved 
 			// quickly, so that we release the memory used by their generator for them.
 			EUrgency Urgency = PackageData.GetUrgency() > EUrgency::High ? PackageData.GetUrgency() : EUrgency::High;
 
 			// Queue the package for cooking
-			QueueDiscoveredPackage(*ChildPackageData, FInstigator(ChildPackageData->GetInstigator()),
+			QueueDiscoveredPackage(*ChildPackageData, FInstigator(ChildPackageData->GetInstigator(EReachability::Runtime)),
 				EDiscoveredPlatformSet::CopyFromInstigator, Urgency, &GenerationHelper);
 		}
 		GenerationHelper.EndQueueGeneratedPackages(*this);
@@ -4478,21 +4511,27 @@ void UCookOnTheFlyServer::DumpCrashContext(FCrashContextExtendedWriter& Writer)
 void UCookOnTheFlyServer::ProcessUnsolicitedPackages(TArray<FName>* OutDiscoveredPackageNames,
 	TMap<FName, UE::Cook::FInstigator>* OutInstigators)
 {
+	using namespace UE::Cook;
+
 	if (bRunningAsShaderServer)
 	{
 		return;
 	}
 
-	using namespace UE::Cook;
+	if (GetCookPhase() != ECookPhase::Cook)
+	{
+		// We no longer add unsolicited packages to the cook when we have started committing build dependencies
+		return;
+	}
 
-	TMap<FName, UE::Cook::FInstigator> NewPackages = PackageTracker->GetNewPackages();
+	TMap<FName, FInstigator> NewPackages = PackageTracker->GetNewPackages();
 
-	for (TPair<FName, UE::Cook::FInstigator>& PackageWithInstigator : NewPackages)
+	for (TPair<FName, FInstigator>& PackageWithInstigator : NewPackages)
 	{
 		FName PackageName = PackageWithInstigator.Key;
 		FInstigator& Instigator = PackageWithInstigator.Value;
 
-		UE::Cook::FPackageData* PackageData = PackageDatas->TryAddPackageDataByPackageName(PackageName);
+		FPackageData* PackageData = PackageDatas->TryAddPackageDataByPackageName(PackageName);
 		if (!PackageData)
 		{
 			continue; // Getting the PackageData will fail if e.g. it is a script package
@@ -4542,7 +4581,7 @@ void UCookOnTheFlyServer::ProcessUnsolicitedPackages(TArray<FName>* OutDiscovere
 			// Queue it for discovery and do not log a warning about it.
 			QueueDiscoveredPackage(*PackageData, FInstigator(Instigator), EDiscoveredPlatformSet::CopyFromInstigator);
 		}
-		else if (PackageData->FindOrAddPlatformData(CookerLoadingPlatformKey).IsReachable())
+		else if (PackageData->FindOrAddPlatformData(CookerLoadingPlatformKey).IsReachable(EReachability::Runtime))
 		{
 			// This load was expected so we do not need to add a hidden dependency for it.
 			// If we are using legacy WhatGetsCookedRules, queue the package for cooking because it was loaded.
@@ -4572,6 +4611,19 @@ void UCookOnTheFlyServer::ProcessUnsolicitedPackages(TArray<FName>* OutDiscovere
 }
 
 void UCookOnTheFlyServer::PumpSaves(UE::Cook::FTickStackData& StackData, uint32 DesiredQueueLength, int32& OutNumPushed, bool& bOutBusy)
+{
+	if (GetCookPhase() == UE::Cook::ECookPhase::Cook)
+	{
+		PumpRuntimeSaves(StackData, DesiredQueueLength, OutNumPushed, bOutBusy);
+	}
+	else
+	{
+		// After we enter the build dependencies phase, we no longer save packages, we just commit them with dependencies but no cooked data.
+		PumpBuildDependencySaves(StackData, DesiredQueueLength, OutNumPushed, bOutBusy);
+	}
+}
+
+void UCookOnTheFlyServer::PumpRuntimeSaves(UE::Cook::FTickStackData& StackData, uint32 DesiredQueueLength, int32& OutNumPushed, bool& bOutBusy)
 {
 	using namespace UE::Cook;
 	OutNumPushed = 0;
@@ -4600,7 +4652,7 @@ void UCookOnTheFlyServer::PumpSaves(UE::Cook::FTickStackData& StackData, uint32
 
 		FPoppedPackageDataScope PoppedScope(PackageData);
 		UPackage* Package = PackageData.GetPackage();
-		
+
 		check(Package != nullptr);
 		++HandledCount;
 
@@ -4608,16 +4660,8 @@ void UCookOnTheFlyServer::PumpSaves(UE::Cook::FTickStackData& StackData, uint32
 		UE_LOG(LogCook, Display, TEXT("Processing save for package %s"), *Package->GetName());
 #endif
 
-		if (PackageTracker->NeverCookPackageList.Contains(Package->GetFName()))
-		{
-			// refuse to save this package, it's clearly one of the undesirables
-			DemoteToIdle(PackageData, ESendFlags::QueueAdd, ESuppressCookReason::NeverCook);
-			++OutNumPushed;
-			continue;
-		}
-
 		// Cook only the session platforms that have not yet been cooked for the given package
-		PackageData.GetPlatformsNeedingCooking(PlatformsForPackage);
+		PackageData.GetPlatformsNeedingCommit(PlatformsForPackage, GetCookPhase());
 		if (PlatformsForPackage.Num() == 0)
 		{
 			// We've already saved all possible platforms for this package; this should not be possible.
@@ -4774,7 +4818,7 @@ void UCookOnTheFlyServer::PumpSaves(UE::Cook::FTickStackData& StackData, uint32
 			}
 		}
 
- 		FSaveCookedPackageContext Context(*this, PackageData, PlatformsForPackage, StackData);
+		FSaveCookedPackageContext Context(*this, PackageData, PlatformsForPackage, StackData, EReachability::Runtime);
 		SaveCookedPackage(Context);
 		if (Context.bHasTimeOut)
 		{
@@ -4808,6 +4852,126 @@ void UCookOnTheFlyServer::PumpSaves(UE::Cook::FTickStackData& StackData, uint32
 	}
 }
 
+void UCookOnTheFlyServer::PumpBuildDependencySaves(UE::Cook::FTickStackData& StackData, uint32 DesiredQueueLength, int32& OutNumPushed, bool& bOutBusy)
+{
+	using namespace UE::Cook;
+	OutNumPushed = 0;
+	bOutBusy = false;
+
+	UE_SCOPED_HIERARCHICAL_COOKTIMER(SavingPackages);
+	check(IsInGameThread());
+	ON_SCOPE_EXIT
+	{
+		PumpHasExceededMaxMemory(StackData.ResultFlags);
+	};
+
+	// Commit as many packages as we can during our time slice
+	FPackageDataQueue& SaveQueue = PackageDatas->GetSaveQueue();
+	TArray<const ITargetPlatform*, TInlineAllocator<ExpectedMaxNumPlatforms>> PlatformsForPackage;
+	COOK_STAT(DetailedCookStats::PeakSaveQueueSize = FMath::Max(DetailedCookStats::PeakSaveQueueSize, SaveQueue.Num()));
+	while (SaveQueue.Num() > static_cast<int32>(DesiredQueueLength))
+	{
+		FPackageData& PackageData(*SaveQueue.PopFrontValue());
+		if (TryCreateRequestCluster(PackageData))
+		{
+			continue;
+		}
+
+		FPoppedPackageDataScope PoppedScope(PackageData);
+		UPackage* Package = PackageData.GetPackage();
+
+		check(Package != nullptr);
+
+		// Commit only the session platforms that are requested and have not yet been committed
+		PackageData.GetPlatformsNeedingCommit(PlatformsForPackage, GetCookPhase());
+		if (PlatformsForPackage.Num() == 0)
+		{
+			// We've already committed all possible platforms for this package; this should not be possible.
+			// All places that add a package to the save queue check for existence of incomplete platforms before adding
+			UE_LOG(LogCook, Warning, TEXT("Package '%s' in SaveQueue has no more platforms left to commit; this should not be possible!"),
+				*PackageData.GetFileName().ToString());
+			DemoteToIdle(PackageData, ESendFlags::QueueAdd, ESuppressCookReason::AlreadyCooked);
+			++OutNumPushed;
+			continue;
+		}
+
+		if (StackData.Timer.IsActionTimeUp())
+		{
+			SaveQueue.AddFront(&PackageData);
+			return;
+		}
+
+		FSaveCookedPackageContext Context(*this, PackageData, PlatformsForPackage, StackData, EReachability::Build);
+		CommitUncookedPackage(Context);
+		if (Context.bHasTimeOut)
+		{
+			// Timeouts can occur because of new objects created during the save, so we need to update our object cache,
+			// so we call ReleaseCookedPlatformData and ClearObjectCache to clear it and recache on next attempt.
+			check(PackageData.GetState() == EPackageState::SaveActive);
+			SaveQueue.Add(&PackageData);
+			continue;
+		}
+
+		PromoteToSaveComplete(PackageData, ESendFlags::QueueAdd);
+		++OutNumPushed;
+#if ENABLE_LOW_LEVEL_MEM_TRACKER
+		FLowLevelMemTracker::Get().UpdateStatsPerFrame();
+#endif
+	}
+}
+
+void UCookOnTheFlyServer::KickBuildDependencies(UE::Cook::FTickStackData& StackData)
+{
+	using namespace UE::Cook;
+
+	check(!IsCookWorkerMode()); // KickBuildDependencies should only be called on the Director.
+	if (CookDirector)
+	{
+		CookDirector->BroadcastMessage(FDirectorEventMessage(EDirectorEvent::KickBuildDependencies));
+	}
+	bKickedBuildDependencies = true;
+
+	if (!bHybridIterativeEnabled)
+	{
+		// Uncooked BuildDependencies are not committed in legacy cooks that are not using IncrementalCook
+		return;
+	}
+
+	// Switch to turn off the commit of uncooked build dependencies, in case it causes a performance problem.
+	constexpr bool bCommitUncookedBuildDependenciesEnabled = true;
+	if (!bCommitUncookedBuildDependenciesEnabled)
+	{
+		return;
+	}
+
+	FRequestQueue& RequestQueue = PackageDatas->GetRequestQueue();
+	// KickBuildDependencies is only called when !RequestQueue.HasRequestsToExplore(). We rely on that and give a warning
+	// during BuildDependencies phase if it is ever non-empty.
+	check(RequestQueue.GetDiscoveryQueue().IsEmpty());
+
+	TArray<FPackageData*> BuildPackages;
+	PackageDatas->LockAndEnumeratePackageDatas([&BuildPackages](FPackageData* PackageData)
+		{
+			// KickBuildDependencies is only called when no packages are in progress.
+			check(!PackageData->IsInProgress());
+			for (const TPair<const ITargetPlatform*, FPackagePlatformData>& PlatformPair :
+				PackageData->GetPlatformDatas())
+			{
+				const FPackagePlatformData& PlatformData = PlatformPair.Value;
+				if (PlatformData.IsReachable(EReachability::Build) && !PlatformData.IsCommitted())
+				{
+					BuildPackages.Add(PackageData);
+					break;
+				}
+			}
+		});
+	if (!BuildPackages.IsEmpty())
+	{
+		UE_LOG(LogCook, Display, TEXT("UncookedBuildDependencies: Queueing %d packages for load-only commit."), BuildPackages.Num());
+		PackageDatas->GetRequestQueue().GetBuildDependencyDiscoveryQueue().MoveAppendRange(BuildPackages.GetData(), BuildPackages.Num());
+	}
+}
+
 void UCookOnTheFlyServer::PostLoadPackageFixup(UE::Cook::FPackageData& PackageData, UPackage* Package)
 {
 	if (Package->ContainsMap() == false)
@@ -6011,6 +6175,7 @@ void UCookOnTheFlyServer::InitializeSession()
 
 	NumObjectsHistory.Initialize(GUObjectArray.GetObjectArrayNumMinusAvailable());
 	VirtualMemoryHistory.Initialize(FPlatformMemory::GetStats().UsedVirtual);
+	bKickedBuildDependencies = false;
 }
 
 bool UCookOnTheFlyServer::Exec_Editor(class UWorld* InWorld, const TCHAR* Cmd, FOutputDevice& Ar)
@@ -6096,6 +6261,11 @@ bool UCookOnTheFlyServer::Exec_Editor(class UWorld* InWorld, const TCHAR* Cmd, F
 }
 
 UE::Cook::FInstigator UCookOnTheFlyServer::GetInstigator(FName PackageName)
+{
+	return GetInstigator(PackageName, UE::Cook::EReachability::All);
+}
+
+UE::Cook::FInstigator UCookOnTheFlyServer::GetInstigator(FName PackageName, UE::Cook::EReachability Reachability)
 {
 	using namespace UE::Cook;
 
@@ -6104,7 +6274,7 @@ UE::Cook::FInstigator UCookOnTheFlyServer::GetInstigator(FName PackageName)
 	{
 		return FInstigator(EInstigator::NotYetRequested);
 	}
-	return PackageData->GetInstigator();
+	return PackageData->GetInstigator(Reachability);
 }
 
 TArray<UE::Cook::FInstigator> UCookOnTheFlyServer::GetInstigatorChain(FName PackageName)
@@ -6122,7 +6292,7 @@ TArray<UE::Cook::FInstigator> UCookOnTheFlyServer::GetInstigatorChain(FName Pack
 			Result.Add(FInstigator(EInstigator::NotYetRequested));
 			return Result;
 		}
-		const FInstigator& Last = Result.Add_GetRef(PackageData->GetInstigator());
+		const FInstigator& Last = Result.Add_GetRef(PackageData->GetInstigator(EReachability::All));
 		bool bGetNext = false;
 		switch (Last.Category)
 		{
@@ -9019,35 +9189,12 @@ void UCookOnTheFlyServer::CookByTheBookFinishedInternal()
 	check(IsInSession());
 	check(PackageDatas->GetRequestQueue().IsEmpty());
 	check(PackageDatas->GetRequestQueue().GetDiscoveryQueue().IsEmpty());
+	check(PackageDatas->GetRequestQueue().GetBuildDependencyDiscoveryQueue().IsEmpty());
 	check(PackageDatas->GetAssignedToWorkerSet().IsEmpty());
 	check(PackageDatas->GetLoadQueue().IsEmpty());
 	check(PackageDatas->GetSaveQueue().IsEmpty());
 	check(PackageDatas->GetSaveStalledSet().IsEmpty());
 
-	if (bHybridIterativeEnabled)
-	{
-		// Report the packages that are build dependencies and that were not saved
-		// TODO: cook these packages instead of just reporting them
-		TArray<FPackageData*> UncookedBuildDependencies;
-		PackageDatas->LockAndEnumeratePackageDatas([&UncookedBuildDependencies](FPackageData* PackageData)
-			{
-				for (const TPair<const ITargetPlatform*, FPackagePlatformData>& PlatformPair :
-					PackageData->GetPlatformDatas())
-				{
-					const FPackagePlatformData& Data = PlatformPair.Value;
-					if (Data.IsBuildDependency())
-					{
-						if (!Data.IsCookAttempted())
-						{
-							UncookedBuildDependencies.Add(PackageData);
-							break;
-						}
-					}
-				}
-			});
-		UE_LOG(LogCook, Display, TEXT("UncookedBuildDependencies: %d packages."), UncookedBuildDependencies.Num());
-	}
-
 	TArray<FPackageData*> DanglingGenerationHelpers;
 	PackageDatas->LockAndEnumeratePackageDatas([&DanglingGenerationHelpers](FPackageData* PackageData)
 		{
@@ -9172,7 +9319,7 @@ void UCookOnTheFlyServer::CookByTheBookFinishedInternal()
 		{
 			UE_SCOPED_HIERARCHICAL_COOKTIMER(CalculateReferencedSet);
 
-			PackageDatas->GetCookedPackagesForPlatform(TargetPlatform, CookedPackageDatas, IgnorePackageDatas);
+			PackageDatas->GetCommittedPackagesForPlatform(TargetPlatform, CookedPackageDatas, IgnorePackageDatas);
 
 			if (IsCookingDLC())
 			{
@@ -9485,7 +9632,8 @@ TMap<FName, TSet<FName>> UCookOnTheFlyServer::BuildMapDependencyGraph(const ITar
 	TMap<FName, TSet<FName>> MapDependencyGraph;
 
 	TArray<UE::Cook::FPackageData*> PlatformCookedPackages;
-	PackageDatas->GetCookedPackagesForPlatform(TargetPlatform, PlatformCookedPackages, PlatformCookedPackages);
+	TArray<UE::Cook::FPackageData*> FailedPackages;
+	PackageDatas->GetCommittedPackagesForPlatform(TargetPlatform, PlatformCookedPackages, FailedPackages);
 
 	// assign chunks for all the map packages
 	for (const UE::Cook::FPackageData* const CookedPackage : PlatformCookedPackages)
@@ -9560,7 +9708,7 @@ void UCookOnTheFlyServer::GenerateCachedEditorThumbnails()
 			TArray<UE::Cook::FPackageData*> CookedPackages;
 			{
 				TArray<UE::Cook::FPackageData*> FailedPackages;
-				PackageDatas->GetCookedPackagesForPlatform(Platform, CookedPackages, FailedPackages);
+				PackageDatas->GetCommittedPackagesForPlatform(Platform, CookedPackages, FailedPackages);
 			}
 
 			for (const UE::Cook::FPackageData* const CookedPackage : CookedPackages)
@@ -9699,6 +9847,7 @@ void UCookOnTheFlyServer::CancelAllQueues()
 		DemoteToIdle(*PackageData, ESendFlags::QueueAddAndRemove, ESuppressCookReason::CookCanceled);
 	}
 	RequestQueue.GetDiscoveryQueue().Empty();
+	RequestQueue.GetBuildDependencyDiscoveryQueue().Empty();
 	TRingBuffer<FRequestCluster>& RequestClusters = RequestQueue.GetRequestClusters();
 	for (FRequestCluster& RequestCluster : RequestClusters)
 	{
@@ -9749,7 +9898,7 @@ void UCookOnTheFlyServer::ResetCook(TConstArrayView<TPair<const ITargetPlatform*
 
 	PackageDatas->LockAndEnumeratePackageDatas([TargetPlatforms](UE::Cook::FPackageData* PackageData)
 	{
-		PackageData->FindOrAddPlatformData(CookerLoadingPlatformKey).ResetReachable();
+		PackageData->FindOrAddPlatformData(CookerLoadingPlatformKey).ResetReachable(EReachability::All);
 
 		for (const TPair<const ITargetPlatform*, bool>& Pair : TargetPlatforms)
 		{
@@ -9758,8 +9907,7 @@ void UCookOnTheFlyServer::ResetCook(TConstArrayView<TPair<const ITargetPlatform*
 			if (PlatformData)
 			{
 				bool bResetResults = Pair.Value;
-				PlatformData->ResetReachable();
-				PlatformData->SetIsBuildDependency(false);
+				PlatformData->ResetReachable(EReachability::All);
 				if (bResetResults)
 				{
 					PackageData->ClearCookResults(TargetPlatform);
@@ -10230,7 +10378,7 @@ UE::Cook::FCookSavePackageContext* UCookOnTheFlyServer::CreateSaveContext(const
 					UClass* Class = Cast<UClass>(Object);
 					if (Class && Class->IsNative() && SuppressedNativeScriptPackages.Contains(Class->GetPackage()->GetFName()))
 					{
-						FInstigator Instigator = GetInstigator(ValidationContext.Package->GetFName());
+						FInstigator Instigator = GetInstigator(ValidationContext.Package->GetFName(), EReachability::Runtime);
 						bool bIsError = true;
 						if (Instigator.Category == EInstigator::StartupPackage)
 						{
diff --git a/Engine/Source/Editor/UnrealEd/Private/Cooker/CookDirector.cpp b/Engine/Source/Editor/UnrealEd/Private/Cooker/CookDirector.cpp
--- a/Engine/Source/Editor/UnrealEd/Private/Cooker/CookDirector.cpp
+++ b/Engine/Source/Editor/UnrealEd/Private/Cooker/CookDirector.cpp
@@ -644,11 +644,16 @@ void FCookDirector::RemoveFromWorker(FPackageData& PackageData)
 	bForceNextHeartbeat = true;
 }
 
-void FCookDirector::BroadcastGeneratorMessage(FGeneratorEventMessage&& Message)
+void FCookDirector::BroadcastMessage(const IMPCollectorMessage& Message, ECookBroadcastTiming Timing)
+{
+	BroadcastMessage(MarshalToCompactBinaryTCP(Message), Timing);
+}
+
+void FCookDirector::BroadcastMessage(UE::CompactBinaryTCP::FMarshalledMessage&& Message, ECookBroadcastTiming Timing)
 {
 	if (bReceivingMessages)
 	{
-		QueuedGeneratorBroadcasts.Add(MoveTemp(Message));
+		QueuedBroadcasts.Emplace(MoveTemp(Message), Timing);
 		return;
 	}
 
@@ -659,7 +664,23 @@ void FCookDirector::BroadcastGeneratorMessage(FGeneratorEventMessage&& Message)
 	TArray<TRefCountPtr<FCookWorkerServer>> LocalRemoteWorkers = CopyRemoteWorkers();
 	for (TRefCountPtr<FCookWorkerServer>& RemoteWorker : LocalRemoteWorkers)
 	{
-		RemoteWorker->AppendMessage(Message, ECookDirectorThread::SchedulerThread);
+		UE::CompactBinaryTCP::FMarshalledMessage CopyMessage(Message);
+		switch (Timing)
+		{
+		case ECookBroadcastTiming::Immediate:
+		{
+			RemoteWorker->SendMessage(MoveTemp(CopyMessage), ECookDirectorThread::SchedulerThread);
+			break;
+		}
+		case ECookBroadcastTiming::AfterAssignPackages:
+		{
+			RemoteWorker->AppendMessage(MoveTemp(CopyMessage), ECookDirectorThread::SchedulerThread);
+			break;
+		}
+		default:
+			checkNoEntry();
+			break;
+		}
 	}
 }
 
@@ -755,11 +776,11 @@ void FCookDirector::TickFromSchedulerThread()
 		}
 		bReceivingMessages = false;
 	}
-	for (FGeneratorEventMessage& Message : QueuedGeneratorBroadcasts)
+	for (TPair<UE::CompactBinaryTCP::FMarshalledMessage, ECookBroadcastTiming>& Pair : QueuedBroadcasts)
 	{
-		BroadcastGeneratorMessage(MoveTemp(Message));
+		BroadcastMessage(MoveTemp(Pair.Key), Pair.Value);
 	}
-	QueuedGeneratorBroadcasts.Empty();
+	QueuedBroadcasts.Empty();
 
 	// Process any queued messages
 	WorkersWithMessage.Empty();
@@ -2275,4 +2296,26 @@ bool FRetractionResultsMessage::TryRead(FCbObjectView Object)
 
 FGuid FRetractionResultsMessage::MessageType(TEXT("CBFB840A4FB94903A757C490514A4B86"));
 
+void FDirectorEventMessage::Write(FCbWriter& Writer) const
+{
+	static_assert(sizeof(EDirectorEvent) <= sizeof(uint8), "We are storing it in a uint8");
+	Writer << "T" << (uint8)Event;
+}
+bool FDirectorEventMessage::TryRead(FCbObjectView Object)
+{
+	bool bOk = true;
+	uint8 EventInt = Object["T"].AsUInt8(static_cast<uint8>(EDirectorEvent::Count));
+	if (EventInt < static_cast<uint8>(EDirectorEvent::Count))
+	{
+		Event = static_cast<EDirectorEvent>(EventInt);
+	}
+	else
+	{
+		bOk = false;
+	}
+	return bOk;
+}
+
+FGuid FDirectorEventMessage::MessageType(TEXT("58A03A38FEE045B08331BB8457AEBE35"));
+
 }
diff --git a/Engine/Source/Editor/UnrealEd/Private/Cooker/CookDirector.h b/Engine/Source/Editor/UnrealEd/Private/Cooker/CookDirector.h
--- a/Engine/Source/Editor/UnrealEd/Private/Cooker/CookDirector.h
+++ b/Engine/Source/Editor/UnrealEd/Private/Cooker/CookDirector.h
@@ -2,11 +2,11 @@
 
 #pragma once
 
-#include "CompactBinaryTCP.h"
 #include "Containers/Array.h"
 #include "Containers/ArrayView.h"
 #include "Containers/Map.h"
 #include "Containers/UnrealString.h"
+#include "Cooker/CompactBinaryTCP.h"
 #include "Cooker/CookSockets.h"
 #include "Cooker/CookTypes.h"
 #include "Cooker/MPCollector.h"
@@ -30,6 +30,7 @@ class UCookOnTheFlyServer;
 namespace UE::Cook { class FCookWorkerServer; }
 namespace UE::Cook { struct FAssignPackageExtraData; }
 namespace UE::Cook { struct FCookWorkerProfileData; }
+namespace UE::Cook { struct FDirectorEventMessage; }
 namespace UE::Cook { struct FGeneratorEventMessage; }
 namespace UE::Cook { struct FGenerationHelper; }
 namespace UE::Cook { struct FHeartbeatMessage; }
@@ -54,6 +55,16 @@ enum class ECookDirectorThread : uint8
 	Invalid,
 };
 
+/**
+ * Timings for messages sent to CookWorker or CookDirector through the public Broadcast functions.
+ * The timing indicates at which point in the network tick the message should be sent.
+ */
+enum class ECookBroadcastTiming : uint8
+{
+	Immediate,
+	AfterAssignPackages,
+};
+
 /**
  * Helper for CookOnTheFlyServer that sends requests to CookWorker processes for load/save and merges
  * their replies into the local process's cook results.
@@ -78,10 +89,14 @@ public:
 	void RemoveFromWorker(FPackageData& PackageData);
 
 	/**
-	 * Report EGeneratorEvent::QueuedGeneratedPackagesFencePassed to all CookWorkers so any of them
-	 * waiting on the fence to be passed before clearing their data can clear their data.
+	 * Send a message to all CookWorkers to match a state change on the Director.
+	 * The message will be sent at the time indicated by the Timing argument; either Immediate, before any other messages
+	 * from the director that are triggered after the call, or at the specified later time.
 	 */
-	void BroadcastGeneratorMessage(FGeneratorEventMessage&& Message);
+	void BroadcastMessage(const IMPCollectorMessage& Message,
+		ECookBroadcastTiming Timing = ECookBroadcastTiming::Immediate);
+	void BroadcastMessage(UE::CompactBinaryTCP::FMarshalledMessage&& Message,
+		ECookBroadcastTiming Timing = ECookBroadcastTiming::Immediate);
 
 	/** Periodic tick function. Sends/Receives messages to CookWorkers. */
 	void TickFromSchedulerThread();
@@ -245,7 +260,7 @@ private:
 	TArray<FCookWorkerProfileData> RemoteWorkerProfileDatas;
 	TArray<FPendingConnection> PendingConnections;
 	TUniquePtr<FCookWorkerProfileData> LocalWorkerProfileData;
-	TArray<FGeneratorEventMessage> QueuedGeneratorBroadcasts;
+	TArray<TPair<UE::CompactBinaryTCP::FMarshalledMessage, ECookBroadcastTiming>> QueuedBroadcasts;
 	UCookOnTheFlyServer& COTFS;
 	double WorkersStalledStartTimeSeconds = 0.;
 	double WorkersStalledWarnTimeSeconds = 0.;
@@ -344,4 +359,34 @@ public:
 	static FGuid MessageType;
 };
 
+/** Director status Events that can be broadcast from CookDirector to CookWorkers. */
+enum class EDirectorEvent : uint8
+{
+	KickBuildDependencies,
+	Count,
+};
+/** Message sent from CookDirector to CookWorkers to notify them of an EDirectorEvent. */
+struct FDirectorEventMessage : public IMPCollectorMessage
+{
+	explicit FDirectorEventMessage(EDirectorEvent InEvent = EDirectorEvent::Count)
+		: Event(InEvent)
+	{
+	}
+
+	virtual void Write(FCbWriter& Writer) const override;
+	virtual bool TryRead(FCbObjectView Object) override;
+	virtual FGuid GetMessageType() const override
+	{
+		return MessageType;
+	}
+	virtual const TCHAR* GetDebugName() const override
+	{
+		return TEXT("DirectorEventMessage");
+	}
+
+public:
+	EDirectorEvent Event;
+	static FGuid MessageType;
+};
+
 } // namespace UE::Cook
\ No newline at end of file
diff --git a/Engine/Source/Editor/UnrealEd/Private/Cooker/CookGenerationHelper.cpp b/Engine/Source/Editor/UnrealEd/Private/Cooker/CookGenerationHelper.cpp
--- a/Engine/Source/Editor/UnrealEd/Private/Cooker/CookGenerationHelper.cpp
+++ b/Engine/Source/Editor/UnrealEd/Private/Cooker/CookGenerationHelper.cpp
@@ -208,7 +208,7 @@ void FGenerationHelper::ModifyNumSaved(int32 Delta)
 			{
 				FName PackageName = GetOwner().GetPackageName();
 				FGeneratorEventMessage Message(EGeneratorEvent::AllSavesCompleted, PackageName);
-				COTFS.CookDirector->BroadcastGeneratorMessage(MoveTemp(Message));
+				COTFS.CookDirector->BroadcastMessage(MoveTemp(Message), ECookBroadcastTiming::AfterAssignPackages);
 			}
 			OnAllSavesCompleted(COTFS);
 		}
@@ -771,9 +771,6 @@ bool FGenerationHelper::TryGenerateList()
 	}
 	PackagesToGenerate.Reserve(GeneratorDatas.Num());
 
-	TArray<const ITargetPlatform*, TInlineAllocator<1>> PlatformsToCook;
-	OwnerPackageData.GetPlatformsNeedingCooking(PlatformsToCook);
-
 	NumSaved = 0;
 	for (ICookPackageSplitter::FGeneratedPackage& SplitterData : GeneratorDatas)
 	{
@@ -1049,7 +1046,7 @@ void FGenerationHelper::StartQueueGeneratedPackages(UCookOnTheFlyServer& COTFS)
 
 		FPackageData& OwnerPackageData = GetOwner();
 		TArray<const ITargetPlatform*, TInlineAllocator<1>> PlatformsToCook;
-		OwnerPackageData.GetPlatformsNeedingCooking(PlatformsToCook);
+		OwnerPackageData.GetPlatformsNeedingCommit(PlatformsToCook, COTFS.GetCookPhase());
 
 		int32 NumIterativeUnmodified = 0;
 		int32 NumIterativeModified = 0;
@@ -1178,7 +1175,7 @@ void FGenerationHelper::OnRequestFencePassed(UCookOnTheFlyServer& COTFS)
 		{
 			FName PackageName = GetOwner().GetPackageName();
 			FGeneratorEventMessage Message(EGeneratorEvent::QueuedGeneratedPackagesFencePassed, PackageName);
-			COTFS.CookDirector->BroadcastGeneratorMessage(MoveTemp(Message));
+			COTFS.CookDirector->BroadcastMessage(MoveTemp(Message), ECookBroadcastTiming::AfterAssignPackages);
 		}
 		OnQueuedGeneratedPackagesFencePassed(COTFS);
 	}
diff --git a/Engine/Source/Editor/UnrealEd/Private/Cooker/CookPackageData.cpp b/Engine/Source/Editor/UnrealEd/Private/Cooker/CookPackageData.cpp
--- a/Engine/Source/Editor/UnrealEd/Private/Cooker/CookPackageData.cpp
+++ b/Engine/Source/Editor/UnrealEd/Private/Cooker/CookPackageData.cpp
@@ -56,38 +56,49 @@ static FAutoConsoleVariableRef CVarPollAsyncPeriod(
 //////////////////////////////////////////////////////////////////////////
 // FPackageData
 FPackagePlatformData::FPackagePlatformData()
-	: bReachable(0), bIsBuildDependency(0), bVisitedByCluster(0), bSaveTimedOut(0), bCookable(1), bExplorable(1)
-	, bExplorableOverride(0) , bIterativelyUnmodified(0), bIterativelySkipped(0), bRegisteredForCachedObjectsInOuter(0)
+	: Reachability((uint8)EReachability::None), ReachabilityVisitedByCluster((uint8)EReachability::None)
+	, bSaveTimedOut(0), bCookable(1), bExplorable(1), bExplorableOverride(0) , bIterativelyUnmodified(0)
+	, bIterativelySkipped(0), bRegisteredForCachedObjectsInOuter(0), bCommitted(0)
 	, CookResults((uint8)ECookResult::NotAttempted)
 {
 }
 
-void FPackagePlatformData::ResetReachable()
+void FPackagePlatformData::ResetReachable(EReachability InReachability)
 {
-	SetReachable(false);
-	SetVisitedByCluster(false);
-	SetCookable(true);
-	SetExplorable(true);
+	ClearReachability(InReachability);
+	ClearVisitedByCluster(InReachability);
+	if (EnumHasAnyFlags(InReachability, EReachability::Runtime))
+	{
+		SetCookable(true);
+		SetExplorable(true);
+	}
 }
 
 void FPackagePlatformData::MarkAsExplorable()
 {
-	ResetReachable();
+	ResetReachable(EReachability::Runtime);
 	SetExplorableOverride(true);
 }
 
-void FPackagePlatformData::MarkCookableForWorker(FCookWorkerClient& CookWorkerClient)
+void FPackagePlatformData::MarkCommittableForWorker(EReachability InReachability, FCookWorkerClient& CookWorkerClient)
 {
-	SetReachable(true);
-	SetVisitedByCluster(true);
-	SetExplorable(true);
-	SetCookable(true);
+	AddReachability(InReachability);
+	AddVisitedByCluster(InReachability);
+	if (EnumHasAnyFlags(InReachability, EReachability::Runtime))
+	{
+		SetExplorable(true);
+		SetCookable(true);
+	}
+	SetCommitted(false);
 	SetCookResults(ECookResult::NotAttempted);
 }
 
-bool FPackagePlatformData::NeedsCooking(const ITargetPlatform* PlatformItBelongsTo) const
+bool FPackagePlatformData::NeedsCommit(const ITargetPlatform* PlatformItBelongsTo, EReachability InReachability) const
 {
-	return IsReachable() && PlatformItBelongsTo != CookerLoadingPlatformKey && IsCookable() && !IsCookAttempted();
+	return !IsCommitted()
+		&& PlatformItBelongsTo != CookerLoadingPlatformKey
+		&& IsReachable(InReachability)
+		&& !(InReachability == EReachability::Runtime && !IsCookable());
 }
 
 FPackageData::FPackageData(FPackageDatas& PackageDatas, const FName& InPackageName, const FName& InFileName)
@@ -140,12 +151,12 @@ void FPackageData::ClearReferences()
 	}
 }
 
-int32 FPackageData::GetPlatformsNeedingCookingNum() const
+int32 FPackageData::GetPlatformsNeedingCommitNum(EReachability Reachability) const
 {
 	int32 Result = 0;
 	for (const TPair<const ITargetPlatform*, FPackagePlatformData>& Pair : PlatformDatas)
 	{
-		if (Pair.Value.NeedsCooking(Pair.Key))
+		if (Pair.Value.NeedsCommit(Pair.Key, Reachability))
 		{
 			++Result;
 		}
@@ -153,13 +164,14 @@ int32 FPackageData::GetPlatformsNeedingCookingNum() const
 	return Result;
 }
 
-bool FPackageData::IsPlatformVisitedByCluster(const ITargetPlatform* Platform) const
+bool FPackageData::IsPlatformVisitedByCluster(const ITargetPlatform* Platform, EReachability InReachability) const
 {
 	const FPackagePlatformData* PlatformData = FindPlatformData(Platform);
-	return PlatformData && PlatformData->IsVisitedByCluster();
+	return PlatformData && PlatformData->IsVisitedByCluster(InReachability);
 }
 
-bool FPackageData::HasReachablePlatforms(const TArrayView<const ITargetPlatform* const>& Platforms) const
+bool FPackageData::HasReachablePlatforms(EReachability InReachability,
+	const TArrayView<const ITargetPlatform* const>& Platforms) const
 {
 	if (Platforms.Num() == 0)
 	{
@@ -173,7 +185,7 @@ bool FPackageData::HasReachablePlatforms(const TArrayView<const ITargetPlatform*
 	for (const ITargetPlatform* QueryPlatform : Platforms)
 	{
 		const FPackagePlatformData* PlatformData = PlatformDatas.Find(QueryPlatform);
-		if (!PlatformData || !PlatformData->IsReachable())
+		if (!PlatformData || !PlatformData->IsReachable(InReachability))
 		{
 			return false;
 		}
@@ -181,11 +193,11 @@ bool FPackageData::HasReachablePlatforms(const TArrayView<const ITargetPlatform*
 	return true;
 }
 
-bool FPackageData::AreAllReachablePlatformsVisitedByCluster() const
+bool FPackageData::AreAllReachablePlatformsVisitedByCluster(EReachability InReachability) const
 {
 	for (const TPair<const ITargetPlatform*, FPackagePlatformData>& Pair : PlatformDatas)
 	{
-		if (Pair.Value.IsReachable() && !Pair.Value.IsVisitedByCluster())
+		if (EnumHasAllFlags(Pair.Value.GetReachability(), InReachability) && !Pair.Value.IsVisitedByCluster(InReachability))
 		{
 			return false;
 		}
@@ -193,19 +205,19 @@ bool FPackageData::AreAllReachablePlatformsVisitedByCluster() const
 	return true;
 }
 
-void FPackageData::GetReachablePlatformsForInstigator(UCookOnTheFlyServer& COTFS, FName InInstigator,
-	TArray<const ITargetPlatform*>& Platforms)
+void FPackageData::GetReachablePlatformsForInstigator(EReachability InReachability, UCookOnTheFlyServer& COTFS,
+	FName InInstigator, TArray<const ITargetPlatform*>& Platforms)
 {
-	return GetReachablePlatformsForInstigator(COTFS,
+	return GetReachablePlatformsForInstigator(InReachability, COTFS,
 		COTFS.PackageDatas->TryAddPackageDataByPackageName(InInstigator), Platforms);
 }
 
-void FPackageData::GetReachablePlatformsForInstigator(UCookOnTheFlyServer& COTFS, UE::Cook::FPackageData* InInstigator,
-	TArray<const ITargetPlatform*>& Platforms)
+void FPackageData::GetReachablePlatformsForInstigator(EReachability InReachability, UCookOnTheFlyServer& COTFS,
+	UE::Cook::FPackageData* InInstigator, TArray<const ITargetPlatform*>& Platforms)
 {
 	if (InInstigator)
 	{
-		InInstigator->GetReachablePlatforms(Platforms);
+		InInstigator->GetReachablePlatforms(InReachability, Platforms);
 	}
 	else
 	{
@@ -215,14 +227,13 @@ void FPackageData::GetReachablePlatformsForInstigator(UCookOnTheFlyServer& COTFS
 	}
 }
 
-
-void FPackageData::AddReachablePlatforms(FRequestCluster& RequestCluster,
+void FPackageData::AddReachablePlatforms(FRequestCluster& RequestCluster, EReachability InReachability,
 	TConstArrayView<const ITargetPlatform*> Platforms, FInstigator&& InInstigator)
 {
-	AddReachablePlatformsInternal(*this, Platforms, MoveTemp(InInstigator));
+	AddReachablePlatformsInternal(*this, InReachability, Platforms, MoveTemp(InInstigator));
 }
 
-void FPackageData::AddReachablePlatformsInternal(FPackageData& PackageData,
+void FPackageData::AddReachablePlatformsInternal(FPackageData& PackageData, EReachability InReachability,
 	TConstArrayView<const ITargetPlatform*> Platforms, FInstigator&& InInstigator)
 {
 	// This is a static helper function to make it impossible to make a typo and use this->Instigator instead of
@@ -231,57 +242,15 @@ void FPackageData::AddReachablePlatformsInternal(FPackageData& PackageData,
 	for (const ITargetPlatform* Platform : Platforms)
 	{
 		FPackagePlatformData& PlatformData = PackageData.FindOrAddPlatformData(Platform);
-		bSessionPlatformModified |= (Platform != CookerLoadingPlatformKey && !PlatformData.IsReachable());
-		PlatformData.SetReachable(true);
+		bSessionPlatformModified |= (Platform != CookerLoadingPlatformKey && !PlatformData.IsReachable(InReachability));
+		PlatformData.AddReachability(InReachability);
 	}
 	if (bSessionPlatformModified)
 	{
-		PackageData.SetInstigatorInternal(MoveTemp(InInstigator));
-	}
-}
-
-void FPackageData::AddBuildDependencyPlatforms(TConstArrayView<const ITargetPlatform*> Platforms,
-	FInstigator&& InInstigator)
-{
-	if (Platforms.IsEmpty())
-	{
-		return;
-	}
-	if (BuildInstigator.Category == EInstigator::NotYetRequested)
-	{
-		BuildInstigator = MoveTemp(InInstigator);
-	}
-
-	for (const ITargetPlatform* Platform : Platforms)
-	{
-		FPackagePlatformData& PlatformData = FindOrAddPlatformData(Platform);
-		PlatformData.SetIsBuildDependency(true);
+		PackageData.SetInstigatorInternal(InReachability, MoveTemp(InInstigator));
 	}
 }
 
-bool FPackageData::IsBuildDependency(const TArrayView<const ITargetPlatform* const>& Platforms) const
-{
-	if (Platforms.Num() == 0)
-	{
-		return true;
-	}
-	if (PlatformDatas.Num() == 0)
-	{
-		return false;
-	}
-
-	for (const ITargetPlatform* QueryPlatform : Platforms)
-	{
-		const FPackagePlatformData* PlatformData = PlatformDatas.Find(QueryPlatform);
-		if (!PlatformData || !PlatformData->IsBuildDependency())
-		{
-			return false;
-		}
-	}
-	return true;
-}
-
-
 void FPackageData::QueueAsDiscovered(FInstigator&& InInstigator, FDiscoveredPlatformSet&& ReachablePlatforms,
 	EUrgency InUrgency)
 {
@@ -293,8 +262,48 @@ void FPackageData::QueueAsDiscoveredInternal(FPackageData& PackageData, FInstiga
 {
 	// This is a static helper function to make it impossible to make a typo and use this->Instigator instead of
 	// InInstigator
-	TRingBuffer<FDiscoveryQueueElement>& Queue = PackageData.PackageDatas.GetRequestQueue().GetDiscoveryQueue();
-	Queue.Add(FDiscoveryQueueElement{ &PackageData, MoveTemp(InInstigator), MoveTemp(ReachablePlatforms), InUrgency });
+	FPackageDatas& LocalPackageDatas = PackageData.PackageDatas;
+	FRequestQueue& RequestQueue = LocalPackageDatas.GetRequestQueue();
+	UCookOnTheFlyServer& COTFS = LocalPackageDatas.GetCookOnTheFlyServer();
+	if (InInstigator.Category != EInstigator::BuildDependency)
+	{
+		TRingBuffer<FDiscoveryQueueElement>& Queue = RequestQueue.GetDiscoveryQueue();
+		if (COTFS.GetCookPhase() == ECookPhase::BuildDependencies)
+		{
+			UE_LOG(LogCook, Warning, TEXT("Package was added to the runtime discovery queue after starting BuildDependencies phase.")
+				TEXT("\n\tPackage: %s"), *PackageData.GetPackageName().ToString());
+			constexpr int32 MaxCount = 5;
+			static int32 Count = 0;
+			if (Count++ < MaxCount)
+			{
+				FDebug::DumpStackTraceToLog(ELogVerbosity::Warning);
+			}
+		}
+		Queue.Add(FDiscoveryQueueElement{ &PackageData, MoveTemp(InInstigator), MoveTemp(ReachablePlatforms), InUrgency });
+	}
+	else
+	{
+		// Build dependencies always immediately mark the package as being reachable, rather than needing to wait for
+		// the discoveryqueue. Waiting for the discovery queue is only necessary for runtime dependencies because
+		// we need to know whether the package was expected.
+		TArray<const ITargetPlatform*, TInlineAllocator<ExpectedMaxNumPlatforms>> BufferPlatforms;
+		TConstArrayView<const ITargetPlatform*> PlatformArray =
+			ReachablePlatforms.GetPlatforms(COTFS, &InInstigator, TConstArrayView<const ITargetPlatform*>(),
+				EReachability::Build, &BufferPlatforms);
+		bool bHasNewPlatforms = !PackageData.HasAllCommittedPlatforms(PlatformArray);
+		if (bHasNewPlatforms)
+		{
+			AddReachablePlatformsInternal(PackageData, EReachability::Build, PlatformArray, MoveTemp(InInstigator));
+
+			// If we have already kicked build dependencies, send the package to the discovery queue.
+			// Otherwise it will be added to the discoveryqueue when we kick build dependencies, if it hasn't been committed by then.
+			if (COTFS.GetCookPhase() == ECookPhase::BuildDependencies)
+			{
+				TRingBuffer<FPackageData*>& Queue = RequestQueue.GetBuildDependencyDiscoveryQueue();
+				Queue.Add(&PackageData);
+			}
+		}
+	}
 }
 
 void FPackageData::SetUrgency(EUrgency NewUrgency, ESendFlags SendFlags, bool bAllowUrgencyInIdle)
@@ -329,26 +338,27 @@ void FPackageData::SetIsCookLast(bool bValue)
 	}
 }
 
-void FPackageData::SetInstigator(FRequestCluster& Cluster, FInstigator&& InInstigator)
+void FPackageData::SetInstigator(FRequestCluster& Cluster, EReachability InReachability, FInstigator&& InInstigator)
 {
-	SetInstigatorInternal(MoveTemp(InInstigator));
+	SetInstigatorInternal(InReachability, MoveTemp(InInstigator));
 }
 
-void FPackageData::SetInstigator(FCookWorkerClient& Client, FInstigator&& InInstigator)
+void FPackageData::SetInstigator(FCookWorkerClient& Client, EReachability InReachability, FInstigator&& InInstigator)
 {
-	SetInstigatorInternal(MoveTemp(InInstigator));
+	SetInstigatorInternal(InReachability, MoveTemp(InInstigator));
 }
 
-void FPackageData::SetInstigator(FGenerationHelper& InGenerationHelper, FInstigator&& InInstigator)
+void FPackageData::SetInstigator(FGenerationHelper& InHelper, EReachability InReachability, FInstigator&& InInstigator)
 {
-	SetInstigatorInternal(MoveTemp(InInstigator));
+	SetInstigatorInternal(InReachability, MoveTemp(InInstigator));
 }
 
-void FPackageData::SetInstigatorInternal(FInstigator&& InInstigator)
+void FPackageData::SetInstigatorInternal(EReachability InReachability, FInstigator&& InInstigator)
 {
-	if (this->Instigator.Category == EInstigator::NotYetRequested)
+	if ((InReachability == EReachability::Runtime && this->Instigator.Category == EInstigator::NotYetRequested)
+		|| (InReachability == EReachability::Build && this->BuildInstigator.Category == EInstigator::NotYetRequested))
 	{
-		OnPackageDataFirstMarkedReachable(MoveTemp(InInstigator));
+		OnPackageDataFirstMarkedReachable(InReachability, MoveTemp(InInstigator));
 	}
 }
 
@@ -439,27 +449,25 @@ void FPackageData::SetPlatformCooked(
 	}
 }
 
-void FPackageData::ClearCookResults(const TConstArrayView<const ITargetPlatform*> TargetPlatforms)
+void FPackageData::SetPlatformCommitted(const ITargetPlatform* TargetPlatform)
 {
-	for (const ITargetPlatform* TargetPlatform : TargetPlatforms)
-	{
-		ClearCookResults(TargetPlatform);
-	}
+	FPackagePlatformData& Value = PlatformDatas.FindOrAdd(TargetPlatform);
+	Value.SetCommitted(true);
 }
 
-void FPackageData::ResetReachable()
+void FPackageData::ClearCookResults(const TConstArrayView<const ITargetPlatform*> TargetPlatforms)
 {
-	for (TPair<const ITargetPlatform*, FPackagePlatformData>& Pair : PlatformDatas)
+	for (const ITargetPlatform* TargetPlatform : TargetPlatforms)
 	{
-		Pair.Value.ResetReachable();
+		ClearCookResults(TargetPlatform);
 	}
 }
 
-void FPackageData::ResetIsBuildDependency()
+void FPackageData::ResetReachable(EReachability InReachability)
 {
 	for (TPair<const ITargetPlatform*, FPackagePlatformData>& Pair : PlatformDatas)
 	{
-		Pair.Value.SetIsBuildDependency(false);
+		Pair.Value.ResetReachable(InReachability);
 	}
 }
 
@@ -470,6 +478,7 @@ void FPackageData::ClearCookResults()
 	{
 		bModifiedCookAttempted = bModifiedCookAttempted | Pair.Value.IsCookAttempted();
 		Pair.Value.SetCookResults(ECookResult::NotAttempted);
+		Pair.Value.SetCommitted(false);
 		Pair.Value.SetSaveTimedOut(false);
 	}
 	if (bModifiedCookAttempted)
@@ -490,6 +499,7 @@ void FPackageData::ClearCookResults(const ITargetPlatform* TargetPlatform)
 		{
 			bModifiedCookAttempted = bModifiedCookAttempted | Pair.Value.IsCookAttempted();
 			Pair.Value.SetCookResults(ECookResult::NotAttempted);
+			Pair.Value.SetCommitted(false);
 			Pair.Value.SetSaveTimedOut(false);
 		}
 		else
@@ -540,6 +550,15 @@ bool FPackageData::HasAnyCookedPlatform() const
 		});
 }
 
+bool FPackageData::HasAnyCommittedPlatforms() const
+{
+	return Algo::AnyOf(PlatformDatas,
+		[](const TPair<const ITargetPlatform*, FPackagePlatformData>& Pair)
+		{
+			return Pair.Key != CookerLoadingPlatformKey && Pair.Value.IsCommitted();
+		});
+}
+
 bool FPackageData::HasAnyCookedPlatforms(const TArrayView<const ITargetPlatform* const>& Platforms,
 	bool bIncludeFailed) const
 {
@@ -596,6 +615,37 @@ ECookResult FPackageData::GetCookResults(const ITargetPlatform* Platform) const
 	return ECookResult::NotAttempted;
 }
 
+bool FPackageData::HasAllCommittedPlatforms(const TArrayView<const ITargetPlatform* const>& Platforms) const
+{
+	if (Platforms.Num() == 0)
+	{
+		return true;
+	}
+	if (PlatformDatas.Num() == 0)
+	{
+		return false;
+	}
+
+	for (const ITargetPlatform* QueryPlatform : Platforms)
+	{
+		if (!HasCommittedPlatform(QueryPlatform))
+		{
+			return false;
+		}
+	}
+	return true;
+}
+
+bool FPackageData::HasCommittedPlatform(const ITargetPlatform* Platform) const
+{
+	const FPackagePlatformData* PlatformData = PlatformDatas.Find(Platform);
+	if (PlatformData)
+	{
+		return PlatformData->IsCommitted();
+	}
+	return false;
+}
+
 UPackage* FPackageData::GetPackage() const
 {
 	return Package.Get();
@@ -1097,12 +1147,20 @@ void FPackageData::OnExitSaving(EStateChangeReason ReleaseSaveReason, EPackageSt
 	SetPackage(nullptr);
 }
 
-void FPackageData::OnPackageDataFirstMarkedReachable(FInstigator&& InInstigator)
+void FPackageData::OnPackageDataFirstMarkedReachable(EReachability InReachability, FInstigator&& InInstigator)
 {
-	TracePackage(GetPackageName().ToUnstableInt(), GetPackageName().ToString());
-	Instigator = MoveTemp(InInstigator);
-	PackageDatas.DebugInstigator(*this);
-	PackageDatas.UpdateThreadsafePackageData(*this);
+	if (InReachability == EReachability::Runtime)
+	{
+		TracePackage(GetPackageName().ToUnstableInt(), GetPackageName().ToString());
+		Instigator = MoveTemp(InInstigator);
+		PackageDatas.DebugInstigator(*this);
+		PackageDatas.UpdateThreadsafePackageData(*this);
+	}
+	else
+	{
+		check(InReachability == EReachability::Build);
+		BuildInstigator = MoveTemp(InInstigator);
+	}
 }
 
 void FPackageData::OnEnterAssignedToWorkerProperty()
@@ -1135,9 +1193,9 @@ void FPackageData::AddCompletionCallback(TConstArrayView<const ITargetPlatform*>
 	for (const ITargetPlatform* TargetPlatform : TargetPlatforms)
 	{
 		FPackagePlatformData* PlatformData = FindPlatformData(TargetPlatform);
-		// Adding a completion callback is only allowed after marking the requested platforms as reachable
+		// Adding a completion callback is only allowed after marking the requested platforms as runtime reachable
 		check(PlatformData);
-		check(PlatformData->IsReachable());
+		check(PlatformData->IsReachable(EReachability::Runtime));
 		// Adding a completion callback is only allowed after putting the PackageData in progress.
 		// If it's not in progress because it already finished the desired platforms, that is allowed.
 		check(IsInProgress() || PlatformData->IsCookAttempted() || !PlatformData->IsCookable());
@@ -2540,7 +2598,7 @@ void FPackageDatas::UpdateThreadsafePackageData(const FPackageData& PackageData)
 	UpdateThreadsafePackageData(PackageData.GetPackageName(),
 		[&PackageData](FThreadsafePackageData& ThreadsafeData, bool bNew)
 		{
-			ThreadsafeData.Instigator = PackageData.GetInstigator();
+			ThreadsafeData.Instigator = PackageData.GetInstigator(EReachability::Runtime);
 			ThreadsafeData.Generator = PackageData.GetParentGenerator();
 		});
 }
@@ -2560,16 +2618,17 @@ int32 FPackageDatas::GetNumCooked(ECookResult CookResult)
 	return Monitor.GetNumCooked(CookResult);
 }
 
-void FPackageDatas::GetCookedPackagesForPlatform(const ITargetPlatform* Platform,
+void FPackageDatas::GetCommittedPackagesForPlatform(const ITargetPlatform* Platform,
 	TArray<FPackageData*>& SucceededPackages,
 	TArray<FPackageData*>& FailedPackages)
 {
 	LockAndEnumeratePackageDatas(
 	[Platform, &SucceededPackages, &FailedPackages](FPackageData* PackageData)
 	{
-		ECookResult CookResults = PackageData->GetCookResults(Platform);
-		if (CookResults != ECookResult::NotAttempted)
+		FPackagePlatformData* PlatformData = PackageData->FindPlatformData(Platform);
+		if (PlatformData && PlatformData->IsCommitted())
 		{
+			ECookResult CookResults = PackageData->GetCookResults(Platform);
 			(CookResults == ECookResult::Succeeded ? SucceededPackages : FailedPackages).Add(PackageData);
 		}
 	});
@@ -2606,8 +2665,7 @@ void FPackageDatas::ClearCookedPlatforms()
 {
 	LockAndEnumeratePackageDatas([](FPackageData* PackageData)
 	{
-		PackageData->ResetReachable();
-		PackageData->ResetIsBuildDependency();
+		PackageData->ResetReachable(EReachability::All);
 		PackageData->ClearCookResults();
 	});
 }
@@ -2828,6 +2886,7 @@ void FRequestQueue::Empty()
 {
 	RestartedRequests.Empty();
 	DiscoveryQueue.Empty();
+	BuildDependencyDiscoveryQueue.Empty();
 	RequestClusters.Empty();
 	RequestFencePackageListeners.Empty();
 	NormalRequests.Empty();
@@ -2867,12 +2926,6 @@ bool FRequestQueue::Contains(const FPackageData* InPackageData) const
 	return false;
 }
 
-bool FRequestQueue::DiscoveryQueueContains(FPackageData* PackageData) const
-{
-	return Algo::FindBy(DiscoveryQueue, PackageData,
-		[](const FDiscoveryQueueElement& Element) { return Element.PackageData; }) != nullptr;
-}
-
 uint32 FRequestQueue::RemoveRequestExceptFromCluster(FPackageData* PackageData, FRequestCluster* ExceptFromCluster)
 {
 	uint32 OriginalNum = Num();
@@ -2909,7 +2962,7 @@ bool FRequestQueue::IsReadyRequestsEmpty() const
 bool FRequestQueue::HasRequestsToExplore() const
 {
 	return !RequestClusters.IsEmpty() | !RestartedRequests.IsEmpty() | !DiscoveryQueue.IsEmpty()
-		| !RequestFencePackageListeners.IsEmpty();
+		| !BuildDependencyDiscoveryQueue.IsEmpty() | !RequestFencePackageListeners.IsEmpty();
 }
 
 uint32 FRequestQueue::ReadyRequestsNum() const
diff --git a/Engine/Source/Editor/UnrealEd/Private/Cooker/CookPackageData.h b/Engine/Source/Editor/UnrealEd/Private/Cooker/CookPackageData.h
--- a/Engine/Source/Editor/UnrealEd/Private/Cooker/CookPackageData.h
+++ b/Engine/Source/Editor/UnrealEd/Private/Cooker/CookPackageData.h
@@ -12,6 +12,7 @@
 #include "Containers/UnrealString.h"
 #include "HAL/CriticalSection.h"
 #include "HAL/Platform.h"
+#include "HAL/PlatformMath.h"
 #include "IO/IoHash.h"
 #include "Math/NumericLimits.h"
 #include "Misc/EnumClassFlags.h"
@@ -142,6 +143,31 @@ enum class ESendFlags : uint8
 };
 ENUM_CLASS_FLAGS(ESendFlags);
 
+
+/**
+ * Reachability can be set for multiple properties. Reachability of each property is initially assigned to a package
+ * from initial requests, by e.g. the AssetManager requests in StartCookByTheBook, and it transitively is assigned to
+ * dependencies of the packages to which it is assigned.
+ */
+enum class EReachability : uint8
+{
+	None,
+	/**
+	 * A transitive build dependency from a cooked package. Build dependencies are commitable even if the package is
+	 * not cookable.
+	 */
+	Build = 0x01,
+	/**
+	 * Reachable as a runtime dependency; the package should be saved and made available at runtime. Packages
+	 * that are runtime reachable might still be not cooked if they are not cookable due to  e.g. NeverCook settings.
+	 */
+	Runtime = 0x02,
+	MaxBit = Runtime,
+	All = Build | Runtime,
+	NumBits = FPlatformMath::ConstExprCeilLogTwo(EReachability::MaxBit) + 1,
+};
+ENUM_CLASS_FLAGS(EReachability);
+
 /** Data necessary to create an FPackageData without checking the disk, for e.g. AddExistingPackageDatasForPlatform */
 struct FConstructPackageData
 {
@@ -160,17 +186,18 @@ struct FPackagePlatformData
 {
 	FPackagePlatformData();
 
-	/** Requested from StartCookByTheBook or CookOnTheFly, or is a transitive dependency of those requests. */
-	bool IsReachable() const;
-	void SetReachable(bool bValue);
-
-	/** Marked as a transitive build dependency from a reachable package, so we need to save its dependency data.*/
-	bool IsBuildDependency() const;
-	void SetIsBuildDependency(bool bValue);
+	/**
+	 * Get or modify the Reachability for specified reachability bits. @see UE::Cook::EReachability.
+	 */
+	EReachability GetReachability() const;
+	bool IsReachable(EReachability InReachability) const;
+	void AddReachability(EReachability InReachability);
+	void ClearReachability(EReachability InReachability);
 
-	/** The package has passed through a RequestCluster and its transitive dependencies were added. */
-	bool IsVisitedByCluster() const;
-	void SetVisitedByCluster(bool bValue);
+	/** Query/Modify whether the package has passed through a RequestCluster for all of the given Reachability bits. */
+	bool IsVisitedByCluster(EReachability InReachability) const;
+	void AddVisitedByCluster(EReachability InReachability);
+	void ClearVisitedByCluster(EReachability InReachability);
 
 	/** UPackage::Save was called on the package, but timedout and needs to be retried. */
 	bool IsSaveTimedOut() const;
@@ -194,8 +221,8 @@ struct FPackagePlatformData
 	bool IsExplorableOverride() const;
 	void SetExplorableOverride(bool bValue);
 
-	/** All flags modified by reachability calculations are returned to default. */
-	void ResetReachable();
+	/** All flags modified by reachability calculations for the given Reachability bits are returned to default. */
+	void ResetReachable(EReachability InReachability);
 	/**
 	 * Mark platform as ExplorableOverride=true and reset all data necessary to reexplore it, including reachability.
 	 * Caller is responsible for marking it again as reachable.
@@ -203,7 +230,7 @@ struct FPackagePlatformData
 	void MarkAsExplorable();
 
 	/** Called on CookWorkers to indicate reachable,cookable,etc for packages sent from Director. */
-	void MarkCookableForWorker(FCookWorkerClient& CookWorkerClient);
+	void MarkCommittableForWorker(EReachability InReachability, FCookWorkerClient& CookWorkerClient);
 
 	/** The package was found to be unmodified in the current iterative cook. */
 	bool IsIterativelyUnmodified() const;
@@ -221,7 +248,14 @@ struct FPackagePlatformData
 	bool IsCookSucceeded() const;
 	void SetCookResults(ECookResult Value);
 
-	/** Return if we need to call SavePackage on it (reachable, cookable, not yet cooked) */
+	bool IsCommitted() const;
+	void SetCommitted(bool bValue);
+
+	/** Return if we need to commit for the given reachability - reachable but not yet cooked/committed. */
+	bool NeedsCommit(const ITargetPlatform* PlatformItBelongsTo, EReachability InReachability) const;
+	bool NeedsCommit(const ITargetPlatform* PlatformItBelongsTo, ECookPhase CookPhase) const;
+
+	/** Return if we need to cook the package (Reachable for Runtime and not yet cooked), aka NeedsCommit(Runtime). */
 	bool NeedsCooking(const ITargetPlatform* PlatformItBelongsTo) const;
 
 	bool IsRegisteredForCachedObjectsInOuter() const;
@@ -232,9 +266,8 @@ struct FPackagePlatformData
 	void SetReportedToDirector(bool bValue);
 
 private:
-	uint32 bReachable : 1;
-	uint32 bIsBuildDependency : 1;
-	uint32 bVisitedByCluster : 1;
+	uint32 Reachability : (int)EReachability::NumBits;
+	uint32 ReachabilityVisitedByCluster : (int)EReachability::NumBits;
 	uint32 bSaveTimedOut : 1;
 	uint32 bCookable : 1;
 	uint32 bExplorable : 1;
@@ -243,6 +276,7 @@ private:
 	uint32 bIterativelySkipped : 1;
 	uint32 bRegisteredForCachedObjectsInOuter : 1;
 	uint32 bReportedToDirector : 1;
+	uint32 bCommitted : 1;
 	uint32 CookResults : (int)ECookResult::NumBits;
 };
 
@@ -338,27 +372,40 @@ public:
 	/** Set the LeafToRootRank for this package data. */
 	void SetLeafToRootRank(uint32 Value);
 
-	/** Reset OutPlatforms and copy current set of reachable & cookable & not yet cooked platforms into it. */
+	/** Reset OutPlatforms and copy NeedsCommit platforms into it, for the given CookPhase (@see NeedsCommit). */
+	template <typename ArrayType>
+	void GetPlatformsNeedingCommit(ArrayType& OutPlatforms, ECookPhase CookPhase) const;
+	/** Number of platforms that would be returned by GetPlatformsNeedingCommit. */
+	int32 GetPlatformsNeedingCommitNum(ECookPhase CookPhase) const;
+
+	/** Reset OutPlatforms and copy NeedsCommit platforms into it, for the given Reachability (@see NeedsCommit). */
 	template <typename ArrayType>
-	void GetPlatformsNeedingCooking(ArrayType& OutPlatforms) const;
-	/** Number of platforms that would be returned by GetPlatformsNeedingCooking. */
-	int32 GetPlatformsNeedingCookingNum() const;
+	void GetPlatformsNeedingCommit(ArrayType& OutPlatforms, EReachability Reachability) const;
+	/** Number of platforms that would be returned by GetPlatformsNeedingCommit. */
+	int32 GetPlatformsNeedingCommitNum(EReachability Reachability) const;
 
 	/** Reset OutPlatforms and copy current set of reachable platforms into it. */
 	template <typename ArrayType>
-	void GetReachablePlatforms(ArrayType& OutPlatforms) const;
+	void GetReachablePlatforms(EReachability InReachability, ArrayType& OutPlatforms) const;
 
-	/** Return true if and only if the Platform has been visited by a cluster (and hence is reachable and explored). */
-	bool IsPlatformVisitedByCluster(const ITargetPlatform* Platform) const;
+	/**
+	 * Return true if and only if the Platform has been visited by a cluster (and hence is reachable and explored)
+	 * for all of the requested Reachability bits.
+	 */
+	bool IsPlatformVisitedByCluster(const ITargetPlatform* Platform, EReachability InReachability) const;
 
 	/**
-	 * Return true if and only if every element of Platforms is currently reachable.
-	 * Returns true if Platforms is empty.
+	 * Return true if and only if every element of Platforms is currently reachable for all of the given
+	 * reachability bits. Returns true if Platforms is empty.
 	 */
-	bool HasReachablePlatforms(const TArrayView<const ITargetPlatform* const>& Platforms) const;
+	bool HasReachablePlatforms(EReachability InReachability,
+		const TArrayView<const ITargetPlatform* const>& Platforms) const;
 
-	/** Return whether all Platforms that have been marked reachable have also been explored in an FRequestCluster. */
-	bool AreAllReachablePlatformsVisitedByCluster() const;
+	/**
+	 * Return whether all Platforms that have been marked reachable for the given reachability property have also been
+	 * explored in an FRequestCluster for that reachability.
+	 */
+	bool AreAllReachablePlatformsVisitedByCluster(EReachability InReachability) const;
 
 	/** Get/Set the urgency of a package. Always EUrgency::Normal for pacakges that are not in progress. */
 	EUrgency GetUrgency() const;
@@ -367,21 +414,8 @@ public:
 	void RaiseUrgency(EUrgency NewUrgency, ESendFlags SendFlags, bool bAllowUrgencyInIdle = false);
 
 	/** Accessor for RequestClusters to add reachable platforms directly without modifying dependent data. */
-	void AddReachablePlatforms(FRequestCluster& RequestCluster, TConstArrayView<const ITargetPlatform*> Platforms,
-		FInstigator&& InInstigator);
-	/**
-	 * Set that the packagedata is reachable as a build dependency for the given platform. This setting is ignored
-	 * if the package is already reachable for the given platform, and is cleared if it ever becomes reachable.
-	 */
-	void AddBuildDependencyPlatforms(TConstArrayView<const ITargetPlatform*> Platforms, FInstigator&& InInstigator);
-	/** Clear the IsBuildDependency flag for all platforms on the packagedata. */
-	void ResetIsBuildDependency();
-
-	/**
-	 * Return true if and only if the package is a BuildDependency for every element of Platforms.
-	 * Returns true if Platforms is empty.
-	 */
-	bool IsBuildDependency(const TArrayView<const ITargetPlatform* const>& Platforms) const;
+	void AddReachablePlatforms(FRequestCluster& RequestCluster, EReachability InReachability,
+		TConstArrayView<const ITargetPlatform*> Platforms, FInstigator&& InInstigator);
 
 	/** Add the given reachable platforms to this PackageData and send it back to Request state for exploration. */
 	void QueueAsDiscovered(FInstigator&& InInstigator, FDiscoveredPlatformSet&& ReachablePlatforms, EUrgency InUrgency);
@@ -402,6 +436,8 @@ public:
 		ECookResult Result, bool bInWasCookedThisSession = true);
 	void SetPlatformCooked(const ITargetPlatform* TargetPlatform, ECookResult Result,
 		bool bInWasCookedThisSession = true);
+	void SetPlatformCommitted(const ITargetPlatform* TargetPlatform);
+
 	/**
 	 * FindOrAdd each TargetPlatform and set its flags: CookAttempted=false.
 	 * In Version that takes no TargetPlatform, CookAttempted is cleared from all existing platforms.
@@ -409,8 +445,8 @@ public:
 	void ClearCookResults(const TConstArrayView<const ITargetPlatform*> TargetPlatforms);
 	void ClearCookResults();
 	void ClearCookResults(const ITargetPlatform* TargetPlatform);
-	/** Clear reachable and related fields from all platforms. */
-	void ResetReachable();
+	/** Clear reachable and related fields from all platforms for the given reachability bits. */
+	void ResetReachable(EReachability InReachability);
 
 	/** Access the information about platforms interacted with by *this. */
 	const TSortedMap<const ITargetPlatform*, FPackagePlatformData, TInlineAllocator<1>>& GetPlatformDatas() const;
@@ -428,6 +464,10 @@ public:
 
 	/** Return true if and only if at least one platform has been cooked. */
 	bool HasAnyCookedPlatform() const;
+
+	/** Return true if and only if at least one platform has been committed. */
+	bool HasAnyCommittedPlatforms() const;
+
 	/**
 	 * Return true if and only if at least one element of Platforms has been cooked, and with its
 	 * succeeded flag set to true if bIncludeFailed is false. Returns false if Platforms is empty.
@@ -452,6 +492,10 @@ public:
 	ESuppressCookReason GetSuppressCookReason() const;
 	/** Get/Set the SuppressCookReason for the package. Defaults to ESuppressCookReason::NotSuppressed. */
 	void SetSuppressCookReason(ESuppressCookReason Reason);
+	/** Return true iff every element of Platforms has been committed. Returns true if Platforms is empty. */
+	bool HasAllCommittedPlatforms(const TArrayView<const ITargetPlatform* const>& Platforms) const;
+	/** Return true if and only if the given Platform has been committed. */
+	bool HasCommittedPlatform(const ITargetPlatform* Platform) const;
 
 	/**
 	 * Return the package pointer. By contract it will be non-null if and only if the PackageData's state is
@@ -534,7 +578,7 @@ public:
 	TArray<FCachedObjectInOuter>& GetCachedObjectsInOuter();
 	const TArray<FCachedObjectInOuter>& GetCachedObjectsInOuter() const;
 	template <typename ArrayType>
-	/** The list of platforms that were recorded as needscooking when CachedObjeObjectsInOuter was recorded. */
+	/** The list of platforms that were recorded as NeedsCooking when CachedObjeObjectsInOuter was recorded. */
 	void GetCachedObjectsInOuterPlatforms(ArrayType& OutPlatforms) const;
 
 	/** Validate that the CachedObjectsInOuter-dependent variables are empty, when entering save. */
@@ -671,22 +715,17 @@ public:
 	void SetDoesGeneratedRequireGenerator(ICookPackageSplitter::EGeneratedRequiresGenerator Value);
 
 	/**
-	 * Return the instigator for this package. The Instigator is the first code location or
-	 * referencing package that causes the package to enter the requested state.
+	 * Return the instigator for this package for the given reachability bits. The Instigator is the first code
+	 * location or referencing package that causes the package to enter the requested state for the given
+	 * reachability bit. If multiple bits of Reachability are specified, the most important one that has an
+	 * instigator is returned. EReachability::None is invalid.
 	 */
-	const FInstigator& GetInstigator() const;
-	bool HasInstigator() const;
+	const FInstigator& GetInstigator(EReachability InReachability) const;
+	bool HasInstigator(EReachability InReachability) const;
 	/** Setting the instigator is mostly private - it should only be done during clustering. */
-	void SetInstigator(FRequestCluster& Cluster, FInstigator&& InInstigator);
-	void SetInstigator(FCookWorkerClient& Client, FInstigator&& InInstigator);
-	void SetInstigator(FGenerationHelper& InGenerationHelper, FInstigator&& InInstigator);
-
-	/**
-	 * Return the BuildInstigator for this package. The BuildInstigator is the first other package
-	 * that marked this package as a BuildDependency.
-	 */
-	const FInstigator& GetBuildInstigator() const;
-	bool HasBuildInstigator() const;
+	void SetInstigator(FRequestCluster& Cluster, EReachability InReachability, FInstigator&& InInstigator);
+	void SetInstigator(FCookWorkerClient& Client, EReachability InReachability, FInstigator&& InInstigator);
+	void SetInstigator(FGenerationHelper& InHelper, EReachability InReachability, FInstigator&& InInstigator);
 
 	/** Get whether COTFS is keeping this package referenced referenced during GC. */
 	bool IsKeepReferencedDuringGC() const;
@@ -723,16 +762,16 @@ public:
 	 * Return the platforms for which the given Package has been marked reachable.
 	 * If the package does not exist, return the COTFS's list of Session platforms
 	 */
-	static void GetReachablePlatformsForInstigator(UCookOnTheFlyServer& COTFS, FPackageData* InInstigator,
-		TArray<const ITargetPlatform*>& Platforms);
-	static void GetReachablePlatformsForInstigator(UCookOnTheFlyServer& COTFS, FName InInstigator,
-		TArray<const ITargetPlatform*>& Platforms);
+	static void GetReachablePlatformsForInstigator(EReachability InReachability, UCookOnTheFlyServer& COTFS,
+		FPackageData* InInstigator, TArray<const ITargetPlatform*>& Platforms);
+	static void GetReachablePlatformsForInstigator(EReachability InReachability, UCookOnTheFlyServer& COTFS,
+		FName InInstigator, TArray<const ITargetPlatform*>& Platforms);
 private:
 	friend struct UE::Cook::FPackageDatas;
 
 	void UpdateContainerUrgency(EUrgency OldUrgency, EUrgency NewUrgency);
-	void SetInstigatorInternal(FInstigator&& InInstigator);
-	static void AddReachablePlatformsInternal(FPackageData& PackageData,
+	void SetInstigatorInternal(EReachability InReachability, FInstigator&& InInstigator);
+	static void AddReachablePlatformsInternal(FPackageData& PackageData, EReachability InReachability,
 		TConstArrayView<const ITargetPlatform*> Platforms, FInstigator&& InInstigator);
 	static void QueueAsDiscoveredInternal(FPackageData& PackageData, FInstigator&& InInstigator,
 		FDiscoveredPlatformSet&& ReachablePlatforms, EUrgency InUrgency);
@@ -785,7 +824,7 @@ private:
 	void OnExitAssignedToWorkerProperty();
 	void OnEnterAssignedToWorkerProperty();
 
-	void OnPackageDataFirstMarkedReachable(FInstigator&& InInstigator);
+	void OnPackageDataFirstMarkedReachable(EReachability InReachability, FInstigator&& InInstigator);
 
 	FGenerationHelper* GenerationHelper = nullptr;
 	TRefCountPtr<FGenerationHelper> ParentGenerationHelper;
@@ -997,7 +1036,6 @@ public:
 	uint32 Num() const;
 	uint32 Remove(FPackageData* PackageData);
 	bool Contains(const FPackageData* PackageData) const;
-	bool DiscoveryQueueContains(FPackageData* PackageData) const;
 	void Empty();
 
 	void AddRequest(FPackageData* PackageData, bool bForceUrgent=false);
@@ -1013,11 +1051,13 @@ public:
 
 	FPackageDataSet& GetRestartedRequests();
 	/**
-	 * Unlike other containers on PackageData, GetDiscoveryQueue is not an ownership container.
-	 * PackageDatas in the DiscoveryQueue are PackageDatas we need to look at - in the right order during
+	 * Unlike non-discovery containers on PackageData, GetDiscoveryQueue and GetBuildDependencyDiscoveryQueue
+	 * are not ownership containers.
+	 * PackageDatas in the DiscoveryQueues are PackageDatas we need to look at - in the right order during
 	 * PumpRequests - they can be in any state and are owned by another container (or are in the idle state).
 	 */
 	TRingBuffer<FDiscoveryQueueElement>& GetDiscoveryQueue();
+	TRingBuffer<FPackageData*>& GetBuildDependencyDiscoveryQueue();
 	TRingBuffer<FRequestCluster>& GetRequestClusters();
 	FPackageDataSet& GetReadyRequestsUrgent();
 	FPackageDataSet& GetReadyRequestsNormal();
@@ -1032,6 +1072,7 @@ public:
 private:
 	FPackageDataSet RestartedRequests;
 	TRingBuffer<FDiscoveryQueueElement> DiscoveryQueue;
+	TRingBuffer<FPackageData*> BuildDependencyDiscoveryQueue;
 	TRingBuffer<FRequestCluster> RequestClusters;
 	TSet<FName> RequestFencePackageListeners;
 	FPackageDataSet UrgentRequests;
@@ -1293,10 +1334,10 @@ public:
 	int32 GetNumCooked(ECookResult CookResult);
 	/**
 	 * Append to SucceededPackages all packages that have cooked any platform with success, and to
-	 * FailedPackages all packages that have cooked any platform with other results. The output variables are permitted
-	 * to point to the same array.
+	 * FailedPackages all packages that were committed for any platform with other results. The output variables are
+	 * permitted to point to the same array.
 	 */
-	void GetCookedPackagesForPlatform(const ITargetPlatform* Platform, TArray<FPackageData*>& SucceededPackages,
+	void GetCommittedPackagesForPlatform(const ITargetPlatform* Platform, TArray<FPackageData*>& SucceededPackages,
 		TArray<FPackageData*>& FailedPackages);
 
 	/**
@@ -1398,34 +1439,39 @@ struct FPoppedPackageDataScope
 ///////////////////////////////////////////////////////
 
 
-inline bool FPackagePlatformData::IsReachable() const
+inline EReachability FPackagePlatformData::GetReachability() const
+{
+	return (EReachability)Reachability;
+}
+
+inline bool FPackagePlatformData::IsReachable(EReachability InReachability) const
 {
-	return bReachable != 0;
+	return EnumHasAllFlags((EReachability)Reachability, InReachability);
 }
 
-inline void FPackagePlatformData::SetReachable(bool bValue)
+inline void FPackagePlatformData::AddReachability(EReachability InReachability)
 {
-	bReachable = (uint32)bValue;
+	Reachability |= (uint8)InReachability;
 }
 
-inline bool FPackagePlatformData::IsBuildDependency() const
+inline void FPackagePlatformData::ClearReachability(EReachability InReachability)
 {
-	return bIsBuildDependency != 0;
+	Reachability &= ~(uint8)InReachability;
 }
 
-inline void FPackagePlatformData::SetIsBuildDependency(bool bValue)
+inline bool FPackagePlatformData::IsVisitedByCluster(EReachability InReachability) const
 {
-	bIsBuildDependency = (uint32)bValue;
+	return EnumHasAllFlags((EReachability)ReachabilityVisitedByCluster, InReachability);
 }
 
-inline bool FPackagePlatformData::IsVisitedByCluster() const
+inline void FPackagePlatformData::AddVisitedByCluster(EReachability InReachability)
 {
-	return bVisitedByCluster != 0;
+	ReachabilityVisitedByCluster |= (uint8)InReachability;
 }
 
-inline void FPackagePlatformData::SetVisitedByCluster(bool bValue)
+inline void FPackagePlatformData::ClearVisitedByCluster(EReachability InReachability)
 {
-	bVisitedByCluster = (uint32)bValue;
+	ReachabilityVisitedByCluster &= ~((uint8)InReachability);
 }
 
 inline bool FPackagePlatformData::IsSaveTimedOut() const
@@ -1508,9 +1554,34 @@ inline void FPackagePlatformData::SetCookResults(ECookResult Value)
 	// ECookResult::Invalid is only used in replication and is not allowed in FPackagePlatformData
 	check(Value != ECookResult::Invalid);
 	CookResults = (uint32)Value;
+	if (Value == ECookResult::Succeeded || Value == ECookResult::Failed)
+	{
+		SetCommitted(true);
+	}
+}
+
+inline bool FPackagePlatformData::IsCommitted() const
+{
+	return bCommitted != 0;
+}
+
+inline void FPackagePlatformData::SetCommitted(bool bValue)
+{
+	bCommitted = (uint32)bValue;
 	SetReportedToDirector(false);
 }
 
+inline bool FPackagePlatformData::NeedsCooking(const ITargetPlatform* PlatformItBelongsTo) const
+{
+	return NeedsCommit(PlatformItBelongsTo, EReachability::Runtime);
+}
+
+inline bool FPackagePlatformData::NeedsCommit(const ITargetPlatform* PlatformItBelongsTo, ECookPhase CookPhase) const
+{
+	return NeedsCommit(PlatformItBelongsTo,
+		CookPhase == ECookPhase::Cook ? EReachability::Runtime : EReachability::Build);
+}
+
 inline bool FPackagePlatformData::IsRegisteredForCachedObjectsInOuter() const
 {
 	return bRegisteredForCachedObjectsInOuter != 0;
@@ -1562,12 +1633,25 @@ inline void FPackageData::SetLeafToRootRank(uint32 Value)
 }
 
 template <typename ArrayType>
-inline void FPackageData::GetPlatformsNeedingCooking(ArrayType& OutPlatforms) const
+inline void FPackageData::GetPlatformsNeedingCommit(ArrayType& OutPlatforms, ECookPhase CookPhase) const
+{
+	return GetPlatformsNeedingCommit(OutPlatforms,
+		CookPhase == ECookPhase::Cook ? EReachability::Runtime : EReachability::Build);
+}
+
+inline int32 FPackageData::GetPlatformsNeedingCommitNum(ECookPhase CookPhase) const
+{
+	return GetPlatformsNeedingCommitNum(
+		CookPhase == ECookPhase::Cook ? EReachability::Runtime : EReachability::Build);
+}
+
+template <typename ArrayType>
+inline void FPackageData::GetPlatformsNeedingCommit(ArrayType& OutPlatforms, EReachability Reachability) const
 {
 	OutPlatforms.Reset(PlatformDatas.Num());
 	for (const TPair<const ITargetPlatform*, FPackagePlatformData>& Pair : PlatformDatas)
 	{
-		if (Pair.Value.NeedsCooking(Pair.Key))
+		if (Pair.Value.NeedsCommit(Pair.Key, Reachability))
 		{
 			OutPlatforms.Add(Pair.Key);
 		}
@@ -1575,12 +1659,12 @@ inline void FPackageData::GetPlatformsNeedingCooking(ArrayType& OutPlatforms) co
 }
 
 template <typename ArrayType>
-inline void FPackageData::GetReachablePlatforms(ArrayType& OutPlatforms) const
+inline void FPackageData::GetReachablePlatforms(EReachability InReachability, ArrayType& OutPlatforms) const
 {
 	OutPlatforms.Reset(PlatformDatas.Num());
 	for (const TPair<const ITargetPlatform*, FPackagePlatformData>& Pair : PlatformDatas)
 	{
-		if (Pair.Value.IsReachable())
+		if (Pair.Value.IsReachable(InReachability))
 		{
 			OutPlatforms.Add(Pair.Key);
 		}
@@ -1703,24 +1787,50 @@ inline void FPackageData::SetDoesGeneratedRequireGenerator(ICookPackageSplitter:
 	DoesGeneratedRequireGeneratorValue = static_cast<uint32>(Value);
 }
 
-inline const FInstigator& FPackageData::GetInstigator() const
-{
-	return Instigator;
-}
-
-inline bool FPackageData::HasInstigator() const
+inline const FInstigator& FPackageData::GetInstigator(EReachability InReachability) const
 {
-	return Instigator.Category != EInstigator::NotYetRequested;
-}
+	if (InReachability == EReachability::None)
+	{
+		checkf(false, TEXT("Invalid argument EReachability::None."));
+		return Instigator;
+	}
+	if (EnumHasAnyFlags(InReachability, EReachability::Runtime)
+		&& Instigator.Category != EInstigator::NotYetRequested)
+	{
+		return Instigator;
+	}
+	if (EnumHasAnyFlags(InReachability, EReachability::Build)
+		&& BuildInstigator.Category != EInstigator::NotYetRequested)
+	{
+		return BuildInstigator;
+	}
 
-inline const FInstigator& FPackageData::GetBuildInstigator() const
-{
+	// return a const reference to the (empty) Instigator for the most important set bit.
+	if (EnumHasAnyFlags(InReachability, EReachability::Runtime))
+	{
+		return Instigator;
+	}
 	return BuildInstigator;
 }
 
-inline bool FPackageData::HasBuildInstigator() const
+inline bool FPackageData::HasInstigator(EReachability InReachability) const
 {
-	return BuildInstigator.Category != EInstigator::NotYetRequested;
+	if (InReachability == EReachability::None)
+	{
+		checkf(false, TEXT("Invalid argument EReachability::None."));
+		return false;
+	}
+	if (EnumHasAnyFlags(InReachability, EReachability::Runtime)
+		&& Instigator.Category != EInstigator::NotYetRequested)
+	{
+		return true;
+	}
+	if (EnumHasAnyFlags(InReachability, EReachability::Build)
+		&& BuildInstigator.Category != EInstigator::NotYetRequested)
+	{
+		return true;
+	}
+	return false;
 }
 
 inline bool FPackageData::IsKeepReferencedDuringGC() const
@@ -1775,6 +1885,11 @@ inline TRingBuffer<FDiscoveryQueueElement>& FRequestQueue::GetDiscoveryQueue()
 	return DiscoveryQueue;
 }
 
+inline TRingBuffer<FPackageData*>& FRequestQueue::GetBuildDependencyDiscoveryQueue()
+{
+	return BuildDependencyDiscoveryQueue;
+}
+
 inline TRingBuffer<FRequestCluster>& FRequestQueue::GetRequestClusters()
 {
 	return RequestClusters;
diff --git a/Engine/Source/Editor/UnrealEd/Private/Cooker/CookRequestCluster.cpp b/Engine/Source/Editor/UnrealEd/Private/Cooker/CookRequestCluster.cpp
--- a/Engine/Source/Editor/UnrealEd/Private/Cooker/CookRequestCluster.cpp
+++ b/Engine/Source/Editor/UnrealEd/Private/Cooker/CookRequestCluster.cpp
@@ -37,7 +37,7 @@
 namespace UE::Cook
 {
 
-FRequestCluster::FRequestCluster(UCookOnTheFlyServer& InCOTFS)
+FRequestCluster::FRequestCluster(UCookOnTheFlyServer& InCOTFS, EReachability ExploreReachability)
 	: COTFS(InCOTFS)
 	, PackageDatas(*InCOTFS.PackageDatas)
 	, AssetRegistry(*IAssetRegistry::Get())
@@ -57,12 +57,43 @@ FRequestCluster::FRequestCluster(UCookOnTheFlyServer& InCOTFS)
 		bAllowSoftDependencies = false;
 	}
 
+	if (COTFS.IsCookWorkerMode())
+	{
+		if (ExploreReachability == EReachability::Build)
+		{
+			TraversalTier = ETraversalTier::MarkForBuildDependency;
+		}
+		else
+		{
+			check(ExploreReachability == EReachability::Runtime);
+			TraversalTier = ETraversalTier::MarkForRuntime;
+		}
+	}
+	else if (ExploreReachability == EReachability::Build)
+	{
+		TraversalTier = ETraversalTier::BuildDependencies;
+	}
+	else
+	{
+		check(ExploreReachability == EReachability::Runtime);
+		TraversalTier = bAllowHardDependencies ? ETraversalTier::RuntimeFollowDependencies :
+			ETraversalTier::RuntimeVisitVertices;
+	}
+
 	if (bErrorOnEngineContentUse)
 	{
 		DLCPath = FPaths::Combine(*COTFS.GetBaseDirectoryForDLC(), TEXT("Content"));
 		FPaths::MakeStandardFilename(DLCPath);
 	}
-	GConfig->GetBool(TEXT("CookSettings"), TEXT("PreQueueBuildDefinitions"), bPreQueueBuildDefinitions, GEditorIni);
+	if (TraversalMarkCookable())
+	{
+		GConfig->GetBool(TEXT("CookSettings"), TEXT("PreQueueBuildDefinitions"), 
+			bPreQueueBuildDefinitions, GEditorIni);
+	}
+	else
+	{
+		bPreQueueBuildDefinitions = false;
+	}
 
 	bAllowIterativeResults = true;
 	bool bFirst = true;
@@ -87,15 +118,16 @@ FRequestCluster::FRequestCluster(UCookOnTheFlyServer& InCOTFS)
 }
 
 FRequestCluster::FRequestCluster(UCookOnTheFlyServer& InCOTFS, TArray<FFilePlatformRequest>&& InRequests)
-	: FRequestCluster(InCOTFS)
+	: FRequestCluster(InCOTFS, EReachability::Runtime)
 {
 	ReserveInitialRequests(InRequests.Num());
 	FilePlatformRequests = MoveTemp(InRequests);
 	InRequests.Empty();
 }
 
-FRequestCluster::FRequestCluster(UCookOnTheFlyServer& InCOTFS, FPackageDataSet&& InRequests)
-	: FRequestCluster(InCOTFS)
+FRequestCluster::FRequestCluster(UCookOnTheFlyServer& InCOTFS, FPackageDataSet&& InRequests,
+	EReachability InExploreReachability)
+	: FRequestCluster(InCOTFS, InExploreReachability)
 {
 	ReserveInitialRequests(InRequests.Num());
 	for (FPackageData* PackageData : InRequests)
@@ -109,7 +141,7 @@ FRequestCluster::FRequestCluster(UCookOnTheFlyServer& InCOTFS, FPackageDataSet&&
 }
 
 FRequestCluster::FRequestCluster(UCookOnTheFlyServer& InCOTFS, TRingBuffer<FDiscoveryQueueElement>& DiscoveryQueue)
-	: FRequestCluster(InCOTFS)
+	: FRequestCluster(InCOTFS, EReachability::Runtime)
 {
 	TArray<const ITargetPlatform*, TInlineAllocator<ExpectedMaxNumPlatforms>> BufferPlatforms;
 	if (!COTFS.bSkipOnlyEditorOnly)
@@ -127,7 +159,7 @@ FRequestCluster::FRequestCluster(UCookOnTheFlyServer& InCOTFS, TRingBuffer<FDisc
 		if (COTFS.bSkipOnlyEditorOnly)
 		{
 			NewReachablePlatforms = Discovery->ReachablePlatforms.GetPlatforms(COTFS, &Discovery->Instigator,
-				TConstArrayView<const ITargetPlatform*>(), &BufferPlatforms);
+				TConstArrayView<const ITargetPlatform*>(), EReachability::Runtime, &BufferPlatforms);
 		}
 		else
 		{
@@ -149,14 +181,14 @@ FRequestCluster::FRequestCluster(UCookOnTheFlyServer& InCOTFS, TRingBuffer<FDisc
 			}
 		}
 
-		if (PackageData.HasReachablePlatforms(NewReachablePlatforms))
+		if (PackageData.HasReachablePlatforms(EReachability::Runtime, NewReachablePlatforms))
 		{
 			// If there are no new reachable platforms, add it to the cluster for cooking if it needs
 			// it, otherwise let it remain where it is
 			EUrgency Urgency = Discovery->Urgency;
 			DiscoveryQueue.PopFrontValue();
 			Discovery = nullptr;
-			if (!PackageData.IsInProgress() && PackageData.GetPlatformsNeedingCookingNum() == 0)
+			if (!PackageData.IsInProgress() && PackageData.GetPlatformsNeedingCommitNum(EReachability::Runtime) == 0)
 			{
 				PackageData.SendToState(EPackageState::Request, ESendFlags::QueueRemove,
 					EStateChangeReason::RequestCluster);
@@ -180,7 +212,8 @@ FRequestCluster::FRequestCluster(UCookOnTheFlyServer& InCOTFS, TRingBuffer<FDisc
 			bool bExpectedDiscoveryType = false;
 			if (Discovery->Instigator.Category == EInstigator::Unsolicited)
 			{
-				bExpectedDiscoveryType = PackageData.FindOrAddPlatformData(CookerLoadingPlatformKey).IsReachable();
+				bExpectedDiscoveryType = PackageData.FindOrAddPlatformData(CookerLoadingPlatformKey)
+					.IsReachable(EReachability::Runtime);
 				if (bExpectedDiscoveryType && COTFS.bSkipOnlyEditorOnly)
 				{
 					// In SkipOnlyEditorOnly mode, expected-load unsolicited packages are skipped; merely loading
@@ -230,7 +263,8 @@ FRequestCluster::FRequestCluster(UCookOnTheFlyServer& InCOTFS, TRingBuffer<FDisc
 			}
 		}
 		// Add the new reachable platforms
-		PackageData.AddReachablePlatforms(*this, NewReachablePlatforms, MoveTemp(Discovery->Instigator));
+		PackageData.AddReachablePlatforms(*this, EReachability::Runtime, NewReachablePlatforms,
+			MoveTemp(Discovery->Instigator));
 
 		// Pop it off the list; note that this invalidates the pointers we had into the DiscoveryQueueElement
 		FDiscoveryQueueElement PoppedDiscovery = DiscoveryQueue.PopFrontValue();
@@ -246,6 +280,87 @@ FRequestCluster::FRequestCluster(UCookOnTheFlyServer& InCOTFS, TRingBuffer<FDisc
 	}
 }
 
+FRequestCluster::FRequestCluster(UCookOnTheFlyServer& InCOTFS, EBuildDependencyQueueConstructorType,
+	TRingBuffer<FPackageData*>& BuildDependencyDiscoveryQueue)
+	: FRequestCluster(InCOTFS, EReachability::Build)
+{
+	while (!BuildDependencyDiscoveryQueue.IsEmpty())
+	{
+		FPackageData& PackageData = *BuildDependencyDiscoveryQueue.PopFrontValue();
+		if (PackageData.IsInProgress() ||
+			PackageData.GetPlatformsNeedingCommitNum(EReachability::Build) == 0)
+		{
+			// Already kicked or committed since being queued
+			continue;
+		}
+
+		PackageData.SendToState(EPackageState::Request, ESendFlags::QueueRemove,
+			EStateChangeReason::RequestCluster);
+		// SetPackageDataSuppressReason adds it in the proper container of *this
+		SetPackageDataSuppressReason(PackageData, ESuppressCookReason::NotSuppressed);
+	}
+}
+
+bool FRequestCluster::TraversalExploreRuntimeDependencies()
+{
+	switch (TraversalTier)
+	{
+	case ETraversalTier::MarkForRuntime:
+		return false;
+	case ETraversalTier::MarkForBuildDependency:
+		return false;
+	case ETraversalTier::BuildDependencies:
+		return false;
+	case ETraversalTier::RuntimeVisitVertices:
+		return false;
+	case ETraversalTier::RuntimeFollowDependencies:
+		return true;
+	default:
+		checkNoEntry();
+		return false;
+	}
+}
+
+bool FRequestCluster::TraversalExploreIncremental()
+{
+	switch (TraversalTier)
+	{
+	case ETraversalTier::MarkForRuntime:
+		return false;
+	case ETraversalTier::MarkForBuildDependency:
+		return false;
+	case ETraversalTier::BuildDependencies:
+		return IsIncrementalCook();
+	case ETraversalTier::RuntimeVisitVertices:
+		return IsIncrementalCook();
+	case ETraversalTier::RuntimeFollowDependencies:
+		return IsIncrementalCook();
+	default:
+		checkNoEntry();
+		return false;
+	}
+}
+
+bool FRequestCluster::TraversalMarkCookable()
+{
+	switch (TraversalTier)
+	{
+	case ETraversalTier::MarkForRuntime:
+		return true;
+	case ETraversalTier::MarkForBuildDependency:
+		return false;
+	case ETraversalTier::BuildDependencies:
+		return false;
+	case ETraversalTier::RuntimeVisitVertices:
+		return true;
+	case ETraversalTier::RuntimeFollowDependencies:
+		return true;
+	default:
+		checkNoEntry();
+		return false;
+	}
+}
+
 FName GInstigatorRequestCluster(TEXT("RequestCluster"));
 
 void FRequestCluster::Process(const FCookerTimer& CookerTimer, bool& bOutComplete)
@@ -272,6 +387,7 @@ void FRequestCluster::FetchPackageNames(const FCookerTimer& CookerTimer, bool& b
 		return;
 	}
 
+	check(TraversalMarkCookable() || FilePlatformRequests.IsEmpty());
 	constexpr int32 TimerCheckPeriod = 100; // Do not incur the cost of checking the timer on every package
 	int32 NextRequest = 0;
 	for (; NextRequest < FilePlatformRequests.Num(); ++NextRequest)
@@ -303,9 +419,10 @@ void FRequestCluster::FetchPackageNames(const FCookerTimer& CookerTimer, bool& b
 		}
 
 		// If it has new reachable platforms we definitely need to explore it
-		if (!PackageData->HasReachablePlatforms(Request.GetPlatforms()))
+		if (!PackageData->HasReachablePlatforms(EReachability::Runtime, Request.GetPlatforms()))
 		{
-			PackageData->AddReachablePlatforms(*this, Request.GetPlatforms(), MoveTemp(Request.GetInstigator()));
+			PackageData->AddReachablePlatforms(*this, EReachability::Runtime, Request.GetPlatforms(),
+				MoveTemp(Request.GetInstigator()));
 			PullIntoCluster(*PackageData);
 			if (Request.IsUrgent())
 			{
@@ -323,7 +440,7 @@ void FRequestCluster::FetchPackageNames(const FCookerTimer& CookerTimer, bool& b
 					PackageData->SetUrgency(EUrgency::Blocking, ESendFlags::QueueAddAndRemove);
 				}
 			}
-			else if (PackageData->GetPlatformsNeedingCookingNum() > 0)
+			else if (PackageData->GetPlatformsNeedingCommitNum(EReachability::Runtime) > 0)
 			{
 				// If it's missing cookable platforms and not in progress we need to add it to the cluster for cooking
 				PullIntoCluster(*PackageData);
@@ -381,6 +498,11 @@ void FRequestCluster::StartAsync(const FCookerTimer& CookerTimer, bool& bOutComp
 		return;
 	}
 
+	if (!TraversalMarkCookable())
+	{
+		return;
+	}
+
 	FEditorDomain* EditorDomain = FEditorDomain::Get();
 	if (EditorDomain && EditorDomain->IsReadingPackages())
 	{
@@ -445,21 +567,35 @@ void FRequestCluster::SetPackageDataSuppressReason(FPackageData& PackageData, ES
 	}
 }
 
-void FRequestCluster::SetPackageDataWasMarkedCooked(FPackageData& PackageData, bool bValue,
+void FRequestCluster::SetPackageDataWasMarkedSkipped(FPackageData& PackageData, bool bValue,
 	bool* bOutExisted)
 {
-	FProcessingFlags& Existing = OwnedPackageDatas.FindOrAdd(&PackageData);
+	// We don't want to bring the package into the cluster if it's not already there, which is governed by existence in
+	// OwnedPackageDatas. So early exit if it is not there.
+	// TODO: Change name of OwnedPackageDatas and allow packages not in the cluster to exist in it,
+	// with a separate flag for whether they are owned, so we can keep track of SetWasMarkedSkipped for packages
+	// that are not in the cluster when first evaluated, but are later added.
+	FProcessingFlags* Existing = OwnedPackageDatas.Find(&PackageData);
+	if (!Existing)
+	{
+		if (bOutExisted)
+		{
+			*bOutExisted = false;
+		}
+		return;
+	}
+
 	if (bOutExisted)
 	{
-		*bOutExisted = Existing.IsValid();
+		*bOutExisted = Existing->IsValid();
 	}
-	if (Existing.ShouldMarkNotInProgress())
+	if (Existing->ShouldMarkNotInProgress())
 	{
 		--PackagesToMarkNotInProgressCount;
 	}
-	Existing.SetValid();
-	Existing.SetWasMarkedCooked(bValue);
-	if (Existing.ShouldMarkNotInProgress())
+	Existing->SetValid();
+	Existing->SetWasMarkedSkipped(bValue);
+	if (Existing->ShouldMarkNotInProgress())
 	{
 		++PackagesToMarkNotInProgressCount;
 	}
@@ -565,21 +701,11 @@ void FRequestCluster::PumpExploration(const FCookerTimer& CookerTimer, bool& bOu
 
 	if (!GraphSearch)
 	{
-		ETraversalTier TraversalTier = ETraversalTier::None;
-		if (COTFS.IsCookWorkerMode())
-		{
-			TraversalTier = ETraversalTier::None;
-		}
-		else
-		{
-			TraversalTier = bAllowHardDependencies ? ETraversalTier::FollowDependencies :
-				ETraversalTier::FetchEdgeData;
-		}
-		GraphSearch.Reset(new FGraphSearch(*this, TraversalTier));
+		GraphSearch.Reset(new FGraphSearch(*this));
 
-		if (TraversalTier == ETraversalTier::None)
+		if (!TraversalExploreIncremental() && !TraversalExploreRuntimeDependencies())
 		{
-			GraphSearch->VisitWithoutDependencies();
+			GraphSearch->VisitWithoutFetching();
 			GraphSearch.Reset();
 			bDependenciesComplete = true;
 			return;
@@ -649,9 +775,8 @@ void FRequestCluster::PumpExploration(const FCookerTimer& CookerTimer, bool& bOu
 	bDependenciesComplete = true;
 }
 
-FRequestCluster::FGraphSearch::FGraphSearch(FRequestCluster& InCluster, ETraversalTier InTraversalTier)
+FRequestCluster::FGraphSearch::FGraphSearch(FRequestCluster& InCluster)
 	: Cluster(InCluster)
-	, TraversalTier(InTraversalTier)
 	, ExploreEdgesContext(InCluster, *this)
 	, AsyncResultsReadyEvent(EEventMode::ManualReset)
 {
@@ -682,7 +807,7 @@ FRequestCluster::FGraphSearch::FGraphSearch(FRequestCluster& InCluster, ETravers
 	check(FetchPlatforms[CookerLoadingPlatformIndex].bIsCookerLoadingPlatform);
 }
 
-void FRequestCluster::FGraphSearch::VisitWithoutDependencies()
+void FRequestCluster::FGraphSearch::VisitWithoutFetching()
 {
 	// PumpExploration is responsible for marking all requests as explored and cookable/uncoookable.
 	// If we're skipping the dependencies search, handle that responsibility for the initial requests and return.
@@ -996,7 +1121,7 @@ void FRequestCluster::FGraphSearch::ResolveTransitiveBuildDependencyCycle()
 	}
 	check(FirstVertex); // This function should not be called if PendingTransitiveBuildDependencyVertices is empty.
 	PendingTransitiveBuildDependencyVertices.Empty();
-	UE_LOG(LogCook, Display,
+	UE_LOG(LogCook, Verbose,
 		TEXT("Cycle detected in the graph of transitive build dependencies.")
 		TEXT(" No vertices in the cycle are invalidated by their direct dependencies, so marking them all as iteratively skippable.")
 		TEXT("\n\tVertex in the cycle: %s"),
@@ -1034,6 +1159,8 @@ void FRequestCluster::FGraphSearch::VisitVertex(FVertexData& Vertex)
 		return;
 	}
 
+	const EReachability ClusterReachability = Cluster.TraversalMarkCookable()
+		? EReachability::Runtime : EReachability::Build;
 	int32 LocalNumFetchPlatforms = NumFetchPlatforms();
 	TBitArray<> ShouldFetchPlatforms(false, LocalNumFetchPlatforms);
 	
@@ -1050,7 +1177,7 @@ void FRequestCluster::FGraphSearch::VisitVertex(FVertexData& Vertex)
 		{
 			CookerLoadingPlatform = &PlatformData;
 		}
-		else if (PlatformData.IsReachable())
+		else if (PlatformData.IsReachable(ClusterReachability))
 		{
 			int32 PlatformIndex = Algo::BinarySearchBy(FetchPlatforms, TargetPlatform, [](const FFetchPlatformData& D)
 				{
@@ -1062,13 +1189,13 @@ void FRequestCluster::FGraphSearch::VisitVertex(FVertexData& Vertex)
 			{
 				FirstReachableSessionPlatform = TargetPlatform;
 			}
-			if (!PlatformData.IsVisitedByCluster())
+			if (!PlatformData.IsVisitedByCluster(ClusterReachability))
 			{
-				VisitVertexForPlatform(Vertex, TargetPlatform, PlatformData, SuppressCookReason);
+				VisitVertexForPlatform(Vertex, TargetPlatform, ClusterReachability, PlatformData,
+					SuppressCookReason);
 
-				if ((TraversalTier >= ETraversalTier::FetchEdgeData) && 
-					(((TraversalTier >= ETraversalTier::FollowDependencies) && PlatformData.IsExplorable())
-						|| Cluster.IsIncrementalCook()))
+				if (Cluster.TraversalExploreIncremental() ||
+					(Cluster.TraversalExploreRuntimeDependencies() && PlatformData.IsExplorable()))
 				{
 					ShouldFetchPlatforms[PlatformIndex] = true;
 					Vertex.PlatformData[PlatformIndex].bExploreRequested = true;
@@ -1083,58 +1210,65 @@ void FRequestCluster::FGraphSearch::VisitVertex(FVertexData& Vertex)
 			}
 		}
 	}
-	bool bAnyCookable = (FirstReachableSessionPlatform == nullptr) | !bAllReachablesUncookable;
-	if (bAnyCookable != Vertex.bAnyCookable)
+
+	if (Cluster.TraversalMarkCookable())
 	{
-		if (!bAnyCookable)
+		bool bAnyCookable = (FirstReachableSessionPlatform == nullptr) | !bAllReachablesUncookable;
+		if (bAnyCookable != Vertex.bAnyCookable)
 		{
-			if (SuppressCookReason == ESuppressCookReason::Invalid)
+			if (!bAnyCookable)
 			{
-				// We need the SuppressCookReason for reporting. If we didn't calculate it this Visit and
-				// we don't have it stored in this->OwnedPackageDatas, then we must have calculated it in
-				// a previous cluster, but we don't store it anywhere. Recalculate it from the
-				// FirstReachableSessionPlatform. FirstReachableSessionPlatform must be non-null, otherwise
-				// bAnyCookable would be true.
-				check(FirstReachableSessionPlatform);
-				bool bCookable;
-				bool bExplorable;
-				Cluster.IsRequestCookable(FirstReachableSessionPlatform, Vertex.PackageData->GetPackageName(),
-					*Vertex.PackageData, SuppressCookReason, bCookable, bExplorable);
-				check(!bCookable); // We don't support bCookable changing for a given package and platform
-				check(SuppressCookReason != ESuppressCookReason::Invalid);
+				if (SuppressCookReason == ESuppressCookReason::Invalid)
+				{
+					// We need the SuppressCookReason for reporting. If we didn't calculate it this Visit and
+					// we don't have it stored in this->OwnedPackageDatas, then we must have calculated it in
+					// a previous cluster, but we don't store it anywhere. Recalculate it from the
+					// FirstReachableSessionPlatform. FirstReachableSessionPlatform must be non-null, otherwise
+					// bAnyCookable would be true.
+					check(FirstReachableSessionPlatform);
+					bool bCookable;
+					bool bExplorable;
+					Cluster.IsRequestCookable(FirstReachableSessionPlatform, Vertex.PackageData->GetPackageName(),
+						*Vertex.PackageData, SuppressCookReason, bCookable, bExplorable);
+					check(!bCookable); // We don't support bCookable changing for a given package and platform
+					check(SuppressCookReason != ESuppressCookReason::Invalid);
+				}
 			}
+			else
+			{
+				check(SuppressCookReason == ESuppressCookReason::NotSuppressed);
+			}
+			Cluster.SetPackageDataSuppressReason(*Vertex.PackageData, SuppressCookReason);
+			Vertex.bAnyCookable = bAnyCookable;
 		}
-		else
-		{
-			check(SuppressCookReason == ESuppressCookReason::NotSuppressed);
-		}
-		Cluster.SetPackageDataSuppressReason(*Vertex.PackageData, SuppressCookReason);
-		Vertex.bAnyCookable = bAnyCookable;
-	}
 
-	// If any target platform is cookable, then we need to mark the CookerLoadingPlatform as reachable because we will
-	// need to load the package to cook it
-	if (bAnyCookable)
-	{
-		if (!CookerLoadingPlatform)
+		// If any target platform is cookable, then we need to mark the CookerLoadingPlatform as reachable because we will
+		// need to load the package to cook it
+		if (bAnyCookable)
 		{
-			CookerLoadingPlatform = &Vertex.PackageData->FindOrAddPlatformData(CookerLoadingPlatformKey);
+			if (!CookerLoadingPlatform)
+			{
+				CookerLoadingPlatform = &Vertex.PackageData->FindOrAddPlatformData(CookerLoadingPlatformKey);
+			}
+			CookerLoadingPlatform->AddReachability(EReachability::Runtime);
 		}
-		CookerLoadingPlatform->SetReachable(true);
-	}
-	if (CookerLoadingPlatform && CookerLoadingPlatform->IsReachable() && !CookerLoadingPlatform->IsVisitedByCluster())
-	{
-		CookerLoadingPlatform->SetCookable(true);
-		CookerLoadingPlatform->SetExplorable(true);
-		CookerLoadingPlatform->SetVisitedByCluster(true);
-		if (TraversalTier >= ETraversalTier::FollowDependencies)
+		if (CookerLoadingPlatform && CookerLoadingPlatform->IsReachable(EReachability::Runtime)
+			&& !CookerLoadingPlatform->IsVisitedByCluster(EReachability::Runtime))
 		{
-			ShouldFetchPlatforms[CookerLoadingPlatformIndex] = true;
-			Vertex.PlatformData[CookerLoadingPlatformIndex].bExploreRequested = true;
+			CookerLoadingPlatform->SetCookable(true);
+			CookerLoadingPlatform->SetExplorable(true);
+			CookerLoadingPlatform->AddVisitedByCluster(EReachability::Runtime);
+			if (Cluster.TraversalExploreRuntimeDependencies())
+			{
+				ShouldFetchPlatforms[CookerLoadingPlatformIndex] = true;
+				Vertex.PlatformData[CookerLoadingPlatformIndex].bExploreRequested = true;
+			}
 		}
 	}
 
-	if (TraversalTier >= ETraversalTier::FetchEdgeData)
+	const bool bMightNeedToFetch = Cluster.TraversalExploreIncremental()
+		|| Cluster.TraversalExploreRuntimeDependencies();
+	if (bMightNeedToFetch)
 	{
 		for (int32 PlatformIndex = 0; PlatformIndex < LocalNumFetchPlatforms; ++PlatformIndex)
 		{
@@ -1172,30 +1306,34 @@ void FRequestCluster::FGraphSearch::VisitVertex(FVertexData& Vertex)
 }
 
 void FRequestCluster::FGraphSearch::VisitVertexForPlatform(FVertexData& Vertex, const ITargetPlatform* Platform,
-	FPackagePlatformData& PlatformData, ESuppressCookReason& AccumulatedSuppressCookReason)
+	EReachability ClusterReachability, FPackagePlatformData& PlatformData,
+	ESuppressCookReason& AccumulatedSuppressCookReason)
 {
-	FPackageData& PackageData = *Vertex.PackageData;
-	ESuppressCookReason SuppressCookReason = ESuppressCookReason::Invalid;
-	bool bCookable;
-	bool bExplorable;
-	Cluster.IsRequestCookable(Platform, Vertex.PackageData->GetPackageName(), PackageData, SuppressCookReason,
-		bCookable, bExplorable);
-	PlatformData.SetCookable(bCookable);
-	PlatformData.SetExplorable(bExplorable);
-	if (bCookable)
+	if (Cluster.TraversalMarkCookable())
 	{
-		AccumulatedSuppressCookReason = ESuppressCookReason::NotSuppressed;
-	}
-	else
-	{
-		check(SuppressCookReason != ESuppressCookReason::Invalid
-			&& SuppressCookReason != ESuppressCookReason::NotSuppressed);
-		if (AccumulatedSuppressCookReason == ESuppressCookReason::Invalid)
+		FPackageData& PackageData = *Vertex.PackageData;
+		ESuppressCookReason SuppressCookReason = ESuppressCookReason::Invalid;
+		bool bCookable;
+		bool bExplorable;
+		Cluster.IsRequestCookable(Platform, Vertex.PackageData->GetPackageName(), PackageData, SuppressCookReason,
+			bCookable, bExplorable);
+		PlatformData.SetCookable(bCookable);
+		PlatformData.SetExplorable(bExplorable);
+		if (bCookable)
 		{
-			AccumulatedSuppressCookReason = SuppressCookReason;
+			AccumulatedSuppressCookReason = ESuppressCookReason::NotSuppressed;
+		}
+		else
+		{
+			check(SuppressCookReason != ESuppressCookReason::Invalid
+				&& SuppressCookReason != ESuppressCookReason::NotSuppressed);
+			if (AccumulatedSuppressCookReason == ESuppressCookReason::Invalid)
+			{
+				AccumulatedSuppressCookReason = SuppressCookReason;
+			}
 		}
 	}
-	PlatformData.SetVisitedByCluster(true);
+	PlatformData.AddVisitedByCluster(ClusterReachability);
 }
 
 FRequestCluster::FGraphSearch::FExploreEdgesContext::FExploreEdgesContext(FRequestCluster& InCluster,
@@ -1346,12 +1484,6 @@ bool FRequestCluster::FGraphSearch::FExploreEdgesContext::TryCalculateIterativel
 		const ITargetPlatform* TargetPlatform = FetchPlatformData.Platform;
 		FPackagePlatformData& PackagePlatformData = PackageData->FindOrAddPlatformData(TargetPlatform);
 
-		if (!PackagePlatformData.IsCookable())
-		{
-			SetIsIterativelyUnmodified(PlatformIndex, false, PackagePlatformData);
-			continue;
-		}
-
 		UE::TargetDomain::FCookDependencies& CookDependencies = QueryPlatformData.CookAttachments.Dependencies;
 		const FAssetPackageData* OverrideAssetPackageData = nullptr;
 		FPackageData* ParentPackageData = nullptr;
@@ -1525,7 +1657,7 @@ void FRequestCluster::FGraphSearch::FExploreEdgesContext::CalculatePackageDataDe
 {
 	using namespace UE::AssetRegistry;
 
-	if (!bFetchAnyTargetPlatform)
+	if (!bFetchAnyTargetPlatform || !Cluster.TraversalMarkCookable())
 	{
 		return;
 	}
@@ -1586,7 +1718,7 @@ void FRequestCluster::FGraphSearch::FExploreEdgesContext::CalculateDependenciesA
 		FFetchPlatformData& FetchPlatformData = GraphSearch.FetchPlatforms[PlatformIndex];
 		const ITargetPlatform* TargetPlatform = FetchPlatformData.Platform;
 		FPackagePlatformData& PackagePlatformData = PackageData->FindOrAddPlatformData(TargetPlatform);
-		if ((GraphSearch.TraversalTier < ETraversalTier::FollowDependencies) || !PackagePlatformData.IsExplorable())
+		if (!GraphSearch.Cluster.TraversalExploreRuntimeDependencies() || !PackagePlatformData.IsExplorable())
 		{
 			// ExploreVertexEdges is responsible for updating package modification status so we might
 			// have been called for this platform even if not explorable. If not explorable, just update
@@ -1595,7 +1727,7 @@ void FRequestCluster::FGraphSearch::FExploreEdgesContext::CalculateDependenciesA
 			if (PlatformIndex != CookerLoadingPlatformIndex)
 			{
 				ProcessPlatformAttachments(PlatformIndex, TargetPlatform, FetchPlatformData, PackagePlatformData,
-					QueryPlatformData.CookAttachments, false /* bExploreDependencies */);
+					QueryPlatformData.CookAttachments, false /* bExploreRuntimeDependencies */);
 			}
 			continue;
 		}
@@ -1631,7 +1763,7 @@ void FRequestCluster::FGraphSearch::FExploreEdgesContext::CalculateDependenciesA
 			AddPlatformDependencyRange(HardEditorDependencies, PlatformIndex, EInstigator::HardEditorOnlyDependency);
 			AddPlatformDependencyRange(SoftGameDependencies, PlatformIndex, EInstigator::SoftDependency);
 			ProcessPlatformAttachments(PlatformIndex, TargetPlatform, FetchPlatformData, PackagePlatformData,
-				QueryPlatformData.CookAttachments, true /* bExploreDependencies  */);
+				QueryPlatformData.CookAttachments, true /* bExploreRuntimeDependencies  */);
 		}
 		if (DiscoveredDependencies)
 		{
@@ -1647,13 +1779,17 @@ void FRequestCluster::FGraphSearch::FExploreEdgesContext::QueueVisitsOfDependenc
 		return;
 	}
 
+	const EReachability ClusterReachability = Cluster.TraversalMarkCookable()
+		? EReachability::Runtime : EReachability::Build;
 	TArray<FPackageData*>* Edges = nullptr;
 	TRefCountPtr<FGenerationHelper> GenerationHelper = PackageData->GetGenerationHelper();
 	for (TPair<FName, FScratchPlatformDependencyBits>& PlatformDependencyPair : PlatformDependencyMap)
 	{
 		FName DependencyName = PlatformDependencyPair.Key;
-		TBitArray<>& HasPlatformByIndex = PlatformDependencyPair.Value.HasPlatformByIndex;
+		TBitArray<>& HasRuntimePlatformByIndex = PlatformDependencyPair.Value.HasRuntimePlatformByIndex;
+		TBitArray<>& HasBuildPlatformByIndex = PlatformDependencyPair.Value.HasBuildPlatformByIndex;
 		EInstigator InstigatorType = PlatformDependencyPair.Value.InstigatorType;
+		EInstigator BuildInstigatorType = PlatformDependencyPair.Value.BuildInstigatorType;
 
 		// Process any CoreRedirects before checking whether the package exists
 		FName Redirected = FCoreRedirects::GetRedirectedName(ECoreRedirectFlags::Type_Package,
@@ -1671,48 +1807,82 @@ void FRequestCluster::FGraphSearch::FExploreEdgesContext::QueueVisitsOfDependenc
 
 		for (int32 PlatformIndex = 0; PlatformIndex < LocalNumFetchPlatforms; ++PlatformIndex)
 		{
-			if (!HasPlatformByIndex[PlatformIndex])
+			if (!HasRuntimePlatformByIndex[PlatformIndex] &&
+				!HasBuildPlatformByIndex[PlatformIndex])
 			{
 				continue;
 			}
+
 			FFetchPlatformData& FetchPlatformData = GraphSearch.FetchPlatforms[PlatformIndex];
 			const ITargetPlatform* TargetPlatform = FetchPlatformData.Platform;
 			FPackagePlatformData& PlatformData = DependencyPackageData.FindOrAddPlatformData(TargetPlatform);
 
-			if (PlatformIndex == CookerLoadingPlatformIndex)
+			if (HasRuntimePlatformByIndex[PlatformIndex] && ClusterReachability == EReachability::Runtime)
 			{
-				if (!Edges)
+				// RuntimeDependencies only cause edges and reachability in EReachability::Runtime clusters.
+				if (PlatformIndex == CookerLoadingPlatformIndex)
 				{
-					Edges = &GraphSearch.GraphEdges.FindOrAdd(PackageData);
-					Edges->Reset(PlatformDependencyMap.Num());
+					if (!Edges)
+					{
+						Edges = &GraphSearch.GraphEdges.FindOrAdd(PackageData);
+						Edges->Reset(PlatformDependencyMap.Num());
+					}
+					Edges->Add(&DependencyPackageData);
 				}
-				Edges->Add(&DependencyPackageData);
-			}
 
-			if (!PlatformData.IsReachable())
-			{
-				PlatformData.SetReachable(true);
-				if (!DependencyPackageData.HasInstigator() && TargetPlatform != CookerLoadingPlatformKey)
+				if (!PlatformData.IsReachable(EReachability::Runtime))
+				{
+					PlatformData.AddReachability(EReachability::Runtime);
+					if (InstigatorType != EInstigator::InvalidCategory
+						&& !DependencyPackageData.HasInstigator(EReachability::Runtime)
+						&& TargetPlatform != CookerLoadingPlatformKey)
+					{
+						DependencyPackageData.SetInstigator(Cluster, EReachability::Runtime,
+							FInstigator(InstigatorType, PackageName));
+					}
+				}
+				if (!PlatformData.IsVisitedByCluster(EReachability::Runtime))
 				{
-					DependencyPackageData.SetInstigator(Cluster, FInstigator(InstigatorType, PackageName));
+					bAddToVisitVertexQueue = true;
 				}
 			}
-			if (!PlatformData.IsVisitedByCluster())
+			if (HasBuildPlatformByIndex[PlatformIndex] && PlatformIndex != CookerLoadingPlatformIndex)
 			{
-				bAddToVisitVertexQueue = true;
+				// CookerLoadingPlatform does not cause BuildDependencies.
+
+				// BuildDependencies from session platforms set the packages build reachability no matter what kind of
+				// cluster we're in, but they only get added to cluster if we're in a EReachability::Build cluster,
+				// Otherwise they will need to get picked up later by UCookOnTheFlyServer::KickBuildDependencies.
+				if (!PlatformData.IsReachable(EReachability::Build))
+				{
+					PlatformData.AddReachability(EReachability::Build);
+					if (BuildInstigatorType != EInstigator::InvalidCategory
+						&& !DependencyPackageData.HasInstigator(EReachability::Build))
+					{
+						DependencyPackageData.SetInstigator(Cluster, EReachability::Build,
+							FInstigator(BuildInstigatorType, PackageName));
+					}
+				}
+				if (ClusterReachability == EReachability::Build)
+				{
+					// Being Reachable for EReachability::Build does not necessarily mean that it needs to be visited by
+					// a cluster, as it does for EReachability::Runtime. We only need to visit BuildDependencies that 
+					// were not committed
+					if (!PlatformData.IsCommitted() && !PlatformData.IsVisitedByCluster(EReachability::Build))
+					{
+						bAddToVisitVertexQueue = true;
+					}
+				}
 			}
 		}
 		if (bAddToVisitVertexQueue)
 		{
-			if (DependencyVertex.PackageData)
+			// Only pull the vertex into the cluster if it has not already been pulled into the cluster.
+			// This prevents us from trying to readd a packagedata after COTFS called Cluster->RemovePackageData.
+			if (!DependencyVertex.bPulledIntoCluster)
 			{
-				// Only pull the vertex into the cluster if it has not already been pulled into the cluster.
-				// This prevents us from trying to readd a packagedata after COTFS called Cluster->RemovePackageData.
-				if (!DependencyVertex.bPulledIntoCluster)
-				{
-					DependencyVertex.bPulledIntoCluster = true;
-					Cluster.PullIntoCluster(*DependencyVertex.PackageData);
-				}
+				DependencyVertex.bPulledIntoCluster = true;
+				Cluster.PullIntoCluster(DependencyPackageData);
 			}
 			GraphSearch.AddToVisitVertexQueue(DependencyVertex);
 		}
@@ -1731,36 +1901,53 @@ void FRequestCluster::FGraphSearch::FExploreEdgesContext::AddPlatformDependency(
 	int32 PlatformIndex, EInstigator InstigatorType)
 {
 	FScratchPlatformDependencyBits& PlatformDependencyBits = PlatformDependencyMap.FindOrAdd(DependencyName);
-	if (PlatformDependencyBits.HasPlatformByIndex.Num() != LocalNumFetchPlatforms)
+	if (PlatformDependencyBits.HasRuntimePlatformByIndex.Num() != LocalNumFetchPlatforms)
 	{
-		PlatformDependencyBits.HasPlatformByIndex.Init(false, LocalNumFetchPlatforms);
-		PlatformDependencyBits.InstigatorType = EInstigator::SoftDependency;
+		PlatformDependencyBits.HasRuntimePlatformByIndex.Init(false, LocalNumFetchPlatforms);
+		PlatformDependencyBits.HasBuildPlatformByIndex.Init(false, LocalNumFetchPlatforms);
+		PlatformDependencyBits.InstigatorType = EInstigator::InvalidCategory;
+		PlatformDependencyBits.BuildInstigatorType = EInstigator::InvalidCategory;
 	}
-	PlatformDependencyBits.HasPlatformByIndex[PlatformIndex] = true;
 
-	// Calculate PlatformDependencyType.InstigatorType == 
-	// Max(InstigatorType, PlatformDependencyType.InstigatorType)
-	// based on the enum values, from least required to most: [ Soft, HardEditorOnly, Hard ]
-	switch (InstigatorType)
+	if (InstigatorType != EInstigator::BuildDependency)
 	{
-	case EInstigator::HardDependency:
-		PlatformDependencyBits.InstigatorType = InstigatorType;
-		break;
-	case EInstigator::HardEditorOnlyDependency:
-		if (PlatformDependencyBits.InstigatorType != EInstigator::HardDependency)
+		PlatformDependencyBits.HasRuntimePlatformByIndex[PlatformIndex] = true;
+
+		// For runtime dependencies, calculate PlatformDependencyType.InstigatorType == 
+		// Max(InstigatorType, PlatformDependencyType.InstigatorType)
+		// based on the enum values, from least required to most: [ Soft, HardEditorOnly, Hard ]
+		switch (InstigatorType)
 		{
+		case EInstigator::HardDependency:
 			PlatformDependencyBits.InstigatorType = InstigatorType;
+			break;
+		case EInstigator::HardEditorOnlyDependency:
+			if (PlatformDependencyBits.InstigatorType != EInstigator::HardDependency)
+			{
+				PlatformDependencyBits.InstigatorType = InstigatorType;
+			}
+			break;
+		case EInstigator::SoftDependency:
+			if (PlatformDependencyBits.InstigatorType != EInstigator::HardDependency
+				&& PlatformDependencyBits.InstigatorType != EInstigator::HardEditorOnlyDependency)
+			{
+				PlatformDependencyBits.InstigatorType = InstigatorType;
+			}
+			break;
+		case EInstigator::InvalidCategory:
+			// Caller indicated they do not want to set the InstigatorType
+			break;
+		default:
+			checkNoEntry();
+			break;
 		}
-		break;
-	case EInstigator::SoftDependency:
-		// New value is minimum, so keep the old value
-		break;
-	case EInstigator::InvalidCategory:
-		// Caller indicated they do not want to set the InstigatorType
-		break;
-	default:
-		checkNoEntry();
-		break;
+	}
+	else
+	{
+		PlatformDependencyBits.HasBuildPlatformByIndex[PlatformIndex] = true;
+
+		// For build dependencies, there is only one instigator type so just set it.
+		PlatformDependencyBits.BuildInstigatorType = InstigatorType;
 	}
 }
 
@@ -1776,33 +1963,45 @@ void FRequestCluster::FGraphSearch::FExploreEdgesContext::AddPlatformDependencyR
 void FRequestCluster::FGraphSearch::FExploreEdgesContext::ProcessPlatformAttachments(int32 PlatformIndex,
 	const ITargetPlatform* TargetPlatform, FFetchPlatformData& FetchPlatformData,
 	FPackagePlatformData& PackagePlatformData, UE::TargetDomain::FIncrementalCookAttachments& PlatformAttachments,
-	bool bExploreDependencies)
+	bool bExploreRuntimeDependencies)
 {
 	bool bFoundBuildDefinitions = false;
 	ICookedPackageWriter* PackageWriter = FetchPlatformData.Writer;
 	FQueryPlatformData& QueryPlatformData = Vertex->PlatformData[PlatformIndex];
 
-	if (Cluster.IsIncrementalCook() && PackagePlatformData.IsCookable())
+	if (Cluster.IsIncrementalCook())
 	{
 		check(QueryPlatformData.bIterativelyUnmodified.IsSet());
 		bool bIterativelyUnmodified = QueryPlatformData.bIterativelyUnmodified.GetValue();
 		if (bIterativelyUnmodified)
 		{
+			// Queue runtime dependencies if we are exploring runtime dependencies, build definitions if
+			// we are queuing those, and always queue build dependencies
 			UE::TargetDomain::FCookDependencies& CookDependencies = PlatformAttachments.Dependencies;
-			if (bExploreDependencies && Cluster.bAllowSoftDependencies)
+			if (bExploreRuntimeDependencies && PackagePlatformData.IsCookable() && Cluster.bAllowSoftDependencies)
 			{
 				AddPlatformDependencyRange(CookDependencies.GetRuntimePackageDependencies(), PlatformIndex,
 					EInstigator::SoftDependency);
 			}
 
-			if (Cluster.bPreQueueBuildDefinitions)
+			if (Cluster.bPreQueueBuildDefinitions && PackagePlatformData.IsCookable())
 			{
 				bFoundBuildDefinitions = true;
 				Cluster.BuildDefinitions.AddBuildDefinitionList(PackageName, TargetPlatform,
 					PlatformAttachments.BuildDefinitions.Definitions);
 			}
+
+			for (const UE::Cook::FCookDependency& TransitiveBuildDependency : 
+				CookDependencies.GetTransitiveBuildDependencies())
+			{
+				AddPlatformDependency(TransitiveBuildDependency.GetPackageName(), PlatformIndex,
+					EInstigator::BuildDependency);
+			}
 		}
 		bool bShouldIterativelySkip = bIterativelyUnmodified;
+		ECookResult CookResult =
+			QueryPlatformData.CookAttachments.CommitStatus == IPackageWriter::ECommitStatus::Success
+			? ECookResult::Succeeded : ECookResult::Failed;
 		if (QueryPlatformData.CookAttachments.CommitStatus == IPackageWriter::ECommitStatus::Error)
 		{
 			// Recook packages with errors even if they have not changed, so that the error is not hidden from
@@ -1814,8 +2013,20 @@ void FRequestCluster::FGraphSearch::FExploreEdgesContext::ProcessPlatformAttachm
 			// and will require care to avoid causing a softlock in the presence of a transitive build dependency
 			// cycle.
 		}
-		PackageWriter->UpdatePackageModificationStatus(PackageName, bIterativelyUnmodified,
-			bShouldIterativelySkip);
+		if (!QueryPlatformData.CookAttachments.Dependencies.HasSaveResults() && Cluster.TraversalMarkCookable())
+		{
+			// The package was previously committed only as a build dependency, with no attempt made to save its data
+			// because it was not marked reachable in the previous cook. But during this cook, it is marked reachable
+			// so we need to know the results of its save attempt. Therefore we need to cook it and it is not
+			// iteratively skippable.
+			bShouldIterativelySkip = false;
+		}
+
+		if (Cluster.TraversalMarkCookable())
+		{
+			PackageWriter->UpdatePackageModificationStatus(PackageName, bIterativelyUnmodified,
+				bShouldIterativelySkip);
+		}
 
 		TRefCountPtr<FGenerationHelper> ParentGenerationHelper;
 		if (PackageData->IsGenerated())
@@ -1834,25 +2045,34 @@ void FRequestCluster::FGraphSearch::FExploreEdgesContext::ProcessPlatformAttachm
 		}
 		if (bShouldIterativelySkip)
 		{
-			// Call SetPlatformCooked instead of just PackagePlatformData.SetCookResults because we might also need
-			// to set OnFirstCookedPlatformAdded
-			PackageData->SetPlatformCooked(TargetPlatform, ECookResult::Succeeded);
-			PackagePlatformData.SetIterativelySkipped(true);
-			if (TRefCountPtr<FGenerationHelper> GenerationHelper = PackageData->GetGenerationHelper())
+			if (Cluster.TraversalMarkCookable())
 			{
-				GenerationHelper->MarkPackageIterativelySkipped(*PackageData);
-			}
-			if (ParentGenerationHelper)
-			{
-				ParentGenerationHelper->MarkPackageIterativelySkipped(*PackageData);
+				// Call SetPlatformCooked instead of just PackagePlatformData.SetCookResults because we might also need
+				// to set OnFirstCookedPlatformAdded
+				PackageData->SetPlatformCooked(TargetPlatform, CookResult);
+				PackagePlatformData.SetIterativelySkipped(true);
+				if (TRefCountPtr<FGenerationHelper> GenerationHelper = PackageData->GetGenerationHelper())
+				{
+					GenerationHelper->MarkPackageIterativelySkipped(*PackageData);
+				}
+				if (ParentGenerationHelper)
+				{
+					ParentGenerationHelper->MarkPackageIterativelySkipped(*PackageData);
+				}
+				if (CookResult == ECookResult::Succeeded && PlatformIndex == FirstSessionPlatformIndex)
+				{
+					COOK_STAT(++DetailedCookStats::NumPackagesIterativelySkipped);
+				}
+				// Declare the package to the EDLCookInfo verification so we don't warn about missing exports from it
+				UE::SavePackageUtilities::EDLCookInfoAddIterativelySkippedPackage(PackageName);
 			}
-			Cluster.SetPackageDataWasMarkedCooked(*PackageData, true);
-			if (PlatformIndex == FirstSessionPlatformIndex)
+			else
 			{
-				COOK_STAT(++DetailedCookStats::NumPackagesIterativelySkipped);
+				// Mark the package as already committed.
+				PackagePlatformData.SetCommitted(true);
 			}
-			// Declare the package to the EDLCookInfo verification so we don't warn about missing exports from it
-			UE::SavePackageUtilities::EDLCookInfoAddIterativelySkippedPackage(PackageName);
+
+			Cluster.SetPackageDataWasMarkedSkipped(*PackageData, true);
 		}
 	}
 
diff --git a/Engine/Source/Editor/UnrealEd/Private/Cooker/CookRequestCluster.h b/Engine/Source/Editor/UnrealEd/Private/Cooker/CookRequestCluster.h
--- a/Engine/Source/Editor/UnrealEd/Private/Cooker/CookRequestCluster.h
+++ b/Engine/Source/Editor/UnrealEd/Private/Cooker/CookRequestCluster.h
@@ -31,6 +31,7 @@ class ICookedPackageWriter;
 class ITargetPlatform;
 class UCookOnTheFlyServer;
 namespace UE::Cook { class FRequestQueue; }
+namespace UE::Cook { enum class EReachability : uint8; }
 namespace UE::Cook { struct FDiscoveryQueueElement; }
 namespace UE::Cook { struct FFilePlatformRequest; }
 namespace UE::Cook { struct FPackageData; }
@@ -49,9 +50,15 @@ class FRequestCluster
 {
 public:
 	FRequestCluster(UCookOnTheFlyServer& COTFS, TArray<FFilePlatformRequest>&& InRequests);
-	FRequestCluster(UCookOnTheFlyServer& COTFS, FPackageDataSet&& InRequests);
+	FRequestCluster(UCookOnTheFlyServer& COTFS, FPackageDataSet&& InRequests, EReachability InExploreReachability);
 	FRequestCluster(UCookOnTheFlyServer& COTFS, TRingBuffer<FDiscoveryQueueElement>& DiscoveryQueue);
 	FRequestCluster(FRequestCluster&&) = default;
+	enum EBuildDependencyQueueConstructorType
+	{
+		BuildDependencyQueue
+	};
+	FRequestCluster(UCookOnTheFlyServer& COTFS, EBuildDependencyQueueConstructorType,
+		TRingBuffer<FPackageData*>& BuildDependencyDiscoveryQueue);
 
 	/**
 	 * Calculate the information needed to create a PackageData, and transitive search dependencies for all requests.
@@ -204,16 +211,32 @@ private:
 	/** How much traversal the GraphSearch should do based on settings for the entire cook. */
 	enum class ETraversalTier
 	{
-		/** Do not fetch any edgedata. Used on CookWorkers; the director already did the fetch. */
-		None,
 		/**
-		 * Fetch the edgedata and use it for ancillary calculation like updating whether a package is iteratively
-		 * unmodified. Do not explore the discovered dependencies.
+		 * Do not fetch any edgedata, do not evaluate skippability. Mark each input vertex as should-be-cooked.
+		 * Used on CookWorkers when saving runtime packages.
+		 */
+		MarkForRuntime,
+		/**
+		 * Do not fetch any edgedata, do not evaluate skippability. Mark each input vertex as should-be-committed.
+		 * Used on CookWorkers when committing build dependencies without saving them.
+		 */
+		MarkForBuildDependency,
+		/**
+		 * Mark vertices as skippable if they have uptodate dependencies, even without a saveresult. 
+		 * Explore dependencies necessary for evaluating modification status, otherwise do not explore dependencies.
+		 */
+		BuildDependencies,
+		/**
+		 * Mark vertices as skippable only if they have uptodate dependencies and a saveresult.
+		 * Explore dependencies necessary for evaluating modification status, otherwise do not explore dependencies.
+		 * Used when traversing runtime packages to save, with debug cooking flag such as -cooksinglepackagenorefs.
 		 */
-		FetchEdgeData,
-		/** Fetch the edgedata, update ancillary calculations, and explore the discovered dependencies. */
-		FollowDependencies,
-		All=FollowDependencies,
+		RuntimeVisitVertices,
+		/**
+		 * Mark vertices as skippable only if they have uptodate dependencies and a saveresult. Explore runtime dependencies
+		 * and add them to the cluster. Used when traversing runtime packages to save on the cookdirector.
+		 */
+		RuntimeFollowDependencies,
 	};
 
 	/**
@@ -226,12 +249,12 @@ private:
 	struct FGraphSearch
 	{
 	public:
-		FGraphSearch(FRequestCluster& InCluster, ETraversalTier TraversalTier);
+		FGraphSearch(FRequestCluster& InCluster);
 		~FGraphSearch();
 
 		// All public functions are callable only from the process thread
 		/** Skip the entire GraphSearch and just visit the Cluster's current OwnedPackageDatas. */
-		void VisitWithoutDependencies();
+		void VisitWithoutFetching();
 		/** Start a search from the Cluster's current OwnedPackageDatas. */
 		void StartSearch();
 		void OnNewReachablePlatforms(FPackageData* PackageData);
@@ -254,8 +277,10 @@ private:
 		// Scratch data structures used to avoid dynamic allocations; lifetime of each use is only on the stack
 		struct FScratchPlatformDependencyBits
 		{
-			TBitArray<> HasPlatformByIndex;
-			EInstigator InstigatorType = EInstigator::SoftDependency;
+			TBitArray<> HasRuntimePlatformByIndex;
+			TBitArray<> HasBuildPlatformByIndex;
+			EInstigator InstigatorType = EInstigator::InvalidCategory;
+			EInstigator BuildInstigatorType = EInstigator::InvalidCategory;
 		};
 		struct FExploreEdgesContext
 		{
@@ -319,7 +344,8 @@ private:
 		void VisitVertex(FVertexData& VertexData);
 		/** Calculate and store the vertex's PackageData's cookability for the platform. */
 		void VisitVertexForPlatform(FVertexData& VertexData, const ITargetPlatform* Platform,
-			FPackagePlatformData& PlatformData, ESuppressCookReason& AccumulatedSuppressCookReason);
+			EReachability ClusterReachability, FPackagePlatformData& PlatformData,
+			ESuppressCookReason& AccumulatedSuppressCookReason);
 		void ResolveTransitiveBuildDependencyCycle();
 
 		/** Find or add a Vertex for PackageName. If PackageData is provided, use it, otherwise look it up. */
@@ -357,7 +383,6 @@ private:
 		// Variables that are read-only during multithreading
 		TArray<FFetchPlatformData> FetchPlatforms;
 		FRequestCluster& Cluster;
-		ETraversalTier TraversalTier = ETraversalTier::All;
 
 		// Variables that are accessible only from the Process thread
 		/** A set of stack and scratch variables used when calculating and exploring the edges of a vertex. */
@@ -390,23 +415,23 @@ private:
 		bool IsValid() const;
 		/** The package's SuppressCookReason, either NotSuppressed or a reason it was suppressed. */
 		ESuppressCookReason GetSuppressReason() const;
-		/** Whether the package was marked as cooked for any platform by this cluster. */
-		bool WasMarkedCooked() const;
+		/** Whether the package was marked as committed for any platform by this cluster. */
+		bool WasMarkedSkipped() const;
 		/** Whether the values indicate the package should be added to the cluster's PackagesToMarkNotInProgress. */
 		bool ShouldMarkNotInProgress() const;
 
 		void SetValid();
 		void SetSuppressReason(ESuppressCookReason Value);
-		void SetWasMarkedCooked(bool bValue);
+		void SetWasMarkedSkipped(bool bValue);
 
 	private:
 		ESuppressCookReason SuppressCookReason = ESuppressCookReason::NotSuppressed;
 		bool bValid = false;
-		bool bWasMarkedCooked = false;
+		bool bWasMarkedSkipped = false;
 	};
 
 private:
-	explicit FRequestCluster(UCookOnTheFlyServer& COTFS);
+	FRequestCluster(UCookOnTheFlyServer& COTFS, EReachability ExploreReachability);
 	void ReserveInitialRequests(int32 RequestNum);
 	void PullIntoCluster(FPackageData& PackageData);
 	void FetchPackageNames(const FCookerTimer& CookerTimer, bool& bOutComplete);
@@ -415,12 +440,26 @@ private:
 	bool IsIncrementalCook() const;
 	void SetPackageDataSuppressReason(FPackageData& PackageData, ESuppressCookReason Reason,
 		bool* bOutExisted = nullptr);
-	void SetPackageDataWasMarkedCooked(FPackageData& PackageData, bool bValue, bool* bOutExisted = nullptr);
+	void SetPackageDataWasMarkedSkipped(FPackageData& PackageData, bool bValue, bool* bOutExisted = nullptr);
 	void IsRequestCookable(const ITargetPlatform* TargetPlatform, FName PackageName, FPackageData& PackageData,
 		ESuppressCookReason& OutReason, bool& bOutCookable, bool& bOutExplorable);
 	static void IsRequestCookable(const ITargetPlatform* TargetPlatform, FName PackageName, FPackageData& PackageData,
 		UCookOnTheFlyServer& InCOTFS, FStringView InDLCPath, ESuppressCookReason& OutReason, bool& bOutCookable,
 		bool& bOutExplorable);
+	/**
+	 * TraversalTier property: runtime dependencies of visited vertices should be explored and their targets
+	 * added to the cluster.
+	 */
+	bool TraversalExploreRuntimeDependencies();
+	/**
+	 * TraversalTier property: True if we need to test packages for iteratively skippable, false if we don't.
+	 */
+	bool TraversalExploreIncremental();
+	/**
+	 * TraversalTier property: True if we are marking packages as should-be-saved for runtime, false if
+	 * we are committing packages just to record CookDependencies instead of saving them for runtime.
+	 */
+	bool TraversalMarkCookable();
 
 	TArray<FFilePlatformRequest> FilePlatformRequests;
 	/**
@@ -439,6 +478,7 @@ private:
 	FPackageTracker& PackageTracker;
 	FBuildDefinitions& BuildDefinitions;
 	int32 PackagesToMarkNotInProgressCount = 0;
+	ETraversalTier TraversalTier = ETraversalTier::RuntimeFollowDependencies;
 	bool bAllowHardDependencies = true;
 	bool bAllowSoftDependencies = true;
 	bool bErrorOnEngineContentUse = false;
@@ -469,15 +509,15 @@ inline ESuppressCookReason FRequestCluster::FProcessingFlags::GetSuppressReason(
 	return SuppressCookReason;
 }
 
-inline bool FRequestCluster::FProcessingFlags::WasMarkedCooked() const
+inline bool FRequestCluster::FProcessingFlags::WasMarkedSkipped() const
 {
-	return bWasMarkedCooked;
+	return bWasMarkedSkipped;
 }
 
 inline bool FRequestCluster::FProcessingFlags::ShouldMarkNotInProgress() const
 {
 	return bValid & 
-		((SuppressCookReason != ESuppressCookReason::NotSuppressed) | bWasMarkedCooked);
+		((SuppressCookReason != ESuppressCookReason::NotSuppressed) | bWasMarkedSkipped);
 }
 
 inline void FRequestCluster::FProcessingFlags::SetValid()
@@ -490,9 +530,9 @@ inline void FRequestCluster::FProcessingFlags::SetSuppressReason(ESuppressCookRe
 	SuppressCookReason = Value;
 }
 
-inline void FRequestCluster::FProcessingFlags::SetWasMarkedCooked(bool bValue)
+inline void FRequestCluster::FProcessingFlags::SetWasMarkedSkipped(bool bValue)
 {
-	bWasMarkedCooked = bValue;
+	bWasMarkedSkipped = bValue;
 }
 
 inline int32 FRequestCluster::FGraphSearch::NumFetchPlatforms() const
diff --git a/Engine/Source/Editor/UnrealEd/Private/Cooker/CookSavePackage.cpp b/Engine/Source/Editor/UnrealEd/Private/Cooker/CookSavePackage.cpp
--- a/Engine/Source/Editor/UnrealEd/Private/Cooker/CookSavePackage.cpp
+++ b/Engine/Source/Editor/UnrealEd/Private/Cooker/CookSavePackage.cpp
@@ -73,7 +73,6 @@ void UCookOnTheFlyServer::SaveCookedPackage(UE::Cook::FSaveCookedPackageContext&
 
 	Context.SetupPackage();
 
-	ITargetPlatformManagerModule& TPM = GetTargetPlatformManagerRef();
 	TGuardValue<bool> ScopedOutputCookerWarnings(GOutputCookingWarnings, IsCookFlagSet(ECookInitializationFlags::OutputVerboseCookerWarnings));
 	// SavePackage can CollectGarbage, so we need to store the currently-unqueued PackageData in a separate variable that we register for garbage collection
 	TGuardValue<UE::Cook::FPackageData*> ScopedSavingPackageData(SavingPackageData, &Context.PackageData);
@@ -175,11 +174,30 @@ void UCookOnTheFlyServer::SaveCookedPackage(UE::Cook::FSaveCookedPackageContext&
 		*WriteToString<256>(Package->GetFName()), SaveDurationSeconds);
 }
 
+void UCookOnTheFlyServer::CommitUncookedPackage(UE::Cook::FSaveCookedPackageContext& Context)
+{
+	using namespace UE::Cook;
+
+	UPackage* Package = Context.Package;
+	Context.SetupPackage();
+
+	for (int32 PlatformIndex = 0; PlatformIndex < Context.PlatformsForPackage.Num(); ++PlatformIndex)
+	{
+		const ITargetPlatform* TargetPlatform = Context.PlatformsForPackage[PlatformIndex];
+		Context.SetupPlatform(TargetPlatform, PlatformIndex);
+		Context.SavePackageResult = ESavePackageResult::Canceled;
+		Context.FinishPlatform();
+	}
+
+	Context.FinishPackage();
+}
+
 namespace UE::Cook
 {
 
 FSaveCookedPackageContext::FSaveCookedPackageContext(UCookOnTheFlyServer& InCOTFS, UE::Cook::FPackageData& InPackageData,
-	TArrayView<const ITargetPlatform*> InPlatformsForPackage, UE::Cook::FTickStackData& InStackData)
+	TArrayView<const ITargetPlatform*> InPlatformsForPackage, UE::Cook::FTickStackData& InStackData,
+	EReachability InCommitType)
 	: COTFS(InCOTFS)
 	, PackageData(InPackageData)
 	, PlatformsForPackage(InPlatformsForPackage)
@@ -187,8 +205,10 @@ FSaveCookedPackageContext::FSaveCookedPackageContext(UCookOnTheFlyServer& InCOTF
 	, Package(PackageData.GetPackage())
 	, PackageName(Package ? Package->GetName() : FString())
 	, Filename(PackageData.GetFileName().ToString())
+	, CommitType(InCommitType)
 {
 	PlatformDependencies.SetNum(InPlatformsForPackage.Num());
+	check(CommitType == EReachability::Runtime || CommitType == EReachability::Build);
 }
 
 void FSaveCookedPackageContext::SetupPackage()
@@ -196,23 +216,26 @@ void FSaveCookedPackageContext::SetupPackage()
 	check(Package && Package->IsFullyLoaded()); // PackageData should not be in the save state if Package is not fully loaded
 	check(Package->GetPathName().Equals(PackageName)); // We should only be saving outermost packages, so the path name should be the same as the package name
 	check(!Filename.IsEmpty()); // PackageData guarantees FileName is non-empty; if not found it should never make it into save state
-	if (Package->HasAnyPackageFlags(PKG_ReloadingForCooker))
+	// Use SandboxFile to do path conversion to properly handle sandbox paths (outside of standard paths in particular).
+	Filename = COTFS.ConvertToFullSandboxPath(*Filename, true);
+
+	if (CommitType == EReachability::Runtime)
 	{
-		UE_LOG(LogCook, Warning, TEXT("Package %s marked as reloading for cook was requested to save"), *PackageName);
-		UE_LOG(LogCook, Fatal, TEXT("Package %s marked as reloading for cook was requested to save"), *PackageName);
-	}
+		if (Package->HasAnyPackageFlags(PKG_ReloadingForCooker))
+		{
+			UE_LOG(LogCook, Warning, TEXT("Package %s marked as reloading for cook was requested to save"), *PackageName);
+			UE_LOG(LogCook, Fatal, TEXT("Package %s marked as reloading for cook was requested to save"), *PackageName);
+		}
 
-	SaveFlags = SAVE_KeepGUID | SAVE_Async
-		| (COTFS.IsCookFlagSet(ECookInitializationFlags::Unversioned) ? SAVE_Unversioned : 0);
-	SaveFlags |= COTFS.IsCookFlagSet(ECookInitializationFlags::CookEditorOptional) ? SAVE_Optional : SAVE_None;
+		SaveFlags = SAVE_KeepGUID | SAVE_Async
+			| (COTFS.IsCookFlagSet(ECookInitializationFlags::Unversioned) ? SAVE_Unversioned : 0);
+		SaveFlags |= COTFS.IsCookFlagSet(ECookInitializationFlags::CookEditorOptional) ? SAVE_Optional : SAVE_None;
 
-	if (COTFS.CookByTheBookOptions->bCookSoftPackageReferences)
-	{
-		SaveFlags |= SAVE_CookSoftPackageReferences;
+		if (COTFS.CookByTheBookOptions->bCookSoftPackageReferences)
+		{
+			SaveFlags |= SAVE_CookSoftPackageReferences;
+		}
 	}
-
-	// Use SandboxFile to do path conversion to properly handle sandbox paths (outside of standard paths in particular).
-	Filename = COTFS.ConvertToFullSandboxPath(*Filename, true);
 }
 
 void FSaveCookedPackageContext::SetupPlatform(const ITargetPlatform* InTargetPlatform, int32 InPlatformIndex)
@@ -222,79 +245,103 @@ void FSaveCookedPackageContext::SetupPlatform(const ITargetPlatform* InTargetPla
 	PlatFilename = Filename.Replace(TEXT("[Platform]"), *TargetPlatform->PlatformName());
 	bPlatformSetupSuccessful = false;
 
-	ArchiveCookContext.Emplace(Package,
-		COTFS.IsDirectorCookByTheBook() ? UE::Cook::ECookType::ByTheBook : UE::Cook::ECookType::OnTheFly,
-		COTFS.IsCookingDLC() ? UE::Cook::ECookingDLC::Yes : UE::Cook::ECookingDLC::No,
-		TargetPlatform);
-
-	// don't save Editor resources from the Engine if the target doesn't have editoronly data
-	if (COTFS.IsCookFlagSet(ECookInitializationFlags::SkipEditorContent) &&
-		(PackageName.StartsWith(TEXT("/Engine/Editor")) || PackageName.StartsWith(TEXT("/Engine/VREditor"))) &&
-		!TargetPlatform->HasEditorOnlyData())
+	if (CommitType == EReachability::Runtime)
 	{
-		SavePackageResult = ESavePackageResult::ContainsEditorOnlyData;
-		const TCHAR* RejectedReason = TEXT("EngineEditorContent");
-		if (GCookProgressDisplay & (int32)ECookProgressDisplayMode::Instigators)
-		{
-			UE_LOG(LogCook, Display, TEXT("Cooking %s, Instigator: { %s } -> Rejected %s"), *PackageName,
-				*(PackageData.GetInstigator().ToString()), RejectedReason);
-		}
-		else
+		ArchiveCookContext.Emplace(Package,
+			COTFS.IsDirectorCookByTheBook() ? UE::Cook::ECookType::ByTheBook : UE::Cook::ECookType::OnTheFly,
+			COTFS.IsCookingDLC() ? UE::Cook::ECookingDLC::Yes : UE::Cook::ECookingDLC::No,
+			TargetPlatform);
+
+		// don't save Editor resources from the Engine if the target doesn't have editoronly data
+		if (COTFS.IsCookFlagSet(ECookInitializationFlags::SkipEditorContent) &&
+			(PackageName.StartsWith(TEXT("/Engine/Editor")) || PackageName.StartsWith(TEXT("/Engine/VREditor"))) &&
+			!TargetPlatform->HasEditorOnlyData())
 		{
-			UE_LOG(LogCook, Display, TEXT("Cooking %s -> Rejected %s"), *PackageName, RejectedReason);
+			SavePackageResult = ESavePackageResult::ContainsEditorOnlyData;
+			const TCHAR* RejectedReason = TEXT("EngineEditorContent");
+			if (GCookProgressDisplay & (int32)ECookProgressDisplayMode::Instigators)
+			{
+				UE_LOG(LogCook, Display, TEXT("Cooking %s, Instigator: { %s } -> Rejected %s"), *PackageName,
+					*(PackageData.GetInstigator(EReachability::Runtime).ToString()), RejectedReason);
+			}
+			else
+			{
+				UE_LOG(LogCook, Display, TEXT("Cooking %s -> Rejected %s"), *PackageName, RejectedReason);
+			}
+			return;
 		}
-		return;
-	}
-	// Check whether or not game-specific behaviour should prevent this package from being cooked for the target platform
-	else if (!UAssetManager::Get().ShouldCookForPlatform(Package, TargetPlatform))
-	{
-		SavePackageResult = ESavePackageResult::ContainsEditorOnlyData;
-		const TCHAR* RejectedReason = TEXT("NotAssetManagerShouldCookForPlatform");
-		if (GCookProgressDisplay & (int32)ECookProgressDisplayMode::Instigators)
+		// Check whether or not game-specific behaviour should prevent this package from being cooked for the target platform
+		else if (!UAssetManager::Get().ShouldCookForPlatform(Package, TargetPlatform))
 		{
-			UE_LOG(LogCook, Display, TEXT("Cooking %s, Instigator: { %s } -> Rejected %s"), *PackageName,
-				*(PackageData.GetInstigator().ToString()), RejectedReason);
+			SavePackageResult = ESavePackageResult::ContainsEditorOnlyData;
+			const TCHAR* RejectedReason = TEXT("NotAssetManagerShouldCookForPlatform");
+			if (GCookProgressDisplay & (int32)ECookProgressDisplayMode::Instigators)
+			{
+				UE_LOG(LogCook, Display, TEXT("Cooking %s, Instigator: { %s } -> Rejected %s"), *PackageName,
+					*(PackageData.GetInstigator(EReachability::Runtime).ToString()), RejectedReason);
+			}
+			else
+			{
+				UE_LOG(LogCook, Display, TEXT("Cooking %s -> Rejected %s"), *PackageName, RejectedReason);
+			}
+			return;
 		}
+		// check if this package is unsupported for the target platform (typically plugin content)
 		else
 		{
-			UE_LOG(LogCook, Display, TEXT("Cooking %s -> Rejected %s"), *PackageName, RejectedReason);
-		}
-		return;
-	}
-	// check if this package is unsupported for the target platform (typically plugin content)
-	else
-	{
-		if (TSet<FName>* NeverCookPackages = COTFS.PackageTracker->PlatformSpecificNeverCookPackages.Find(TargetPlatform))
-		{
-			FGenerationHelper* GenerationHelper =
-				PackageData.IsGenerated() ? PackageData.GetParentGenerationHelper() : nullptr;
-
-			if (NeverCookPackages->Find(Package->GetFName()) ||
-				(GenerationHelper && NeverCookPackages->Find(GenerationHelper->GetOwner().GetPackageName())))
+			if (TSet<FName>* NeverCookPackages = COTFS.PackageTracker->PlatformSpecificNeverCookPackages.Find(TargetPlatform))
 			{
-				SavePackageResult = ESavePackageResult::ContainsEditorOnlyData;
-				const TCHAR* RejectedReason = TEXT("PlatformSpecificNeverCook");
-				if (GCookProgressDisplay & (int32)ECookProgressDisplayMode::Instigators)
-				{
-					UE_LOG(LogCook, Display, TEXT("Cooking %s, Instigator: { %s } -> Rejected %s"), *PackageName,
-						*(PackageData.GetInstigator().ToString()), RejectedReason);
-				}
-				else
+				FGenerationHelper* GenerationHelper =
+					PackageData.IsGenerated() ? PackageData.GetParentGenerationHelper() : nullptr;
+
+				if (NeverCookPackages->Find(Package->GetFName()) ||
+					(GenerationHelper && NeverCookPackages->Find(GenerationHelper->GetOwner().GetPackageName())))
 				{
-					UE_LOG(LogCook, Display, TEXT("Cooking %s -> Rejected %s"), *PackageName, RejectedReason);
+					SavePackageResult = ESavePackageResult::ContainsEditorOnlyData;
+					const TCHAR* RejectedReason = TEXT("PlatformSpecificNeverCook");
+					if (GCookProgressDisplay & (int32)ECookProgressDisplayMode::Instigators)
+					{
+						UE_LOG(LogCook, Display, TEXT("Cooking %s, Instigator: { %s } -> Rejected %s"), *PackageName,
+							*(PackageData.GetInstigator(EReachability::Runtime).ToString()), RejectedReason);
+					}
+					else
+					{
+						UE_LOG(LogCook, Display, TEXT("Cooking %s -> Rejected %s"), *PackageName, RejectedReason);
+					}
+					return;
 				}
-				return;
 			}
 		}
-	}
 
-	const FString FullFilename = FPaths::ConvertRelativePathToFull(PlatFilename);
-	if (FullFilename.Len() >= FPlatformMisc::GetMaxPathLength())
-	{
-		LogCookerMessage(FString::Printf(TEXT("Couldn't save package, filename is too long (%d >= %d): %s"),
-			FullFilename.Len(), FPlatformMisc::GetMaxPathLength(), *FullFilename), EMessageSeverity::Error);
-		SavePackageResult = ESavePackageResult::Error;
-		return;
+		const FString FullFilename = FPaths::ConvertRelativePathToFull(PlatFilename);
+		if (FullFilename.Len() >= FPlatformMisc::GetMaxPathLength())
+		{
+			LogCookerMessage(FString::Printf(TEXT("Couldn't save package, filename is too long (%d >= %d): %s"),
+				FullFilename.Len(), FPlatformMisc::GetMaxPathLength(), *FullFilename), EMessageSeverity::Error);
+			SavePackageResult = ESavePackageResult::Error;
+			return;
+		}
+
+		bEndianSwap = (!TargetPlatform->IsLittleEndian()) ^ (!PLATFORM_LITTLE_ENDIAN);
+
+		if (!TargetPlatform->HasEditorOnlyData())
+		{
+			Package->SetPackageFlags(PKG_FilterEditorOnly);
+		}
+		else
+		{
+			Package->ClearPackageFlags(PKG_FilterEditorOnly);
+		}
+
+		// Set platform-specific save flags
+		FPackagePlatformData& PlatformData = PackageData.FindOrAddPlatformData(TargetPlatform);
+		uint32 PlatformSaveFlagsMask = SAVE_AllowTimeout;
+		SaveFlags = (SaveFlags & ~PlatformSaveFlagsMask);
+		if (!PlatformData.IsSaveTimedOut())
+		{
+			// If we timedout before, do not allow another timeout, otherwise do allow it
+			SaveFlags |= SAVE_AllowTimeout;
+		}
 	}
 
 	if (!bHasDelayLoaded)
@@ -309,21 +356,21 @@ void FSaveCookedPackageContext::SetupPlatform(const ITargetPlatform* InTargetPla
 		bHasDelayLoaded = true;
 	}
 
-	UE_CLOG((GCookProgressDisplay & (int32)ECookProgressDisplayMode::Instigators) && PlatformIndex == 0,
-		LogCook, Display, TEXT("Cooking %s, Instigator: { %s }"),
-		*PackageName, *(PackageData.GetInstigator().ToString()));
-	UE_CLOG(GCookProgressDisplay & (int32)ECookProgressDisplayMode::PackageNames, LogCook, Display,
-		TEXT("Cooking %s"), *PackageName);
-
-	bEndianSwap = (!TargetPlatform->IsLittleEndian()) ^ (!PLATFORM_LITTLE_ENDIAN);
-
-	if (!TargetPlatform->HasEditorOnlyData())
+	if (CommitType == EReachability::Runtime)
 	{
-		Package->SetPackageFlags(PKG_FilterEditorOnly);
+		UE_CLOG((GCookProgressDisplay & (int32)ECookProgressDisplayMode::Instigators) && PlatformIndex == 0,
+			LogCook, Display, TEXT("Cooking %s, Instigator: { %s }"),
+			*PackageName, *(PackageData.GetInstigator(EReachability::Runtime).ToString()));
+		UE_CLOG(GCookProgressDisplay & (int32)ECookProgressDisplayMode::PackageNames, LogCook, Display,
+			TEXT("Cooking %s"), *PackageName);
 	}
 	else
 	{
-		Package->ClearPackageFlags(PKG_FilterEditorOnly);
+		UE_CLOG((GCookProgressDisplay & (int32)ECookProgressDisplayMode::Instigators) && PlatformIndex == 0,
+			LogCook, Display, TEXT("Committing BuildDependencies for %s, Instigator: { %s }"),
+			*PackageName, *(PackageData.GetInstigator(EReachability::Build).ToString()));
+		UE_CLOG(GCookProgressDisplay & (int32)ECookProgressDisplayMode::PackageNames, LogCook, Display,
+			TEXT("Committing BuildDependencies for %s"), *PackageName);
 	}
 
 	CookContext = &COTFS.FindOrCreateSaveContext(TargetPlatform);
@@ -337,15 +384,6 @@ void FSaveCookedPackageContext::SetupPlatform(const ITargetPlatform* InTargetPla
 	{
 		CookContext->DeterminismManager->BeginPackage(Package, TargetPlatform, PackageWriter);
 	}
-	// Set platform-specific save flags
-	FPackagePlatformData& PlatformData = PackageData.FindOrAddPlatformData(TargetPlatform);
-	uint32 PlatformSaveFlagsMask = SAVE_AllowTimeout;
-	SaveFlags = (SaveFlags & ~PlatformSaveFlagsMask);
-	if (!PlatformData.IsSaveTimedOut())
-	{
-		// If we timedout before, do not allow another timeout, otherwise do allow it
-		SaveFlags |= SAVE_AllowTimeout;
-	}
 
 	// Indicate Setup was successful
 	bPlatformSetupSuccessful = true;
@@ -362,8 +400,8 @@ void FSaveCookedPackageContext::FinishPlatform()
 	TRACE_CPUPROFILER_EVENT_SCOPE(FSaveCookedPackageContext::FinishPlatform);
 	check(PlatformIndex >= 0 && PlatformDependencies.IsValidIndex(PlatformIndex));
 
-	bool bSuccessful = SavePackageResult.IsSuccessful();
-	ECookResult CookResult = bSuccessful ? ECookResult::Succeeded : ECookResult::Failed;
+	bool bSuccessful = HasSavePackageResult() ? SavePackageResult.IsSuccessful() : false;
+	ECookResult CookResult;
 
 	// Calculate up-to-date assetregistry data for Generator and Generated packages
 	TOptional<TArray<FAssetDependency>> OverridePackageDependencies;
@@ -371,25 +409,34 @@ void FSaveCookedPackageContext::FinishPlatform()
 	TOptional<FAssetPackageData> AssetPackageDataBuffer;
 	TOptional<FAssetPackageData> OverrideAssetPackageData;
 	const FAssetPackageData* AssetPackageData = nullptr;
-	if (TRefCountPtr<FGenerationHelper> GenerationHelper = PackageData.GetGenerationHelper();
-		GenerationHelper)
+	if (CommitType == EReachability::Runtime)
 	{
-		OverridePackageDependencies.Emplace();
-		GenerationHelper->FinishGeneratorPlatformSave(PackageData, PlatformIndex == 0, *OverridePackageDependencies);
-	}
-	else if (TRefCountPtr<FGenerationHelper> ParentGenerationHelper = PackageData.GetParentGenerationHelper();
-		ParentGenerationHelper)
-	{
-		GeneratedResult.Emplace();
-		ParentGenerationHelper->FinishGeneratedPlatformSave(PackageData, *GeneratedResult);
-		OverridePackageDependencies.Emplace(GeneratedResult->PackageDependencies);
-		OverrideAssetPackageData.Emplace(GeneratedResult->AssetPackageData);
-		AssetPackageData = &GeneratedResult->AssetPackageData;
+		CookResult = bSuccessful ? ECookResult::Succeeded : ECookResult::Failed;
+
+		if (TRefCountPtr<FGenerationHelper> GenerationHelper = PackageData.GetGenerationHelper();
+			GenerationHelper)
+		{
+			OverridePackageDependencies.Emplace();
+			GenerationHelper->FinishGeneratorPlatformSave(PackageData, PlatformIndex == 0, *OverridePackageDependencies);
+		}
+		else if (TRefCountPtr<FGenerationHelper> ParentGenerationHelper = PackageData.GetParentGenerationHelper();
+			ParentGenerationHelper)
+		{
+			GeneratedResult.Emplace();
+			ParentGenerationHelper->FinishGeneratedPlatformSave(PackageData, *GeneratedResult);
+			OverridePackageDependencies.Emplace(GeneratedResult->PackageDependencies);
+			OverrideAssetPackageData.Emplace(GeneratedResult->AssetPackageData);
+			AssetPackageData = &GeneratedResult->AssetPackageData;
+		}
+		if (!AssetPackageData)
+		{
+			AssetPackageDataBuffer = COTFS.AssetRegistry->GetAssetPackageDataCopy(Package->GetFName());
+			AssetPackageData = AssetPackageDataBuffer.GetPtrOrNull();
+		}
 	}
-	if (!AssetPackageData)
+	else
 	{
-		AssetPackageDataBuffer = COTFS.AssetRegistry->GetAssetPackageDataCopy(Package->GetFName());
-		AssetPackageData = AssetPackageDataBuffer.GetPtrOrNull();
+		CookResult = ECookResult::NotAttempted;
 	}
 
 	// Commit the saved bytes and the iterative cook data to the PackageWriter
@@ -403,7 +450,8 @@ void FSaveCookedPackageContext::FinishPlatform()
 
 		ICookedPackageWriter::FCommitPackageInfo Info;
 		Info.Attachments = GetCommitAttachments();
-		Info.Status = PackageResultToCommitStatus(SavePackageResult.Result);
+		Info.Status = HasSavePackageResult() ? PackageResultToCommitStatus(SavePackageResult.Result)
+			: IPackageWriter::ECommitStatus::NothingToCook;
 		Info.PackageName = Package->GetFName();
 		Info.PackageHash = AssetPackageData ? AssetPackageData->GetPackageSavedHash() : FIoHash();
 		Info.WriteOptions = GetCommitWriteOptions();
@@ -429,54 +477,67 @@ void FSaveCookedPackageContext::FinishPlatform()
 			COTFS);
 	}
 
-	// If not retrying, mark the package as cooked, either successfully or with failure
-	bool bIsRetryErrorCode = IsRetryErrorCode(SavePackageResult.Result);
-	if (!bIsRetryErrorCode)
+	if (CommitType == EReachability::Runtime)
 	{
-		PackageData.SetPlatformCooked(TargetPlatform, CookResult);
-	}
+		// If not retrying, mark the package as cooked, either successfully or with failure
+		bool bIsRetryErrorCode = IsRetryErrorCode(SavePackageResult.Result);
+		if (!bIsRetryErrorCode)
+		{
+			PackageData.SetPlatformCooked(TargetPlatform, CookResult);
+		}
 
-	// Update flags used to determine garbage collection.
-	if (bSuccessful)
-	{
-		if (bContainsMap)
+		// Update flags used to determine garbage collection.
+		if (bSuccessful)
 		{
-			StackData.ResultFlags |= UCookOnTheFlyServer::COSR_CookedMap;
+			if (bContainsMap)
+			{
+				StackData.ResultFlags |= UCookOnTheFlyServer::COSR_CookedMap;
+			}
+			else
+			{
+				++COTFS.CookedPackageCountSinceLastGC;
+				StackData.ResultFlags |= UCookOnTheFlyServer::COSR_CookedPackage;
+			}
 		}
-		else
+
+		// Accumulate results for SaveCookedPackage_Finish
+		if (SavePackageResult.Result == ESavePackageResult::Timeout)
 		{
-			++COTFS.CookedPackageCountSinceLastGC;
-			StackData.ResultFlags |= UCookOnTheFlyServer::COSR_CookedPackage;
+			PackageData.FindOrAddPlatformData(TargetPlatform).SetSaveTimedOut(true);
+			bHasTimeOut = true;
 		}
-	}
+		bAnySaveSucceeded |= bSuccessful;
+
+		bHasRetryErrorCode |= bIsRetryErrorCode;
 
-	// Accumulate results for SaveCookedPackage_Finish
-	if (SavePackageResult.Result == ESavePackageResult::Timeout)
+		ArchiveCookContext.Reset();
+	}
+	else
 	{
-		PackageData.FindOrAddPlatformData(TargetPlatform).SetSaveTimedOut(true);
-		bHasTimeOut = true;
+		PackageData.SetPlatformCommitted(TargetPlatform);
 	}
-	bAnySaveSucceeded |= bSuccessful;
-
-	bHasRetryErrorCode |= bIsRetryErrorCode;
-
-	ArchiveCookContext.Reset();
 	PlatformIndex = -1;
 }
 
 void FSaveCookedPackageContext::FinishPackage()
 {
 	// If any save succeeded, add all dependencies from all platforms to the cook for the platform
-	if (bAnySaveSucceeded && !COTFS.CookByTheBookOptions->bSkipSoftReferences)
+	if (CommitType == EReachability::Build || bAnySaveSucceeded)
 	{
+		bool bAddSoftReferences = CommitType == EReachability::Runtime && bAnySaveSucceeded
+			&& !COTFS.CookByTheBookOptions->bSkipSoftReferences;
+
 		FName PackageFName = Package->GetFName();
 		if (PlatformsForPackage.Num() == 1)
 		{
 			TConstArrayView<const ITargetPlatform*> ReachablePlatforms(&TargetPlatform, 1);
-			for (const TPair<FPackageData*, EInstigator>& DependencyPair : PlatformDependencies[0].RuntimeDependencies)
+			if (bAddSoftReferences)
 			{
-				COTFS.QueueDiscoveredPackage(*DependencyPair.Key,
-					FInstigator(DependencyPair.Value, PackageFName), FDiscoveredPlatformSet(ReachablePlatforms));
+				for (const TPair<FPackageData*, EInstigator>& DependencyPair : PlatformDependencies[0].RuntimeDependencies)
+				{
+					COTFS.QueueDiscoveredPackage(*DependencyPair.Key,
+						FInstigator(DependencyPair.Value, PackageFName), FDiscoveredPlatformSet(ReachablePlatforms));
+				}
 			}
 			for (FPackageData* BuildDependency : PlatformDependencies[0].BuildDependencies)
 			{
@@ -492,14 +553,17 @@ void FSaveCookedPackageContext::FinishPackage()
 				const FPlatformDiscoveryData& DData = PlatformDependencies[LocalIndex];
 				const ITargetPlatform* CurrentPlatform = PlatformsForPackage[LocalIndex];
 
-				// Merge the runtime dependencies for each package into a single QueueDiscoveredPackage call, if possible
-				// It will not be possible if the different platforms have different instigators, so track a list of
-				// platforms for each package for each instigator type.
-				for (const TPair<FPackageData*, EInstigator>& PackagePlatformPair : DData.RuntimeDependencies)
+				if (bAddSoftReferences)
 				{
-					TMap<EInstigator, TArray<const ITargetPlatform*>>& TargetMap =
-						PackagePlatformsForInstigator.FindOrAdd(PackagePlatformPair.Key);
-					TargetMap.FindOrAdd(PackagePlatformPair.Value).Add(CurrentPlatform);
+					// Merge the runtime dependencies for each package into a single QueueDiscoveredPackage call, if possible
+					// It will not be possible if the different platforms have different instigators, so track a list of
+					// platforms for each package for each instigator type.
+					for (const TPair<FPackageData*, EInstigator>& PackagePlatformPair : DData.RuntimeDependencies)
+					{
+						TMap<EInstigator, TArray<const ITargetPlatform*>>& TargetMap =
+							PackagePlatformsForInstigator.FindOrAdd(PackagePlatformPair.Key);
+						TargetMap.FindOrAdd(PackagePlatformPair.Value).Add(CurrentPlatform);
+					}
 				}
 
 				// Merge the build dependencies for each package into a single QueueDiscoveredPackage call.
@@ -523,38 +587,41 @@ void FSaveCookedPackageContext::FinishPackage()
 		}
 	}
 
-	if (COTFS.IsDebugRecordUnsolicited())
+	if (CommitType == EReachability::Runtime)
 	{
-		COTFS.ProcessUnsolicitedPackages();
-		TMap<FPackageData*, EInstigator> AllPlatformDependenciesBuffer;
-		TMap<FPackageData*, EInstigator>* AllPlatformDependencies = &PlatformDependencies[0].RuntimeDependencies;
-		if (PlatformsForPackage.Num() > 1)
+		if (COTFS.IsDebugRecordUnsolicited())
 		{
-			AllPlatformDependencies = &AllPlatformDependenciesBuffer;
-			for (int32 LocalIndex = 0; LocalIndex < PlatformsForPackage.Num(); ++LocalIndex)
+			COTFS.ProcessUnsolicitedPackages();
+			TMap<FPackageData*, EInstigator> AllPlatformDependenciesBuffer;
+			TMap<FPackageData*, EInstigator>* AllPlatformDependencies = &PlatformDependencies[0].RuntimeDependencies;
+			if (PlatformsForPackage.Num() > 1)
 			{
-				AllPlatformDependenciesBuffer.Append(PlatformDependencies[LocalIndex].RuntimeDependencies);
+				AllPlatformDependencies = &AllPlatformDependenciesBuffer;
+				for (int32 LocalIndex = 0; LocalIndex < PlatformsForPackage.Num(); ++LocalIndex)
+				{
+					AllPlatformDependenciesBuffer.Append(PlatformDependencies[LocalIndex].RuntimeDependencies);
+				}
 			}
-		}
 
-		FDiagnostics::AnalyzeHiddenDependencies(COTFS, PackageData, PackageData.DetachUnsolicited(),
-			*AllPlatformDependencies, PlatformsForPackage, COTFS.bOnlyEditorOnlyDebug, COTFS.bHiddenDependenciesDebug);
-	}
+			FDiagnostics::AnalyzeHiddenDependencies(COTFS, PackageData, PackageData.DetachUnsolicited(),
+				*AllPlatformDependencies, PlatformsForPackage, COTFS.bOnlyEditorOnlyDebug, COTFS.bHiddenDependenciesDebug);
+		}
 
-	if (!bHasRetryErrorCode)
-	{
-		if (COTFS.IsCookOnTheFlyMode() && PackageData.GetUrgency() != EUrgency::Blocking &&
-			(!COTFS.CookOnTheFlyRequestManager || COTFS.CookOnTheFlyRequestManager->ShouldUseLegacyScheduling()))
+		if (!bHasRetryErrorCode)
 		{
-			// this is an unsolicited package
-			if (FPaths::FileExists(Filename))
+			if (COTFS.IsCookOnTheFlyMode() && PackageData.GetUrgency() != EUrgency::Blocking &&
+				(!COTFS.CookOnTheFlyRequestManager || COTFS.CookOnTheFlyRequestManager->ShouldUseLegacyScheduling()))
 			{
-				COTFS.PackageTracker->UnsolicitedCookedPackages.AddCookedPackage(
-					FFilePlatformRequest(PackageData.GetFileName(), EInstigator::Unspecified, PlatformsForPackage));
+				// this is an unsolicited package
+				if (FPaths::FileExists(Filename))
+				{
+					COTFS.PackageTracker->UnsolicitedCookedPackages.AddCookedPackage(
+						FFilePlatformRequest(PackageData.GetFileName(), EInstigator::Unspecified, PlatformsForPackage));
 
 #if DEBUG_COOKONTHEFLY
-				UE_LOG(LogCook, Display, TEXT("UnsolicitedCookedPackages: %s"), *Filename);
+					UE_LOG(LogCook, Display, TEXT("UnsolicitedCookedPackages: %s"), *Filename);
 #endif
+				}
 			}
 		}
 	}
@@ -629,9 +696,9 @@ void FSaveCookedPackageContext::CalculatePlatformRuntimeDependencies()
 	CurrentDependencies.Append(PlatformAgnosticDependencies);
 
 	// Add imports and softobjectpaths from the save to the SaveDependencies for the current Platform and for all Platforms
-	if (COTFS.bSkipOnlyEditorOnly)
+	FName PackageFName = Package->GetFName();
+	if (COTFS.bSkipOnlyEditorOnly && HasSavePackageResult())
 	{
-		FName PackageFName = Package->GetFName();
 		TConstArrayView<const ITargetPlatform*> ReachablePlatforms(&TargetPlatform, 1);
 		for (const TArray<FName>* DependencyNames : { &SavePackageResult.ImportPackages, &SavePackageResult.SoftPackageReferences })
 		{
@@ -675,7 +742,8 @@ void FSaveCookedPackageContext::CalculateCookDependencies(UE::TargetDomain::FGen
 			}
 		}
 		PlatformCookAttachments = UE::TargetDomain::FIncrementalCookAttachments::Collect(Package, TargetPlatform,
-			&SavePackageResult, GeneratedResult, MoveTemp(PlatformRuntimeDependencies));
+			HasSavePackageResult() ? &SavePackageResult : nullptr, GeneratedResult,
+			MoveTemp(PlatformRuntimeDependencies));
 	}
 }
 
@@ -741,4 +809,9 @@ void FSaveCookedPackageContext::AddDependency(TMap<FPackageData*, EInstigator>&
 	}
 }
 
+bool FSaveCookedPackageContext::HasSavePackageResult() const
+{
+	return CommitType == EReachability::Runtime;
+}
+
 } // namespace UE::Cook
diff --git a/Engine/Source/Editor/UnrealEd/Private/Cooker/CookSavePackage.h b/Engine/Source/Editor/UnrealEd/Private/Cooker/CookSavePackage.h
--- a/Engine/Source/Editor/UnrealEd/Private/Cooker/CookSavePackage.h
+++ b/Engine/Source/Editor/UnrealEd/Private/Cooker/CookSavePackage.h
@@ -20,6 +20,7 @@ class FSavePackageContext;
 class ITargetPlatform;
 class UCookOnTheFlyServer;
 class UWorld;
+namespace UE::Cook { enum class EReachability : uint8; }
 namespace UE::Cook { struct FCookSavePackageContext; }
 namespace UE::Cook { struct FPackageData; }
 namespace UE::Cook { struct FTickStackData; }
@@ -35,7 +36,8 @@ class FSaveCookedPackageContext
 private: // Used only by UCookOnTheFlyServer, which has private access
 
 	FSaveCookedPackageContext(UCookOnTheFlyServer& InCOTFS, UE::Cook::FPackageData& InPackageData,
-		TArrayView<const ITargetPlatform*> InPlatformsForPackage, UE::Cook::FTickStackData& StackData);
+		TArrayView<const ITargetPlatform*> InPlatformsForPackage, UE::Cook::FTickStackData& StackData,
+		EReachability InCommitType);
 
 	// Hooks used by friends
 	void SetupPackage();
@@ -50,6 +52,7 @@ private: // Used only by UCookOnTheFlyServer, which has private access
 	void CalculateCookDependencies(UE::TargetDomain::FGeneratedPackageResultStruct* GeneratedResult,
 		TArray<FAssetDependency>* OverrideDependencies);
 	void RecordPlatformBuildDependencies();
+	bool HasSavePackageResult() const;
 
 	TArray<IPackageWriter::FCommitAttachmentInfo> GetCommitAttachments();
 	IPackageWriter::EWriteOptions GetCommitWriteOptions() const;
@@ -93,6 +96,7 @@ private: // Used only by UCookOnTheFlyServer, which has private access
 	FString PlatFilename;
 	FSavePackageResultStruct SavePackageResult;
 	int32 PlatformIndex = -1;
+	EReachability CommitType = static_cast<EReachability>(0);
 	bool bPlatformSetupSuccessful = false;
 	bool bEndianSwap = false;
 
diff --git a/Engine/Source/Editor/UnrealEd/Private/Cooker/CookTypes.cpp b/Engine/Source/Editor/UnrealEd/Private/Cooker/CookTypes.cpp
--- a/Engine/Source/Editor/UnrealEd/Private/Cooker/CookTypes.cpp
+++ b/Engine/Source/Editor/UnrealEd/Private/Cooker/CookTypes.cpp
@@ -155,6 +155,18 @@ const TCHAR* LexToString(UE::Cook::EUrgency Urgency)
 	}
 }
 
+const TCHAR* LexToString(UE::Cook::ECookPhase Phase)
+{
+	static_assert(static_cast<int>(ECookPhase::Count) == 2);
+	switch (Phase)
+	{
+	case ECookPhase::Cook: return TEXT("Cook");
+	case ECookPhase::BuildDependencies: return TEXT("BuildDependencies");
+	default: return TEXT("Invalid");
+	}
+}
+
+
 EStateChangeReason ConvertToStateChangeReason(ESuppressCookReason Reason)
 {
 	switch (Reason)
@@ -564,7 +576,7 @@ void FDiscoveredPlatformSet::ConvertToBitfield(TConstArrayView<const ITargetPlat
 }
 
 TConstArrayView<const ITargetPlatform*> FDiscoveredPlatformSet::GetPlatforms(UCookOnTheFlyServer& COTFS,
-	FInstigator* Instigator, TConstArrayView<const ITargetPlatform*> OrderedPlatforms,
+	FInstigator* Instigator, TConstArrayView<const ITargetPlatform*> OrderedPlatforms, EReachability Reachability,
 	TArray<const ITargetPlatform*, TInlineAllocator<ExpectedMaxNumPlatforms>>* OutBuffer)
 {
 	switch (Source)
@@ -589,7 +601,7 @@ TConstArrayView<const ITargetPlatform*> FDiscoveredPlatformSet::GetPlatforms(UCo
 	case EDiscoveredPlatformSet::CopyFromInstigator:
 		Platforms.Reset();
 		check(Instigator);
-		FPackageData::GetReachablePlatformsForInstigator(COTFS, Instigator->Referencer, Platforms);
+		FPackageData::GetReachablePlatformsForInstigator(Reachability, COTFS, Instigator->Referencer, Platforms);
 		return Platforms;
 	default:
 		checkNoEntry();
diff --git a/Engine/Source/Editor/UnrealEd/Private/Cooker/CookTypes.h b/Engine/Source/Editor/UnrealEd/Private/Cooker/CookTypes.h
--- a/Engine/Source/Editor/UnrealEd/Private/Cooker/CookTypes.h
+++ b/Engine/Source/Editor/UnrealEd/Private/Cooker/CookTypes.h
@@ -29,6 +29,7 @@ class ICookArtifactReader;
 class ITargetPlatform;
 struct FWeakObjectPtr;
 namespace UE::DerivedData { class FBuildDefinition; }
+namespace UE::Cook { enum class EReachability : uint8; }
 namespace UE::Cook { struct FBeginCookConfigSettings; }
 namespace UE::Cook { struct FCookByTheBookOptions; }
 namespace UE::Cook { struct FCookOnTheFlyOptions; }
@@ -302,6 +303,31 @@ namespace UE::Cook
 	};
 	const TCHAR* LexToString(UE::Cook::EUrgency Urgency);
 
+	/**
+	 * Which phase the cook is in. CookPhases change the rules for how the Cooker follows dependencies and what steps
+	 * it takes to commit packages. CookPhase is different from CookMode: CookMode is constant for the entire process,
+	 * while CookPhase can change throughout a single CookSession.
+	 */
+	enum class ECookPhase : uint8
+	{
+		/**
+		 * Packages are Saved into the TargetPlatform, and the runtime and build dependencies they report are
+		 * transitively followed, respectively, to add other packages to the cook or to add CookDependencies to the
+		 * oplog. The Saved packages and their Cook Metadata is Committed into the Oplog. Only Cooked (aka
+		 * runtime-referenced) packages are committed during this phase.
+		 */
+		Cook,
+		/**
+		 * Packages are Loaded but not Saved. The build dependencies they report are transitively followed to add
+		 * CookDependencies to the oplog. The Cook Metadata of each package is Committed into the Oplog. Only
+		 * non-Cooked packages that are needed as transitive build dependencies of Cooked packages are Committed during
+		 * this phase.
+		 */
+		BuildDependencies,
+		Count,
+	};
+	const TCHAR* LexToString(UE::Cook::ECookPhase Phase);
+
 	/** Used as a helper to timeslice cooker functions. */
 	struct FCookerTimer
 	{
@@ -598,7 +624,7 @@ struct FDiscoveredPlatformSet
 	void OnRemoveSessionPlatform(const ITargetPlatform* Platform, int32 RemovedIndex);
 	void OnPlatformAddedToSession(const ITargetPlatform* Platform);
 	TConstArrayView<const ITargetPlatform*> GetPlatforms(UCookOnTheFlyServer& COTFS,
-		FInstigator* Instigator, TConstArrayView<const ITargetPlatform*> OrderedPlatforms,
+		FInstigator* Instigator, TConstArrayView<const ITargetPlatform*> OrderedPlatforms, EReachability Reachability,
 		TArray<const ITargetPlatform*, TInlineAllocator<ExpectedMaxNumPlatforms>>* OutBuffer);
 	/** If the current type is EmbeddedBitField, change it to EmbeddedList. */
 	void ConvertFromBitfield(TConstArrayView<const ITargetPlatform*> OrderedPlatforms);
diff --git a/Engine/Source/Editor/UnrealEd/Private/Cooker/CookWorkerClient.cpp b/Engine/Source/Editor/UnrealEd/Private/Cooker/CookWorkerClient.cpp
--- a/Engine/Source/Editor/UnrealEd/Private/Cooker/CookWorkerClient.cpp
+++ b/Engine/Source/Editor/UnrealEd/Private/Cooker/CookWorkerClient.cpp
@@ -209,14 +209,16 @@ void FCookWorkerClient::ReportPromoteToSaveComplete(FPackageData& PackageData)
 		ITargetPlatform* TargetPlatform = OrderedSessionPlatforms[PlatformIndex];
 		FPackageRemoteResult::FPlatformResult& PlatformResults = Result->GetPlatforms()[PlatformIndex];
 		FPackagePlatformData& PackagePlatformData = PackageData.FindOrAddPlatformData(TargetPlatform);
-		if (!PackagePlatformData.IsCookAttempted() || PackagePlatformData.IsReportedToDirector())
+		if (!PackagePlatformData.IsCommitted() || PackagePlatformData.IsReportedToDirector())
 		{
-			// We didn't attempt to cook this platform for this package, or we cooked it previously and already sent
-			// the information about it
+			// We didn't attempt to commit this platform for this package, or we committed it previously and already
+			// sent the information about it
+			PlatformResults.SetWasCommitted(false);
 			PlatformResults.SetCookResults(ECookResult::Invalid);
 		}
 		else
 		{
+			PlatformResults.SetWasCommitted(true);
 			PlatformResults.SetCookResults(PackagePlatformData.GetCookResults());
 			PackagePlatformData.SetReportedToDirector(true);
 		}
@@ -338,6 +340,19 @@ void FCookWorkerClient::ReportGeneratorQueuedGeneratedPackages(FGenerationHelper
 		GenerationHelper.GetOwner().GetPackageName()));
 }
 
+void FCookWorkerClient::HandleDirectorMessage(FDirectorEventMessage&& DirectorMessage)
+{
+	switch (DirectorMessage.Event)
+	{
+	case EDirectorEvent::KickBuildDependencies:
+		COTFS.bKickedBuildDependencies = true;
+		break;
+	default:
+		checkNoEntry();
+		break;
+	}
+}
+
 void FCookWorkerClient::HandleGeneratorMessage(FGeneratorEventMessage&& GeneratorMessage)
 {
 	FPackageData* PackageData = COTFS.PackageDatas->FindPackageDataByPackageName(GeneratorMessage.PackageName);
@@ -743,6 +758,18 @@ void FCookWorkerClient::HandleReceiveMessages(TArray<UE::CompactBinaryTCP::FMars
 					AssignPackages(AssignPackagesMessage);
 				}
 			}
+			else if (Message.MessageType == FDirectorEventMessage::MessageType)
+			{
+				FDirectorEventMessage DirectorMessage;
+				if (!DirectorMessage.TryRead(Message.Object))
+				{
+					LogInvalidMessage(TEXT("FDirectorEventMessage"));
+				}
+				else
+				{
+					HandleDirectorMessage(MoveTemp(DirectorMessage));
+				}
+			}
 			else if (Message.MessageType == FGeneratorEventMessage::MessageType)
 			{
 				FGeneratorEventMessage GeneratorMessage;
@@ -888,7 +915,7 @@ void FCookWorkerClient::AssignPackages(FAssignPackagesMessage& Message)
 	}
 	if (!Message.PackageDatas.IsEmpty())
 	{
-		TArray<const ITargetPlatform*, TInlineAllocator<ExpectedMaxNumPlatforms>> NeedCookPlatformsBuffer;
+		TArray<const ITargetPlatform*, TInlineAllocator<ExpectedMaxNumPlatforms>> NeedCommitPlatformsBuffer;
 		for (FAssignPackageData& AssignData : Message.PackageDatas)
 		{
 			FPackageData& PackageData = COTFS.PackageDatas->FindOrAddPackageData(AssignData.ConstructData.PackageName,
@@ -907,16 +934,17 @@ void FCookWorkerClient::AssignPackages(FAssignPackagesMessage& Message)
 			{
 				HandleReceiveMessages(MoveTemp(AssignData.PerPackageCollectorMessages), AssignData.ConstructData.PackageName);
 			}
-			TConstArrayView<const ITargetPlatform*> NeedCookPlatforms =
-				AssignData.NeedCookPlatforms.GetPlatforms(COTFS, nullptr, OrderedSessionPlatforms,
-					&NeedCookPlatformsBuffer);
+			EReachability Reachability = AssignData.Reachability;
+			TConstArrayView<const ITargetPlatform*> NeedCommitPlatforms =
+				AssignData.NeedCommitPlatforms.GetPlatforms(COTFS, nullptr, OrderedSessionPlatforms,
+					Reachability, &NeedCommitPlatformsBuffer);
 			if (PackageData.IsInProgress())
 			{
 				// If already in progress but there are new platforms requested, demote the package back to Load
-				for (const ITargetPlatform* TargetPlatform : NeedCookPlatforms)
+				for (const ITargetPlatform* TargetPlatform : NeedCommitPlatforms)
 				{
 					check(TargetPlatform != CookerLoadingPlatformKey);
-					if (!PackageData.FindOrAddPlatformData(TargetPlatform).IsReachable())
+					if (!PackageData.FindOrAddPlatformData(TargetPlatform).IsReachable(Reachability))
 					{
 						if (PackageData.IsInStateProperty(EPackageStateProperty::Saving))
 						{
@@ -935,12 +963,28 @@ void FCookWorkerClient::AssignPackages(FAssignPackagesMessage& Message)
 
 			// We do not want CookWorkers to explore dependencies in CookRequestCluster because the Director
 			// did it already. Mark the PackageDatas we get from the Director as already explored.
-			for (const ITargetPlatform* TargetPlatform : NeedCookPlatforms)
+			for (const ITargetPlatform* TargetPlatform : NeedCommitPlatforms)
+			{
+				PackageData.FindOrAddPlatformData(TargetPlatform).MarkCommittableForWorker(Reachability, *this);
+			}
+			if (Reachability == EReachability::Runtime)
+			{
+				checkf(COTFS.GetCookPhase() == ECookPhase::Cook,
+					TEXT("CookDirector has assigned a package for EReachability::Runtime cooking after the CookWorker has entered ECookPhase::BuildDependencies.")
+					TEXT(" This would soft-lock the cook so we assert instead.")
+					TEXT(" Package: %s"), *PackageData.GetPackageName().ToString());
+				// Also mark that CookerLoadingPlatformKey is reachable, since we do expect to need to load the package
+				PackageData.FindOrAddPlatformData(CookerLoadingPlatformKey).MarkCommittableForWorker(EReachability::Runtime, *this);
+			}
+			else
 			{
-				PackageData.FindOrAddPlatformData(TargetPlatform).MarkCookableForWorker(*this);
+				checkf(COTFS.GetCookPhase() == ECookPhase::BuildDependencies,
+					TEXT("CookDirector has assigned a package for EReachability::Build committing before the CookWorker has entered ECookPhase::BuildDependencies.")
+					TEXT(" This would soft-lock the cook so we assert instead.")
+					TEXT(" Package: %s"), *PackageData.GetPackageName().ToString());
+				// BuildDependency phase does not use the CookerLoadingPlatformKey, so we do not need to mark it reachable.
 			}
-			PackageData.FindOrAddPlatformData(CookerLoadingPlatformKey).MarkCookableForWorker(*this);
-			PackageData.SetInstigator(*this, FInstigator(AssignData.Instigator));
+			PackageData.SetInstigator(*this, Reachability, FInstigator(AssignData.Instigator));
 			PackageData.RaiseUrgency(AssignData.Urgency, ESendFlags::QueueAddAndRemove, true /* bAllowUrgencyInIdle */);
 			PackageData.SendToState(EPackageState::Request, ESendFlags::QueueAddAndRemove,
 				EStateChangeReason::DirectorRequest);
@@ -1089,8 +1133,7 @@ void FCookWorkerClient::HandleRetractionMessage(FMPCollectorClientMessageContext
 		{
 			COTFS.DemoteToIdle(*PackageData, ESendFlags::QueueAddAndRemove,
 				ESuppressCookReason::RetractedByCookDirector);
-			PackageData->ResetReachable();
-			PackageData->ResetIsBuildDependency();
+			PackageData->ResetReachable(EReachability::All);
 		}
 	}
 
diff --git a/Engine/Source/Editor/UnrealEd/Private/Cooker/CookWorkerClient.h b/Engine/Source/Editor/UnrealEd/Private/Cooker/CookWorkerClient.h
--- a/Engine/Source/Editor/UnrealEd/Private/Cooker/CookWorkerClient.h
+++ b/Engine/Source/Editor/UnrealEd/Private/Cooker/CookWorkerClient.h
@@ -18,6 +18,7 @@ namespace UE::Cook { class IMPCollector; }
 namespace UE::Cook { struct FAbortPackagesMessage; }
 namespace UE::Cook { struct FAssignPackagesMessage; }
 namespace UE::Cook { struct FDirectorConnectionInfo; }
+namespace UE::Cook { struct FDirectorEventMessage; }
 namespace UE::Cook { struct FDiscoveredPackageReplication; }
 namespace UE::Cook { struct FHeartbeatMessage; }
 namespace UE::Cook { struct FInitialConfigMessage; }
@@ -134,6 +135,7 @@ private:
 		FRetractionRequestMessage&& Message);
 	void HandleHeartbeatMessage(FMPCollectorClientMessageContext& Context, bool bReadSuccessful,
 		FHeartbeatMessage&& Message);
+	void HandleDirectorMessage(FDirectorEventMessage&& GeneratorMessage);
 	void HandleGeneratorMessage(FGeneratorEventMessage&& GeneratorMessage);
 
 private:
diff --git a/Engine/Source/Editor/UnrealEd/Private/Cooker/CookWorkerServer.cpp b/Engine/Source/Editor/UnrealEd/Private/Cooker/CookWorkerServer.cpp
--- a/Engine/Source/Editor/UnrealEd/Private/Cooker/CookWorkerServer.cpp
+++ b/Engine/Source/Editor/UnrealEd/Private/Cooker/CookWorkerServer.cpp
@@ -697,26 +697,39 @@ void FCookWorkerServer::SendPendingPackages()
 
 	TArray<FAssignPackageData> AssignDatas;
 	AssignDatas.Reserve(PackagesToAssign.Num());
-	TBitArray<> SessionPlatformNeedsCook;
+	TBitArray<> SessionPlatformNeedsCommit;
 	TArray<FPackageDataExistenceInfo> ExistenceInfos;
 	ExistenceInfos.Reserve(PackagesToAssignInfoPackages.Num());
 
+	ECookPhase CookPhase = Director.COTFS.GetCookPhase();
+	EReachability Reachability = CookPhase == ECookPhase::Cook ? EReachability::Runtime : EReachability::Build;
+
 	for (FPackageData* PackageData : PackagesToAssign)
 	{
 		FAssignPackageData& AssignData = AssignDatas.Emplace_GetRef();
 		AssignData.ConstructData = PackageData->CreateConstructData();
 		AssignData.ParentGenerator = PackageData->GetParentGenerator();
 		AssignData.DoesGeneratedRequireGenerator = PackageData->DoesGeneratedRequireGenerator();
-		AssignData.Instigator = PackageData->GetInstigator();
+		AssignData.Reachability = Reachability;
+		AssignData.Instigator = PackageData->GetInstigator(Reachability);
 		AssignData.Urgency = PackageData->GetUrgency();
-		SessionPlatformNeedsCook.Init(false, OrderedSessionPlatforms.Num());
+		SessionPlatformNeedsCommit.Init(false, OrderedSessionPlatforms.Num());
+
 		int32 PlatformIndex = 0;
+		int32 NumNeedCommitPlatforms = 0;
 		for (const ITargetPlatform* SessionPlatform : OrderedSessionPlatforms)
 		{
 			FPackagePlatformData* PlatformData = PackageData->FindPlatformData(SessionPlatform);
-			SessionPlatformNeedsCook[PlatformIndex++] = PlatformData && PlatformData->NeedsCooking(SessionPlatform);
+			if (PlatformData && PlatformData->NeedsCommit(SessionPlatform, CookPhase))
+			{
+				SessionPlatformNeedsCommit[PlatformIndex] = true;
+				++NumNeedCommitPlatforms;
+			}
+			++PlatformIndex;
 		}
-		AssignData.NeedCookPlatforms = FDiscoveredPlatformSet(SessionPlatformNeedsCook);
+		// It should not have been added to PackagesToAssign if there are no platforms to commit
+		check(NumNeedCommitPlatforms > 0);
+		AssignData.NeedCommitPlatforms = FDiscoveredPlatformSet(SessionPlatformNeedsCommit);
 		FAssignPackageExtraData* ExtraData = PackagesToAssignExtraDatas.Find(PackageData);
 		if (ExtraData)
 		{
@@ -896,27 +909,44 @@ void FCookWorkerServer::HandleReceivedPackagePlatformMessages(FPackageData& Pack
 }
 
 void FCookWorkerServer::SendMessage(const IMPCollectorMessage& Message, ECookDirectorThread TickThread)
+{
+	SendMessage(MarshalToCompactBinaryTCP(Message), TickThread);
+}
+
+void FCookWorkerServer::SendMessage(UE::CompactBinaryTCP::FMarshalledMessage&& Message,
+	ECookDirectorThread TickThread)
 {
 	FCommunicationScopeLock ScopeLock(this, TickThread, ETickAction::Tick);
-	SendMessageInLock(Message);
+	SendMessageInLock(MoveTemp(Message));
 }
 
 void FCookWorkerServer::AppendMessage(const IMPCollectorMessage& Message, ECookDirectorThread TickThread)
+{
+	AppendMessage(MarshalToCompactBinaryTCP(Message), TickThread);
+}
+
+void FCookWorkerServer::AppendMessage(UE::CompactBinaryTCP::FMarshalledMessage&& Message,
+	ECookDirectorThread TickThread)
 {
 	FCommunicationScopeLock ScopeLock(this, TickThread, ETickAction::Queue);
-	QueuedMessagesToSendAfterPackagesToAssign.Add(MarshalToCompactBinaryTCP(Message));
+	QueuedMessagesToSendAfterPackagesToAssign.Add(MoveTemp(Message));
 }
 
 void FCookWorkerServer::SendMessageInLock(const IMPCollectorMessage& Message)
+{
+	SendMessageInLock(MarshalToCompactBinaryTCP(Message));
+}
+
+void FCookWorkerServer::SendMessageInLock(UE::CompactBinaryTCP::FMarshalledMessage&& Message)
 {
 	if (TickState.TickAction == ETickAction::Tick)
 	{
-		UE::CompactBinaryTCP::TryWritePacket(Socket, SendBuffer, MarshalToCompactBinaryTCP(Message));
+		UE::CompactBinaryTCP::TryWritePacket(Socket, SendBuffer, MoveTemp(Message));
 	}
 	else
 	{
 		check(TickState.TickAction == ETickAction::Queue);
-		UE::CompactBinaryTCP::QueueMessage(SendBuffer, MarshalToCompactBinaryTCP(Message));
+		UE::CompactBinaryTCP::QueueMessage(SendBuffer, MoveTemp(Message));
 	}
 }
 
@@ -932,7 +962,8 @@ FString WriteCookStatus(FPackageData& PackageData, TConstArrayView<const ITarget
 	{
 		FPackagePlatformData* PlatformData = PackageData.FindPlatformData(TargetPlatform);
 		Result.Appendf(TEXT("[ %s: { Reachable: %s, Cookable: %s, CookResult: %s }, "),
-			*TargetPlatform->PlatformName(), BoolToString(PlatformData && PlatformData->IsReachable()),
+			*TargetPlatform->PlatformName(),
+			BoolToString(PlatformData && PlatformData->IsReachable(EReachability::Runtime)),
 			BoolToString(PlatformData && PlatformData->IsCookable()),
 			(PlatformData ? ::LexToString(PlatformData->GetCookResults()) : TEXT("<NotCooked>")));
 	}
@@ -945,7 +976,7 @@ FString WriteCookStatus(FPackageData& PackageData, TConstArrayView<const ITarget
 	return Result;
 }
 
-TArray<const ITargetPlatform*> GetCookedPlatformListFromPlatformResults(
+TArray<const ITargetPlatform*> GetCommittedPlatformListFromPlatformResults(
 	TConstArrayView<const ITargetPlatform*> OrderedPlatforms,
 	TConstArrayView<FPackageRemoteResult::FPlatformResult> PlatformResults)
 {
@@ -958,7 +989,7 @@ TArray<const ITargetPlatform*> GetCookedPlatformListFromPlatformResults(
 	for (int32 PlatformIndex = 0; PlatformIndex < NumPlatforms; ++PlatformIndex)
 	{
 		const FPackageRemoteResult::FPlatformResult& PlatformResult = PlatformResults[PlatformIndex];
-		if (PlatformResult.GetCookResults() != ECookResult::Invalid)
+		if (PlatformResult.WasCommitted())
 		{
 			PlatformList.Add(OrderedPlatforms[PlatformIndex]);
 		}
@@ -1026,23 +1057,38 @@ void FCookWorkerServer::RecordResults(FPackageResultsMessage& Message)
 					ITargetPlatform* TargetPlatform = OrderedSessionPlatforms[PlatformIndex];
 					FPackageRemoteResult::FPlatformResult& PlatformResult = Result.GetPlatforms()[PlatformIndex];
 					FPackagePlatformData& ExistingData = PackageData->FindOrAddPlatformData(TargetPlatform);
-					if (!ExistingData.NeedsCooking(TargetPlatform))
+					if (ExistingData.IsCommitted())
 					{
-						if (PlatformResult.GetCookResults() != ECookResult::Invalid)
+						if (PlatformResult.WasCommitted())
 						{
 							UE_LOG(LogCook, Display,
-								TEXT("CookWorkerServer %d received FPackageResultsMessage for package %s, platform %s, but that platform has already been cooked. Ignoring the results for that platform."),
+								TEXT("CookWorkerServer %d received FPackageResultsMessage for package %s, platform %s, but that platform has already been committed. Ignoring the results for that platform."),
 								ProfileId, *Result.GetPackageName().ToString(), *TargetPlatform->PlatformName());
 						}
 					}
 					else
 					{
-						if (PlatformResult.GetCookResults() != ECookResult::Invalid)
+						bool bWasCooked = PlatformResult.GetCookResults() != ECookResult::Invalid
+							&& PlatformResult.GetCookResults() != ECookResult::NotAttempted;
+						if (!ExistingData.NeedsCooking(TargetPlatform) && bWasCooked)
+						{
+							UE_LOG(LogCook, Display,
+								TEXT("CookWorkerServer %d received FPackageResultsMessage for package %s, platform %s, but that platform has already been cooked. Ignoring the results for that platform."),
+								ProfileId, *Result.GetPackageName().ToString(), *TargetPlatform->PlatformName());
+						}
+						else
 						{
-							PackageData->SetPlatformCooked(TargetPlatform, PlatformResult.GetCookResults());
+							if (bWasCooked)
+							{
+								PackageData->SetPlatformCooked(TargetPlatform, PlatformResult.GetCookResults());
+							}
+							else if (PlatformResult.WasCommitted())
+							{
+								PackageData->SetPlatformCommitted(TargetPlatform);
+							}
+							HandleReceivedPackagePlatformMessages(*PackageData, TargetPlatform,
+								PlatformResult.ReleaseMessages());
 						}
-						HandleReceivedPackagePlatformMessages(*PackageData, TargetPlatform,
-							PlatformResult.ReleaseMessages());
 					}
 				}
 				COTFS.RecordExternalActorDependencies(Result.GetExternalActorDependencies());
@@ -1091,14 +1137,15 @@ void FCookWorkerServer::RecordResults(FPackageResultsMessage& Message)
 		PackageData->SetWorkerAssignment(FWorkerId::Invalid(), ESendFlags::QueueNone);
 		if (Result.GetSuppressCookReason() == ESuppressCookReason::NotSuppressed)
 		{
-			if (PackageData->GetPlatformsNeedingCookingNum() > 0)
+			ECookPhase CookPhase = COTFS.GetCookPhase();
+			if (PackageData->GetPlatformsNeedingCommitNum(CookPhase) > 0)
 			{
 				TArray<const ITargetPlatform*> Remaining;
-				PackageData->GetPlatformsNeedingCooking(Remaining);
-				UE_LOG(LogCook, Display, TEXT("Package %s was completed by CookWorker %d for platforms { %s }, but it still needs to cook platforms { %s }.")
+				PackageData->GetPlatformsNeedingCommit(Remaining, CookPhase);
+				UE_LOG(LogCook, Display, TEXT("Package %s was completed by CookWorker %d for platforms { %s }, but it still needs to commit platforms { %s }.")
 					TEXT(" Sending it back to the request state."),
 					*PackageData->GetPackageName().ToString(), ProfileId,
-					*PlatformListToString(GetCookedPlatformListFromPlatformResults(OrderedSessionPlatforms, Result.GetPlatforms())),
+					*PlatformListToString(GetCommittedPlatformListFromPlatformResults(OrderedSessionPlatforms, Result.GetPlatforms())),
 					*PlatformListToString(Remaining));
 				PackageData->SendToState(EPackageState::Request, ESendFlags::QueueAddAndRemove, EStateChangeReason::Discovered);
 			}
@@ -1152,6 +1199,8 @@ void FCookWorkerServer::QueueDiscoveredPackage(FDiscoveredPackageReplication&& D
 
 	TArray<const ITargetPlatform*, TInlineAllocator<ExpectedMaxNumPlatforms>> BufferPlatforms;
 	TConstArrayView<const ITargetPlatform*> DiscoveredPlatforms;
+	EReachability DiscoveredReachability =
+		Instigator.Category == EInstigator::BuildDependency ? EReachability::Build : EReachability::Runtime;
 	if (!COTFS.bSkipOnlyEditorOnly)
 	{
 		DiscoveredPlatforms = OrderedSessionAndSpecialPlatforms;
@@ -1159,37 +1208,26 @@ void FCookWorkerServer::QueueDiscoveredPackage(FDiscoveredPackageReplication&& D
 	else
 	{
 		DiscoveredPlatforms = Platforms.GetPlatforms(COTFS, &Instigator, OrderedSessionAndSpecialPlatforms,
-			&BufferPlatforms);
+			DiscoveredReachability, &BufferPlatforms);
 	}
 
-	if (Instigator.Category == EInstigator::BuildDependency)
+	if (Instigator.Category != EInstigator::ForceExplorableSaveTimeSoftDependency &&
+		PackageData.HasReachablePlatforms(DiscoveredReachability, DiscoveredPlatforms))
 	{
-		if (PackageData.IsBuildDependency(DiscoveredPlatforms))
-		{
-			// The CookWorker thought there were some new builddependency platforms, but the Director already knows
-			// about all of them; ignore the report
-			return;
-		}
-	}
-	else if (Instigator.Category != EInstigator::ForceExplorableSaveTimeSoftDependency)
-	{
-		if (PackageData.HasReachablePlatforms(DiscoveredPlatforms))
-		{
-			// The CookWorker thought there were some new reachable platforms, but the Director already knows about
-			// all of them; ignore the report
-			return;
-		}
+		// The CookWorker thought there were some new reachable platforms, but the Director already knows about
+		// all of them; ignore the report
+		return;
 	}
 
 	if (COTFS.bSkipOnlyEditorOnly &&
 		Instigator.Category == EInstigator::Unsolicited &&
 		Platforms.GetSource() == EDiscoveredPlatformSet::CopyFromInstigator &&
-		PackageData.FindOrAddPlatformData(CookerLoadingPlatformKey).IsReachable())
+		PackageData.FindOrAddPlatformData(CookerLoadingPlatformKey).IsReachable(EReachability::Runtime))
 	{
 		// The CookWorker thought this package was new (previously unreachable even by editoronly references),
 		// and it is not marked as a known used-in-game or editor-only issue, so it fell back to reporting it
 		// as used-in-game-because-its-not-a-known-issue (see UCookOnTheFlyServer::ProcessUnsolicitedPackages's
-		// use of PackageData->FindOrAddPlatformData(CookerLoadingPlatformKey).IsReachable()).
+		// use of PackageData->FindOrAddPlatformData(CookerLoadingPlatformKey).IsReachable(EReachability::Runtime)).
 		// But we only do that fall back for unexpected packages not found by the search of editor-only AssetRegistry
 		// dependencies. And this package was found by that search; the director has already marked it as reachable by
 		// editoronly references. Correct the heuristic: ignore the unmarked load because the load is expected as an
@@ -1355,7 +1393,9 @@ void FAssignPackageData::Write(FCbWriter& Writer,
 	Writer << Instigator;
 	Writer << static_cast<uint8>(Urgency);
 	static_assert(sizeof(EUrgency) <= sizeof(uint8), "We are storing it in a uint8");
-	WriteToCompactBinary(Writer, NeedCookPlatforms, OrderedSessionPlatforms);
+	Writer << static_cast<uint8>(Reachability);
+	static_assert(sizeof(EReachability) <= sizeof(uint8), "We are storing it in a uint8");
+	WriteToCompactBinary(Writer, NeedCommitPlatforms, OrderedSessionPlatforms);
 	{
 		Writer.BeginArray();
 		for (const TPair<FName, FAssetPackageData>& Pair : GeneratorPreviousGeneratedPackages)
@@ -1389,7 +1429,16 @@ bool FAssignPackageData::TryRead(FCbFieldView Field, TConstArrayView<const ITarg
 	{
 		bOk = false;
 	}
-	bOk = LoadFromCompactBinary(*It++, NeedCookPlatforms, OrderedSessionPlatforms) & bOk;
+	uint8 ReachabilityInt = It->AsUInt8();
+	if (!(It++)->HasError() && ReachabilityInt < (static_cast<uint8>(EReachability::MaxBit) << 1))
+	{
+		Reachability = static_cast<EReachability>(ReachabilityInt);
+	}
+	else
+	{
+		bOk = false;
+	}
+	bOk = LoadFromCompactBinary(*It++, NeedCommitPlatforms, OrderedSessionPlatforms) & bOk;
 	{
 		FCbFieldView ArrayFieldView = *It++;
 		bool bGeneratorPreviousGeneratedPackagesOk = false;
diff --git a/Engine/Source/Editor/UnrealEd/Private/Cooker/CookWorkerServer.h b/Engine/Source/Editor/UnrealEd/Private/Cooker/CookWorkerServer.h
--- a/Engine/Source/Editor/UnrealEd/Private/Cooker/CookWorkerServer.h
+++ b/Engine/Source/Editor/UnrealEd/Private/Cooker/CookWorkerServer.h
@@ -3,12 +3,12 @@
 #pragma once
 
 #include "AssetRegistry/AssetData.h"
-#include "CompactBinaryTCP.h"
 #include "Containers/Array.h"
 #include "Containers/ArrayView.h"
 #include "Containers/Map.h"
 #include "Containers/RingBuffer.h"
 #include "Containers/Set.h"
+#include "Cooker/CompactBinaryTCP.h"
 #include "Cooker/CookPackageData.h"
 #include "Cooker/CookSockets.h"
 #include "Cooker/CookTypes.h"
@@ -88,11 +88,13 @@ public:
 
 	/** Send the message immediately to the Socket. If cannot complete immediately, it will be finished during Tick. */
 	void SendMessage(const IMPCollectorMessage& Message, ECookDirectorThread TickThread);
+	void SendMessage(UE::CompactBinaryTCP::FMarshalledMessage&& Message, ECookDirectorThread TickThread);
 	/**
 	 * Queue the given message for sending; it will be sent during tick. Only needed for messages that have to be sent
 	 * after queued package assignments or other appended messages.
 	 */
 	void AppendMessage(const IMPCollectorMessage& Message, ECookDirectorThread TickThread);
+	void AppendMessage(UE::CompactBinaryTCP::FMarshalledMessage&& Message, ECookDirectorThread TickThread);
 
 	/** Periodic Tick function to send and receive messages to the Client. */
 	void TickCommunication(ECookDirectorThread TickThread);
@@ -192,6 +194,7 @@ private:
 	void AbortAllAssignmentsInLock(TSet<FPackageData*>& OutPendingPackages, int32 CurrentHeartbeat);
 	/** Send the message immediately to the Socket. If cannot complete immediately, it will be finished during Tick. */
 	void SendMessageInLock(const IMPCollectorMessage& Message);
+	void SendMessageInLock(UE::CompactBinaryTCP::FMarshalledMessage&& Message);
 	/** Send this into the given state. Update any state-dependent variables. */
 	void SendToState(EConnectStatus TargetStatus);
 	/** Close the connection and connection resources to the remote process. Does not kill the process. */
@@ -257,7 +260,8 @@ struct FAssignPackageData
 	FName ParentGenerator;
 	FInstigator Instigator;
 	EUrgency Urgency;
-	FDiscoveredPlatformSet NeedCookPlatforms;
+	EReachability Reachability;
+	FDiscoveredPlatformSet NeedCommitPlatforms;
 	TMap<FName, FAssetPackageData> GeneratorPreviousGeneratedPackages;
 	TArray<UE::CompactBinaryTCP::FMarshalledMessage> PerPackageCollectorMessages;
 	ICookPackageSplitter::EGeneratedRequiresGenerator DoesGeneratedRequireGenerator
@@ -460,8 +464,11 @@ public:
 enum class EGeneratorEvent
 {
 	Invalid,
+	/** Sent from Worker to Director to let the Director know to start waiting for AllSavesCompleted. */
 	QueuedGeneratedPackages,
+	/** Broadcast to Workers so any waiting on the fence to pass before clearing their data can clear their data. */
 	QueuedGeneratedPackagesFencePassed,
+	/** Broadcast to Workers to unconditionally clear their data. */
 	AllSavesCompleted,
 	Num,
 };
diff --git a/Engine/Source/Editor/UnrealEd/Private/Cooker/PackageResultsMessage.cpp b/Engine/Source/Editor/UnrealEd/Private/Cooker/PackageResultsMessage.cpp
--- a/Engine/Source/Editor/UnrealEd/Private/Cooker/PackageResultsMessage.cpp
+++ b/Engine/Source/Editor/UnrealEd/Private/Cooker/PackageResultsMessage.cpp
@@ -141,6 +141,7 @@ void FPackageResultsMessage::Write(FCbWriter& Writer) const
 			for (const FPackageRemoteResult::FPlatformResult& PlatformResult : Result.Platforms)
 			{
 				Writer.BeginObject();
+				Writer << "C" << PlatformResult.WasCommitted();
 				Writer << "S" << (uint8) PlatformResult.GetCookResults();
 				WriteMessagesArray(Writer, PlatformResult.Messages);
 				Writer.EndObject();
@@ -200,6 +201,7 @@ bool FPackageResultsMessage::TryRead(FCbObjectView Object)
 			FPackageRemoteResult::FPlatformResult& PlatformResult = Result.Platforms.Emplace_GetRef();
 			FCbObjectView PlatformObject = PlatformField.AsObjectView();
 
+			PlatformResult.SetWasCommitted(PlatformObject["C"].AsBool());
 			uint8 CookResultsInt = PlatformObject["S"].AsUInt8();
 			if (CookResultsInt >= (uint8)ECookResult::Count)
 			{
diff --git a/Engine/Source/Editor/UnrealEd/Private/Cooker/PackageResultsMessage.h b/Engine/Source/Editor/UnrealEd/Private/Cooker/PackageResultsMessage.h
--- a/Engine/Source/Editor/UnrealEd/Private/Cooker/PackageResultsMessage.h
+++ b/Engine/Source/Editor/UnrealEd/Private/Cooker/PackageResultsMessage.h
@@ -37,6 +37,9 @@ public:
 		TConstArrayView<UE::CompactBinaryTCP::FMarshalledMessage> GetMessages() const;
 		TArray<UE::CompactBinaryTCP::FMarshalledMessage> ReleaseMessages();
 
+		bool WasCommitted() const;
+		void SetWasCommitted(bool bValue);
+
 		ECookResult GetCookResults() const;
 		void SetCookResults(ECookResult Value);
 
@@ -44,6 +47,7 @@ public:
 		TArray<UE::CompactBinaryTCP::FMarshalledMessage> Messages;
 		const ITargetPlatform* Platform = nullptr;
 		ECookResult CookResults = ECookResult::NotAttempted;
+		bool bWasCommitted = false;
 
 		friend FPackageRemoteResult;
 		friend FPackageResultsMessage;
@@ -183,6 +187,16 @@ FPackageRemoteResult::FPlatformResult::GetMessages() const
 	return Messages;
 }
 
+inline bool FPackageRemoteResult::FPlatformResult::WasCommitted() const
+{
+	return bWasCommitted;
+}
+
+inline void FPackageRemoteResult::FPlatformResult::SetWasCommitted(bool bValue)
+{
+	bWasCommitted = bValue;
+}
+
 inline ECookResult FPackageRemoteResult::FPlatformResult::GetCookResults() const
 {
 	return CookResults;
diff --git a/Engine/Source/Editor/UnrealEd/Private/TargetDomain/TargetDomainUtils.cpp b/Engine/Source/Editor/UnrealEd/Private/TargetDomain/TargetDomainUtils.cpp
--- a/Engine/Source/Editor/UnrealEd/Private/TargetDomain/TargetDomainUtils.cpp
+++ b/Engine/Source/Editor/UnrealEd/Private/TargetDomain/TargetDomainUtils.cpp
@@ -39,7 +39,7 @@ namespace UE::TargetDomain
 {
 
 // Change TargetDomainVersion to a new guid when all TargetDomain keys in an incremental cook need to be invalidated
-FGuid TargetDomainVersion(TEXT("C9B0281696234067A3A888CEEAAA50F9"));
+FGuid TargetDomainVersion(0xD5476CB0, 0xD7D4482F, 0x966DC671, 0xBBA1BC40);
 
 // Bump CookDependenciesVersion version when the serialization of CookDependencies has changed and we want to add
 // backwards compatibility rather than invalidating everything.
@@ -109,6 +109,11 @@ bool FCookDependencies::IsValid() const
 	return bValid;
 }
 
+bool FCookDependencies::HasSaveResults() const
+{
+	return bHasSaveResults;
+}
+
 bool FCookDependencies::HasKeyMatch(const FAssetPackageData* OverrideAssetPackageData)
 {
 	if (!bValid)
@@ -124,6 +129,7 @@ bool FCookDependencies::HasKeyMatch(const FAssetPackageData* OverrideAssetPackag
 		ECurrentKeyResult Result = TryCalculateCurrentKey(OverrideAssetPackageData);
 		if (Result != ECurrentKeyResult::Success)
 		{
+			CurrentKey.Reset();
 			return false;
 		}
 	}
@@ -365,6 +371,7 @@ FCookDependencies FCookDependencies::Collect(UPackage* Package, const ITargetPla
 
 	FCookDependencies Result;
 	Result.PackageName = Package->GetFName();
+	Result.bHasSaveResults = SaveResult != nullptr;
 	enum class EBuildDependency
 	{
 		PackageOnly,
@@ -736,6 +743,13 @@ bool LoadFromCompactBinary(FCbObjectView ObjectView, UE::TargetDomain::FCookDepe
 				return false;
 			}
 		}
+		if (FieldView.GetName().Equals(UTF8TEXTVIEW("HasSaveResults")))
+		{
+			if (!LoadFromCompactBinary(FieldView++, Dependencies.bHasSaveResults))
+			{
+				return false;
+			}
+		}
 		if (FieldView.GetName().Equals(UTF8TEXTVIEW("BuildPackageDependencies")))
 		{
 			if (!LoadFromCompactBinary(FieldView++, Dependencies.BuildPackageDependencies))
@@ -812,6 +826,7 @@ FCbWriter& operator<<(FCbWriter& Writer, const UE::TargetDomain::FCookDependenci
 	Writer.BeginObject();
 	Writer << "Version" << CookDependenciesVersion;
 	Writer << "StoredKey" << CookDependencies.StoredKey;
+	Writer << "HasSaveResults" << CookDependencies.bHasSaveResults;
 	if (!CookDependencies.BuildPackageDependencies.IsEmpty())
 	{
 		Writer << "BuildPackageDependencies" << CookDependencies.BuildPackageDependencies;
diff --git a/Engine/Source/Editor/UnrealEd/Private/TargetDomain/TargetDomainUtils.h b/Engine/Source/Editor/UnrealEd/Private/TargetDomain/TargetDomainUtils.h
--- a/Engine/Source/Editor/UnrealEd/Private/TargetDomain/TargetDomainUtils.h
+++ b/Engine/Source/Editor/UnrealEd/Private/TargetDomain/TargetDomainUtils.h
@@ -86,6 +86,7 @@ public:
 	 * key for the package. False if the stucture is default, has been reset, or was marked invalid.
 	 */
 	bool IsValid() const;
+	bool HasSaveResults() const;
 	FName GetPackageName() const;
 	const FIoHash& GetStoredKey() const;
 	const FIoHash& GetCurrentKey() const;
@@ -129,6 +130,7 @@ private:
 	FName PackageName;
 	FIoHash StoredKey;
 	FIoHash CurrentKey;
+	bool bHasSaveResults = false;
 	bool bValid = false;
 
 	friend bool ::LoadFromCompactBinary(FCbObjectView ObjectView, FCookDependencies& CookAttachments);
