Fixed Automation failures in Verse Scene Graph Framework Test Map

#rnx

diff --git a/Engine/Source/Editor/UnrealEd/Private/FileHelpers.cpp b/Engine/Source/Editor/UnrealEd/Private/FileHelpers.cpp
--- a/Engine/Source/Editor/UnrealEd/Private/FileHelpers.cpp
+++ b/Engine/Source/Editor/UnrealEd/Private/FileHelpers.cpp
@@ -5397,8 +5397,8 @@ void FEditorFileUtils::GetDirtyContentPackages(TArray<UPackage*>& OutDirtyPackag
 		bShouldIgnorePackage |= Package == GetTransientPackage();
 		// Don't try to save packages with the RF_Transient flag
 		bShouldIgnorePackage |= Package->HasAnyFlags(RF_Transient);
-		// Ignore PIE packages, or packages containing map data
-		bShouldIgnorePackage |= Package->HasAnyPackageFlags(PKG_PlayInEditor | PKG_ContainsMapData);
+		// Ignore UHT/Verse generated packages, PIE packages, or packages containing map data
+		bShouldIgnorePackage |= Package->HasAnyPackageFlags(PKG_CompiledIn | PKG_PlayInEditor | PKG_ContainsMapData);
 		// Ignore packages that haven't been modified.
 		bShouldIgnorePackage |= !Package->IsDirty();
 
