[Viewport Toolbar] Add missing icons for Camera FOV and Camera Far View Plane Entries

diff --git a/Engine/Source/Editor/UnrealEd/Private/ViewportToolbar/UnrealEdViewportToolbar.cpp b/Engine/Source/Editor/UnrealEd/Private/ViewportToolbar/UnrealEdViewportToolbar.cpp
--- a/Engine/Source/Editor/UnrealEd/Private/ViewportToolbar/UnrealEdViewportToolbar.cpp
+++ b/Engine/Source/Editor/UnrealEd/Private/ViewportToolbar/UnrealEdViewportToolbar.cpp
@@ -2719,19 +2719,23 @@ void ExtendCameraSubmenu(FName InCameraOptionsSubmenuName, bool bInShowFarViewPl
 				FToolMenuSection& PostOrthoSection = InDynamicMenu->FindOrAddSection("CameraSubmenuPostOrtho", FText(), InsertPosition);
 				PostOrthoSection.AddSeparator("CameraSubmenuSeparator");
 
-				PostOrthoSection.AddEntry(FToolMenuEntry::InitWidget(
+				FToolMenuEntry CameraFOV = FToolMenuEntry::InitWidget(
 					"CameraFOV",
 					UE::UnrealEd::CreateFOVMenuWidget(EditorViewport.ToSharedRef()),
 					LOCTEXT("CameraSubmenu_FieldOfViewLabel", "Field of View")
-				));
+				);
+				CameraFOV.Icon = FSlateIcon(FAppStyle::GetAppStyleSetName(), "ViewportToolbar.FieldOfView");
+				PostOrthoSection.AddEntry(CameraFOV);
 
 				if (bInShowFarViewPlaneEntry)
 				{
-					PostOrthoSection.AddEntry(FToolMenuEntry::InitWidget(
+					FToolMenuEntry CameraFarViewPlane =FToolMenuEntry::InitWidget(
 						"CameraFarViewPlane",
 						UE::UnrealEd::CreateFarViewPlaneMenuWidget(EditorViewport.ToSharedRef()),
 						LOCTEXT("CameraSubmenu_FarViewPlaneLabel", "Far View Plane")
-					));
+					);
+					CameraFarViewPlane.Icon = FSlateIcon(FAppStyle::GetAppStyleSetName(), "ViewportToolbar.FarViewPlane");
+					PostOrthoSection.AddEntry(CameraFarViewPlane);
 				}
 			}
 		)
