Don't checkout/add packages into /Temp (which maps to Project\Saved), only happens in the editor when running automation tests.

diff --git a/Engine/Source/Editor/UnrealEd/Private/FileHelpers.cpp b/Engine/Source/Editor/UnrealEd/Private/FileHelpers.cpp
--- a/Engine/Source/Editor/UnrealEd/Private/FileHelpers.cpp
+++ b/Engine/Source/Editor/UnrealEd/Private/FileHelpers.cpp
@@ -7,6 +7,7 @@
 #include "Misc/MessageDialog.h"
 #include "HAL/FileManager.h"
 #include "Algo/Copy.h"
+#include "Algo/Transform.h"
 #include "Misc/CommandLine.h"
 #include "Misc/Paths.h"
 #include "Misc/ConfigCacheIni.h"
@@ -5565,7 +5566,9 @@ static bool InternalCheckoutAndSavePackages(const TArray<UPackage*>& PackagesToS
 
 			if (ISourceControlModule::Get().IsEnabled())
 			{
-				FEditorFileUtils::CheckoutPackages(PackagesToSave, &PackagesCheckedOut, bErrorIfAlreadyCheckedOut, bConfirmPackageBranchCheckOutStatus);
+				TArray<UPackage*> PackagesToCheckout;
+				Algo::CopyIf(PackagesToSave, PackagesToCheckout, [](UPackage* Package) { return !FPackageName::IsTempPackage(Package->GetName()); });
+				FEditorFileUtils::CheckoutPackages(PackagesToCheckout, &PackagesCheckedOut, bErrorIfAlreadyCheckedOut, bConfirmPackageBranchCheckOutStatus);
 			}
 
 			// Cannot mark files for add until after packages saved
@@ -5573,7 +5576,7 @@ static bool InternalCheckoutAndSavePackages(const TArray<UPackage*>& PackagesToS
 			for (UPackage* Package : PackagesToSave)
 			{
 				// List unsaved packages that were not checked out and are not going to be deleted
-				if (!PackagesCheckedOut.Contains(Package) && !UPackage::IsEmptyPackage(Package))
+				if (!PackagesCheckedOut.Contains(Package) && !UPackage::IsEmptyPackage(Package) && !FPackageName::IsTempPackage(Package->GetName()))
 				{
 					PackagesToMarkForAdd.Add(Package);
 				}
