Always give quick access to "Camera > Camera Speed" in the viewport toolbar

Note that UE::UnrealEd::CreateCameraSpeedMenu was never released, so we don't have to deprecate the old signature.

diff --git a/Engine/Source/Editor/Kismet/Private/SSCSEditorViewport.cpp b/Engine/Source/Editor/Kismet/Private/SSCSEditorViewport.cpp
--- a/Engine/Source/Editor/Kismet/Private/SSCSEditorViewport.cpp
+++ b/Engine/Source/Editor/Kismet/Private/SSCSEditorViewport.cpp
@@ -64,19 +64,7 @@ void CreateCameraSpeedMenu(UToolMenu* InMenu, TWeakPtr<FSCSEditorViewportClient>
 	FToolMenuSection& PositioningSection =
 		InMenu->FindOrAddSection("Positioning", LOCTEXT("PositioningLabel", "Positioning"));
 
-	PositioningSection.AddSubMenu(
-		"CameraSpeed",
-		LOCTEXT("CameraSpeedSubMenu", "Camera Speed"),
-		LOCTEXT("CameraSpeedSubMenu_ToolTip", "Camera Speed related actions"),
-		FNewToolMenuDelegate::CreateLambda(
-			[](UToolMenu* InMenu)
-			{
-				UE::UnrealEd::CreateCameraSpeedMenu(InMenu);
-			}
-		),
-		false,
-		FSlateIcon(FSlateIcon(FAppStyle::GetAppStyleSetName(), "LevelViewport.ToggleActorPilotCameraView"))
-	);
+	PositioningSection.AddEntry(UE::UnrealEd::CreateCameraSpeedMenu());
 }
 
 void ExtendCameraSubmenu(FName InCameraOptionsSubmenuName, TSharedPtr<FSCSEditorViewportClient> InViewportClient)
diff --git a/Engine/Source/Editor/LevelEditor/Private/ViewportToolbar/LevelEditorSubmenus.cpp b/Engine/Source/Editor/LevelEditor/Private/ViewportToolbar/LevelEditorSubmenus.cpp
--- a/Engine/Source/Editor/LevelEditor/Private/ViewportToolbar/LevelEditorSubmenus.cpp
+++ b/Engine/Source/Editor/LevelEditor/Private/ViewportToolbar/LevelEditorSubmenus.cpp
@@ -2530,31 +2530,7 @@ void ExtendCameraSubmenu(FName InCameraOptionsSubmenuName)
 					FToolMenuSection& PositioningSection =
 						InDynamicMenu->FindOrAddSection("Positioning", LOCTEXT("PositioningLabel", "Positioning"));
 
-					// Camera Speed Submenu
-					{
-						FToolMenuEntry& Entry = PositioningSection.AddSubMenu(
-							"CameraSpeed",
-							LOCTEXT("CameraSpeedSubMenu", "Camera Speed"),
-							LOCTEXT("CameraSpeedSubMenu_ToolTip", "Camera Speed related actions"),
-							FNewToolMenuDelegate::CreateLambda(
-								[LevelViewportWeak](UToolMenu* InMenu)
-								{
-									UE::UnrealEd::CreateCameraSpeedMenu(InMenu);
-								}
-							),
-							false,
-							FSlateIcon(FSlateIcon(FAppStyle::GetAppStyleSetName(), "ViewportToolbar.CameraSpeed"))
-						);
-
-						Entry.ToolBarData.LabelOverride = TAttribute<FText>::CreateLambda(
-							[LevelViewportWeak]()
-							{
-								return UE::UnrealEd::GetCameraSpeedLabel(LevelViewportWeak);
-							}
-						);
-						Entry.SetShowInToolbarTopLevel(true);
-						Entry.ToolBarData.PlacementOverride = MenuPlacement_BelowRightAnchor;
-					}
+					PositioningSection.AddEntry(UE::UnrealEd::CreateCameraSpeedMenu());
 
 					PositioningSection.AddSeparator("PositioningSeparator_1");
 
diff --git a/Engine/Source/Editor/Persona/Private/ViewportToolbar/AnimationEditorMenus.cpp b/Engine/Source/Editor/Persona/Private/ViewportToolbar/AnimationEditorMenus.cpp
--- a/Engine/Source/Editor/Persona/Private/ViewportToolbar/AnimationEditorMenus.cpp
+++ b/Engine/Source/Editor/Persona/Private/ViewportToolbar/AnimationEditorMenus.cpp
@@ -1066,22 +1066,7 @@ void UE::AnimationEditor::ExtendCameraMenu(FName InCameraOptionsMenuName)
 				FToolMenuSection& PositioningSection =
 					InDynamicMenu->FindOrAddSection("Positioning", LOCTEXT("PositioningLabel", "Positioning"));
 
-				// Camera Speed Submenu
-				{
-					PositioningSection.AddSubMenu(
-						"CameraSpeed",
-						LOCTEXT("CameraSpeedSubMenu", "Camera Speed"),
-						LOCTEXT("CameraSpeedSubMenu_ToolTip", "Camera Speed related actions"),
-						FNewToolMenuDelegate::CreateLambda(
-							[](UToolMenu* InMenu)
-							{
-								UE::UnrealEd::CreateCameraSpeedMenu(InMenu);
-							}
-						),
-						false,
-						FSlateIcon(FSlateIcon(FAppStyle::GetAppStyleSetName(), "EditorViewport.Speed"))
-					);
-				}
+				PositioningSection.AddEntry(UE::UnrealEd::CreateCameraSpeedMenu());
 
 				PositioningSection.AddSeparator("PositioningSeparator_1");
 
diff --git a/Engine/Source/Editor/UnrealEd/Private/ViewportToolbar/UnrealEdViewportToolbar.cpp b/Engine/Source/Editor/UnrealEd/Private/ViewportToolbar/UnrealEdViewportToolbar.cpp
--- a/Engine/Source/Editor/UnrealEd/Private/ViewportToolbar/UnrealEdViewportToolbar.cpp
+++ b/Engine/Source/Editor/UnrealEd/Private/ViewportToolbar/UnrealEdViewportToolbar.cpp
@@ -3595,29 +3595,50 @@ FToolMenuEntry CreateCameraSpeedScalarSlider(TWeakPtr<::SEditorViewport> InEdito
 	);
 }
 
-void CreateCameraSpeedMenu(UToolMenu* InMenu)
+FToolMenuEntry CreateCameraSpeedMenu()
 {
-	if (!InMenu)
-	{
-		return;
-	}
-
-	FToolMenuSection& CameraSpeedSection = InMenu->AddSection("CameraSpeed", LOCTEXT("CameraSpeedLabel", "Camera Speed"));
-
-	UUnrealEdViewportToolbarContext* const EditorViewportContext = InMenu->FindContext<UUnrealEdViewportToolbarContext>();
-	if (!EditorViewportContext)
-	{
-		return;
-	}
+	return FToolMenuEntry::InitDynamicEntry(
+		"CameraSpeedMenuDynamicSection",
+		FNewToolMenuSectionDelegate::CreateLambda(
+			[](FToolMenuSection& InDynamicSection)
+			{
+				FToolMenuEntry& Entry = InDynamicSection.AddSubMenu(
+					"CameraSpeed",
+					LOCTEXT("CameraSpeedSubMenu", "Camera Speed"),
+					LOCTEXT("CameraSpeedSubMenu_ToolTip", "Camera Speed related actions"),
+					FNewToolMenuDelegate::CreateLambda(
+						[](UToolMenu* InMenu)
+						{
+							FToolMenuSection& CameraSpeedSection =
+								InMenu->AddSection("CameraSpeed", LOCTEXT("CameraSpeedLabel", "Camera Speed"));
 
-	const TSharedPtr<SEditorViewport> EditorViewport = EditorViewportContext->Viewport.Pin();
-	if (!EditorViewport)
-	{
-		return;
-	}
+							if (UUnrealEdViewportToolbarContext* const Context =
+									InMenu->FindContext<UUnrealEdViewportToolbarContext>())
+							{
+								CameraSpeedSection.AddEntry(CreateCameraSpeedSlider(Context->Viewport));
+								CameraSpeedSection.AddEntry(CreateCameraSpeedScalarSlider(Context->Viewport));
+							}
+						}
+					),
+					false,
+					FSlateIcon(FSlateIcon(FAppStyle::GetAppStyleSetName(), "ViewportToolbar.CameraSpeed"))
+				);
 
-	CameraSpeedSection.AddEntry(CreateCameraSpeedSlider(EditorViewport));
-	CameraSpeedSection.AddEntry(CreateCameraSpeedScalarSlider(EditorViewport));
+				if (UUnrealEdViewportToolbarContext* Context =
+						InDynamicSection.FindContext<UUnrealEdViewportToolbarContext>())
+				{
+					Entry.ToolBarData.LabelOverride = TAttribute<FText>::CreateLambda(
+						[WeakViewport = Context->Viewport]()
+						{
+							return UE::UnrealEd::GetCameraSpeedLabel(WeakViewport);
+						}
+					);
+				}
+				Entry.SetShowInToolbarTopLevel(true);
+				Entry.ToolBarData.PlacementOverride = MenuPlacement_BelowRightAnchor;
+			}
+		)
+	);
 }
 
 } // namespace UE::UnrealEd
diff --git a/Engine/Source/Editor/UnrealEd/Public/ViewportToolbar/UnrealEdViewportToolbar.h b/Engine/Source/Editor/UnrealEd/Public/ViewportToolbar/UnrealEdViewportToolbar.h
--- a/Engine/Source/Editor/UnrealEd/Public/ViewportToolbar/UnrealEdViewportToolbar.h
+++ b/Engine/Source/Editor/UnrealEd/Public/ViewportToolbar/UnrealEdViewportToolbar.h
@@ -196,7 +196,7 @@ UNREALED_API FToolMenuEntry CreateRemoveRealtimeOverrideEntry(TWeakPtr<SEditorVi
 UNREALED_API FOnViewportClientCamSpeedChanged& OnViewportClientCamSpeedChanged();
 UNREALED_API FOnViewportClientCamSpeedScalarChanged& OnViewportClientCamSpeedScalarChanged();
 
-UNREALED_API void CreateCameraSpeedMenu(UToolMenu* InMenu);
+UNREALED_API FToolMenuEntry CreateCameraSpeedMenu();
 
 UNREALED_API void ConstructScreenPercentageMenu(UToolMenu* InMenu);
 UNREALED_API FToolMenuEntry CreateScreenPercentageSubmenu();
