Cooker: CookRequestCluster: Refactor per-package storage to be simpler: Change OwnedPackageDatas into ClusterPackages that have a flag for ownership rather than always implying ownership, and merge FVertexData into this container. Change FGraphSearch to be an inline member rather than a pointer with limited lifetime, so that we don't have to handle it being null in some places.
Fix typo: SetMaxBlockSize -> SetMinBlockSize.
#rnx

diff --git a/Engine/Source/Editor/UnrealEd/Private/CookOnTheFlyServer.cpp b/Engine/Source/Editor/UnrealEd/Private/CookOnTheFlyServer.cpp
--- a/Engine/Source/Editor/UnrealEd/Private/CookOnTheFlyServer.cpp
+++ b/Engine/Source/Editor/UnrealEd/Private/CookOnTheFlyServer.cpp
@@ -852,9 +852,9 @@ void UCookOnTheFlyServer::OnTargetPlatformsInvalidated()
 	PackageDatas->RemapTargetPlatforms(Remap);
 	PackageTracker->RemapTargetPlatforms(Remap);
 	WorkerRequests->RemapTargetPlatforms(Remap);
-	for (UE::Cook::FRequestCluster& Cluster : PackageDatas->GetRequestQueue().GetRequestClusters())
+	for (TUniquePtr<FRequestCluster>& Cluster : PackageDatas->GetRequestQueue().GetRequestClusters())
 	{
-		Cluster.RemapTargetPlatforms(Remap);
+		Cluster->RemapTargetPlatforms(Remap);
 	}
 	for (FDiscoveryQueueElement& Element : PackageDatas->GetRequestQueue().GetDiscoveryQueue())
 	{
@@ -1848,9 +1848,9 @@ void UCookOnTheFlyServer::UpdateDisplay(UE::Cook::FTickStackData& StackData, boo
 		// When a RequestCluster is doing a graph search, it marks uncookable packages as to-be-demoted, and
 		// iteratively skippable packages as cooked, but those packages remain in the request state until the cluster
 		// search is complete so we still count them as inprogress. Subtract them from the inprogress count.
-		for (const FRequestCluster& Cluster : PackageDatas->GetRequestQueue().GetRequestClusters())
+		for (const TUniquePtr<FRequestCluster>& Cluster : PackageDatas->GetRequestQueue().GetRequestClusters())
 		{
-			CookPendingCount -= Cluster.GetPackagesToMarkNotInProgress();
+			CookPendingCount -= Cluster->GetPackagesToMarkNotInProgress();
 		}
 		UE_CLOG(!(StackData.TickFlags & ECookTickFlags::HideProgressDisplay)
 			&& (GCookProgressDisplay & (int32)ECookProgressDisplayMode::RemainingPackages),
@@ -2544,45 +2544,47 @@ int32 UCookOnTheFlyServer::NumMultiprocessLocalWorkerAssignments() const
 		PackageDatas->GetRequestQueue().Num() +
 		PackageDatas->GetLoadQueue().Num() +
 		PackageDatas->GetSaveQueue().Num();
-	for (const FRequestCluster& Cluster : PackageDatas->GetRequestQueue().GetRequestClusters())
+	for (const TUniquePtr<FRequestCluster>& Cluster : PackageDatas->GetRequestQueue().GetRequestClusters())
 	{
-		Result -= Cluster.GetPackagesToMarkNotInProgress();
+		Result -= Cluster->GetPackagesToMarkNotInProgress();
 	}
 	return Result;
 }
 
 void UCookOnTheFlyServer::PumpExternalRequests(const UE::Cook::FCookerTimer& CookerTimer)
 {
+	using namespace UE::Cook;
+
 	if (!WorkerRequests->HasExternalRequests())
 	{
 		return;
 	}
 	UE_SCOPED_COOKTIMER(PumpExternalRequests);
 
-	TArray<UE::Cook::FFilePlatformRequest> BuildRequests;
-	TArray<UE::Cook::FSchedulerCallback> SchedulerCallbacks;
-	UE::Cook::EExternalRequestType RequestType;
+	TArray<FFilePlatformRequest> BuildRequests;
+	TArray<FSchedulerCallback> SchedulerCallbacks;
+	EExternalRequestType RequestType;
 	while (!CookerTimer.IsActionTimeUp())
 	{
 		BuildRequests.Reset();
 		SchedulerCallbacks.Reset();
 		RequestType = WorkerRequests->DequeueNextCluster(SchedulerCallbacks, BuildRequests);
-		if (RequestType == UE::Cook::EExternalRequestType::None)
+		if (RequestType == EExternalRequestType::None)
 		{
 			// No more requests to process
 			break;
 		}
-		else if (RequestType == UE::Cook::EExternalRequestType::Callback)
+		else if (RequestType == EExternalRequestType::Callback)
 		{
 			// An array of TickCommands to process; execute through them all
-			for (UE::Cook::FSchedulerCallback& SchedulerCallback : SchedulerCallbacks)
+			for (FSchedulerCallback& SchedulerCallback : SchedulerCallbacks)
 			{
 				SchedulerCallback();
 			}
 		}
 		else
 		{
-			check(RequestType == UE::Cook::EExternalRequestType::Cook && BuildRequests.Num() > 0);
+			check(RequestType == EExternalRequestType::Cook && BuildRequests.Num() > 0);
 #if PROFILE_NETWORK
 			if (NetworkRequestEvent)
 			{
@@ -2590,8 +2592,9 @@ void UCookOnTheFlyServer::PumpExternalRequests(const UE::Cook::FCookerTimer& Coo
 			}
 #endif
 			bool bRequestsAreUrgent = IsCookOnTheFlyMode() && IsUsingLegacyCookOnTheFlyScheduling();
-			TRingBuffer<UE::Cook::FRequestCluster>& RequestClusters = PackageDatas->GetRequestQueue().GetRequestClusters();
-			RequestClusters.Add(UE::Cook::FRequestCluster(*this, MoveTemp(BuildRequests)));
+			TRingBuffer<TUniquePtr<FRequestCluster>>& RequestClusters
+				= PackageDatas->GetRequestQueue().GetRequestClusters();
+			RequestClusters.Add(MakeUnique<FRequestCluster>(*this, MoveTemp(BuildRequests)));
 		}
 	}
 }
@@ -2620,7 +2623,7 @@ void UCookOnTheFlyServer::PumpRequests(UE::Cook::FTickStackData& StackData, int3
 	FPackageDataSet& RestartedRequests = RequestQueue.GetRestartedRequests();
 	TRingBuffer<FDiscoveryQueueElement>& DiscoveryQueue = RequestQueue.GetDiscoveryQueue();
 	TRingBuffer<FPackageData*>& BuildDependencyDiscoveryQueue = RequestQueue.GetBuildDependencyDiscoveryQueue();
-	TRingBuffer<FRequestCluster>& RequestClusters = RequestQueue.GetRequestClusters();
+	TRingBuffer<TUniquePtr<FRequestCluster>>& RequestClusters = RequestQueue.GetRequestClusters();
 	const FCookerTimer& CookerTimer = StackData.Timer;
 
 	// First pump all requestclusters and unclustered/discovered requests that need to create a requestcluster
@@ -2631,7 +2634,7 @@ void UCookOnTheFlyServer::PumpRequests(UE::Cook::FTickStackData& StackData, int3
 		// back and forth from each other.
 		if (!RequestClusters.IsEmpty())
 		{
-			FRequestCluster& RequestCluster = RequestClusters.First();
+			FRequestCluster& RequestCluster = *RequestClusters.First();
 			bool bComplete = false;
 			RequestCluster.Process(CookerTimer, bComplete);
 			if (bComplete)
@@ -2669,15 +2672,15 @@ void UCookOnTheFlyServer::PumpRequests(UE::Cook::FTickStackData& StackData, int3
 			EReachability ClusterReachability = GetCookPhase() == ECookPhase::Cook
 				? EReachability::Runtime
 				: EReachability::Build;
-			RequestClusters.Add(FRequestCluster(*this, MoveTemp(RestartedRequests), ClusterReachability));
+			RequestClusters.Add(MakeUnique<FRequestCluster>(*this, MoveTemp(RestartedRequests), ClusterReachability));
 			RestartedRequests.Empty();
 		}
 		else if (!DiscoveryQueue.IsEmpty())
 		{
 			if (GetCookPhase() == ECookPhase::Cook)
 			{
-				FRequestCluster Cluster(*this, DiscoveryQueue);
-				if (Cluster.NumPackageDatas() > 0)
+				TUniquePtr<FRequestCluster> Cluster = MakeUnique<FRequestCluster>(*this, DiscoveryQueue);
+				if (Cluster->NumPackageDatas() > 0)
 				{
 					RequestClusters.Add(MoveTemp(Cluster));
 				}
@@ -2701,8 +2704,9 @@ void UCookOnTheFlyServer::PumpRequests(UE::Cook::FTickStackData& StackData, int3
 			}
 			else
 			{
-				FRequestCluster Cluster(*this, FRequestCluster::BuildDependencyQueue, BuildDependencyDiscoveryQueue);
-				if (Cluster.NumPackageDatas() > 0)
+				TUniquePtr<FRequestCluster> Cluster = MakeUnique<FRequestCluster>(*this,
+					FRequestCluster::BuildDependencyQueue, BuildDependencyDiscoveryQueue);
+				if (Cluster->NumPackageDatas() > 0)
 				{
 					RequestClusters.Add(MoveTemp(Cluster));
 				}
@@ -3393,9 +3397,9 @@ void UCookOnTheFlyServer::OnRemoveSessionPlatform(const ITargetPlatform* TargetP
 	{
 		Element.ReachablePlatforms.OnRemoveSessionPlatform(TargetPlatform, RemovedIndex);
 	}
-	for (UE::Cook::FRequestCluster& Cluster : PackageDatas->GetRequestQueue().GetRequestClusters())
+	for (TUniquePtr<FRequestCluster>& Cluster : PackageDatas->GetRequestQueue().GetRequestClusters())
 	{
-		Cluster.OnRemoveSessionPlatform(TargetPlatform);
+		Cluster->OnRemoveSessionPlatform(TargetPlatform);
 	}
 	PackageDatas->OnRemoveSessionPlatform(TargetPlatform);
 	WorkerRequests->OnRemoveSessionPlatform(TargetPlatform);
@@ -3409,9 +3413,9 @@ void UCookOnTheFlyServer::OnPlatformAddedToSession(const ITargetPlatform* Target
 	{
 		Element.ReachablePlatforms.OnPlatformAddedToSession(TargetPlatform);
 	}
-	for (UE::Cook::FRequestCluster& Cluster : PackageDatas->GetRequestQueue().GetRequestClusters())
+	for (TUniquePtr<FRequestCluster>& Cluster : PackageDatas->GetRequestQueue().GetRequestClusters())
 	{
-		Cluster.OnPlatformAddedToSession(TargetPlatform);
+		Cluster->OnPlatformAddedToSession(TargetPlatform);
 	}
 }
 
@@ -9848,13 +9852,13 @@ void UCookOnTheFlyServer::CancelAllQueues()
 	}
 	RequestQueue.GetDiscoveryQueue().Empty();
 	RequestQueue.GetBuildDependencyDiscoveryQueue().Empty();
-	TRingBuffer<FRequestCluster>& RequestClusters = RequestQueue.GetRequestClusters();
-	for (FRequestCluster& RequestCluster : RequestClusters)
+	TRingBuffer<TUniquePtr<FRequestCluster>>& RequestClusters = RequestQueue.GetRequestClusters();
+	for (TUniquePtr<FRequestCluster>& RequestCluster : RequestClusters)
 	{
 		TArray<FPackageData*> RequestsToLoad;
 		TArray<TPair<FPackageData*, ESuppressCookReason>> RequestsToDemote;
 		TMap<FPackageData*, TArray<FPackageData*>> UnusedRequestGraph;
-		RequestCluster.ClearAndDetachOwnedPackageDatas(RequestsToLoad, RequestsToDemote, UnusedRequestGraph);
+		RequestCluster->ClearAndDetachOwnedPackageDatas(RequestsToLoad, RequestsToDemote, UnusedRequestGraph);
 		for (FPackageData* PackageData : RequestsToLoad)
 		{
 			DemoteToIdle(*PackageData, ESendFlags::QueueAdd, ESuppressCookReason::CookCanceled);
diff --git a/Engine/Source/Editor/UnrealEd/Private/Cooker/CookPackageData.cpp b/Engine/Source/Editor/UnrealEd/Private/Cooker/CookPackageData.cpp
--- a/Engine/Source/Editor/UnrealEd/Private/Cooker/CookPackageData.cpp
+++ b/Engine/Source/Editor/UnrealEd/Private/Cooker/CookPackageData.cpp
@@ -2901,9 +2901,9 @@ bool FRequestQueue::IsEmpty() const
 uint32 FRequestQueue::Num() const
 {
 	uint32 Count = RestartedRequests.Num() + ReadyRequestsNum();
-	for (const FRequestCluster& RequestCluster : RequestClusters)
+	for (const TUniquePtr<FRequestCluster>& RequestCluster : RequestClusters)
 	{
-		Count += RequestCluster.NumPackageDatas();
+		Count += RequestCluster->NumPackageDatas();
 	}
 	return Count;
 }
@@ -2916,9 +2916,9 @@ bool FRequestQueue::Contains(const FPackageData* InPackageData) const
 	{
 		return true;
 	}
-	for (const FRequestCluster& RequestCluster : RequestClusters)
+	for (const TUniquePtr<FRequestCluster>& RequestCluster : RequestClusters)
 	{
-		if (RequestCluster.Contains(PackageData))
+		if (RequestCluster->Contains(PackageData))
 		{
 			return true;
 		}
@@ -2932,11 +2932,11 @@ uint32 FRequestQueue::RemoveRequestExceptFromCluster(FPackageData* PackageData,
 	RestartedRequests.Remove(PackageData);
 	NormalRequests.Remove(PackageData);
 	UrgentRequests.Remove(PackageData);
-	for (FRequestCluster& RequestCluster : RequestClusters)
+	for (TUniquePtr<FRequestCluster>& RequestCluster : RequestClusters)
 	{
-		if (&RequestCluster != ExceptFromCluster)
+		if (RequestCluster.Get() != ExceptFromCluster)
 		{
-			RequestCluster.RemovePackageData(PackageData);
+			RequestCluster->RemovePackageData(PackageData);
 		}
 	}
 	uint32 Result = OriginalNum - Num();
diff --git a/Engine/Source/Editor/UnrealEd/Private/Cooker/CookPackageData.h b/Engine/Source/Editor/UnrealEd/Private/Cooker/CookPackageData.h
--- a/Engine/Source/Editor/UnrealEd/Private/Cooker/CookPackageData.h
+++ b/Engine/Source/Editor/UnrealEd/Private/Cooker/CookPackageData.h
@@ -1058,7 +1058,7 @@ public:
 	 */
 	TRingBuffer<FDiscoveryQueueElement>& GetDiscoveryQueue();
 	TRingBuffer<FPackageData*>& GetBuildDependencyDiscoveryQueue();
-	TRingBuffer<FRequestCluster>& GetRequestClusters();
+	TRingBuffer<TUniquePtr<FRequestCluster>>& GetRequestClusters();
 	FPackageDataSet& GetReadyRequestsUrgent();
 	FPackageDataSet& GetReadyRequestsNormal();
 	/**
@@ -1073,7 +1073,7 @@ private:
 	FPackageDataSet RestartedRequests;
 	TRingBuffer<FDiscoveryQueueElement> DiscoveryQueue;
 	TRingBuffer<FPackageData*> BuildDependencyDiscoveryQueue;
-	TRingBuffer<FRequestCluster> RequestClusters;
+	TRingBuffer<TUniquePtr<FRequestCluster>> RequestClusters; // Not trivially relocatable, so must be TUniquePtr
 	TSet<FName> RequestFencePackageListeners;
 	FPackageDataSet UrgentRequests;
 	FPackageDataSet NormalRequests;
@@ -1890,7 +1890,7 @@ inline TRingBuffer<FPackageData*>& FRequestQueue::GetBuildDependencyDiscoveryQue
 	return BuildDependencyDiscoveryQueue;
 }
 
-inline TRingBuffer<FRequestCluster>& FRequestQueue::GetRequestClusters()
+inline TRingBuffer<TUniquePtr<FRequestCluster>>& FRequestQueue::GetRequestClusters()
 {
 	return RequestClusters;
 }
diff --git a/Engine/Source/Editor/UnrealEd/Private/Cooker/CookRequestCluster.cpp b/Engine/Source/Editor/UnrealEd/Private/Cooker/CookRequestCluster.cpp
--- a/Engine/Source/Editor/UnrealEd/Private/Cooker/CookRequestCluster.cpp
+++ b/Engine/Source/Editor/UnrealEd/Private/Cooker/CookRequestCluster.cpp
@@ -38,12 +38,19 @@ namespace UE::Cook
 {
 
 FRequestCluster::FRequestCluster(UCookOnTheFlyServer& InCOTFS, EReachability ExploreReachability)
-	: COTFS(InCOTFS)
+	: GraphSearch(*this)
+	, COTFS(InCOTFS)
 	, PackageDatas(*InCOTFS.PackageDatas)
 	, AssetRegistry(*IAssetRegistry::Get())
 	, PackageTracker(*InCOTFS.PackageTracker)
 	, BuildDefinitions(*InCOTFS.BuildDefinitions)
 {
+	TConstArrayView<const ITargetPlatform*> SessionPlatforms = COTFS.PlatformManager->GetSessionPlatforms();
+	check(SessionPlatforms.Num() > 0);
+	NumFetchPlatforms = SessionPlatforms.Num() + 2;
+	VertexAllocator.SetMinBlockSize(1024);
+	VertexAllocator.SetMaxBlockSize(65536);
+
 	// CookByTheBookOptions is always available; in other modes it is set to the default values
 	UE::Cook::FCookByTheBookOptions& Options = *COTFS.CookByTheBookOptions;
 	bAllowHardDependencies = !Options.bSkipHardReferences;
@@ -117,6 +124,26 @@ FRequestCluster::FRequestCluster(UCookOnTheFlyServer& InCOTFS, EReachability Exp
 	}
 }
 
+FRequestCluster::~FRequestCluster()
+{
+	EmptyClusterPackages();
+}
+
+void FRequestCluster::EmptyClusterPackages()
+{
+	// Call the FVertexData destructors, but do not bother calling DeleteElement or Free on the VertexAllocator
+	// since we are destructing the VertexAllocator.
+	for (TPair<FName, FVertexData*>& VertexPair : ClusterPackages)
+	{
+		check(VertexPair.Value);
+		FVertexData& Vertex = *VertexPair.Value;
+		Vertex.~FVertexData();
+	}
+	ClusterPackages.Empty();
+	// Empty frees the struct memory for each FVertexData we allocated, but it does not call the destructor.
+	VertexAllocator.Empty();
+}
+
 FRequestCluster::FRequestCluster(UCookOnTheFlyServer& InCOTFS, TArray<FFilePlatformRequest>&& InRequests)
 	: FRequestCluster(InCOTFS, EReachability::Runtime)
 {
@@ -133,9 +160,12 @@ FRequestCluster::FRequestCluster(UCookOnTheFlyServer& InCOTFS, FPackageDataSet&&
 	for (FPackageData* PackageData : InRequests)
 	{
 		check(PackageData);
-		bool bExisted;
-		SetPackageDataSuppressReason(*PackageData, ESuppressCookReason::NotSuppressed, &bExisted);
-		check(!bExisted);
+		FVertexData& Vertex = FindOrAddVertex(*PackageData);
+		// Setting bNeedsStateChange=false is important to avoid a crash: changing the state will try to remove it
+		// from InRequests while we are iterating over it trigger an assertion/crash. Calling InRequests::Empty below will
+		// accomplish what bNeedsStateChange=true would have tried to do more slowly.
+		check(PackageData->GetState() == EPackageState::Request);
+		SetOwnedByCluster(Vertex, true /* bOwnedByCluster */, false /* bNeedsStateChange */);
 	}
 	InRequests.Empty();
 }
@@ -193,8 +223,8 @@ FRequestCluster::FRequestCluster(UCookOnTheFlyServer& InCOTFS, TRingBuffer<FDisc
 				PackageData.SendToState(EPackageState::Request, ESendFlags::QueueRemove,
 					EStateChangeReason::RequestCluster);
 				PackageData.RaiseUrgency(Urgency, ESendFlags::QueueNone);
-				// SetPackageDataSuppressReason adds it in the proper container of *this
-				SetPackageDataSuppressReason(PackageData, ESuppressCookReason::NotSuppressed);
+				FVertexData& Vertex = FindOrAddVertex(PackageData);
+				SetOwnedByCluster(Vertex, true /* bOwnedByCluster */, false /* bNeedsStateChange */);
 			}
 			else if (PackageData.IsInProgress())
 			{
@@ -236,7 +266,7 @@ FRequestCluster::FRequestCluster(UCookOnTheFlyServer& InCOTFS, TRingBuffer<FDisc
 			// If there are other discovered packages we have already added to this cluster, then defer this one
 			// until we have explored those; add this one to the next cluster. Exploring those earlier discoveries
 			// might add this one through cluster exploration and not require a hidden dependency.
-			if (!OwnedPackageDatas.IsEmpty())
+			if (!ClusterPackages.IsEmpty())
 			{
 				break;
 			}
@@ -275,8 +305,8 @@ FRequestCluster::FRequestCluster(UCookOnTheFlyServer& InCOTFS, TRingBuffer<FDisc
 		// and add it to this cluster.
 		PackageData.SendToState(EPackageState::Request, ESendFlags::QueueRemove, EStateChangeReason::RequestCluster);
 		PackageData.RaiseUrgency(Discovery->Urgency, ESendFlags::QueueNone);
-		// SetPackageDataSuppressReason adds it in the proper container of *this
-		SetPackageDataSuppressReason(PackageData, ESuppressCookReason::NotSuppressed);
+		FVertexData& Vertex = FindOrAddVertex(PackageData);
+		SetOwnedByCluster(Vertex, true /* bOwnedByCluster */, false /* bNeedsStateChange */);
 	}
 }
 
@@ -296,8 +326,8 @@ FRequestCluster::FRequestCluster(UCookOnTheFlyServer& InCOTFS, EBuildDependencyQ
 
 		PackageData.SendToState(EPackageState::Request, ESendFlags::QueueRemove,
 			EStateChangeReason::RequestCluster);
-		// SetPackageDataSuppressReason adds it in the proper container of *this
-		SetPackageDataSuppressReason(PackageData, ESuppressCookReason::NotSuppressed);
+		FVertexData& Vertex = FindOrAddVertex(PackageData);
+		SetOwnedByCluster(Vertex, true /* bOwnedByCluster */, false /* bNeedsStateChange */);
 	}
 }
 
@@ -423,7 +453,8 @@ void FRequestCluster::FetchPackageNames(const FCookerTimer& CookerTimer, bool& b
 		{
 			PackageData->AddReachablePlatforms(*this, EReachability::Runtime, Request.GetPlatforms(),
 				MoveTemp(Request.GetInstigator()));
-			PullIntoCluster(*PackageData);
+			FVertexData& Vertex = FindOrAddVertex(*PackageData);
+			SetOwnedByCluster(Vertex, true);
 			if (Request.IsUrgent())
 			{
 				PackageData->SetUrgency(EUrgency::Blocking, ESendFlags::QueueNone);
@@ -443,7 +474,8 @@ void FRequestCluster::FetchPackageNames(const FCookerTimer& CookerTimer, bool& b
 			else if (PackageData->GetPlatformsNeedingCommitNum(EReachability::Runtime) > 0)
 			{
 				// If it's missing cookable platforms and not in progress we need to add it to the cluster for cooking
-				PullIntoCluster(*PackageData);
+				FVertexData& Vertex = FindOrAddVertex(*PackageData);
+				SetOwnedByCluster(Vertex, true);
 				if (Request.IsUrgent())
 				{
 					PackageData->SetUrgency(EUrgency::Blocking, ESendFlags::QueueNone);
@@ -466,15 +498,34 @@ void FRequestCluster::FetchPackageNames(const FCookerTimer& CookerTimer, bool& b
 
 void FRequestCluster::ReserveInitialRequests(int32 RequestNum)
 {
-	OwnedPackageDatas.Reserve(FMath::Max(RequestNum, 1024));
+	ClusterPackages.Reserve(FMath::Max(RequestNum, 1024));
 }
 
-void FRequestCluster::PullIntoCluster(FPackageData& PackageData)
+void FRequestCluster::AddVertexCounts(FVertexData& Vertex, int32 Delta)
 {
-	bool bExisted;
-	SetPackageDataSuppressReason(PackageData, ESuppressCookReason::NotSuppressed, &bExisted);
-	if (!bExisted)
+	if (Vertex.IsOwnedByCluster())
 	{
+		NumOwned += Delta;
+		if (Vertex.IsOwnedButNotInProgress())
+		{
+			NumOwnedButNotInProgress += Delta;
+		}
+	}
+}
+
+void FRequestCluster::SetOwnedByCluster(FVertexData& Vertex, bool bOwnedByCluster, bool bNeedsStateChange)
+{
+	if (bOwnedByCluster == Vertex.IsOwnedByCluster())
+	{
+		return;
+	}
+	AddVertexCounts(Vertex, -1);
+	Vertex.SetOwnedByCluster(bOwnedByCluster);
+	AddVertexCounts(Vertex, 1);
+
+	if (bOwnedByCluster && bNeedsStateChange && Vertex.GetPackageData())
+	{
+		FPackageData& PackageData = *Vertex.GetPackageData();
 		// Steal it from wherever it is and send it to Request State. It has already been added to this cluster
 		if (PackageData.GetState() == EPackageState::Request)
 		{
@@ -488,6 +539,22 @@ void FRequestCluster::PullIntoCluster(FPackageData& PackageData)
 	}
 }
 
+void FRequestCluster::SetSuppressReason(FVertexData& Vertex, ESuppressCookReason Reason)
+{
+	check(Reason != ESuppressCookReason::Invalid);
+
+	AddVertexCounts(Vertex, -1);
+	Vertex.SetSuppressReason(Reason);
+	AddVertexCounts(Vertex, 1);
+}
+
+void FRequestCluster::SetWasMarkedSkipped(FVertexData& Vertex, bool bValue)
+{
+	AddVertexCounts(Vertex, -1);
+	Vertex.SetWasMarkedSkipped(bValue);
+	AddVertexCounts(Vertex, 1);
+}
+
 void FRequestCluster::StartAsync(const FCookerTimer& CookerTimer, bool& bOutComplete)
 {
 	using namespace UE::DerivedData;
@@ -512,12 +579,13 @@ void FRequestCluster::StartAsync(const FCookerTimer& CookerTimer, bool& bOutComp
 		{
 			// If the EditorDomain is active, then batch-download all packages to cook from remote cache into local
 			TArray<FName> BatchDownload;
-			BatchDownload.Reserve(OwnedPackageDatas.Num());
-			for (TPair<FPackageData*, FProcessingFlags>& Pair : OwnedPackageDatas)
+			BatchDownload.Reserve(ClusterPackages.Num());
+			for (TPair<FName, FVertexData*>& Pair : ClusterPackages)
 			{
-				if (Pair.Value.GetSuppressReason() == ESuppressCookReason::NotSuppressed)
+				FVertexData* Vertex = Pair.Value;
+				if (Vertex->IsOwnedByCluster() && Vertex->GetSuppressReason() == ESuppressCookReason::NotSuppressed)
 				{
-					BatchDownload.Add(Pair.Key->GetPackageName());
+					BatchDownload.Add(Pair.Key);
 				}
 			};
 			EditorDomain->BatchDownload(BatchDownload);
@@ -527,91 +595,33 @@ void FRequestCluster::StartAsync(const FCookerTimer& CookerTimer, bool& bOutComp
 	bStartAsyncComplete = true;
 }
 
-int32 FRequestCluster::NumPackageDatas() const
-{
-	return OwnedPackageDatas.Num();
-}
-
 void FRequestCluster::RemovePackageData(FPackageData* PackageData)
 {
-	FProcessingFlags RemovedFlags;
-	if (!OwnedPackageDatas.RemoveAndCopyValue(PackageData, RemovedFlags))
+	if (!PackageData)
 	{
 		return;
 	}
-	check(RemovedFlags.IsValid());
-	if (RemovedFlags.ShouldMarkNotInProgress())
-	{
-		--PackagesToMarkNotInProgressCount;
-	}
-}
-
-void FRequestCluster::SetPackageDataSuppressReason(FPackageData& PackageData, ESuppressCookReason Reason, bool* bOutExisted)
-{
-	check(Reason != ESuppressCookReason::Invalid);
-
-	FProcessingFlags& Existing = OwnedPackageDatas.FindOrAdd(&PackageData);
-	if (bOutExisted)
-	{
-		*bOutExisted = Existing.IsValid();
-	}
-	if (Existing.ShouldMarkNotInProgress())
-	{
-		--PackagesToMarkNotInProgressCount;
-	}
-	Existing.SetValid();
-	Existing.SetSuppressReason(Reason);
-	if (Existing.ShouldMarkNotInProgress())
-	{
-		++PackagesToMarkNotInProgressCount;
-	}
-}
-
-void FRequestCluster::SetPackageDataWasMarkedSkipped(FPackageData& PackageData, bool bValue,
-	bool* bOutExisted)
-{
-	// We don't want to bring the package into the cluster if it's not already there, which is governed by existence in
-	// OwnedPackageDatas. So early exit if it is not there.
-	// TODO: Change name of OwnedPackageDatas and allow packages not in the cluster to exist in it,
-	// with a separate flag for whether they are owned, so we can keep track of SetWasMarkedSkipped for packages
-	// that are not in the cluster when first evaluated, but are later added.
-	FProcessingFlags* Existing = OwnedPackageDatas.Find(&PackageData);
-	if (!Existing)
+	FVertexData** VertexPtr = ClusterPackages.Find(PackageData->GetPackageName());
+	if (!VertexPtr)
 	{
-		if (bOutExisted)
-		{
-			*bOutExisted = false;
-		}
 		return;
 	}
-
-	if (bOutExisted)
-	{
-		*bOutExisted = Existing->IsValid();
-	}
-	if (Existing->ShouldMarkNotInProgress())
-	{
-		--PackagesToMarkNotInProgressCount;
-	}
-	Existing->SetValid();
-	Existing->SetWasMarkedSkipped(bValue);
-	if (Existing->ShouldMarkNotInProgress())
-	{
-		++PackagesToMarkNotInProgressCount;
-	}
+	check(*VertexPtr);
+	FVertexData& Vertex = **VertexPtr;
+	SetOwnedByCluster(Vertex, false);
 }
 
 void FRequestCluster::OnNewReachablePlatforms(FPackageData* PackageData)
 {
-	if (GraphSearch)
+	if (GraphSearch.IsInitialized())
 	{
-		GraphSearch->OnNewReachablePlatforms(PackageData);
+		GraphSearch.OnNewReachablePlatforms(PackageData);
 	}
 }
 
 void FRequestCluster::OnPlatformAddedToSession(const ITargetPlatform* TargetPlatform)
 {
-	if (GraphSearch)
+	if (GraphSearch.IsStarted())
 	{
 		FCookerTimer CookerTimer(FCookerTimer::Forever);
 		bool bComplete;
@@ -625,7 +635,7 @@ void FRequestCluster::OnPlatformAddedToSession(const ITargetPlatform* TargetPlat
 
 void FRequestCluster::OnRemoveSessionPlatform(const ITargetPlatform* TargetPlatform)
 {
-	if (GraphSearch)
+	if (GraphSearch.IsStarted())
 	{
 		FCookerTimer CookerTimer(FCookerTimer::Forever);
 		bool bComplete;
@@ -640,7 +650,7 @@ void FRequestCluster::OnRemoveSessionPlatform(const ITargetPlatform* TargetPlatf
 
 void FRequestCluster::RemapTargetPlatforms(TMap<ITargetPlatform*, ITargetPlatform*>& Remap)
 {
-	if (GraphSearch)
+	if (GraphSearch.IsStarted())
 	{
 		// The platforms have already been invalidated, which means we can't wait for GraphSearch to finish
 		// Need to wait for all async operations to finish, then remap all the platforms
@@ -650,7 +660,17 @@ void FRequestCluster::RemapTargetPlatforms(TMap<ITargetPlatform*, ITargetPlatfor
 
 bool FRequestCluster::Contains(FPackageData* PackageData) const
 {
-	return OwnedPackageDatas.Contains(PackageData);
+	if (!PackageData)
+	{
+		return false;
+	}
+	FVertexData*const* VertexPtr = ClusterPackages.Find(PackageData->GetPackageName());
+	if (!VertexPtr)
+	{
+		return false;
+	}
+	check(*VertexPtr);
+	return (*VertexPtr)->IsOwnedByCluster();
 }
 
 void FRequestCluster::ClearAndDetachOwnedPackageDatas(TArray<FPackageData*>& OutRequestsToLoad,
@@ -659,18 +679,23 @@ void FRequestCluster::ClearAndDetachOwnedPackageDatas(TArray<FPackageData*>& Out
 {
 	if (bStartAsyncComplete)
 	{
-		check(!GraphSearch);
+		check(!GraphSearch.IsStarted());
 		OutRequestsToLoad.Reset();
 		OutRequestsToDemote.Reset();
-		for (TPair<FPackageData*, FProcessingFlags>& Pair : OwnedPackageDatas)
+		for (TPair<FName, FVertexData*>& Pair : ClusterPackages)
 		{
-			if (Pair.Value.GetSuppressReason() == ESuppressCookReason::NotSuppressed)
-			{
-				OutRequestsToLoad.Add(Pair.Key);
-			}
-			else
+			check(Pair.Value);
+			FVertexData& Vertex = *Pair.Value;
+			if (Vertex.IsOwnedByCluster() && Vertex.GetPackageData())
 			{
-				OutRequestsToDemote.Add({ Pair.Key, Pair.Value.GetSuppressReason() });
+				if (Vertex.GetSuppressReason() == ESuppressCookReason::NotSuppressed)
+				{
+					OutRequestsToLoad.Add(Vertex.GetPackageData());
+				}
+				else
+				{
+					OutRequestsToDemote.Add({ Vertex.GetPackageData(), Vertex.GetSuppressReason()});
+				}
 			}
 		}
 		OutRequestGraph = MoveTemp(RequestGraph);
@@ -678,16 +703,22 @@ void FRequestCluster::ClearAndDetachOwnedPackageDatas(TArray<FPackageData*>& Out
 	else
 	{
 		OutRequestsToLoad.Reset();
-		for (TPair<FPackageData*, FProcessingFlags>& Pair : OwnedPackageDatas)
+		for (TPair<FName, FVertexData*>& Pair : ClusterPackages)
 		{
-			OutRequestsToLoad.Add(Pair.Key);
+			check(Pair.Value);
+			FVertexData& Vertex = *Pair.Value;
+			if (Vertex.IsOwnedByCluster() && Vertex.GetPackageData())
+			{
+				OutRequestsToLoad.Add(Vertex.GetPackageData());
+			}
 		}
 		OutRequestsToDemote.Reset();
 		OutRequestGraph.Reset();
 	}
 	FilePlatformRequests.Empty();
-	OwnedPackageDatas.Empty();
-	PackagesToMarkNotInProgressCount = 0;
+	EmptyClusterPackages();
+	NumOwned = 0;
+	NumOwnedButNotInProgress = 0;
 	GraphSearch.Reset();
 	RequestGraph.Reset();
 }
@@ -699,25 +730,23 @@ void FRequestCluster::PumpExploration(const FCookerTimer& CookerTimer, bool& bOu
 		return;
 	}
 
-	if (!GraphSearch)
+	if (!GraphSearch.IsStarted())
 	{
-		GraphSearch.Reset(new FGraphSearch(*this));
-
+		GraphSearch.Initialize();
 		if (!TraversalExploreIncremental() && !TraversalExploreRuntimeDependencies())
 		{
-			GraphSearch->VisitWithoutFetching();
-			GraphSearch.Reset();
+			GraphSearch.VisitWithoutFetching();
 			bDependenciesComplete = true;
 			return;
 		}
-		GraphSearch->StartSearch();
+		GraphSearch.StartSearch();
 	}
 
 	constexpr double WaitTime = 0.50;
 	bool bDone;
-	while (GraphSearch->TickExploration(bDone), !bDone)
+	while (GraphSearch.TickExploration(bDone), !bDone)
 	{
-		GraphSearch->WaitForAsyncQueue(WaitTime);
+		GraphSearch.WaitForAsyncQueue(WaitTime);
 		if (CookerTimer.IsActionTimeUp())
 		{
 			bOutComplete = false;
@@ -726,12 +755,15 @@ void FRequestCluster::PumpExploration(const FCookerTimer& CookerTimer, bool& bOu
 	}
 
 	TArray<FPackageData*> SortedPackages;
-	SortedPackages.Reserve(OwnedPackageDatas.Num());
-	for (TPair<FPackageData*, FProcessingFlags>& Pair : OwnedPackageDatas)
+	SortedPackages.Reserve(ClusterPackages.Num());
+	for (TPair<FName, FVertexData*>& Pair : ClusterPackages)
 	{
-		if (Pair.Value.GetSuppressReason() == ESuppressCookReason::NotSuppressed)
+		check(Pair.Value);
+		FVertexData& Vertex = *Pair.Value;
+		if (Vertex.IsOwnedByCluster() && Vertex.GetPackageData()
+			&& Vertex.GetSuppressReason() == ESuppressCookReason::NotSuppressed)
 		{
-			SortedPackages.Add(Pair.Key);
+			SortedPackages.Add(Vertex.GetPackageData());
 		}
 	}
 
@@ -739,7 +771,7 @@ void FRequestCluster::PumpExploration(const FCookerTimer& CookerTimer, bool& bOu
 	TArray<FPackageData*> Empty;
 	auto GetElementDependencies = [this, &Empty](FPackageData* PackageData) -> const TArray<FPackageData*>&
 	{
-		const TArray<FPackageData*>* VertexEdges = GraphSearch->GetGraphEdges().Find(PackageData);
+		const TArray<FPackageData*>* VertexEdges = GraphSearch.GetGraphEdges().Find(PackageData);
 		return VertexEdges ? *VertexEdges : Empty;
 	};
 
@@ -751,13 +783,13 @@ void FRequestCluster::PumpExploration(const FCookerTimer& CookerTimer, bool& bOu
 	{
 		SortOrder.Add(PackageData, Counter++);
 	}
-	OwnedPackageDatas.KeySort([&SortOrder](const FPackageData& A, const FPackageData& B)
+	ClusterPackages.ValueSort([&SortOrder](const FVertexData& A, const FVertexData& B)
 		{
-			int32* CounterA = SortOrder.Find(&A);
-			int32* CounterB = SortOrder.Find(&B);
+			int32* CounterA = SortOrder.Find(A.GetPackageData());
+			int32* CounterB = SortOrder.Find(B.GetPackageData());
 			if ((CounterA != nullptr) != (CounterB != nullptr))
 			{
-				// Sort the demotes to occur last
+				// Sort the missing packages, unowned vertices, or demotes to occur last
 				return CounterB == nullptr;
 			}
 			else if (CounterA)
@@ -766,11 +798,11 @@ void FRequestCluster::PumpExploration(const FCookerTimer& CookerTimer, bool& bOu
 			}
 			else
 			{
-				return false; // demotes are unsorted
+				return false; // missing packages, unowned vertices and demotes are unsorted
 			}
 		});
 
-	RequestGraph = MoveTemp(GraphSearch->GetGraphEdges());
+	RequestGraph = MoveTemp(GraphSearch.GetGraphEdges());
 	GraphSearch.Reset();
 	bDependenciesComplete = true;
 }
@@ -780,16 +812,24 @@ FRequestCluster::FGraphSearch::FGraphSearch(FRequestCluster& InCluster)
 	, ExploreEdgesContext(InCluster, *this)
 	, AsyncResultsReadyEvent(EEventMode::ManualReset)
 {
+}
+
+void FRequestCluster::FGraphSearch::Initialize()
+{
+	if (bInitialized)
+	{
+		return;
+	}
+
 	AsyncResultsReadyEvent->Trigger();
 	LastActivityTime = FPlatformTime::Seconds();
-	VertexAllocator.SetMaxBlockSize(1024);
-	VertexAllocator.SetMaxBlockSize(65536);
-	BatchAllocator.SetMaxBlockSize(16);
+	BatchAllocator.SetMinBlockSize(16);
 	BatchAllocator.SetMaxBlockSize(16);
 
 	TConstArrayView<const ITargetPlatform*> SessionPlatforms = Cluster.COTFS.PlatformManager->GetSessionPlatforms();
 	check(SessionPlatforms.Num() > 0);
-	FetchPlatforms.SetNum(SessionPlatforms.Num() + 2);
+	check(SessionPlatforms.Num() == Cluster.GetNumSessionPlatforms())
+	FetchPlatforms.SetNum(Cluster.GetNumFetchPlatforms());
 	FetchPlatforms[PlatformAgnosticPlatformIndex].bIsPlatformAgnosticPlatform = true;
 	FetchPlatforms[CookerLoadingPlatformIndex].Platform = CookerLoadingPlatformKey;
 	FetchPlatforms[CookerLoadingPlatformIndex].bIsCookerLoadingPlatform = true;
@@ -805,33 +845,50 @@ FRequestCluster::FGraphSearch::FGraphSearch(FRequestCluster& InCluster)
 		});
 	check(FetchPlatforms[PlatformAgnosticPlatformIndex].bIsPlatformAgnosticPlatform);
 	check(FetchPlatforms[CookerLoadingPlatformIndex].bIsCookerLoadingPlatform);
+
+	bInitialized = true;
 }
 
 void FRequestCluster::FGraphSearch::VisitWithoutFetching()
 {
 	// PumpExploration is responsible for marking all requests as explored and cookable/uncoookable.
 	// If we're skipping the dependencies search, handle that responsibility for the initial requests and return.
-	for (TPair<FPackageData*, FProcessingFlags>& Pair : Cluster.OwnedPackageDatas)
+	for (TPair<FName, FVertexData*>& Pair : Cluster.ClusterPackages)
 	{
-		check(Pair.Key);
-		FVertexData Vertex(Pair.Key->GetPackageName(), Pair.Key, *this);
+		check(Pair.Value);
+		FVertexData& Vertex = *Pair.Value;
+		if (!Vertex.GetPackageData())
+		{
+			continue;
+		}
+		check(Vertex.IsOwnedByCluster());
 		VisitVertex(Vertex);
 	}
 }
 
 void FRequestCluster::FGraphSearch::StartSearch()
 {
-	VisitVertexQueue.Reserve(Cluster.OwnedPackageDatas.Num());
-	for (TPair<FPackageData*, FProcessingFlags>& Pair : Cluster.OwnedPackageDatas)
+	VisitVertexQueue.Reserve(Cluster.ClusterPackages.Num());
+	for (TPair<FName, FVertexData*>& Pair : Cluster.ClusterPackages)
 	{
-		FVertexData& Vertex = FindOrAddVertex(Pair.Key->GetPackageName(), *Pair.Key);
-		check(Vertex.PackageData);
-		Vertex.bPulledIntoCluster = true;
+		check(Pair.Value);
+		FVertexData& Vertex = *Pair.Value;
+		if (!Vertex.GetPackageData())
+		{
+			continue;
+		}
+		check(Vertex.IsOwnedByCluster());
 		AddToVisitVertexQueue(Vertex);
 	}
+	bStarted = true;
 }
 
 FRequestCluster::FGraphSearch::~FGraphSearch()
+{
+	Reset();
+}
+
+void FRequestCluster::FGraphSearch::Reset()
 {
 	for (;;)
 	{
@@ -874,20 +931,26 @@ FRequestCluster::FGraphSearch::~FGraphSearch()
 		WaitForAsyncQueue(WaitTime);
 	}
 
-	// Call the FVertexData destructors, but do not bother calling DeleteElement or Free on the VertexAllocator
-	// since we are destructing the VertexAllocator.
-	for (TPair<FName, FVertexData*>& VertexPair : this->Vertices)
-	{
-		FVertexData* VertexData = VertexPair.Value;
-		VertexData->~FVertexData();
-	}
-	// Empty frees the struct memory for each FVertexData we allocated, but it does not call the destructor.
-	VertexAllocator.Empty();
+	GraphEdges.Empty();
+	VisitVertexQueue.Empty();
+	PendingTransitiveBuildDependencyVertices.Empty();
+	PreAsyncQueue.Empty();
+	RunAwayTickLoopCount = 0;
+	bStarted = false;
+
+	BatchAllocator.Empty();
+	check(AsyncQueueBatches.IsEmpty()); // Emptied by for loop above
+	check(AsyncQueueResults.IsEmpty()); // Emptied by for loop above
+	AsyncResultsReadyEvent->Trigger();
 }
 
 void FRequestCluster::FGraphSearch::OnNewReachablePlatforms(FPackageData* PackageData)
 {
-	FVertexData** VertexPtr = Vertices.Find(PackageData->GetPackageName());
+	if (!PackageData)
+	{
+		return;
+	}
+	FVertexData** VertexPtr = Cluster.ClusterPackages.Find(PackageData->GetPackageName());
 	if (!VertexPtr)
 	{
 		return;
@@ -897,7 +960,7 @@ void FRequestCluster::FGraphSearch::OnNewReachablePlatforms(FPackageData* Packag
 
 void FRequestCluster::FGraphSearch::QueueEdgesFetch(FVertexData& Vertex, TConstArrayView<int32> PlatformIndexes)
 {
-	check(Vertex.PackageData); // Caller must not call without a PackageData; doing so serves no purpose
+	check(Vertex.GetPackageData()); // Caller must not call without a PackageData; doing so serves no purpose
 
 	bool bAnyRequestedNeedsPlatformAgnostic = false;
 	bool bAnyRequested = false;
@@ -906,7 +969,7 @@ void FRequestCluster::FGraphSearch::QueueEdgesFetch(FVertexData& Vertex, TConstA
 	for (int32 PlatformIndex : PlatformIndexes)
 	{
 		// The platform data may have already been requested; request it only if current status is NotRequested
-		FQueryPlatformData& QueryData = Vertex.PlatformData[PlatformIndex];
+		FQueryPlatformData& QueryData = Vertex.GetPlatformData()[PlatformIndex];
 		if (!QueryData.bSchedulerThreadFetchCompleted)
 		{
 			bAllHaveAlreadyCompletedFetch = false;
@@ -949,7 +1012,7 @@ void FRequestCluster::FGraphSearch::TickExploration(bool& bOutDone)
 	// Each time through the loop we either process or delay one platform of one vertex. Once delayed, a vertex is not added to
 	// the queue again until all of its UnreadyDependencies for that platform are cleared, therefore we only delay it once per platform.
 	// An upper bound is therefore 2 * NumPlatforms * NumVertices.
-	int64 RunawayLoopUpperBound = 2 * Vertices.Num() * (FetchPlatforms.Num() - 1);
+	int64 RunawayLoopUpperBound = 2 * Cluster.ClusterPackages.Num() * (FetchPlatforms.Num() - 1);
 	for (;;)
 	{
 		TOptional<FVertexData*> FrontVertex = AsyncQueueResults.Dequeue();
@@ -982,7 +1045,7 @@ void FRequestCluster::FGraphSearch::TickExploration(bool& bOutDone)
 
 	RunawayLoopCount = 0;
 	// Calculation of upper bound: we visit each vertex at most once per platform.
-	RunawayLoopUpperBound = Vertices.Num() * (FetchPlatforms.Num() - 1);
+	RunawayLoopUpperBound = Cluster.ClusterPackages.Num() * (FetchPlatforms.Num() - 1);
 	while (!VisitVertexQueue.IsEmpty())
 	{
 		bHadActivity = true;
@@ -1008,7 +1071,7 @@ void FRequestCluster::FGraphSearch::TickExploration(bool& bOutDone)
 	if (bHadActivity)
 	{
 		++RunAwayTickLoopCount;
-		if (RunAwayTickLoopCount++ > 2 * Vertices.Num()*NumFetchPlatforms())
+		if (RunAwayTickLoopCount++ > 2 * Cluster.ClusterPackages.Num()*Cluster.GetNumFetchPlatforms())
 		{
 			UE_LOG(LogCook, Fatal, TEXT("Infinite loop detected in reentrant calls to FRequestCluster::TickExploration."));
 		}
@@ -1067,7 +1130,7 @@ void FRequestCluster::FGraphSearch::TickExploration(bool& bOutDone)
 		ResolveTransitiveBuildDependencyCycle();
 		bOutDone = false;
 		++RunAwayTickLoopCount;
-		if (RunAwayTickLoopCount++ > 2 * Vertices.Num() * NumFetchPlatforms())
+		if (RunAwayTickLoopCount++ > 2 * Cluster.ClusterPackages.Num() * Cluster.GetNumFetchPlatforms())
 		{
 			UE_LOG(LogCook, Fatal, TEXT("Infinite loop detected in FRequestCluster::PendingTransitiveBuildDependencyVertices."));
 		}
@@ -1116,7 +1179,7 @@ void FRequestCluster::FGraphSearch::ResolveTransitiveBuildDependencyCycle()
 		// We can also empty the IterativelyModifiedListeners since any remaining listeners must be in
 		// PendingTransitiveBuildDependencyVertices. Emptying the list here avoids the expense of kicking
 		// for a second time each of the listeners.
-		CycleVert->IterativelyModifiedListeners.Empty();
+		CycleVert->GetIterativelyModifiedListeners().Empty();
 		KickVertex(CycleVert);
 	}
 	check(FirstVertex); // This function should not be called if PendingTransitiveBuildDependencyVertices is empty.
@@ -1125,7 +1188,7 @@ void FRequestCluster::FGraphSearch::ResolveTransitiveBuildDependencyCycle()
 		TEXT("Cycle detected in the graph of transitive build dependencies.")
 		TEXT(" No vertices in the cycle are invalidated by their direct dependencies, so marking them all as iteratively skippable.")
 		TEXT("\n\tVertex in the cycle: %s"),
-		*FirstVertex->PackageName.ToString());
+		*FirstVertex->GetPackageName().ToString());
 }
 
 void FRequestCluster::FGraphSearch::UpdateDisplay()
@@ -1154,14 +1217,15 @@ void FRequestCluster::FGraphSearch::VisitVertex(FVertexData& Vertex)
 	// Only called from scheduler thread
 
 	// The PackageData will not exist if the package does not exist on disk
-	if (!Vertex.PackageData)
+	if (!Vertex.GetPackageData())
 	{
 		return;
 	}
+	FPackageData& PackageData = *Vertex.GetPackageData();
 
 	const EReachability ClusterReachability = Cluster.TraversalMarkCookable()
 		? EReachability::Runtime : EReachability::Build;
-	int32 LocalNumFetchPlatforms = NumFetchPlatforms();
+	int32 LocalNumFetchPlatforms = Cluster.GetNumFetchPlatforms();
 	TBitArray<> ShouldFetchPlatforms(false, LocalNumFetchPlatforms);
 	
 	FPackagePlatformData* CookerLoadingPlatform = nullptr;
@@ -1169,7 +1233,7 @@ void FRequestCluster::FGraphSearch::VisitVertex(FVertexData& Vertex)
 	ESuppressCookReason SuppressCookReason = ESuppressCookReason::Invalid;
 	bool bAllReachablesUncookable = true;
 	for (TPair<const ITargetPlatform*, FPackagePlatformData>& Pair :
-		Vertex.PackageData->GetPlatformDatasConstKeysMutableValues())
+		PackageData.GetPlatformDatasConstKeysMutableValues())
 	{
 		FPackagePlatformData& PlatformData = Pair.Value;
 		const ITargetPlatform* TargetPlatform = Pair.Key;
@@ -1198,9 +1262,9 @@ void FRequestCluster::FGraphSearch::VisitVertex(FVertexData& Vertex)
 					(Cluster.TraversalExploreRuntimeDependencies() && PlatformData.IsExplorable()))
 				{
 					ShouldFetchPlatforms[PlatformIndex] = true;
-					Vertex.PlatformData[PlatformIndex].bExploreRequested = true;
+					Vertex.GetPlatformData()[PlatformIndex].bExploreRequested = true;
 					// Exploration of any session platform also requires exploration of PlatformAgnosticPlatform
-					Vertex.PlatformData[PlatformAgnosticPlatformIndex].bExploreRequested = true;
+					Vertex.GetPlatformData()[PlatformAgnosticPlatformIndex].bExploreRequested = true;
 				}
 			}
 			if (PlatformData.IsCookable())
@@ -1214,22 +1278,22 @@ void FRequestCluster::FGraphSearch::VisitVertex(FVertexData& Vertex)
 	if (Cluster.TraversalMarkCookable())
 	{
 		bool bAnyCookable = (FirstReachableSessionPlatform == nullptr) | !bAllReachablesUncookable;
-		if (bAnyCookable != Vertex.bAnyCookable)
+		if (bAnyCookable != Vertex.IsAnyCookable())
 		{
 			if (!bAnyCookable)
 			{
 				if (SuppressCookReason == ESuppressCookReason::Invalid)
 				{
 					// We need the SuppressCookReason for reporting. If we didn't calculate it this Visit and
-					// we don't have it stored in this->OwnedPackageDatas, then we must have calculated it in
+					// we don't have it stored in this->ClusterPackages, then we must have calculated it in
 					// a previous cluster, but we don't store it anywhere. Recalculate it from the
 					// FirstReachableSessionPlatform. FirstReachableSessionPlatform must be non-null, otherwise
 					// bAnyCookable would be true.
 					check(FirstReachableSessionPlatform);
 					bool bCookable;
 					bool bExplorable;
-					Cluster.IsRequestCookable(FirstReachableSessionPlatform, Vertex.PackageData->GetPackageName(),
-						*Vertex.PackageData, SuppressCookReason, bCookable, bExplorable);
+					Cluster.IsRequestCookable(FirstReachableSessionPlatform, PackageData, SuppressCookReason, bCookable,
+						bExplorable);
 					check(!bCookable); // We don't support bCookable changing for a given package and platform
 					check(SuppressCookReason != ESuppressCookReason::Invalid);
 				}
@@ -1238,8 +1302,8 @@ void FRequestCluster::FGraphSearch::VisitVertex(FVertexData& Vertex)
 			{
 				check(SuppressCookReason == ESuppressCookReason::NotSuppressed);
 			}
-			Cluster.SetPackageDataSuppressReason(*Vertex.PackageData, SuppressCookReason);
-			Vertex.bAnyCookable = bAnyCookable;
+			Cluster.SetSuppressReason(Vertex, SuppressCookReason);
+			Vertex.SetAnyCookable(bAnyCookable);
 		}
 
 		// If any target platform is cookable, then we need to mark the CookerLoadingPlatform as reachable because we will
@@ -1248,7 +1312,7 @@ void FRequestCluster::FGraphSearch::VisitVertex(FVertexData& Vertex)
 		{
 			if (!CookerLoadingPlatform)
 			{
-				CookerLoadingPlatform = &Vertex.PackageData->FindOrAddPlatformData(CookerLoadingPlatformKey);
+				CookerLoadingPlatform = &PackageData.FindOrAddPlatformData(CookerLoadingPlatformKey);
 			}
 			CookerLoadingPlatform->AddReachability(EReachability::Runtime);
 		}
@@ -1261,7 +1325,7 @@ void FRequestCluster::FGraphSearch::VisitVertex(FVertexData& Vertex)
 			if (Cluster.TraversalExploreRuntimeDependencies())
 			{
 				ShouldFetchPlatforms[CookerLoadingPlatformIndex] = true;
-				Vertex.PlatformData[CookerLoadingPlatformIndex].bExploreRequested = true;
+				Vertex.GetPlatformData()[CookerLoadingPlatformIndex].bExploreRequested = true;
 			}
 		}
 	}
@@ -1272,7 +1336,7 @@ void FRequestCluster::FGraphSearch::VisitVertex(FVertexData& Vertex)
 	{
 		for (int32 PlatformIndex = 0; PlatformIndex < LocalNumFetchPlatforms; ++PlatformIndex)
 		{
-			FQueryPlatformData& PlatformData = Vertex.PlatformData[PlatformIndex];
+			FQueryPlatformData& PlatformData = Vertex.GetPlatformData()[PlatformIndex];
 
 			// Add on the fetch (but not the explore) of bIterativelyUnmodifiedRequested platforms
 			if (PlatformData.bIterativelyUnmodifiedRequested)
@@ -1311,11 +1375,11 @@ void FRequestCluster::FGraphSearch::VisitVertexForPlatform(FVertexData& Vertex,
 {
 	if (Cluster.TraversalMarkCookable())
 	{
-		FPackageData& PackageData = *Vertex.PackageData;
+		FPackageData& PackageData = *Vertex.GetPackageData();
 		ESuppressCookReason SuppressCookReason = ESuppressCookReason::Invalid;
 		bool bCookable;
 		bool bExplorable;
-		Cluster.IsRequestCookable(Platform, Vertex.PackageData->GetPackageName(), PackageData, SuppressCookReason,
+		Cluster.IsRequestCookable(Platform, PackageData, SuppressCookReason,
 			bCookable, bExplorable);
 		PlatformData.SetCookable(bCookable);
 		PlatformData.SetExplorable(bExplorable);
@@ -1376,10 +1440,10 @@ void FRequestCluster::FGraphSearch::FExploreEdgesContext::Explore(FVertexData& I
 void FRequestCluster::FGraphSearch::FExploreEdgesContext::Initialize(FVertexData& InVertex)
 {
 	Vertex = &InVertex;
+	PackageData = Vertex->GetPackageData();
+	PackageName = Vertex->GetPackageName();
 	// Vertices without a package data are never queued for fetch
-	check(Vertex->PackageData);
-	PackageData = Vertex->PackageData;
-	PackageName = Vertex->PackageName;
+	check(PackageData);
 
 	HardGameDependencies.Reset();
 	HardEditorDependencies.Reset();
@@ -1392,7 +1456,7 @@ void FRequestCluster::FGraphSearch::FExploreEdgesContext::Initialize(FVertexData
 	SkippedPackages.Reset();
 	UnreadyTransitiveBuildVertices.Reset();
 
-	LocalNumFetchPlatforms = GraphSearch.NumFetchPlatforms();
+	LocalNumFetchPlatforms = Cluster.GetNumFetchPlatforms();
 	bFetchAnyTargetPlatform = false;
 
 	DiscoveredDependencies = Cluster.COTFS.DiscoveredDependencies.Find(PackageName);
@@ -1402,14 +1466,14 @@ void FRequestCluster::FGraphSearch::FExploreEdgesContext::Initialize(FVertexData
 
 void FRequestCluster::FGraphSearch::FExploreEdgesContext::CalculatePlatformsToProcess()
 {
-	FQueryPlatformData& PlatformAgnosticQueryData = Vertex->PlatformData[PlatformAgnosticPlatformIndex];
+	FQueryPlatformData& PlatformAgnosticQueryData = Vertex->GetPlatformData()[PlatformAgnosticPlatformIndex];
 	for (int32 PlatformIndex = 0; PlatformIndex < LocalNumFetchPlatforms; ++PlatformIndex)
 	{
 		if (PlatformIndex == PlatformAgnosticPlatformIndex)
 		{
 			continue;
 		}
-		FQueryPlatformData& QueryPlatformData = Vertex->PlatformData[PlatformIndex];
+		FQueryPlatformData& QueryPlatformData = Vertex->GetPlatformData()[PlatformIndex];
 		if (!QueryPlatformData.bSchedulerThreadFetchCompleted)
 		{
 			continue;
@@ -1464,8 +1528,8 @@ bool FRequestCluster::FGraphSearch::FExploreEdgesContext::TryCalculateIterativel
 		return true;
 	}
 
-	Vertex->UnreadyDependencies.Reset();
-	Vertex->bWaitingOnUnreadyDependencies = false;
+	Vertex->GetUnreadyDependencies().Reset();
+	Vertex->SetWaitingOnUnreadyDependencies(false);
 	bool bAllPlatformsAreReady = true;
 	for (int32 PlatformIndex : PlatformsToProcess)
 	{
@@ -1474,7 +1538,7 @@ bool FRequestCluster::FGraphSearch::FExploreEdgesContext::TryCalculateIterativel
 			continue;
 		}
 
-		FQueryPlatformData& QueryPlatformData = Vertex->PlatformData[PlatformIndex];
+		FQueryPlatformData& QueryPlatformData = Vertex->GetPlatformData()[PlatformIndex];
 		if (QueryPlatformData.bIterativelyUnmodified.IsSet())
 		{
 			continue;
@@ -1545,8 +1609,8 @@ bool FRequestCluster::FGraphSearch::FExploreEdgesContext::TryCalculateIterativel
 				CookDependencies.GetTransitiveBuildDependencies())
 			{
 				FName TransitiveBuildPackageName = TransitiveBuildDependency.GetPackageName();
-				FVertexData& TransitiveBuildVertex = GraphSearch.FindOrAddVertex(TransitiveBuildPackageName);
-				if (!TransitiveBuildVertex.PackageData)
+				FVertexData& TransitiveBuildVertex = GraphSearch.Cluster.FindOrAddVertex(TransitiveBuildPackageName);
+				if (!TransitiveBuildVertex.GetPackageData())
 				{
 					// A build dependency on a non-existent package can occur e.g. if the package is in an
 					// unmounted plugin. If the package does not exist we count the transitivebuilddependency
@@ -1557,13 +1621,13 @@ bool FRequestCluster::FGraphSearch::FExploreEdgesContext::TryCalculateIterativel
 						TEXT("TransitiveBuildDependency to non-existent package.")
 						TEXT(" Package %s has a transitive build dependency on package %s, which does not exist or is not mounted.")
 						TEXT(" Package %s will be marked as not iteratively skippable and will be recooked."),
-						*Vertex->PackageName.ToString(), *TransitiveBuildPackageName.ToString(),
-						*Vertex->PackageName.ToString());
+						*Vertex->GetPackageName().ToString(), *TransitiveBuildPackageName.ToString(),
+						*Vertex->GetPackageName().ToString());
 					bAnyTransitiveBuildDependencyIsModified = true;
 					break;
 				}
 
-				FQueryPlatformData& TransitivePlatformData = TransitiveBuildVertex.PlatformData[PlatformIndex];
+				FQueryPlatformData& TransitivePlatformData = TransitiveBuildVertex.GetPlatformData()[PlatformIndex];
 				if (!TransitivePlatformData.bIterativelyUnmodified.IsSet())
 				{
 					UnreadyTransitiveBuildVertices.Add(&TransitiveBuildVertex);
@@ -1586,7 +1650,7 @@ bool FRequestCluster::FGraphSearch::FExploreEdgesContext::TryCalculateIterativel
 				// Add this vertex as a listener to the TransitiveBuildVertices' TryCalculateIterativelyUnmodified
 				for (FVertexData* TransitiveBuildVertex : UnreadyTransitiveBuildVertices)
 				{
-					FQueryPlatformData& TransitivePlatformData = TransitiveBuildVertex->PlatformData[PlatformIndex];
+					FQueryPlatformData& TransitivePlatformData = TransitiveBuildVertex->GetPlatformData()[PlatformIndex];
 
 					// Do not kick the vertex again if it has already been fetched; doing so will create busy work
 					// in the case of a cycle and prevent us from detecting the cycle.
@@ -1596,8 +1660,8 @@ bool FRequestCluster::FGraphSearch::FExploreEdgesContext::TryCalculateIterativel
 						GraphSearch.AddToVisitVertexQueue(*TransitiveBuildVertex);
 					}
 					// It's okay to add duplicates to IterativelyModifiedListeners; we remove them when broadcasting
-					TransitiveBuildVertex->IterativelyModifiedListeners.Add(this->Vertex);
-					Vertex->UnreadyDependencies.Add(TransitiveBuildVertex);
+					TransitiveBuildVertex->GetIterativelyModifiedListeners().Add(this->Vertex);
+					Vertex->GetUnreadyDependencies().Add(TransitiveBuildVertex);
 				}
 
 				bAllPlatformsAreReady = false;
@@ -1611,12 +1675,13 @@ bool FRequestCluster::FGraphSearch::FExploreEdgesContext::TryCalculateIterativel
 	if (!bAllPlatformsAreReady)
 	{
 		GraphSearch.PendingTransitiveBuildDependencyVertices.Add(Vertex);
-		Vertex->bWaitingOnUnreadyDependencies = true;
+		Vertex->SetWaitingOnUnreadyDependencies(true);
 		return false;
 	}
 
-	Vertex->UnreadyDependencies.Empty(); // Already reset above, call empty to free the memory
-	if (!Vertex->IterativelyModifiedListeners.IsEmpty())
+	Vertex->GetUnreadyDependencies().Empty(); // Already reset above, call empty to free the memory
+	TArray<FVertexData*>& Listeners = Vertex->GetIterativelyModifiedListeners();
+	if (!Listeners.IsEmpty())
 	{
 		bool bIterativelyModified = false;
 		for (int32 PlatformIndex = 0; PlatformIndex < LocalNumFetchPlatforms; ++PlatformIndex)
@@ -1625,30 +1690,30 @@ bool FRequestCluster::FGraphSearch::FExploreEdgesContext::TryCalculateIterativel
 			{
 				continue;
 			}
-			FQueryPlatformData& QueryPlatformData = Vertex->PlatformData[PlatformIndex];
+			FQueryPlatformData& QueryPlatformData = Vertex->GetPlatformData()[PlatformIndex];
 			if (QueryPlatformData.bIterativelyUnmodified.IsSet() && !*QueryPlatformData.bIterativelyUnmodified)
 			{
 				bIterativelyModified = true;
 				break;
 			}
 		}
-		Algo::Sort(Vertex->IterativelyModifiedListeners);
-		Vertex->IterativelyModifiedListeners.SetNum(Algo::Unique(Vertex->IterativelyModifiedListeners));
-		for (FVertexData* ListenerVertex : Vertex->IterativelyModifiedListeners)
+		Algo::Sort(Listeners);
+		Listeners.SetNum(Algo::Unique(Listeners));
+		for (FVertexData* ListenerVertex : Listeners)
 		{
-			if (!ListenerVertex->bWaitingOnUnreadyDependencies)
+			if (!ListenerVertex->IsWaitingOnUnreadyDependencies())
 			{
 				continue;
 			}
 
-			ListenerVertex->UnreadyDependencies.Remove(Vertex);
-			if (bIterativelyModified || ListenerVertex->UnreadyDependencies.IsEmpty())
+			ListenerVertex->GetUnreadyDependencies().Remove(Vertex);
+			if (bIterativelyModified || ListenerVertex->GetUnreadyDependencies().IsEmpty())
 			{
-				ListenerVertex->bWaitingOnUnreadyDependencies = false;
+				ListenerVertex->SetWaitingOnUnreadyDependencies(false);
 				GraphSearch.KickVertex(ListenerVertex);
 			}
 		}
-		Vertex->IterativelyModifiedListeners.Empty();
+		Listeners.Empty();
 	}
 	return true;
 }
@@ -1714,7 +1779,7 @@ void FRequestCluster::FGraphSearch::FExploreEdgesContext::CalculateDependenciesA
 
 	for (int32 PlatformIndex : PlatformsToExplore)
 	{
-		FQueryPlatformData& QueryPlatformData = Vertex->PlatformData[PlatformIndex];
+		FQueryPlatformData& QueryPlatformData = Vertex->GetPlatformData()[PlatformIndex];
 		FFetchPlatformData& FetchPlatformData = GraphSearch.FetchPlatforms[PlatformIndex];
 		const ITargetPlatform* TargetPlatform = FetchPlatformData.Platform;
 		FPackagePlatformData& PackagePlatformData = PackageData->FindOrAddPlatformData(TargetPlatform);
@@ -1796,13 +1861,13 @@ void FRequestCluster::FGraphSearch::FExploreEdgesContext::QueueVisitsOfDependenc
 			FCoreRedirectObjectName(NAME_None, NAME_None, DependencyName)).PackageName;
 		DependencyName = Redirected;
 
-		FVertexData& DependencyVertex = GraphSearch.FindOrAddVertex(DependencyName,
+		FVertexData& DependencyVertex = GraphSearch.Cluster.FindOrAddVertex(DependencyName,
 			GenerationHelper.GetReference());
-		if (!DependencyVertex.PackageData)
+		if (!DependencyVertex.GetPackageData())
 		{
 			continue;
 		}
-		FPackageData& DependencyPackageData(*DependencyVertex.PackageData);
+		FPackageData& DependencyPackageData(*DependencyVertex.GetPackageData());
 		bool bAddToVisitVertexQueue = false;
 
 		for (int32 PlatformIndex = 0; PlatformIndex < LocalNumFetchPlatforms; ++PlatformIndex)
@@ -1879,10 +1944,9 @@ void FRequestCluster::FGraphSearch::FExploreEdgesContext::QueueVisitsOfDependenc
 		{
 			// Only pull the vertex into the cluster if it has not already been pulled into the cluster.
 			// This prevents us from trying to readd a packagedata after COTFS called Cluster->RemovePackageData.
-			if (!DependencyVertex.bPulledIntoCluster)
+			if (!DependencyVertex.HasBeenPulledIntoCluster())
 			{
-				DependencyVertex.bPulledIntoCluster = true;
-				Cluster.PullIntoCluster(DependencyPackageData);
+				Cluster.SetOwnedByCluster(DependencyVertex, true);
 			}
 			GraphSearch.AddToVisitVertexQueue(DependencyVertex);
 		}
@@ -1893,7 +1957,7 @@ void FRequestCluster::FGraphSearch::FExploreEdgesContext::MarkExploreComplete()
 {
 	for (int32 PlatformIndex : PlatformsToExplore)
 	{
-		Vertex->PlatformData[PlatformIndex].bExploreCompleted = true;
+		Vertex->GetPlatformData()[PlatformIndex].bExploreCompleted = true;
 	}
 }
 
@@ -1967,7 +2031,7 @@ void FRequestCluster::FGraphSearch::FExploreEdgesContext::ProcessPlatformAttachm
 {
 	bool bFoundBuildDefinitions = false;
 	ICookedPackageWriter* PackageWriter = FetchPlatformData.Writer;
-	FQueryPlatformData& QueryPlatformData = Vertex->PlatformData[PlatformIndex];
+	FQueryPlatformData& QueryPlatformData = Vertex->GetPlatformData()[PlatformIndex];
 
 	if (Cluster.IsIncrementalCook())
 	{
@@ -2072,13 +2136,13 @@ void FRequestCluster::FGraphSearch::FExploreEdgesContext::ProcessPlatformAttachm
 				PackagePlatformData.SetCommitted(true);
 			}
 
-			Cluster.SetPackageDataWasMarkedSkipped(*PackageData, true);
+			Cluster.SetWasMarkedSkipped(*Vertex, true);
 		}
 	}
 
 	if (Cluster.bPreQueueBuildDefinitions && !bFoundBuildDefinitions)
 	{
-		FQueryPlatformData& PlatformAgnosticQueryData = Vertex->PlatformData[PlatformAgnosticPlatformIndex];
+		FQueryPlatformData& PlatformAgnosticQueryData = Vertex->GetPlatformData()[PlatformAgnosticPlatformIndex];
 
 		if (PlatformAgnosticQueryData.bSchedulerThreadFetchCompleted)
 		{
@@ -2105,33 +2169,33 @@ void FRequestCluster::FGraphSearch::FExploreEdgesContext::ProcessPlatformAttachm
 void FRequestCluster::FGraphSearch::FExploreEdgesContext::SetIsIterativelyUnmodified(int32 PlatformIndex,
 	bool bIterativelyUnmodified, FPackagePlatformData& PackagePlatformData)
 {
-	Vertex->PlatformData[PlatformIndex].bIterativelyUnmodified.Emplace(bIterativelyUnmodified);
+	Vertex->GetPlatformData()[PlatformIndex].bIterativelyUnmodified.Emplace(bIterativelyUnmodified);
 	if (bIterativelyUnmodified)
 	{
 		PackagePlatformData.SetIterativelyUnmodified(true);
 	}
 }
 
-FRequestCluster::FVertexData* FRequestCluster::FGraphSearch::AllocateVertex(FName PackageName, FPackageData* PackageData)
+FRequestCluster::FVertexData* FRequestCluster::AllocateVertex(FName PackageName, FPackageData* PackageData)
 {
 	// TODO: Change TypeBlockedAllocator to have an optional Size and Align argument,
-	// and use it to allocate the array of PlatformData, to reduce cpu time of allocating the array.
-	return VertexAllocator.NewElement(PackageName, PackageData, *this);
+	// and use it to allocate the FVertexData's array of PlatformData, to reduce cpu time of allocating the array.
+	return VertexAllocator.NewElement(PackageName, PackageData, GetNumFetchPlatforms());
 }
 
 FRequestCluster::FVertexData::FVertexData(FName InPackageName, UE::Cook::FPackageData* InPackageData,
-	FGraphSearch& GraphSearch)
+	int32 NumFetchPlatforms)
 	: PackageName(InPackageName)
 	, PackageData(InPackageData)
 {
-	PlatformData.Reset(new FQueryPlatformData[GraphSearch.NumFetchPlatforms()]);
+	PlatformData.SetNum(NumFetchPlatforms);
 }
 
 FRequestCluster::FVertexData&
-FRequestCluster::FGraphSearch::FindOrAddVertex(FName PackageName, FGenerationHelper* ParentGenerationHelper)
+FRequestCluster::FindOrAddVertex(FName PackageName, FGenerationHelper* ParentGenerationHelper)
 {
-	// Only called from scheduler thread
-	FVertexData*& ExistingVertex = Vertices.FindOrAdd(PackageName);
+	// Only called from process thread
+	FVertexData*& ExistingVertex = ClusterPackages.FindOrAdd(PackageName);
 	if (ExistingVertex)
 	{
 		return *ExistingVertex;
@@ -2142,7 +2206,7 @@ FRequestCluster::FGraphSearch::FindOrAddVertex(FName PackageName, FGenerationHel
 	PackageName.ToString(NameBuffer);
 	if (!FPackageName::IsScriptPackage(NameBuffer))
 	{
-		PackageData = Cluster.COTFS.PackageDatas->TryAddPackageDataByPackageName(PackageName);
+		PackageData = COTFS.PackageDatas->TryAddPackageDataByPackageName(PackageName);
 		if (!PackageData && ParentGenerationHelper && ICookPackageSplitter::IsUnderGeneratedPackageSubPath(NameBuffer))
 		{
 			const FAssetPackageData* PreviousPackageData =
@@ -2150,7 +2214,7 @@ FRequestCluster::FGraphSearch::FindOrAddVertex(FName PackageName, FGenerationHel
 			if (PreviousPackageData)
 			{
 				bool bIsMap = PreviousPackageData->Extension == EPackageExtension::Map;
-				PackageData = Cluster.COTFS.PackageDatas->TryAddPackageDataByPackageName(PackageName,
+				PackageData = COTFS.PackageDatas->TryAddPackageDataByPackageName(PackageName,
 					false /* bRequireExists */, bIsMap);
 				if (PackageData)
 				{
@@ -2165,13 +2229,14 @@ FRequestCluster::FGraphSearch::FindOrAddVertex(FName PackageName, FGenerationHel
 }
 
 FRequestCluster::FVertexData&
-FRequestCluster::FGraphSearch::FindOrAddVertex(FName PackageName, FPackageData& PackageData)
+FRequestCluster::FindOrAddVertex(FPackageData& PackageData)
 {
-	// Only called from scheduler thread
-	FVertexData*& ExistingVertex = Vertices.FindOrAdd(PackageName);
+	// Only called from process thread
+	FName PackageName = PackageData.GetPackageName();
+	FVertexData*& ExistingVertex = ClusterPackages.FindOrAdd(PackageName);
 	if (ExistingVertex)
 	{
-		check(ExistingVertex->PackageData == &PackageData);
+		check(!ExistingVertex->GetPackageData() || ExistingVertex->GetPackageData() == &PackageData);
 		return *ExistingVertex;
 	}
 
@@ -2235,7 +2300,7 @@ FRequestCluster::FQueryVertexBatch* FRequestCluster::FGraphSearch::CreateBatchOf
 	for (int32 BatchIndex = 0; BatchIndex < BatchSize; ++BatchIndex)
 	{
 		FVertexData* Vertex = PreAsyncQueue.PopFrontValue();
-		FVertexData*& ExistingVert = BatchData->Vertices.FindOrAdd(Vertex->PackageName);
+		FVertexData*& ExistingVert = BatchData->Vertices.FindOrAdd(Vertex->GetPackageName());
 		// Each PackageName should be used by just a single vertex.
 		check(!ExistingVert || ExistingVert == Vertex);
 		// If the vertex was already previously added to the batch that's okay, just ignore the new add.
@@ -2292,7 +2357,7 @@ void FRequestCluster::FQueryVertexBatch::Send()
 		for (int32 PlatformIndex = 0; PlatformIndex < PlatformDatas.Num(); ++PlatformIndex)
 		{
 			// The platform data may have already been requested; request it only if current status is NotRequested
-			FQueryPlatformData& PlatformData = Vertex->PlatformData[PlatformIndex];
+			FQueryPlatformData& PlatformData = Vertex->GetPlatformData()[PlatformIndex];
 			if (!PlatformData.bSchedulerThreadFetchCompleted)
 			{
 				bAllHaveAlreadyCompletedFetch = false;
@@ -2373,7 +2438,7 @@ void FRequestCluster::FQueryVertexBatch::RecordCacheResults(FName PackageName, i
 	UE::TargetDomain::FIncrementalCookAttachments&& CookAttachments)
 {
 	FVertexData* Vertex = Vertices.FindChecked(PackageName);
-	FQueryPlatformData& PlatformData = Vertex->PlatformData[PlatformIndex];
+	FQueryPlatformData& PlatformData = Vertex->GetPlatformData()[PlatformIndex];
 	PlatformData.CookAttachments = MoveTemp(CookAttachments);
 
 	EAsyncQueryStatus Expected = EAsyncQueryStatus::AsyncRequested;
@@ -2387,14 +2452,14 @@ void FRequestCluster::FQueryVertexBatch::RecordCacheResults(FName PackageName, i
 		// that both calls to RecordCacheResults will see all values complete, and we will kick it twice.
 		// Kicking twice is okay; it is supported and is a noop.
 		bool bAllPlatformsComplete = true;
-		int32 LocalNumFetchPlatforms = ThreadSafeOnlyVars.NumFetchPlatforms();
+		int32 LocalNumFetchPlatforms = ThreadSafeOnlyVars.Cluster.GetNumFetchPlatforms();
 		for (int32 OtherPlatformIndex = 0; OtherPlatformIndex < LocalNumFetchPlatforms; ++OtherPlatformIndex)
 		{
 			if (OtherPlatformIndex == PlatformIndex)
 			{
 				continue;
 			}
-			FQueryPlatformData& OtherPlatformData = Vertex->PlatformData[OtherPlatformIndex];
+			FQueryPlatformData& OtherPlatformData = Vertex->GetPlatformData()[OtherPlatformIndex];
 			EAsyncQueryStatus OtherStatus = OtherPlatformData.GetAsyncQueryStatus();
 			if (EAsyncQueryStatus::AsyncRequested <= OtherStatus && OtherStatus < EAsyncQueryStatus::Complete)
 			{
@@ -2435,23 +2500,23 @@ void FRequestCluster::IsRequestCookable(const ITargetPlatform* Platform, FPackag
 		FPaths::MakeStandardFilename(LocalDLCPath);
 	}
 
-	IsRequestCookable(Platform, PackageData.GetPackageName(), PackageData, COTFS,
-		LocalDLCPath, OutReason, bOutCookable, bOutExplorable);
+	IsRequestCookable(Platform, PackageData, COTFS, LocalDLCPath, OutReason, bOutCookable, bOutExplorable);
 }
 
-void FRequestCluster::IsRequestCookable(const ITargetPlatform* Platform, FName PackageName, FPackageData& PackageData,
+void FRequestCluster::IsRequestCookable(const ITargetPlatform* Platform, FPackageData& PackageData,
 	ESuppressCookReason& OutReason, bool& bOutCookable, bool& bOutExplorable)
 {
-	return IsRequestCookable(Platform, PackageName, PackageData, COTFS,
+	return IsRequestCookable(Platform, PackageData, COTFS,
 		DLCPath, OutReason, bOutCookable, bOutExplorable);
 }
 
-void FRequestCluster::IsRequestCookable(const ITargetPlatform* Platform, FName PackageName, FPackageData& PackageData,
+void FRequestCluster::IsRequestCookable(const ITargetPlatform* Platform, FPackageData& PackageData,
 	UCookOnTheFlyServer& InCOTFS, FStringView InDLCPath, ESuppressCookReason& OutReason, bool& bOutCookable,
 	bool& bOutExplorable)
 {
 	// IsRequestCookable should not be called for The CookerLoadingPlatform; it has different rules
 	check(Platform != CookerLoadingPlatformKey);
+	FName PackageName = PackageData.GetPackageName();
 
 	TStringBuilder<256> NameBuffer;
 	// We need to reject packagenames from adding themselves or their transitive dependencies using all the same rules
diff --git a/Engine/Source/Editor/UnrealEd/Private/Cooker/CookRequestCluster.h b/Engine/Source/Editor/UnrealEd/Private/Cooker/CookRequestCluster.h
--- a/Engine/Source/Editor/UnrealEd/Private/Cooker/CookRequestCluster.h
+++ b/Engine/Source/Editor/UnrealEd/Private/Cooker/CookRequestCluster.h
@@ -49,10 +49,12 @@ namespace UE::Cook
 class FRequestCluster
 {
 public:
+	~FRequestCluster();
 	FRequestCluster(UCookOnTheFlyServer& COTFS, TArray<FFilePlatformRequest>&& InRequests);
 	FRequestCluster(UCookOnTheFlyServer& COTFS, FPackageDataSet&& InRequests, EReachability InExploreReachability);
 	FRequestCluster(UCookOnTheFlyServer& COTFS, TRingBuffer<FDiscoveryQueueElement>& DiscoveryQueue);
-	FRequestCluster(FRequestCluster&&) = default;
+	FRequestCluster(FRequestCluster&&) = delete; // Has internal pointers, would have to write manually
+	FRequestCluster(const FRequestCluster&) = delete;
 	enum EBuildDependencyQueueConstructorType
 	{
 		BuildDependencyQueue
@@ -88,7 +90,8 @@ public:
 		TMap<FPackageData*, TArray<FPackageData*>>& OutRequestGraph);
 	/**
 	 * Report packages that are in request state and assigned to this Cluster, but that should not be counted as in
-	 * progress for progress displays because this cluster has marked them as already cooked or as to be demoted. */
+	 * progress for progress displays because this cluster has marked them as already cooked or as to be demoted.
+	 */
 	int32 GetPackagesToMarkNotInProgress() const;
 
 	static TConstArrayView<FName> GetLocalizationReferences(FName PackageName, UCookOnTheFlyServer& InCOTFS);
@@ -140,24 +143,100 @@ private:
 	};
 
 	/**
-	 * GraphSearch data for an package referenced by the cluster. VertexData is created when a package is discovered
-	 * from the dependencies of a referencer package. It remains allocated for the rest of the Cluster's lifetime.
+	 * Extra data about a package owned or referenced by the cluster that is needed for the lifetime of the cluster.
+	 * FVertexDatas are never deallocated while async operations are active, they can only be deallocated after all
+	 * async operations are complete, and all FVertexDatas are deallocated together.
 	 */
 	struct FVertexData
 	{
-		FVertexData(FName InPackageName, UE::Cook::FPackageData* InPackageData, FGraphSearch& GraphSearch);
+	public:
+		// Constructor and initialization that occurs before async work on the vertex is possible
+		FVertexData(FName InPackageName, UE::Cook::FPackageData* InPackageData, int32 NumFetchPlatforms);
+
+
+		// Interface that is readonly once async work on the vertex is possible and is therefore readable from any thread
+
+
+		FName GetPackageName() const;
+
+
+		// Interface that is callable only by the current owner thread, which switches from process thread to
+		// async thread during fetch
+
+
+		/** Settings and Results for each of the GraphSearch's FetchPlatforms. Element n corresponds to FetchPlatform n. */
+		TArrayView<FQueryPlatformData> GetPlatformData();
+
+
+		// Interface that is callable from the process thread only
+
+
+		FPackageData* GetPackageData() const;
+		TArray<FVertexData*>& GetIterativelyModifiedListeners();
+		TSet<FVertexData*> GetUnreadyDependencies();
 		const FAssetPackageData* GetGeneratedAssetPackageData();
 
-		/* Async thread is not allowed to access PackageData, so store its name.The name is immutable for vertex lifetime. */
+		/**
+		 * Whether the package is owned by this cluster and the cluster should decide its next state.
+		 * BuildDependency packages are the example where this is not true; they are tracked by the cluster to decide
+		 * skippability of other packages but are not in the cluster and might be idle or in another state.
+		 */
+		bool IsOwnedByCluster() const;
+		void SetOwnedByCluster(bool bOwned);
+		
+		/**
+		 * Whether the package has already been pulled into the cluster once. If there is access from multiple
+		 * clusters this can be true even if GetOwnedByCluster is false.
+		 */
+		bool HasBeenPulledIntoCluster() const;
+
+		/** The package's SuppressCookReason, either NotSuppressed or a reason it was suppressed. */
+		ESuppressCookReason GetSuppressReason() const;
+		void SetSuppressReason(ESuppressCookReason Value);
+
+		/** Whether the package has been marked cookable by any platform. */
+		bool IsAnyCookable() const;
+		void SetAnyCookable(bool bInCookable);
+
+		/**
+		 * Whether the vertex has already checked its dependencies once for skippability, but found some dependencies
+		 * that need to be evaluated before it can decide, and is now waiting for their evaluation to complete.
+		 */
+		bool IsWaitingOnUnreadyDependencies() const;
+		void SetWaitingOnUnreadyDependencies(bool bWaiting);
+
+		/** Whether the package was marked as committed for any platform by this cluster. */
+		bool WasMarkedSkipped() const;
+		void SetWasMarkedSkipped(bool bValue);
+
+		/**
+		 * Whether this package is owned by the cluster and therefore in progress, but should be subtracted from the
+		 * inprogress count because it will be removed from inprogress when the cluster completes.
+		 * Used by COTFS when displaying number of PackageDatas in each state.
+		 */
+		bool IsOwnedButNotInProgress() const;
+
+	private:
+		// Data that is readonly once async work on the vertex is possible and is therefore readable from any thread
+
 		FName PackageName;
+
+		// Data that is read/write only by the current owner thread, which switches from process thread to
+		// async thread during fetch
+
+		TArray<FQueryPlatformData> PlatformData;
+
+		// Data that is read/write from the process thread only
+
 		TArray<FVertexData*> IterativelyModifiedListeners;
 		TSet<FVertexData*> UnreadyDependencies;
 		UE::Cook::FPackageData* PackageData = nullptr;
+		ESuppressCookReason SuppressCookReason = ESuppressCookReason::NotSuppressed;
+		bool bOwnedByCluster = false;
+		bool bHasBeenPulledIntoCluster = false;
 		bool bAnyCookable = true;
-		bool bPulledIntoCluster = false;
 		bool bWaitingOnUnreadyDependencies = false;
-		/** Settings and Results for each of the GraphSearch's FetchPlatforms. Element n corresponds to FetchPlatform n. */
-		TUniquePtr<FQueryPlatformData[]> PlatformData;
+		bool bWasMarkedSkipped = false;
 	};
 
 	/**
@@ -250,13 +329,17 @@ private:
 	{
 	public:
 		FGraphSearch(FRequestCluster& InCluster);
+		void Initialize();
 		~FGraphSearch();
+		void Reset();
+		bool IsInitialized() const;
 
 		// All public functions are callable only from the process thread
-		/** Skip the entire GraphSearch and just visit the Cluster's current OwnedPackageDatas. */
+		/** Skip the entire GraphSearch and just visit the Cluster's current ClusterPackages. */
 		void VisitWithoutFetching();
-		/** Start a search from the Cluster's current OwnedPackageDatas. */
+		/** Start a search from the Cluster's current ClusterPackages. */
 		void StartSearch();
+		bool IsStarted() const;
 		void OnNewReachablePlatforms(FPackageData* PackageData);
 
 		/**
@@ -341,18 +424,13 @@ private:
 		/** Asynchronously fetch the dependencies and previous incremental results for a vertex */
 		void QueueEdgesFetch(FVertexData& Vertex, TConstArrayView<int32> PlatformIndexes);
 		/** Calculate and store the vertex's PackageData's cookability for each reachable platform. Kick off edges fetch. */
-		void VisitVertex(FVertexData& VertexData);
+		void VisitVertex(FVertexData& Vertex);
 		/** Calculate and store the vertex's PackageData's cookability for the platform. */
-		void VisitVertexForPlatform(FVertexData& VertexData, const ITargetPlatform* Platform,
+		void VisitVertexForPlatform(FVertexData& Vertex, const ITargetPlatform* Platform,
 			EReachability ClusterReachability, FPackagePlatformData& PlatformData,
 			ESuppressCookReason& AccumulatedSuppressCookReason);
 		void ResolveTransitiveBuildDependencyCycle();
 
-		/** Find or add a Vertex for PackageName. If PackageData is provided, use it, otherwise look it up. */
-		FVertexData& FindOrAddVertex(FName PackageName, FGenerationHelper* ParentGenerationHelper = nullptr);
-		FVertexData& FindOrAddVertex(FName PackageName, FPackageData& PackageData);
-		/** Batched allocation for vertices. */
-		FVertexData* AllocateVertex(FName PackageName, FPackageData* PackageData);
 		/** Queue a vertex for visiting and dependency traversal */
 		void AddToVisitVertexQueue(FVertexData& Vertex);
 
@@ -372,11 +450,6 @@ private:
 		/** Notify process thread of vertex completion. */
 		void KickVertex(FVertexData* Vertex);
 
-		/** Total number of platforms known to the cluster, including the special cases. */
-		int32 NumFetchPlatforms() const;
-		/** Total number of non-special-case platforms known to the cluster.Identical to COTFS's session platforms */
-		int32 NumSessionPlatforms() const;
-
 		TArrayView<FQueryPlatformData> GetPlatformDataArray(FVertexData& Vertex);
 
 	private:
@@ -388,15 +461,15 @@ private:
 		/** A set of stack and scratch variables used when calculating and exploring the edges of a vertex. */
 		FExploreEdgesContext ExploreEdgesContext;
 		TMap<FPackageData*, TArray<FPackageData*>> GraphEdges;
-		TMap<FName, FVertexData*> Vertices;
 		TSet<FVertexData*> VisitVertexQueue;
 		TSet<FVertexData*> PendingTransitiveBuildDependencyVertices;
-		TTypedBlockAllocatorFreeList<FVertexData> VertexAllocator;
 		/** Vertices queued for async processing that are not yet numerous enough to fill a batch. */
 		TRingBuffer<FVertexData*> PreAsyncQueue;
 		/** Time-tracker for timeout warnings in Poll */
 		double LastActivityTime = 0.;
 		int32 RunAwayTickLoopCount = 0;
+		bool bInitialized = false;
+		bool bStarted = false;
 
 		// Variables that are accessible from multiple threads, guarded by Lock
 		FCriticalSection Lock;
@@ -408,42 +481,25 @@ private:
 		FEventRef AsyncResultsReadyEvent;
 	};
 
-	/** Tracks flags about this cluster's processing state for its OwnedPackageDatas. */
-	struct FProcessingFlags
-	{
-		/** Whether this struct has been set valid, used to identify whether values exist in a TMap. */
-		bool IsValid() const;
-		/** The package's SuppressCookReason, either NotSuppressed or a reason it was suppressed. */
-		ESuppressCookReason GetSuppressReason() const;
-		/** Whether the package was marked as committed for any platform by this cluster. */
-		bool WasMarkedSkipped() const;
-		/** Whether the values indicate the package should be added to the cluster's PackagesToMarkNotInProgress. */
-		bool ShouldMarkNotInProgress() const;
-
-		void SetValid();
-		void SetSuppressReason(ESuppressCookReason Value);
-		void SetWasMarkedSkipped(bool bValue);
-
-	private:
-		ESuppressCookReason SuppressCookReason = ESuppressCookReason::NotSuppressed;
-		bool bValid = false;
-		bool bWasMarkedSkipped = false;
-	};
-
 private:
 	FRequestCluster(UCookOnTheFlyServer& COTFS, EReachability ExploreReachability);
+	void EmptyClusterPackages();
 	void ReserveInitialRequests(int32 RequestNum);
-	void PullIntoCluster(FPackageData& PackageData);
+	/**
+	 * Track the cluster's count of vertices that depend on the vertex's state. Delta indicates whether the
+	 * vertex is being added or removed from the counts.
+	 */
+	void AddVertexCounts(FVertexData& Vertex, int32 Delta);
+	void SetOwnedByCluster(FVertexData& Vertex, bool bOwnedByCluster, bool bNeedsStateChange = true);
+	void SetSuppressReason(FVertexData& Vertex, ESuppressCookReason Reason);
+	void SetWasMarkedSkipped(FVertexData& Vertex, bool bWasMarkedSkipped);
 	void FetchPackageNames(const FCookerTimer& CookerTimer, bool& bOutComplete);
 	void PumpExploration(const FCookerTimer& CookerTimer, bool& bOutComplete);
 	void StartAsync(const FCookerTimer& CookerTimer, bool& bOutComplete);
 	bool IsIncrementalCook() const;
-	void SetPackageDataSuppressReason(FPackageData& PackageData, ESuppressCookReason Reason,
-		bool* bOutExisted = nullptr);
-	void SetPackageDataWasMarkedSkipped(FPackageData& PackageData, bool bValue, bool* bOutExisted = nullptr);
-	void IsRequestCookable(const ITargetPlatform* TargetPlatform, FName PackageName, FPackageData& PackageData,
+	void IsRequestCookable(const ITargetPlatform* TargetPlatform, FPackageData& PackageData,
 		ESuppressCookReason& OutReason, bool& bOutCookable, bool& bOutExplorable);
-	static void IsRequestCookable(const ITargetPlatform* TargetPlatform, FName PackageName, FPackageData& PackageData,
+	static void IsRequestCookable(const ITargetPlatform* TargetPlatform, FPackageData& PackageData,
 		UCookOnTheFlyServer& InCOTFS, FStringView InDLCPath, ESuppressCookReason& OutReason, bool& bOutCookable,
 		bool& bOutExplorable);
 	/**
@@ -461,24 +517,33 @@ private:
 	 */
 	bool TraversalMarkCookable();
 
+	/** Total number of platforms known to the cluster, including the special cases. */
+	int32 GetNumFetchPlatforms() const;
+	/** Total number of non-special-case platforms known to the cluster.Identical to COTFS's session platforms */
+	int32 GetNumSessionPlatforms() const;
+
+	/** Find or add a Vertex for PackageName. If PackageData is provided, use it, otherwise look it up. */
+	FVertexData& FindOrAddVertex(FName PackageName, FGenerationHelper* ParentGenerationHelper = nullptr);
+	FVertexData& FindOrAddVertex(FPackageData& PackageData);
+	/** Batched allocation for vertices. */
+	FVertexData* AllocateVertex(FName PackageName, FPackageData* PackageData);
+
+
 	TArray<FFilePlatformRequest> FilePlatformRequests;
-	/**
-	 * Set of all packageDatas owned by this cluster (they are in the request state and this is the requeststate
-	 * container that records them). The count of PackageDatas matching certain properties is stored in
-	 * PackagesToMarkNotInProgress and must be updated whenever values change in OwnedPackageDatas. Call
-	 * SetPackageData... functions or RemovePackageData instead of modifying it directly.
-	 */
-	TFastPointerMap<FPackageData*, FProcessingFlags> OwnedPackageDatas;
+	TMap<FName, FVertexData*> ClusterPackages;
 	TMap<FPackageData*, TArray<FPackageData*>> RequestGraph;
+	TTypedBlockAllocatorFreeList<FVertexData> VertexAllocator;
 	FString DLCPath;
-	TUniquePtr<FGraphSearch> GraphSearch; // Needs to be dynamic-allocated because of large alignment
+	FGraphSearch GraphSearch;
 	UCookOnTheFlyServer& COTFS;
 	FPackageDatas& PackageDatas;
 	IAssetRegistry& AssetRegistry;
 	FPackageTracker& PackageTracker;
 	FBuildDefinitions& BuildDefinitions;
-	int32 PackagesToMarkNotInProgressCount = 0;
 	ETraversalTier TraversalTier = ETraversalTier::RuntimeFollowDependencies;
+	int32 NumOwned = 0;
+	int32 NumOwnedButNotInProgress = 0;
+	int32 NumFetchPlatforms = 0;
 	bool bAllowHardDependencies = true;
 	bool bAllowSoftDependencies = true;
 	bool bErrorOnEngineContentUse = false;
@@ -494,61 +559,107 @@ private:
 // Inline implementations
 ///////////////////////////////////////////////////////
 
-inline int32 FRequestCluster::GetPackagesToMarkNotInProgress() const
+inline FName FRequestCluster::FVertexData::GetPackageName() const
 {
-	return PackagesToMarkNotInProgressCount;
+	return PackageName;
 }
 
-inline bool FRequestCluster::FProcessingFlags::IsValid() const
+inline TArrayView<FRequestCluster::FQueryPlatformData> FRequestCluster::FVertexData::GetPlatformData()
 {
-	return bValid;
+	return PlatformData;
 }
 
-inline ESuppressCookReason FRequestCluster::FProcessingFlags::GetSuppressReason() const
+inline FPackageData* FRequestCluster::FVertexData::GetPackageData() const
 {
-	return SuppressCookReason;
+	return PackageData;
 }
 
-inline bool FRequestCluster::FProcessingFlags::WasMarkedSkipped() const
+inline TArray<FRequestCluster::FVertexData*>& FRequestCluster::FVertexData::GetIterativelyModifiedListeners()
 {
-	return bWasMarkedSkipped;
+	return IterativelyModifiedListeners;
 }
 
-inline bool FRequestCluster::FProcessingFlags::ShouldMarkNotInProgress() const
+inline TSet<FRequestCluster::FVertexData*> FRequestCluster::FVertexData::GetUnreadyDependencies()
 {
-	return bValid & 
-		((SuppressCookReason != ESuppressCookReason::NotSuppressed) | bWasMarkedSkipped);
+	return UnreadyDependencies;
 }
 
-inline void FRequestCluster::FProcessingFlags::SetValid()
+inline bool FRequestCluster::FVertexData::IsOwnedByCluster() const
 {
-	bValid = true;
+	return bOwnedByCluster;
 }
 
-inline void FRequestCluster::FProcessingFlags::SetSuppressReason(ESuppressCookReason Value)
+inline void FRequestCluster::FVertexData::SetOwnedByCluster(bool bOwned)
+{
+	bOwnedByCluster = bOwned;
+	bHasBeenPulledIntoCluster |= bOwned;
+}
+
+inline bool FRequestCluster::FVertexData::HasBeenPulledIntoCluster() const
+{
+	return bHasBeenPulledIntoCluster;
+}
+
+inline ESuppressCookReason FRequestCluster::FVertexData::GetSuppressReason() const
+{
+	return SuppressCookReason;
+}
+
+inline void FRequestCluster::FVertexData::SetSuppressReason(ESuppressCookReason Value)
 {
 	SuppressCookReason = Value;
 }
 
-inline void FRequestCluster::FProcessingFlags::SetWasMarkedSkipped(bool bValue)
+inline bool FRequestCluster::FVertexData::IsAnyCookable() const
+{
+	return bAnyCookable;
+}
+
+inline void FRequestCluster::FVertexData::SetAnyCookable(bool bInCookable)
+{
+	bAnyCookable = bInCookable;
+}
+
+inline bool FRequestCluster::FVertexData::IsWaitingOnUnreadyDependencies() const
+{
+	return bWaitingOnUnreadyDependencies;
+}
+
+inline void FRequestCluster::FVertexData::SetWaitingOnUnreadyDependencies(bool bWaiting)
+{
+	bWaitingOnUnreadyDependencies = bWaiting;
+}
+
+inline bool FRequestCluster::FVertexData::WasMarkedSkipped() const
+{
+	return bWasMarkedSkipped;
+}
+
+inline void FRequestCluster::FVertexData::SetWasMarkedSkipped(bool bValue)
 {
 	bWasMarkedSkipped = bValue;
 }
 
-inline int32 FRequestCluster::FGraphSearch::NumFetchPlatforms() const
+inline bool FRequestCluster::FVertexData::IsOwnedButNotInProgress() const
 {
-	return FetchPlatforms.Num();
+	return bOwnedByCluster &
+		((SuppressCookReason != ESuppressCookReason::NotSuppressed) | bWasMarkedSkipped);
 }
 
-inline int32 FRequestCluster::FGraphSearch::NumSessionPlatforms() const
+inline bool FRequestCluster::FGraphSearch::IsInitialized() const
 {
-	return FetchPlatforms.Num() - 2;
+	return bInitialized;
+}
+
+inline bool FRequestCluster::FGraphSearch::IsStarted() const
+{
+	return bStarted;
 }
 
 inline TArrayView<FRequestCluster::FQueryPlatformData> FRequestCluster::FGraphSearch::GetPlatformDataArray(
 	FVertexData& Vertex)
 {
-	return TArrayView< FQueryPlatformData>(Vertex.PlatformData.Get(), NumFetchPlatforms());
+	return Vertex.GetPlatformData();
 }
 
 inline FRequestCluster::EAsyncQueryStatus FRequestCluster::FQueryPlatformData::GetAsyncQueryStatus()
@@ -570,4 +681,24 @@ inline bool FRequestCluster::FQueryPlatformData::CompareExchangeAsyncQueryStatus
 	);
 }
 
+inline int32 FRequestCluster::NumPackageDatas() const
+{
+	return NumOwned;
+}
+
+inline int32 FRequestCluster::GetPackagesToMarkNotInProgress() const
+{
+	return NumOwnedButNotInProgress;
+}
+
+inline int32 FRequestCluster::GetNumFetchPlatforms() const
+{
+	return NumFetchPlatforms;
+}
+
+inline int32 FRequestCluster::GetNumSessionPlatforms() const
+{
+	return NumFetchPlatforms - 2;
+}
+
 }
