[Viewport Toolbar] Copy "Enable LOD View Locking" from main toolbar "Settings" menu to Editor Settings

diff --git a/Engine/Source/Editor/LevelEditor/Private/LevelEditorActions.cpp b/Engine/Source/Editor/LevelEditor/Private/LevelEditorActions.cpp
--- a/Engine/Source/Editor/LevelEditor/Private/LevelEditorActions.cpp
+++ b/Engine/Source/Editor/LevelEditor/Private/LevelEditorActions.cpp
@@ -2989,13 +2989,15 @@ bool FLevelEditorActionCallbacks::OnIsParticleSystemHelpersEnabled()
 
 void FLevelEditorActionCallbacks::OnToggleLODViewLocking()
 {
-	GEditor->bEnableLODLocking = !GEditor->bEnableLODLocking;
+	uint32 bUseLODViewLocking = !GetDefault<ULevelEditorViewportSettings>()->bUseLODViewLocking;
+	GetMutableDefault<ULevelEditorViewportSettings>()->bUseLODViewLocking = bUseLODViewLocking;
+
 	GEditor->RedrawLevelEditingViewports();
 }
 
 bool FLevelEditorActionCallbacks::OnIsLODViewLockingEnabled()
 {
-	return GEditor->bEnableLODLocking;
+	return GetDefault<ULevelEditorViewportSettings>()->bUseLODViewLocking;
 }
 
 void FLevelEditorActionCallbacks::OnToggleLevelStreamingVolumePrevis()
diff --git a/Engine/Source/Editor/UnrealEd/Classes/Editor/EditorEngine.h b/Engine/Source/Editor/UnrealEd/Classes/Editor/EditorEngine.h
--- a/Engine/Source/Editor/UnrealEd/Classes/Editor/EditorEngine.h
+++ b/Engine/Source/Editor/UnrealEd/Classes/Editor/EditorEngine.h
@@ -428,7 +428,8 @@ public:
 	uint32 bEnableSocketSnapping:1;
 
 	/** If true, same type views will be camera tied, and ortho views will use perspective view for LOD parenting */
-	UPROPERTY()
+	UE_DEPRECATED(5.6, "bEnableLODLocking is now handled by ULevelEditorViewportSettings::bUseLODViewLocking. Setting can now be found under Editor Preferences->Level Editor->Viewport->Look And Feel")
+	UPROPERTY(config, meta=(DeprecatedProperty, DeprecationMessage = "bEnableLODLocking can now be found under Editor Preferences->Level Editor->Viewport->Look And Feel"))
 	uint32 bEnableLODLocking:1;
 
 	UPROPERTY(config)
diff --git a/Engine/Source/Editor/UnrealEd/Classes/Settings/LevelEditorViewportSettings.h b/Engine/Source/Editor/UnrealEd/Classes/Settings/LevelEditorViewportSettings.h
--- a/Engine/Source/Editor/UnrealEd/Classes/Settings/LevelEditorViewportSettings.h
+++ b/Engine/Source/Editor/UnrealEd/Classes/Settings/LevelEditorViewportSettings.h
@@ -515,6 +515,12 @@ public:
 	UPROPERTY(EditAnywhere, config, Category=LookAndFeel, meta=(DisplayName = "Link Orthographic Viewport Movement"))
 	uint32 bUseLinkedOrthographicViewports:1;
 
+	// bUseLODViewLocking is added to keep around a legacy command from FLevelEditorCommands, as we are removing the Level Editor Settings menu.
+	// LOD View Locking is possibly presenting some bugs.
+	/** If true, same type views will be camera tied, and ortho views will use perspective view for LOD parenting */
+	UPROPERTY(EditAnywhere, config, Category=LookAndFeel, meta=(DisplayName = "Enable LOD View Locking"))
+	uint32 bUseLODViewLocking:1;
+
 	/** True if viewport box selection requires objects to be fully encompassed by the selection box to be selected */
 	UPROPERTY(config)
 	uint32 bStrictBoxSelection:1;
diff --git a/Engine/Source/Editor/UnrealEd/Private/LevelEditorViewport.cpp b/Engine/Source/Editor/UnrealEd/Private/LevelEditorViewport.cpp
--- a/Engine/Source/Editor/UnrealEd/Private/LevelEditorViewport.cpp
+++ b/Engine/Source/Editor/UnrealEd/Private/LevelEditorViewport.cpp
@@ -2526,7 +2526,7 @@ FSceneView* FLevelEditorViewportClient::CalcSceneView(FSceneViewFamily* ViewFami
 
 	// set all other matching viewports to my location, if the LOD locking is enabled,
 	// unless another viewport already set me this frame (otherwise they fight)
-	if (GEditor->bEnableLODLocking)
+	if (GetDefault<ULevelEditorViewportSettings>()->bUseLODViewLocking)
 	{
 		for (FLevelEditorViewportClient* ViewportClient : GEditor->GetLevelViewportClients())
 		{
diff --git a/Engine/Source/Editor/UnrealEd/Private/Settings/SettingsClasses.cpp b/Engine/Source/Editor/UnrealEd/Private/Settings/SettingsClasses.cpp
--- a/Engine/Source/Editor/UnrealEd/Private/Settings/SettingsClasses.cpp
+++ b/Engine/Source/Editor/UnrealEd/Private/Settings/SettingsClasses.cpp
@@ -40,6 +40,7 @@
 #include "DesktopPlatformModule.h"
 #include "InstalledPlatformInfo.h"
 #include "DrawDebugHelpers.h"
+#include "LevelEditorActions.h"
 #include "ToolMenus.h"
 #include "Subsystems/AssetEditorSubsystem.h"
 
@@ -964,6 +965,7 @@ ULevelEditorViewportSettings::ULevelEditorViewportSettings( const FObjectInitial
 	PreviewMeshes.Add(FSoftObjectPath("/Engine/EditorMeshes/ColorCalibrator/SM_ColorCalibrator.SM_ColorCalibrator"));
 	LastInViewportMenuLocation = FVector2D(EForceInit::ForceInitToZero);
 	PreserveNonUniformScale = true;
+	bUseLODViewLocking = false;
 }
 
 void ULevelEditorViewportSettings::PostInitProperties()
@@ -1057,6 +1059,10 @@ void ULevelEditorViewportSettings::PostEditChangeProperty( struct FPropertyChang
 			Settings2D->bEnableSnapLayers = true;
 		}
 	}
+	else if (Name == GET_MEMBER_NAME_CHECKED(ULevelEditorViewportSettings, bUseLODViewLocking))
+	{
+		GEditor->RedrawLevelEditingViewports();
+	}
 
 	if (!FUnrealEdMisc::Get().IsDeletePreferences())
 	{
