UE: A couple resave commandlet improvements. -PackageFolder can now take a '+' delimited list of folders instead of only one. You can now specifiy -NoExternalPackages on the commandline to skip saving external packages in maps that have them

diff --git a/Engine/Source/Editor/UnrealEd/Private/Commandlets/ContentCommandlets.cpp b/Engine/Source/Editor/UnrealEd/Private/Commandlets/ContentCommandlets.cpp
--- a/Engine/Source/Editor/UnrealEd/Private/Commandlets/ContentCommandlets.cpp
+++ b/Engine/Source/Editor/UnrealEd/Private/Commandlets/ContentCommandlets.cpp
@@ -138,7 +138,7 @@ int32 UResavePackagesCommandlet::InitializeResaveParameters( const TArray<FStrin
 	for( int32 SwitchIdx = 0; SwitchIdx < Switches.Num(); SwitchIdx++ )
 	{
 		FString Package;
-		FString PackageFolder;
+		FString PackageFolderArg;
 		FString Maps;
 		FString File;
 		const FString& CurrentSwitch = Switches[ SwitchIdx ];
@@ -156,20 +156,26 @@ int32 UResavePackagesCommandlet::InitializeResaveParameters( const TArray<FStrin
 			
 			bExplicitPackages = true;
 		}
-		else if( FParse::Value( *CurrentSwitch, TEXT( "PACKAGEFOLDER="), PackageFolder ) )
+		else if( FParse::Value( *CurrentSwitch, TEXT( "PACKAGEFOLDER="), PackageFolderArg) )
 		{
-			FPaths::NormalizeDirectoryName(PackageFolder);
-
-			TArray<FString> FilesInPackageFolder;
-			FPackageName::FindPackagesInDirectory(FilesInPackageFolder, PackageFolder);
-			for( int32 FileIndex = 0; FileIndex < FilesInPackageFolder.Num(); FileIndex++ )
+			// Allow support for -PACKAGEFOLDER=Value1+Value2+Value3
+			TArray<FString> PackageFolders;
+			PackageFolderArg.ParseIntoArray(PackageFolders, TEXT("+"));
+			for (FString& PackageFolder : PackageFolders)
 			{
-				FString PackageFile(FilesInPackageFolder[FileIndex]);
-				FPaths::MakeStandardFilename(PackageFile);
-				PackageNames.Add( *PackageFile );
+				FPaths::NormalizeDirectoryName(PackageFolder);
+
+				TArray<FString> FilesInPackageFolder;
+				FPackageName::FindPackagesInDirectory(FilesInPackageFolder, PackageFolder);
+				for( int32 FileIndex = 0; FileIndex < FilesInPackageFolder.Num(); FileIndex++ )
+				{
+					FString PackageFile(FilesInPackageFolder[FileIndex]);
+					FPaths::MakeStandardFilename(PackageFile);
+					PackageNames.Add( *PackageFile );
+				}
 			}
 
-			UE_CLOG(PackageNames.IsEmpty(), LogContentCommandlet, Warning, TEXT("Failed to find any packages in folder: '%s'"), *PackageFolder);
+			UE_CLOG(PackageNames.IsEmpty(), LogContentCommandlet, Warning, TEXT("Failed to find any packages in folder: '%s'"), *PackageFolderArg);
 			bExplicitPackages = true;
 		}
 		else if (FParse::Value(*CurrentSwitch, TEXT("MAP="), Maps))
@@ -1911,7 +1917,7 @@ void UResavePackagesCommandlet::PerformAdditionalOperations(class UWorld* World,
 
 	bool bRevertCheckedOutFilesIfNotSaving = true;
 
-	const bool bFixupRedirects = (Switches.Contains(TEXT("FixupRedirects")) || Switches.Contains(TEXT("FixupRedirectors")));
+	const bool bAllowExternalPackageResaves = !Switches.Contains(TEXT("FixupRedirects")) && !Switches.Contains(TEXT("FixupRedirectors")) && !Switches.Contains(TEXT("NoExternalPackages"));
 	const bool bShouldBuildTextureStreamingForWorld = bShouldBuildTextureStreaming && !bShouldBuildTextureStreamingForAll;
 	const bool bBuildingNonHLODData = (bShouldBuildLighting || bShouldBuildTextureStreamingForWorld || bShouldBuildReflectionCaptures);
 
@@ -1923,7 +1929,7 @@ void UResavePackagesCommandlet::PerformAdditionalOperations(class UWorld* World,
 	const int32 DefaultExternalActorGCFreq = 2048;
 
 	// Load and Save Level's external packages
- 	if (!bResaveWorldPartitionExternalPackages && !bFixupRedirects)
+ 	if (!bResaveWorldPartitionExternalPackages && bAllowExternalPackageResaves)
 	{
 		// Use a default GC frequency for external actors if GarbageCollectionFrequency is 0.
 		TGuardValue<int32> ScopedGCFreq(GarbageCollectionFrequency, GarbageCollectionFrequency ? GarbageCollectionFrequency : DefaultExternalActorGCFreq);
@@ -1961,7 +1967,7 @@ void UResavePackagesCommandlet::PerformAdditionalOperations(class UWorld* World,
 		// Use a default GC frequency for external actors if GarbageCollectionFrequency is 0.
 		TGuardValue<int32> ScopedGCFreq(GarbageCollectionFrequency, GarbageCollectionFrequency ? GarbageCollectionFrequency : DefaultExternalActorGCFreq);
 
-		if (!bFixupRedirects)
+		if (bAllowExternalPackageResaves)
 		{
 			auto ResaveExternalPackage = [this](const UPackage* Package)
 			{
