Postpone PIE WorldContext removal after all PIE worlds have been tore down so that if a callback happens in one of the world teardown, the worldcontext of already processed worlds are still available

#rnx

diff --git a/Engine/Source/Editor/UnrealEd/Private/PlayLevel.cpp b/Engine/Source/Editor/UnrealEd/Private/PlayLevel.cpp
--- a/Engine/Source/Editor/UnrealEd/Private/PlayLevel.cpp
+++ b/Engine/Source/Editor/UnrealEd/Private/PlayLevel.cpp
@@ -351,8 +351,15 @@ void UEditorEngine::EndPlayMap()
 				UOnlineEngineInterface::Get()->ShutdownOnlineSubsystem(OnlineIdentifier);
 				OnlineIdentifiers.Add(OnlineIdentifier);
 			}
-		
-			// Remove world list after online has shutdown in case any async actions require the world context
+		}
+	}
+
+	// Only remove from list after all PIE worlds have been torn down, to prevent callbacks from trying to access removed worldcontexts
+	for (int32 WorldIdx = WorldList.Num() - 1; WorldIdx >= 0; --WorldIdx)
+	{
+		FWorldContext& ThisContext = WorldList[WorldIdx];
+		if (ThisContext.WorldType == EWorldType::PIE)
+		{
 			WorldList.RemoveAt(WorldIdx);
 		}
 	}
