[Viewport Toolbar] update Raytracing visualization submenu so that it uses UToolMenus

diff --git a/Engine/Source/Editor/UnrealEd/Private/RayTracingDebugVisualizationMenuCommands.cpp b/Engine/Source/Editor/UnrealEd/Private/RayTracingDebugVisualizationMenuCommands.cpp
--- a/Engine/Source/Editor/UnrealEd/Private/RayTracingDebugVisualizationMenuCommands.cpp
+++ b/Engine/Source/Editor/UnrealEd/Private/RayTracingDebugVisualizationMenuCommands.cpp
@@ -9,6 +9,7 @@
 #include "RayTracingVisualizationData.h"
 #include "Styling/AppStyle.h"
 #include "EditorViewportClient.h"
+#include "ToolMenu.h"
 
 int32 GRayTracingVisualizeOther = 0;
 static FAutoConsoleVariableRef CVarRayTracingVisualizeOther(
@@ -82,6 +83,7 @@ void FRayTracingDebugVisualizationMenuCommands::BuildCommandMap()
 	}
 }
 
+// Deprecated. When removed, also remove AddCommandTypeToMenu(FMenuBuilder& Menu, const FVisualizationType Type)
 void FRayTracingDebugVisualizationMenuCommands::BuildVisualisationSubMenu(FMenuBuilder& Menu)
 {
 	const bool bShowOther = GRayTracingVisualizeOther != 0;
@@ -123,6 +125,42 @@ void FRayTracingDebugVisualizationMenuCommands::BuildVisualisationSubMenu(FMenuB
 	}
 }
 
+void FRayTracingDebugVisualizationMenuCommands::BuildVisualisationSubMenu(UToolMenu* InMenu)
+{
+	const FRayTracingDebugVisualizationMenuCommands& Commands = FRayTracingDebugVisualizationMenuCommands::Get();
+	if (!Commands.IsPopulated())
+	{
+		return;
+	}
+
+	{
+		FToolMenuSection& UnnamedSection = InMenu->AddSection(NAME_None);
+		Commands.AddCommandTypeToSection(UnnamedSection, FVisualizationType::Overview);
+	}
+
+	{
+		FToolMenuSection& GeneralSection = InMenu->AddSection(
+			"LevelViewportRayTracingVisualizationGeneral", LOCTEXT("RayTracingVisualizationGeneral", "General")
+		);
+		Commands.AddCommandTypeToSection(GeneralSection, FVisualizationType::Standard);
+	}
+
+	{
+		FToolMenuSection& PerformanceSection = InMenu->AddSection(
+			"LevelViewportRayTracingVisualizationPerformance",
+			LOCTEXT("RayTracingVisualizationPerformance", "Performance")
+		);
+		Commands.AddCommandTypeToSection(PerformanceSection, FVisualizationType::Performance);
+	}
+
+	const bool bShowOther = GRayTracingVisualizeOther != 0;
+	if (bShowOther)
+	{
+		FToolMenuSection& OtherSection = InMenu->AddSection("LevelViewportRayTracingVisualizationOther", LOCTEXT("RayTracingVisualizationOther", "Other"));
+		Commands.AddCommandTypeToSection(OtherSection, FVisualizationType::Other);
+	}
+}
+
 bool FRayTracingDebugVisualizationMenuCommands::AddCommandTypeToMenu(FMenuBuilder& Menu, const FVisualizationType Type) const
 {
 	bool bAddedCommands = false;
@@ -141,6 +179,23 @@ bool FRayTracingDebugVisualizationMenuCommands::AddCommandTypeToMenu(FMenuBuilde
 	return bAddedCommands;
 }
 
+bool FRayTracingDebugVisualizationMenuCommands::AddCommandTypeToSection(FToolMenuSection& InSection, const FVisualizationType Type) const
+{
+	bool bAddedCommands = false;
+
+	for (TCommandConstIterator It = CreateCommandConstIterator(); It; ++It)
+	{
+		const FVisualizationRecord& Record = It.Value();
+		if (Record.Type == Type)
+		{
+			InSection.AddMenuEntry(Record.Command, Record.Command->GetLabel());
+			bAddedCommands = true;
+		}
+	}
+
+	return bAddedCommands;
+}
+
 FRayTracingDebugVisualizationMenuCommands::TCommandConstIterator FRayTracingDebugVisualizationMenuCommands::CreateCommandConstIterator() const
 {
 	return CommandMap.CreateConstIterator();
diff --git a/Engine/Source/Editor/UnrealEd/Private/ViewportToolbar/UnrealEdViewportToolbar.cpp b/Engine/Source/Editor/UnrealEd/Private/ViewportToolbar/UnrealEdViewportToolbar.cpp
--- a/Engine/Source/Editor/UnrealEd/Private/ViewportToolbar/UnrealEdViewportToolbar.cpp
+++ b/Engine/Source/Editor/UnrealEd/Private/ViewportToolbar/UnrealEdViewportToolbar.cpp
@@ -1593,11 +1593,11 @@ void PopulateViewModesMenu(UToolMenu* InMenu)
 		{
 			struct Local
 			{
-				static void BuildRayTracingDebugMenu(FMenuBuilder& Menu) //, TWeakPtr<SViewportToolBar> InParentToolBar)
+				static void BuildRayTracingDebugMenu(UToolMenu* InMenu)
 				{
 					const FRayTracingDebugVisualizationMenuCommands& RtDebugCommands =
 						FRayTracingDebugVisualizationMenuCommands::Get();
-					RtDebugCommands.BuildVisualisationSubMenu(Menu);
+					RtDebugCommands.BuildVisualisationSubMenu(InMenu);
 				}
 			};
 
@@ -1605,7 +1605,7 @@ void PopulateViewModesMenu(UToolMenu* InMenu)
 				"RayTracingDebugSubMenu",
 				LOCTEXT("RayTracingDebugSubMenu", "Ray Tracing Debug"),
 				LOCTEXT("RayTracing_ToolTip", "Select ray tracing buffer visualization view modes"),
-				FNewMenuDelegate::CreateStatic(&Local::BuildRayTracingDebugMenu), //, ParentToolBar)
+				FNewToolMenuDelegate::CreateStatic(&Local::BuildRayTracingDebugMenu),
 				false,
 				FSlateIcon(FAppStyle::GetAppStyleSetName(), "EditorViewport.RayTracingDebugMode")
 			);
diff --git a/Engine/Source/Editor/UnrealEd/Public/RayTracingDebugVisualizationMenuCommands.h b/Engine/Source/Editor/UnrealEd/Public/RayTracingDebugVisualizationMenuCommands.h
--- a/Engine/Source/Editor/UnrealEd/Public/RayTracingDebugVisualizationMenuCommands.h
+++ b/Engine/Source/Editor/UnrealEd/Public/RayTracingDebugVisualizationMenuCommands.h
@@ -10,6 +10,8 @@
 #include "Framework/Commands/Commands.h"
 
 class FEditorViewportClient;
+class UToolMenu;
+struct FToolMenuSection;
 
 class UNREALED_API FRayTracingDebugVisualizationMenuCommands : public TCommands<FRayTracingDebugVisualizationMenuCommands>
 {
@@ -44,8 +46,11 @@ public:
 
 	TCommandConstIterator CreateCommandConstIterator() const;
 
+	UE_DEPRECATED(5.6, "Please use the version taking UToolMenu* as argument")
 	static void BuildVisualisationSubMenu(FMenuBuilder& Menu);
 
+	static void BuildVisualisationSubMenu(UToolMenu* InMenu);
+
 	virtual void RegisterCommands() override;
 
 	void BindCommands(FUICommandList& CommandList, const TSharedPtr<FEditorViewportClient>& Client) const;
@@ -57,8 +62,12 @@ public:
 
 private:
 	void BuildCommandMap();
+
+	/** Remove this function when removing deprecated BuildVisualisationSubMenu(FMenuBuilder& Menu) */
 	bool AddCommandTypeToMenu(FMenuBuilder& Menu, const FVisualizationType Type) const;
 
+	bool AddCommandTypeToSection(FToolMenuSection& InSection, const FVisualizationType Type) const;
+
 	static void ChangeRayTracingDebugVisualizationMode(TWeakPtr<FEditorViewportClient> WeakClient, FName InName);
 	static bool IsRayTracingDebugVisualizationModeSelected(TWeakPtr<FEditorViewportClient> WeakClient, FName InName);
 
