[Viewport Toolbar] Snap menus converted from FMenuiBuilder to UToolMenu

diff --git a/Engine/Source/Editor/UnrealEd/Private/ViewportToolbar/UnrealEdViewportToolbar.cpp b/Engine/Source/Editor/UnrealEd/Private/ViewportToolbar/UnrealEdViewportToolbar.cpp
--- a/Engine/Source/Editor/UnrealEd/Private/ViewportToolbar/UnrealEdViewportToolbar.cpp
+++ b/Engine/Source/Editor/UnrealEd/Private/ViewportToolbar/UnrealEdViewportToolbar.cpp
@@ -1768,46 +1768,55 @@ TSharedRef<SWidget> BuildRotationGridCheckBoxList(
 	const TSharedPtr<FUICommandList>& InCommandList
 )
 {
-	constexpr bool bShouldCloseWindowAfterMenuSelection = true;
-	FMenuBuilder RotationGridMenuBuilder(bShouldCloseWindowAfterMenuSelection, InCommandList);
-
-	RotationGridMenuBuilder.BeginSection(InExtentionHook, InHeading);
-	for (int32 CurrGridAngleIndex = 0; CurrGridAngleIndex < InGridSizes.Num(); ++CurrGridAngleIndex)
-	{
-		const float CurrGridAngle = InGridSizes[CurrGridAngleIndex];
-
-		FText MenuName =
-			FText::Format(LOCTEXT("RotationGridAngle", "{0}\u00b0"), FText::AsNumber(CurrGridAngle)); /*degree symbol*/
-		FText ToolTipText = FText::Format(
-			LOCTEXT("RotationGridAngle_ToolTip", "Sets rotation grid angle to {0}"), MenuName
-		); /*degree symbol*/
-
-		RotationGridMenuBuilder.AddMenuEntry(
-			MenuName,
-			ToolTipText,
-			FSlateIcon(),
-			FUIAction(
-				FExecuteAction::CreateLambda(
-					[CurrGridAngleIndex, InGridMode, InExecuteAction]()
-					{
-						InExecuteAction.Execute(CurrGridAngleIndex, InGridMode);
-					}
-				),
-				FCanExecuteAction(),
-				FIsActionChecked::CreateLambda(
-					[CurrGridAngleIndex, InGridMode, InIsActionChecked]()
-					{
-						return InIsActionChecked.Execute(CurrGridAngleIndex, InGridMode);
-					}
-				)
-			),
-			NAME_None,
-			EUserInterfaceActionType::RadioButton
-		);
+	const FName CheckboxListMenu("RotationGridCheckboxList." + InExtentionHook.ToString());
+	if (!UToolMenus::Get()->IsMenuRegistered(CheckboxListMenu))
+	{
+		if (UToolMenu* Menu = UToolMenus::Get()->RegisterMenu(CheckboxListMenu, NAME_None))
+		{
+			FToolMenuSection& Section = Menu->AddSection(InExtentionHook, InHeading);
+			for (int32 CurrGridAngleIndex = 0; CurrGridAngleIndex < InGridSizes.Num(); ++CurrGridAngleIndex)
+			{
+				const float CurrGridAngle = InGridSizes[CurrGridAngleIndex];
+
+				FText MenuName = FText::Format(
+					LOCTEXT("RotationGridAngle", "{0}\u00b0"), FText::AsNumber(CurrGridAngle)
+				); /*degree symbol*/
+				FText ToolTipText = FText::Format(
+					LOCTEXT("RotationGridAngle_ToolTip", "Sets rotation grid angle to {0}"), MenuName
+				); /*degree symbol*/
+
+				Section.AddMenuEntry(
+					NAME_None,
+					MenuName,
+					ToolTipText,
+					FSlateIcon(),
+					FUIAction(
+						FExecuteAction::CreateLambda(
+							[CurrGridAngleIndex, InGridMode, InExecuteAction]()
+							{
+								InExecuteAction.Execute(CurrGridAngleIndex, InGridMode);
+							}
+						),
+						FCanExecuteAction(),
+						FIsActionChecked::CreateLambda(
+							[CurrGridAngleIndex, InGridMode, InIsActionChecked]()
+							{
+								return InIsActionChecked.Execute(CurrGridAngleIndex, InGridMode);
+							}
+						)
+					),
+					EUserInterfaceActionType::RadioButton
+				);
+			}
+		}
+	}
+
+	FToolMenuContext MenuContext;
+	{
+		MenuContext.AppendCommandList(InCommandList);
 	}
-	RotationGridMenuBuilder.EndSection();
 
-	return RotationGridMenuBuilder.MakeWidget();
+	return UToolMenus::Get()->GenerateWidget(CheckboxListMenu, MenuContext);
 }
 
 FText GetRotationGridLabel()
@@ -1863,45 +1872,53 @@ TSharedRef<SWidget> CreateLocationGridSnapMenu(
 	const TSharedPtr<FUICommandList>& InCommandList
 )
 {
-	constexpr bool bShouldCloseWindowAfterMenuSelection = true;
-	FMenuBuilder LocationGridMenuBuilder(bShouldCloseWindowAfterMenuSelection, InCommandList);
-
-	LocationGridMenuBuilder.BeginSection("Snap", LOCTEXT("LocationSnapText", "Snap Sizes"));
-	for (int32 CurrGridSizeIndex = 0; CurrGridSizeIndex < InGridSizes.Num(); ++CurrGridSizeIndex)
+	const FName CheckboxListMenu("LocationGridCheckboxList");
+	if (!UToolMenus::Get()->IsMenuRegistered(CheckboxListMenu))
 	{
-		const float CurGridSize = InGridSizes[CurrGridSizeIndex];
+		if (UToolMenu* Menu = UToolMenus::Get()->RegisterMenu(CheckboxListMenu, NAME_None))
+		{
+			FToolMenuSection& Section = Menu->AddSection("Snap", LOCTEXT("LocationSnapText", "Snap Sizes"));
+			for (int32 CurrGridSizeIndex = 0; CurrGridSizeIndex < InGridSizes.Num(); ++CurrGridSizeIndex)
+			{
+				const float CurGridSize = InGridSizes[CurrGridSizeIndex];
 
-		LocationGridMenuBuilder.AddMenuEntry(
-			FText::AsNumber(CurGridSize),
-			FText::Format(LOCTEXT("LocationGridSize_ToolTip", "Sets grid size to {0}"), FText::AsNumber(CurGridSize)),
-			FSlateIcon(),
-			FUIAction(
-				FExecuteAction::CreateLambda(
-					[CurrGridSizeIndex, InExecuteDelegate]()
-					{
-						InExecuteDelegate.Execute(CurrGridSizeIndex);
-					}
-				),
-				FCanExecuteAction::CreateLambda(
-					[InIsEnabledDelegate]()
-					{
-						return InIsEnabledDelegate.Get();
-					}
-				),
-				FIsActionChecked::CreateLambda(
-					[CurrGridSizeIndex, InIsCheckedDelegate]()
-					{
-						return InIsCheckedDelegate.Execute(CurrGridSizeIndex);
-					}
-				)
-			),
-			NAME_None,
-			EUserInterfaceActionType::RadioButton
-		);
+				Section.AddMenuEntry(
+					NAME_None,
+					FText::AsNumber(CurGridSize),
+					FText::Format(LOCTEXT("LocationGridSize_ToolTip", "Sets grid size to {0}"), FText::AsNumber(CurGridSize)),
+					FSlateIcon(),
+					FUIAction(
+						FExecuteAction::CreateLambda(
+							[CurrGridSizeIndex, InExecuteDelegate]()
+							{
+								InExecuteDelegate.Execute(CurrGridSizeIndex);
+							}
+						),
+						FCanExecuteAction::CreateLambda(
+							[InIsEnabledDelegate]()
+							{
+								return InIsEnabledDelegate.Get();
+							}
+						),
+						FIsActionChecked::CreateLambda(
+							[CurrGridSizeIndex, InIsCheckedDelegate]()
+							{
+								return InIsCheckedDelegate.Execute(CurrGridSizeIndex);
+							}
+						)
+					),
+					EUserInterfaceActionType::RadioButton
+				);
+			}
+		}
 	}
-	LocationGridMenuBuilder.EndSection();
 
-	return LocationGridMenuBuilder.MakeWidget();
+	FToolMenuContext MenuContext;
+	{
+		MenuContext.AppendCommandList(InCommandList);
+	}
+
+	return UToolMenus::Get()->GenerateWidget(CheckboxListMenu, MenuContext);
 }
 
 FText GetScaleGridLabel()
@@ -1924,82 +1941,98 @@ TSharedRef<SWidget> CreateScaleGridSnapMenu(
 	const FUIAction& PreserveNonUniformScaleUIAction
 )
 {
-	FNumberFormattingOptions NumberFormattingOptions;
-	NumberFormattingOptions.MaximumFractionalDigits = 5;
+	const FName CheckboxListMenu("ScaleGridCheckboxList");
+	if (!UToolMenus::Get()->IsMenuRegistered(CheckboxListMenu))
+	{
+		if (UToolMenu* Menu = UToolMenus::Get()->RegisterMenu(CheckboxListMenu, NAME_None))
+		{
 
-	constexpr bool bShouldCloseWindowAfterMenuSelection = true;
-	FMenuBuilder ScaleGridMenuBuilder(bShouldCloseWindowAfterMenuSelection, InCommandList);
+			FNumberFormattingOptions NumberFormattingOptions;
+			NumberFormattingOptions.MaximumFractionalDigits = 5;
 
-	ScaleGridMenuBuilder.BeginSection("ScaleSnapOptions", LOCTEXT("ScaleSnapOptions", "Scale Snap"));
+			FToolMenuSection& Section = Menu->AddSection("ScaleSnapOptions");
+			for (int32 CurrGridAmountIndex = 0; CurrGridAmountIndex < InGridSizes.Num(); ++CurrGridAmountIndex)
+			{
+				const float CurGridAmount = InGridSizes[CurrGridAmountIndex];
 
-	for (int32 CurrGridAmountIndex = 0; CurrGridAmountIndex < InGridSizes.Num(); ++CurrGridAmountIndex)
-	{
-		const float CurGridAmount = InGridSizes[CurrGridAmountIndex];
+				FText MenuText;
+				FText ToolTipText;
 
-		FText MenuText;
-		FText ToolTipText;
+				if (GEditor->UsePercentageBasedScaling())
+				{
+					MenuText = FText::AsPercent(CurGridAmount / 100.0f, &NumberFormattingOptions);
+					ToolTipText =
+						FText::Format(LOCTEXT("ScaleGridAmountOld_ToolTip", "Snaps scale values to {0}"), MenuText);
+				}
+				else
+				{
+					MenuText = FText::AsNumber(CurGridAmount, &NumberFormattingOptions);
+					ToolTipText = FText::Format(
+						LOCTEXT("ScaleGridAmount_ToolTip", "Snaps scale values to increments of {0}"), MenuText
+					);
+				}
 
-		if (GEditor->UsePercentageBasedScaling())
-		{
-			MenuText = FText::AsPercent(CurGridAmount / 100.0f, &NumberFormattingOptions);
-			ToolTipText = FText::Format(LOCTEXT("ScaleGridAmountOld_ToolTip", "Snaps scale values to {0}"), MenuText);
-		}
-		else
-		{
-			MenuText = FText::AsNumber(CurGridAmount, &NumberFormattingOptions);
-			ToolTipText =
-				FText::Format(LOCTEXT("ScaleGridAmount_ToolTip", "Snaps scale values to increments of {0}"), MenuText);
-		}
+				Section.AddMenuEntry(
+					NAME_None,
+					MenuText,
+					ToolTipText,
+					FSlateIcon(),
+					FUIAction(
+						FExecuteAction::CreateLambda(
+							[CurrGridAmountIndex, InExecuteDelegate]()
+							{
+								InExecuteDelegate.Execute(CurrGridAmountIndex);
+							}
+						),
+						FCanExecuteAction::CreateLambda(
+							[InIsEnabledDelegate]()
+							{
+								return InIsEnabledDelegate.Get();
+							}
+						),
+						FIsActionChecked::CreateLambda(
+							[CurrGridAmountIndex, InIsCheckedDelegate]()
+							{
+								return InIsCheckedDelegate.Execute(CurrGridAmountIndex);
+							}
+						)
+					),
+					EUserInterfaceActionType::RadioButton
+				);
+			}
 
-		ScaleGridMenuBuilder.AddMenuEntry(
-			MenuText,
-			ToolTipText,
-			FSlateIcon(),
-			FUIAction(
-				FExecuteAction::CreateLambda(
-					[CurrGridAmountIndex, InExecuteDelegate]()
+			Menu->AddDynamicSection(
+				NAME_None,
+				FNewToolMenuDelegate::CreateLambda(
+					[ShowPreserveNonUniformScaleOption, PreserveNonUniformScaleUIAction](UToolMenu* InMenu)
 					{
-						InExecuteDelegate.Execute(CurrGridAmountIndex);
-					}
-				),
-				FCanExecuteAction::CreateLambda(
-					[InIsEnabledDelegate]()
-					{
-						return InIsEnabledDelegate.Get();
-					}
-				),
-				FIsActionChecked::CreateLambda(
-					[CurrGridAmountIndex, InIsCheckedDelegate]()
-					{
-						return InIsCheckedDelegate.Execute(CurrGridAmountIndex);
+						if (!GEditor->UsePercentageBasedScaling() && ShowPreserveNonUniformScaleOption.Get())
+						{
+							FToolMenuSection& GeneralOptionsSection =
+								InMenu->AddSection("ScaleGeneralOptions", LOCTEXT("ScaleOptions", "Scaling Options"));
+							GeneralOptionsSection.AddMenuEntry(
+								NAME_None,
+								LOCTEXT("ScaleGridPreserveNonUniformScale", "Preserve Non-Uniform Scale"),
+								LOCTEXT(
+									"ScaleGridPreserveNonUniformScale_ToolTip", "When this option is checked, scaling objects that have a non-uniform scale will preserve the ratios between each axis, snapping the axis with the largest value."
+								),
+								FSlateIcon(),
+								PreserveNonUniformScaleUIAction,
+								EUserInterfaceActionType::ToggleButton
+							);
+						}
 					}
 				)
-			),
-			NAME_None,
-			EUserInterfaceActionType::RadioButton
-		);
+			);
+		}
 	}
-	ScaleGridMenuBuilder.EndSection();
 
-	if (!GEditor->UsePercentageBasedScaling() && ShowPreserveNonUniformScaleOption.Get())
+	FToolMenuContext MenuContext;
 	{
-		ScaleGridMenuBuilder.BeginSection("ScaleGeneralOptions", LOCTEXT("ScaleOptions", "Scaling Options"));
-
-		ScaleGridMenuBuilder.AddMenuEntry(
-			LOCTEXT("ScaleGridPreserveNonUniformScale", "Preserve Non-Uniform Scale"),
-			LOCTEXT(
-				"ScaleGridPreserveNonUniformScale_ToolTip", "When this option is checked, scaling objects that have a non-uniform scale will preserve the ratios between each axis, snapping the axis with the largest value."
-			),
-			FSlateIcon(),
-			PreserveNonUniformScaleUIAction,
-			NAME_None,
-			EUserInterfaceActionType::ToggleButton
-		);
-
-		ScaleGridMenuBuilder.EndSection();
+		MenuContext.AppendCommandList(InCommandList);
 	}
 
-	return ScaleGridMenuBuilder.MakeWidget();
+	return UToolMenus::Get()->GenerateWidget(CheckboxListMenu, MenuContext);
 }
 
 FToolMenuEntry CreateCheckboxSubmenu(
