[ChaosVD] Fix ensure in GLevelEditorModeTools during start up of the Chaos Visual Debugger standalone program.

SEditorViewport and EditorViewportClient where still using the global accessor for the editor mode tools instead of of the mode tool specified in the viewport client.

diff --git a/Engine/Source/Editor/UnrealEd/Private/EditorViewportClient.cpp b/Engine/Source/Editor/UnrealEd/Private/EditorViewportClient.cpp
--- a/Engine/Source/Editor/UnrealEd/Private/EditorViewportClient.cpp
+++ b/Engine/Source/Editor/UnrealEd/Private/EditorViewportClient.cpp
@@ -4376,7 +4376,8 @@ void FEditorViewportClient::Draw(FViewport* InViewport, FCanvas* Canvas)
 
 	// Axes indicators
 	const bool bShouldDrawAxes = (bDrawAxes && !ViewFamily.EngineShowFlags.Game) || (bDrawAxesGame && ViewFamily.EngineShowFlags.Game);
-	if (bShouldDrawAxes && !GLevelEditorModeTools().IsViewportUIHidden() && !IsVisualizeCalibrationMaterialEnabled())
+	const bool bIsViewportUIHidden = ModeTools && ModeTools->IsViewportUIHidden();
+	if (bShouldDrawAxes && !bIsViewportUIHidden && !IsVisualizeCalibrationMaterialEnabled())
 	{
 		switch (GetViewportType())
 		{
diff --git a/Engine/Source/Editor/UnrealEd/Private/SEditorViewport.cpp b/Engine/Source/Editor/UnrealEd/Private/SEditorViewport.cpp
--- a/Engine/Source/Editor/UnrealEd/Private/SEditorViewport.cpp
+++ b/Engine/Source/Editor/UnrealEd/Private/SEditorViewport.cpp
@@ -583,7 +583,10 @@ void SEditorViewport::BindShowCommands( FUICommandList& OutCommandList )
 
 EVisibility SEditorViewport::OnGetViewportContentVisibility() const
 {
-	return GLevelEditorModeTools().IsViewportUIHidden() ? EVisibility::Collapsed : EVisibility::SelfHitTestInvisible;
+	FEditorModeTools* EditorModeTools = Client->GetModeTools();
+	const bool bIsViewportUIHidden = EditorModeTools && EditorModeTools->IsViewportUIHidden();
+
+	return bIsViewportUIHidden ? EVisibility::Collapsed : EVisibility::SelfHitTestInvisible;
 }
 
 void SEditorViewport::OnToggleRealtime()
