SourceControl: Fixed attempted checkout of file for which `!CanCheckout()` returns `false`.

If a file is in the ADDED state, it should not be attempted to be checked out (as the server doesn't know about it). The FSkeinSourceControlState::CanCheckout() function checks this, but this wasn't used. This was generating warnings when acquiring a lock would find that the path does not exist in the project

#rnx

diff --git a/Engine/Source/Editor/UnrealEd/Private/FileHelpers.cpp b/Engine/Source/Editor/UnrealEd/Private/FileHelpers.cpp
--- a/Engine/Source/Editor/UnrealEd/Private/FileHelpers.cpp
+++ b/Engine/Source/Editor/UnrealEd/Private/FileHelpers.cpp
@@ -4801,7 +4801,7 @@ bool FEditorFileUtils::AutomaticCheckoutOrPromptToRevertPackages(const TArray<UP
 			// No need to check it out.
 			PackagesCheckedOutAlready.Add(PackageMap[PackageFilename]);
 		}
-		else if (State->IsCheckedOutOther() || !State->IsCurrent())
+		else if (!State->CanCheckout())
 		{
 			// Can't check it out.
 			PackagesCheckOutImpossible.Add(PackageMap[PackageFilename]);
