World Partition - Fixed duplicated HLOD layers not being saved anymore when saving a map template

diff --git a/Engine/Source/Editor/UnrealEd/Private/FileHelpers.cpp b/Engine/Source/Editor/UnrealEd/Private/FileHelpers.cpp
--- a/Engine/Source/Editor/UnrealEd/Private/FileHelpers.cpp
+++ b/Engine/Source/Editor/UnrealEd/Private/FileHelpers.cpp
@@ -608,7 +608,7 @@ static bool SaveExternalPackages(UPackage* InPackage, TArray<UPackage*>& InOutPa
 
 	if (!bInAutosaving && (!FEditorFileUtils::ShouldSkipExternalObjectSave() || bInNewlyCreated))
 	{
-		InOutPackagesToSave = InPackage->GetExternalPackages();
+		InOutPackagesToSave.Append(InPackage->GetExternalPackages());
 
 		if (InOutPackagesToSave.Num())
 		{
