MPCook: Improve error feedback for busy sockets in the director<-> worker connection: handle SE_EWOULDBLOCK and SE_TRY_AGAIN as non-errors.
#rnx

diff --git a/Engine/Source/Editor/UnrealEd/Private/Cooker/CompactBinaryTCP.cpp b/Engine/Source/Editor/UnrealEd/Private/Cooker/CompactBinaryTCP.cpp
--- a/Engine/Source/Editor/UnrealEd/Private/Cooker/CompactBinaryTCP.cpp
+++ b/Engine/Source/Editor/UnrealEd/Private/Cooker/CompactBinaryTCP.cpp
@@ -226,6 +226,7 @@ EConnectionStatus FCompactBinaryTCPImpl::PollSendBytes(FSocket* Socket, const vo
 	bool bResult = Socket->Send(WriteData, WriteSize, BytesSent);
 	if (BytesSent < WriteSize)
 	{
+		ESocketErrors Error = ISocketSubsystem::Get(PLATFORM_SOCKETSUBSYSTEM)->GetLastErrorCode();
 		if (BytesSent >= 0)
 		{
 			InOutBytesWritten += BytesSent;
@@ -233,6 +234,17 @@ EConnectionStatus FCompactBinaryTCPImpl::PollSendBytes(FSocket* Socket, const vo
 
 		if (!bResult)
 		{
+			if (Socket->GetConnectionState() != SCS_ConnectionError &&
+				(Error == SE_EWOULDBLOCK || Error == SE_TRY_AGAIN))
+			{
+				// Not an error, this can occur if the process on the other end of the socket is not
+				// reading fast enough (breakpoint or a busy tick) and the socket's buffers have filled up.
+				return EConnectionStatus::Incomplete;
+			}
+
+			UE_LOG(LogSockets, Log, TEXT("FCompactBinaryTCPImpl::PollSendBytes: Socket->Send(WriteSize==%d) failed with error %d and BytesSent == %d."),
+				WriteSize, (int32)Error, BytesSent);
+
 			// Return values of false indicate the socket send failed with an error, but the error
 			// might be temporary, such as EWOULDBLOCK. Even when return value is false, return incomplete
 			// unless socket's ConnectionState has been set to error
@@ -244,8 +256,6 @@ EConnectionStatus FCompactBinaryTCPImpl::PollSendBytes(FSocket* Socket, const vo
 
 			// TODO: Decide which errors are fatal even if the socket is still connected. We have seen one case
 			// where the socket send fails forever, but we were not able to see what error it was returning.
-			ESocketErrors Error = ISocketSubsystem::Get(PLATFORM_SOCKETSUBSYSTEM)->GetLastErrorCode();
-			UE_LOG(LogSockets, Log, TEXT("FCompactBinaryTCPImpl::PollSendBytes: Socket->Send failed with error %d."), (int32)Error);
 		}
 
 		return EConnectionStatus::Incomplete;
