fixing pointer referencing the same memory iterating through the actions when outputting blueprint API in JSON format for the Blueprint API reference documentation.

diff --git a/Engine/Source/Editor/UnrealEd/Private/Commandlets/GenerateBlueprintAPICommandlet.cpp b/Engine/Source/Editor/UnrealEd/Private/Commandlets/GenerateBlueprintAPICommandlet.cpp
--- a/Engine/Source/Editor/UnrealEd/Private/Commandlets/GenerateBlueprintAPICommandlet.cpp
+++ b/Engine/Source/Editor/UnrealEd/Private/Commandlets/GenerateBlueprintAPICommandlet.cpp
@@ -33,6 +33,7 @@
 #include "Engine/LevelScriptBlueprint.h"
 #include "Kismet2/BlueprintEditorUtils.h"
 #include "AnimStateConduitNode.h"
+#include "AnimationStateMachineGraph.h"
 #include "AnimStateNode.h"
 #include "K2Node_Event.h"
 #include "K2Node_CustomEvent.h"
@@ -822,7 +823,7 @@ static void GenerateBlueprintAPIUtils::DumpActionList(uint32 Indent, FGraphActio
 	SortedActions.Sort(ActionSortFunctor());
 
 	FString LineEnding("\n");
-	for (TSharedPtr<FEdGraphSchemaAction> const& Action : SortedActions)
+	for (const TSharedPtr<FEdGraphSchemaAction>& Action : SortedActions)
 	{
 		FileOutWriter->Serialize(TCHAR_TO_ANSI(*LineEnding), LineEnding.Len());
 		DumpActionMenuItem(Indent + 2, Action, ActionList, FileOutWriter);
