IncrementalCook: MPCook: Fix spurious warning "Package was added to the runtime discovery queue after starting BuildDependencies phase." at end of ECookPhase::Cook, when we transition to committing BuildDependencies. Some discovery messages from CookWorkers might still be buffered at the point when we find all packages are done and decide to phase transition, and we need to wait for all of those messages to be processed.
#rnx

diff --git a/Engine/Source/Editor/UnrealEd/Classes/CookOnTheSide/CookOnTheFlyServer.h b/Engine/Source/Editor/UnrealEd/Classes/CookOnTheSide/CookOnTheFlyServer.h
--- a/Engine/Source/Editor/UnrealEd/Classes/CookOnTheSide/CookOnTheFlyServer.h
+++ b/Engine/Source/Editor/UnrealEd/Classes/CookOnTheSide/CookOnTheFlyServer.h
@@ -550,6 +550,13 @@ private:
 	/** Find all packages that are build dependencies and request them for cook. */
 	void KickBuildDependencies(UE::Cook::FTickStackData& StackData);
 
+	/**
+	 * Send a fence to CookWorkers if not already sent since the last phase transition or non-idle action.
+	 * Return bOutCompleted=true only when the fence has been passed.
+	 * This is a Noop in MPCook; no fence is used and the function sets bOutComplete=true whenever called.
+	 */
+	void PumpPhaseTransitionFence(bool& bOutComplete);
+
 	/**
 	 * Inspect the given package and queue it for saving if necessary.
 	 *
@@ -1647,9 +1654,11 @@ private:
 	double PollNextTimeSeconds = 0.;
 	double PollNextTimeIdleSeconds = 0.;
 	double IdleStatusStartTime = 0.0;
+	double IdleStatusLastReportTime = 0.0;
 	uint32 CookedPackageCountSinceLastGC = 0;
 	float WaitForAsyncSleepSeconds = 0.0f;
 	float DisplayUpdatePeriodSeconds = 0.0f;
+	int32 PhaseTransitionFence = -1;
 	EIdleStatus IdleStatus = EIdleStatus::Done;
 	TUniquePtr<UE::Cook::FODSCClientData> ODSCClientData;
 
diff --git a/Engine/Source/Editor/UnrealEd/Private/CookOnTheFlyServer.cpp b/Engine/Source/Editor/UnrealEd/Private/CookOnTheFlyServer.cpp
--- a/Engine/Source/Editor/UnrealEd/Private/CookOnTheFlyServer.cpp
+++ b/Engine/Source/Editor/UnrealEd/Private/CookOnTheFlyServer.cpp
@@ -1813,12 +1813,18 @@ void UCookOnTheFlyServer::SetLoadBusy(bool bInLoadBusy)
 
 void UCookOnTheFlyServer::SetIdleStatus(UE::Cook::FTickStackData& StackData, EIdleStatus InStatus)
 {
+	if (InStatus == EIdleStatus::Active)
+	{
+		PhaseTransitionFence = -1;
+	}
+
 	if (InStatus == IdleStatus)
 	{
 		return;
 	}
 
 	IdleStatusStartTime = StackData.LoopStartTime;
+	IdleStatusLastReportTime = IdleStatusStartTime;
 	IdleStatus = InStatus;
 }
 
@@ -2484,6 +2490,14 @@ UCookOnTheFlyServer::ECookAction UCookOnTheFlyServer::DecideNextCookAction(UE::C
 
 	if (IsCookByTheBookMode() && GetCookPhase() != ECookPhase::BuildDependencies)
 	{
+		bool bCompleted;
+		PumpPhaseTransitionFence(bCompleted);
+		if (!bCompleted)
+		{
+			SetIdleStatus(StackData, EIdleStatus::Idle);
+			return ECookAction::WaitForAsync;
+		}
+
 		return ECookAction::KickBuildDependencies;
 	}
 
@@ -4924,6 +4938,50 @@ void UCookOnTheFlyServer::PumpBuildDependencySaves(UE::Cook::FTickStackData& Sta
 	}
 }
 
+void UCookOnTheFlyServer::PumpPhaseTransitionFence(bool& bOutComplete)
+{
+	using namespace UE::Cook;
+
+	if (!CookDirector)
+	{
+		bOutComplete = true;
+		return;
+	}
+
+	if (PhaseTransitionFence == -1)
+	{
+		PhaseTransitionFence = CookDirector->InsertBroadcastFence();
+	}
+
+	TArray<FWorkerId> PendingWorkers;
+	if (CookDirector->IsBroadcastFencePassed(PhaseTransitionFence, &PendingWorkers))
+	{
+		bOutComplete = true;
+		PhaseTransitionFence = -1;
+		return;
+	}
+
+	double CurrentTime = FPlatformTime::Seconds();
+	constexpr float ReportPeriod = 60.f;
+	if (IdleStatusLastReportTime + ReportPeriod < CurrentTime)
+	{
+		TStringBuilder<256> WorkerStr;
+		for (FWorkerId WorkerId: PendingWorkers)
+		{
+			WorkerStr << CookDirector->GetDisplayName(WorkerId) << TEXT(", ");
+		}
+		check(WorkerStr.Len() >= 0);
+		WorkerStr.RemoveSuffix(2); // Remove ", "
+
+		UE_LOG(LogCook, Display,
+			TEXT("Waiting on MPCook fence, but { %s } has not responded to a heartbeat request for %.1f seconds. Continuing to wait..."),
+			*WorkerStr, CurrentTime - IdleStatusStartTime);
+		IdleStatusLastReportTime = CurrentTime;
+	}
+
+	bOutComplete = false;
+}
+
 void UCookOnTheFlyServer::KickBuildDependencies(UE::Cook::FTickStackData& StackData)
 {
 	using namespace UE::Cook;
@@ -5637,6 +5695,7 @@ void UCookOnTheFlyServer::Initialize( ECookMode::Type DesiredCookMode, ECookInit
 
 	IdleStatus = EIdleStatus::Done;
 	IdleStatusStartTime = FPlatformTime::Seconds();
+	IdleStatusLastReportTime = IdleStatusStartTime;
 
 	if (!IsCookOnTheFlyMode() && !IsCookingInEditor() &&
 		FPlatformMisc::SupportsMultithreadedFileHandles() && // Preloading moves file handles between threads
diff --git a/Engine/Source/Editor/UnrealEd/Private/Cooker/CookDirector.cpp b/Engine/Source/Editor/UnrealEd/Private/Cooker/CookDirector.cpp
--- a/Engine/Source/Editor/UnrealEd/Private/Cooker/CookDirector.cpp
+++ b/Engine/Source/Editor/UnrealEd/Private/Cooker/CookDirector.cpp
@@ -684,6 +684,37 @@ void FCookDirector::BroadcastMessage(UE::CompactBinaryTCP::FMarshalledMessage&&
 	}
 }
 
+int32 FCookDirector::InsertBroadcastFence()
+{
+	int32 LocalHeartbeatNumber = HeartbeatNumber;
+
+	// We need to Tick, and have the Tick function send and update the heartbeat number, so that any preceding
+	// messages get sent before the heartbeat.
+	bForceNextHeartbeat = true;
+	TickFromSchedulerThread();
+
+	return LocalHeartbeatNumber;
+}
+
+bool FCookDirector::IsBroadcastFencePassed(int32 Fence, TArray<FWorkerId>* OutPendingWorkers)
+{
+	bool bHasPendingWorker = false;
+	FScopeLock CommunicationScopeLock(&CommunicationLock);
+	for (TPair<int32, TRefCountPtr<FCookWorkerServer>>& Pair : RemoteWorkers)
+	{
+		FCookWorkerServer* RemoteWorker = Pair.Value.GetReference();
+		if (RemoteWorker->GetLastReceivedHeartbeatNumber() < Fence)
+		{
+			bHasPendingWorker = true;
+			if (OutPendingWorkers)
+			{
+				OutPendingWorkers->Add(RemoteWorker->GetWorkerId());
+			}
+		}
+	}
+	return !bHasPendingWorker;
+}
+
 void FCookDirector::TickFromSchedulerThread()
 {
 	double CurrentTime = FPlatformTime::Seconds();
diff --git a/Engine/Source/Editor/UnrealEd/Private/Cooker/CookDirector.h b/Engine/Source/Editor/UnrealEd/Private/Cooker/CookDirector.h
--- a/Engine/Source/Editor/UnrealEd/Private/Cooker/CookDirector.h
+++ b/Engine/Source/Editor/UnrealEd/Private/Cooker/CookDirector.h
@@ -98,6 +98,15 @@ public:
 	void BroadcastMessage(UE::CompactBinaryTCP::FMarshalledMessage&& Message,
 		ECookBroadcastTiming Timing = ECookBroadcastTiming::Immediate);
 
+	/**
+	 * Increment the current heartbeat number, send a heartbeat message to all CookWorkers, and return the number in
+	 * the broadcast. This number can then be queried for whether the CookWorkers have acknowledged it in a round trip
+	 * message.
+	 */
+	int32 InsertBroadcastFence();
+	/** Report whether all CookWorkers have acknowledged a previously inserted fence.. */
+	bool IsBroadcastFencePassed(int32 Fence, TArray<FWorkerId>* OutPendingWorkers);
+
 	/** Periodic tick function. Sends/Receives messages to CookWorkers. */
 	void TickFromSchedulerThread();
 	/** Periodic display function, called from CookOnTheFlyServer.UpdateDisplay. */
@@ -136,6 +145,8 @@ public:
 	};
 	FLaunchInfo GetLaunchInfo(FWorkerId WorkerId, int32 ProfileId);
 
+	FString GetDisplayName(const FWorkerId& WorkerId, int32 PreferredWidth = -1) const;
+
 	/** The message CookWorkerServer sends to the remote process once it is ready to connect. */
 	const FInitialConfigMessage& GetInitialConfigMessage();
 
@@ -241,7 +252,6 @@ private:
 
 	TArray<TRefCountPtr<FCookWorkerServer>> CopyRemoteWorkers() const;
 	void DisplayRemainingPackages() const;
-	FString GetDisplayName(const FWorkerId& WorkerId, int32 PreferredWidth = -1) const;
 	FString GetDisplayName(const FCookWorkerServer& RemoteWorker, int32 PreferredWidth=-1) const;
 	const TRefCountPtr<FCookWorkerServer>* FindRemoteWorkerInLock(const FWorkerId& WorkerId) const;
 
