[shaders] temporarily disable stats aggregation from cook workers to director due to exceeding message size limits

diff --git a/Engine/Source/Editor/UnrealEd/Private/ShaderStatsCollector.cpp b/Engine/Source/Editor/UnrealEd/Private/ShaderStatsCollector.cpp
--- a/Engine/Source/Editor/UnrealEd/Private/ShaderStatsCollector.cpp
+++ b/Engine/Source/Editor/UnrealEd/Private/ShaderStatsCollector.cpp
@@ -166,7 +166,8 @@ void FShaderStatsAggregator::ClientTick(UE::Cook::FMPCollectorClientTickContext&
 		Writer.BeginObject();
 
 		FShaderCompilerStats LocalStats;
-		GShaderCompilingManager->GetLocalStats(LocalStats);
+		// Temporarily disable shader stats aggregation for multi-process cooks due to message sizes exceeding limits.
+		// GShaderCompilingManager->GetLocalStats(LocalStats);
 		LocalStats.WriteToCompactBinary(Writer);
 
 		Writer.EndObject();
