Fix for null incremental cooks not functioning as intended  (non-zero number of packages cooked) due to user-defined structs.

diff --git a/Engine/Source/Editor/UnrealEd/Private/TargetDomain/TargetDomainUtils.cpp b/Engine/Source/Editor/UnrealEd/Private/TargetDomain/TargetDomainUtils.cpp
--- a/Engine/Source/Editor/UnrealEd/Private/TargetDomain/TargetDomainUtils.cpp
+++ b/Engine/Source/Editor/UnrealEd/Private/TargetDomain/TargetDomainUtils.cpp
@@ -1094,6 +1094,15 @@ bool IsIterativeEnabled(FName PackageName, bool bAllowAllClasses, const FAssetPa
 			{
 				ExistingData = ClassDigests.Map.Find(ClassPath);
 			}
+			else if (!ClassName.IsNone())
+			{
+				// All classes are top-level objects, but user-defined structs are not top-level
+				// objects.  In this code we don't need to handle user-defined structs because
+				// we do not support deny-listing user-defined structs.  So if the ClassName is
+				// not a top level asset, then ignore it.
+				continue;
+			}
+
 			if (!ExistingData)
 			{
 				// !ExistingData -> !allowed, because caller has already called CalculatePackageDigest, so all
