[Viewport Toolbar] Polish: Add Cltr + ~ coord space cycle hotkey into tooltip of viewport toolbar spaces menu

diff --git a/Engine/Source/Editor/UnrealEd/Private/ViewportToolbar/UnrealEdViewportToolbar.cpp b/Engine/Source/Editor/UnrealEd/Private/ViewportToolbar/UnrealEdViewportToolbar.cpp
--- a/Engine/Source/Editor/UnrealEd/Private/ViewportToolbar/UnrealEdViewportToolbar.cpp
+++ b/Engine/Source/Editor/UnrealEd/Private/ViewportToolbar/UnrealEdViewportToolbar.cpp
@@ -785,6 +785,66 @@ FToolMenuEntry CreateTransformsSubmenu()
 												}
 											);
 										}
+
+										// Have the Tooltip show the current hotkey(s) for cycling coord spaces
+										{
+											CoordinateSystemSubmenu.ToolTip = TAttribute<FText>::CreateLambda(
+												[]()
+												{
+													const FInputChord& PrimaryChord =
+														*FEditorViewportCommands::Get().CycleTransformGizmoCoordSystem->GetActiveChord(
+															EMultipleKeyBindingIndex::Primary
+														);
+
+													const FInputChord& SecondaryChord =
+														*FEditorViewportCommands::Get().CycleTransformGizmoCoordSystem->GetActiveChord(
+															EMultipleKeyBindingIndex::Secondary
+														);
+
+													if (PrimaryChord.IsValidChord() || SecondaryChord.IsValidChord())
+													{
+
+														// Both Chords are available
+														if (PrimaryChord.IsValidChord() && SecondaryChord.IsValidChord())
+														{
+															FFormatNamedArguments Args;
+															Args.Add(TEXT("PrimaryChord"), PrimaryChord.GetInputText());
+															Args.Add(TEXT("SecondaryChord"), SecondaryChord.GetInputText());
+
+															return FText::Format(
+																LOCTEXT(
+																	"CoordinateSystemTooltipWithBothChords", "Select between coordinate systems. \n{PrimaryChord} or {SecondaryChord} to cycle between them."
+																),
+																Args
+															);
+														}
+
+														// If we got here, only one chord is available (primary or secondary)
+														FText ChordText;
+														if (PrimaryChord.IsValidChord())
+														{
+															ChordText = PrimaryChord.GetInputText();
+														}
+														else if (SecondaryChord.IsValidChord())
+														{
+															ChordText = SecondaryChord.GetInputText();
+														}
+
+														return FText::Format(
+															LOCTEXT(
+																"CoordinateSystemTooltipSingleChord", "Select between coordinate systems. \n{0} to cycle between them."
+															),
+															ChordText
+														);
+													}
+
+													return LOCTEXT(
+														"CoordinateSystemTooltipNoChords", "Select between coordinate systems"
+													);
+												}
+											);
+										}
+
 										CoordinateSystemSubmenu.ToolBarData.LabelOverride = FText();
 										CoordinateSystemSubmenu.SetShowInToolbarTopLevel(true);
 									}
