[Debug Draw] Fix a crash when brushes don't have any polys to draw

diff --git a/Engine/Source/Editor/UnrealEd/Private/LevelEditorViewport.cpp b/Engine/Source/Editor/UnrealEd/Private/LevelEditorViewport.cpp
--- a/Engine/Source/Editor/UnrealEd/Private/LevelEditorViewport.cpp
+++ b/Engine/Source/Editor/UnrealEd/Private/LevelEditorViewport.cpp
@@ -5076,7 +5076,7 @@ void FLevelEditorViewportClient::DrawBrushDetails(const FSceneView* View, FPrimi
 
 		// Brush->Brush is checked to safe from brushes that were created without having their brush members attached.
 		// Check whether it satisfies the ShowBrushMarkersPoly condition of being selected or whether bDisplayShadedVolume is true
-		if (Brush->Brush && (FActorEditorUtils::IsABuilderBrush(Brush) || Brush->IsVolumeBrush()) && 
+		if (Brush->Brush && Brush->Brush->Polys && (FActorEditorUtils::IsABuilderBrush(Brush) || Brush->IsVolumeBrush()) && 
 			(Brush->bDisplayShadedVolume || (GetDefault<ULevelEditorViewportSettings>()->bShowBrushMarkerPolys && ModeTools->GetSelectedActors()->IsSelected(Brush))))
 		{
 			// Build a mesh by basically drawing the triangles of each 
