Provide const version of native parent finding functions

#rnx

diff --git a/Engine/Source/Editor/UnrealEd/Private/Kismet2/BlueprintEditorUtils.cpp b/Engine/Source/Editor/UnrealEd/Private/Kismet2/BlueprintEditorUtils.cpp
--- a/Engine/Source/Editor/UnrealEd/Private/Kismet2/BlueprintEditorUtils.cpp
+++ b/Engine/Source/Editor/UnrealEd/Private/Kismet2/BlueprintEditorUtils.cpp
@@ -852,6 +852,11 @@ UClass* FBlueprintEditorUtils::FindFirstNativeClass(UClass* Class)
 	return Class;
 }
 
+const UClass* FBlueprintEditorUtils::FindFirstNativeClass(const UClass* Class)
+{
+	return FindFirstNativeClass(const_cast<UClass*>(Class));
+}
+
 bool FBlueprintEditorUtils::IsNativeSignature(const UFunction* Fn)
 {
 	if(UClass* OwningClass = Fn->GetOwnerClass())
diff --git a/Engine/Source/Editor/UnrealEd/Public/Kismet2/BlueprintEditorUtils.h b/Engine/Source/Editor/UnrealEd/Public/Kismet2/BlueprintEditorUtils.h
--- a/Engine/Source/Editor/UnrealEd/Public/Kismet2/BlueprintEditorUtils.h
+++ b/Engine/Source/Editor/UnrealEd/Public/Kismet2/BlueprintEditorUtils.h
@@ -685,6 +685,7 @@ public:
 
 	/** return find first native class in the hierarchy */
 	static UNREALED_API UClass* FindFirstNativeClass(UClass* Class);
+	static UNREALED_API const UClass* FindFirstNativeClass(const UClass* Class);
 
 	/** returns true if this blueprints signature (inc. visibility) was determined by UHT, rather than the blueprint compiler */
 	static UNREALED_API bool IsNativeSignature(const UFunction* Fn);
