Fix mesh description crash cause by missing register attributes.

#rnx

diff --git a/Engine/Source/Editor/UnrealEd/Private/Fbx/FbxMaterialExportUtilities.cpp b/Engine/Source/Editor/UnrealEd/Private/Fbx/FbxMaterialExportUtilities.cpp
--- a/Engine/Source/Editor/UnrealEd/Private/Fbx/FbxMaterialExportUtilities.cpp
+++ b/Engine/Source/Editor/UnrealEd/Private/Fbx/FbxMaterialExportUtilities.cpp
@@ -40,6 +40,7 @@
 #include "Materials/MaterialExpressionConstant4Vector.h"
 #include "Materials/MaterialExpressionLinearInterpolate.h"
 #include "Materials/MaterialExpressionScalarParameter.h"
+#include "Model.h"
 #include "TextureCompiler.h"
 
 #include "InterchangeAssetImportData.h"
@@ -596,12 +597,11 @@ namespace UnFbx
 		, LODIndex(InLODIndex)
 	{
 		FMeshDescription Mesh;
-		FStaticMeshAttributes MeshAttributes(Description);
-		MeshAttributes.Register();
+		FStaticMeshAttributes(Mesh).Register();
 		TArray<FStaticMaterial>	Materials;
 		GetBrushMesh(Actor, Model, Mesh, Materials);
 
-		StaticMesh = CreateStaticMesh(Mesh, Materials, GetTransientPackage(), Actor->GetFName());
+		StaticMesh = CreateStaticMesh(Mesh, Materials, GetTransientPackage(), Actor ? Actor->GetFName() : Model->GetFName());
 
 		const IMeshMergeUtilities& MeshMergeUtilities = FModuleManager::Get().LoadModuleChecked<IMeshMergeModule>("MeshMergeUtilities").GetUtilities();
 		FStaticMeshAttributes(Description).Register();
