Reintroduce 'fast' reinstancing, to restore an objects class when its class does not change from compilation (e.g. we recompile but end up with the same class)

[RN] Blueprint reinstancing now skips object replacement when a class layout and constructor is not changed, this optimization can be disabled with by setting the BP.bAllowFastReinstancing console variable to false

diff --git a/Engine/Source/Editor/Kismet/Private/BlueprintCompilationManager.cpp b/Engine/Source/Editor/Kismet/Private/BlueprintCompilationManager.cpp
--- a/Engine/Source/Editor/Kismet/Private/BlueprintCompilationManager.cpp
+++ b/Engine/Source/Editor/Kismet/Private/BlueprintCompilationManager.cpp
@@ -134,6 +134,7 @@ struct FBlueprintCompilationManagerImpl : public FGCObject
 	bool IsGeneratedClassLayoutReady() const;
 	void GetDefaultValue(const UClass* ForClass, const FProperty* Property, FString& OutDefaultValueAsString) const;
 	void VerifyNoQueuedRequests(const TArray<FCompilerData>& CurrentlyCompilingBPs);
+	void SaveConstructionHashes(const TArray<FCompilerData>& CurrentlyCompilingBPs);
 
 	static void ReparentHierarchies(const TMap<UClass*, UClass*>& OldClassToNewClass, EReparentClassOptions Options);
 	static void BuildDSOMap(UObject* OldObject, UObject* NewObject, const TMap<UClass*, UClass*>& InOldToNewClassMap, TMap<UObject*, UObject*>& OutOldToNewDSO);
@@ -167,6 +168,7 @@ struct FBlueprintCompilationManagerImpl : public FGCObject
 	// Data stored for reinstancing, which finishes much later than compilation,
 	// populated by FlushCompilationQueueImpl, cleared by FlushReinstancingQueueImpl:
 	TMap<TObjectPtr<UClass>, TObjectPtr<UClass>> ClassesToReinstance;
+	TMap<TObjectPtr<const UClass>, FBlake3Hash> ClassHashes;
 	
 	// Map to old default values, useful for providing access to this data throughout
 	// the compilation process:
@@ -222,6 +224,7 @@ void FBlueprintCompilationManagerImpl::AddReferencedObjects(FReferenceCollector&
 	}
 
 	Collector.AddReferencedObjects(ClassesToReinstance);
+	Collector.AddReferencedObjects(ClassHashes);
 	Collector.AddReferencedObjects(OldCDOs);
 }
 
@@ -583,6 +586,20 @@ namespace UE::Kismet::BlueprintCompilationManager::Private
 			TEXT("BP.bForceAllDependenciesToRecompile"), bForceAllDependenciesToRecompile,
 			TEXT("If true all dependencies will be bytecode-compiled even when all referenced functions have no signature changes. Intended for compiler development/debugging purposes."),
 			ECVF_Default);
+
+		/** 
+		 * Flag to disable skipping of reinstancing logic when a class is not observabling changed
+		 * by compilation. This can only work when compiling in place (e.g. reusing UClasses) - so
+		 * can't be done for asset reload, live coding, or verse 
+		 */
+		static bool bAllowFastReinstancing = true;
+		static FAutoConsoleVariableRef CVarAllowFastReinstancing(
+			TEXT("BP.bAllowFastReinstancing"), bAllowFastReinstancing,
+			TEXT("True by default, this flag enables a fast path in the logic that updates class \
+				objects at runtime (Reinstancing). When this flag is set we check class layout and \
+				construction data for a match, if the layout and construction data match then we \
+				don't replace instances of the class."),
+			ECVF_Default);
 	}
 }
 
@@ -890,6 +907,9 @@ void FBlueprintCompilationManagerImpl::FlushCompilationQueueImpl(bool bSuppressB
 
 		SlowTask.EnterProgressFrame();
 
+		// store off construction hashes, in case we can do fast reinstancing:
+		SaveConstructionHashes(CurrentlyCompilingBPs);
+
 		// STAGE IV: Set UBlueprint flags (bBeingCompiled, bIsRegeneratingOnLoad)
 		for (FCompilerData& CompilerData : CurrentlyCompilingBPs)
 		{
@@ -1934,6 +1954,7 @@ void FBlueprintCompilationManagerImpl::FlushReinstancingQueueImpl(bool bFindAndR
 	// we can finalize reinstancing now:
 	if (ClassesToReinstance.Num() == 0)
 	{
+		ClassHashes.Reset();
 		return;
 	}
 
@@ -1951,7 +1972,9 @@ void FBlueprintCompilationManagerImpl::FlushReinstancingQueueImpl(bool bFindAndR
 		Options.bArchetypesAreUpToDate = true;
 		Options.bReplaceReferencesToOldCDOs = bFindAndReplaceCDOReferences;
 		Options.OldToNewTemplates = OldToNewTemplates;
+		Options.ConstructionVersioningData = ClassHashes.Num() > 0 ? &ClassHashes : nullptr;
 		FBlueprintCompileReinstancer::BatchReplaceInstancesOfClass(ClassesToReinstanceOwned, Options);
+		ClassHashes.Reset();
 
 		// Special case when we run on ALT, we want to cleanup all classes flagged for reinstantiation right away.
 		const bool bIsInActualAsyncLoadingThread = IsInAsyncLoadingThread() && !IsInGameThread();
@@ -2058,6 +2081,49 @@ void FBlueprintCompilationManagerImpl::VerifyNoQueuedRequests(const TArray<FComp
 	}
 }
 
+void FBlueprintCompilationManagerImpl::SaveConstructionHashes(const TArray<FCompilerData>& CurrentlyCompilingBPs)
+{
+	if( !UE::Kismet::BlueprintCompilationManager::Private::ConsoleVariables::bAllowFastReinstancing )
+	{
+		return;
+	}
+
+	// Don't do fast reinstancing for override serialization types, they
+	// are currently relying on reinstancing to propagate edits. This 
+	// requires that we do slow reconstruction/replacement:
+	for (const FCompilerData& CompilerData : CurrentlyCompilingBPs)
+	{
+		UBlueprint* BP = CompilerData.BP;
+		if(const UClass* OriginalClass = BP->GeneratedClass)
+		{
+			if( const UObject* CDO = OriginalClass->ClassDefaultObject )
+			{
+				if( FOverridableManager::Get().IsEnabled(*CDO) )
+				{
+					return;
+				}
+			}
+		}
+	}
+
+	for (const FCompilerData& CompilerData : CurrentlyCompilingBPs)
+	{
+		UBlueprint* BP = CompilerData.BP;
+		if(const UClass* OriginalClass = BP->GeneratedClass)
+		{
+			if(OriginalClass->ClassDefaultObject)
+			{
+				// if there are no instances, no need to calculate the class hash:
+				FBlake3Hash ConstructionHash = FBlueprintCompileReinstancer::CalculateConstructionHashIfNecessary(OriginalClass);
+				if(ConstructionHash != FBlake3Hash())
+				{
+					ClassHashes.Add(OriginalClass, ConstructionHash);
+				}
+			}
+		}
+	}
+}
+
 void FBlueprintCompilationManagerImpl::ReparentHierarchies(const TMap<UClass*, UClass*>& OldToNewClasses, EReparentClassOptions Options)
 {
 	const bool bReplaceReferencesToOldClasses = (Options & EReparentClassOptions::ReplaceReferencesToOldClasses) != EReparentClassOptions::None;
diff --git a/Engine/Source/Editor/UnrealEd/Private/Kismet2/KismetReinstanceUtilities.cpp b/Engine/Source/Editor/UnrealEd/Private/Kismet2/KismetReinstanceUtilities.cpp
--- a/Engine/Source/Editor/UnrealEd/Private/Kismet2/KismetReinstanceUtilities.cpp
+++ b/Engine/Source/Editor/UnrealEd/Private/Kismet2/KismetReinstanceUtilities.cpp
@@ -5,6 +5,7 @@
 #include "BlueprintCompilationManager.h"
 #include "ComponentInstanceDataCache.h"
 #include "Engine/Blueprint.h"
+#include "JsonObjectGraph/Stringify.h"
 #include "Stats/StatsMisc.h"
 #include "UObject/Package.h"
 #include "Components/SceneComponent.h"
@@ -2584,15 +2585,6 @@ void FBlueprintCompileReinstancer::ReplaceInstancesOfClass_Inner(const TMap<UCla
 
 	// Keep track of non-dirty packages for objects about to be reinstanced, so we can clear the dirty state after reinstancing them
 	TSet<UPackage*> CleanPackageList;
-	auto CheckAndSaveOuterPackageToCleanList = [&CleanPackageList](const UObject* InObject)
-	{
-		check(InObject);
-		UPackage* ObjectPackage = InObject->GetPackage();
-		if (ObjectPackage && !ObjectPackage->IsDirty() && ObjectPackage != GetTransientPackage())
-		{
-			CleanPackageList.Add(ObjectPackage);
-		}
-	};
 
 	struct FObjectRemappingHelper
 	{
@@ -2621,25 +2613,6 @@ void FBlueprintCompileReinstancer::ReplaceInstancesOfClass_Inner(const TMap<UCla
 
 	FDelegateHandle OnObjectsReplacedHandle = FCoreUObjectDelegates::OnObjectsReplaced.AddRaw(&ObjectRemappingHelper, &FObjectRemappingHelper::OnObjectsReplaced);
 
-	auto UpdateObjectBeingDebugged = [](UObject* InOldObject, UObject* InNewObject)
-	{
-		if (UBlueprint* OldObjBlueprint = Cast<UBlueprint>(InOldObject->GetClass()->ClassGeneratedBy))
-		{
-			// For now, don't update the object if the outer BP assets don't match (e.g. after a reload). Otherwise, it will
-			// trigger an ensure() in SetObjectBeginDebugged(). This will be replaced with a better solution in a future release.
-			if (OldObjBlueprint == Cast<UBlueprint>(InNewObject->GetClass()->ClassGeneratedBy))
-			{
-				// The old object may already be PendingKill, but we still want to check the current
-				// ptr value for a match. Otherwise, the selection will get cleared after every compile.
-				const UObject* DebugObj = OldObjBlueprint->GetObjectBeingDebugged(EGetObjectOrWorldBeingDebuggedFlags::IgnorePendingKill);
-				if (DebugObj == InOldObject)
-				{
-					OldObjBlueprint->SetObjectBeingDebugged(InNewObject);
-				}
-			}
-		}
-	};
-
 	{
 		TArray<UObject*> ObjectsToReplace;
 		TSet<UObject*> CachedArchetypeObjects;
@@ -2659,7 +2632,7 @@ void FBlueprintCompileReinstancer::ReplaceInstancesOfClass_Inner(const TMap<UCla
 				{
 					if (const UObject* SelectedActor = SelectedActors->GetSelectedObject(SelectionIdx))
 					{
-						CheckAndSaveOuterPackageToCleanList(SelectedActor);
+						FBlueprintCompileReinstancer::CheckAndSaveOuterPackageToCleanList(CleanPackageList, SelectedActor);
 					}
 				}
 			
@@ -2667,7 +2640,42 @@ void FBlueprintCompileReinstancer::ReplaceInstancesOfClass_Inner(const TMap<UCla
 				SelectedActors->Modify();
 			}
 		}
+		
+		FDelegateHandle OnLevelActorDeletedHandle = GEngine ? GEngine->OnLevelActorDeleted().AddLambda([&OldToNewInstanceMap](AActor* DestroyedActor)
+		{
+			if (UObject** ReplacementObject = OldToNewInstanceMap.Find(DestroyedActor))
+			{
+				AActor* ReplacementActor = CastChecked<AActor>(*ReplacementObject);
+				ReplacementActor->GetWorld()->EditorDestroyActor(ReplacementActor, /*bShouldModifyLevel =*/true);
+			}
+		}) : FDelegateHandle();
 
+		if(Params.ConstructionVersioningData)
+		{
+			check(bArchetypesAreUpToDate); // not supported, can't be
+			check(Params.OldToNewTemplates == nullptr); // not supported, can't be.. new templates requires new instances
+			check(Params.InstancesThatShouldUseOldClass == nullptr); // not supported, could be.. but why?
+			ReinstanceWithVersioningData(InOldToNewClassMap, 
+				{
+					.ObjectsToReplaceScratch = ObjectsToReplace,
+					.CachedArchetypeObjects = CachedArchetypeObjects,
+					.OldToNewInstanceMap = OldToNewInstanceMap,
+					.InOldToNewClassMap = InOldToNewClassMap,
+					.OldToNewNameMap = OldToNewNameMap,
+					.PotentialEditorsForRefreshing = PotentialEditorsForRefreshing,
+					.OwnersToRerunConstructionScript = OwnersToRerunConstructionScript,
+					.ReplacementActors = ReplacementActors,
+					.ReinstancedObjectsWeakReferenceMap = ReinstancedObjectsWeakReferenceMap,
+					.ConstructionVersioningData = *Params.ConstructionVersioningData,
+					.ObjectsReplaced = ObjectsReplaced,
+					.CleanPackageList = CleanPackageList,
+					.bSelectionChanged = bSelectionChanged,
+					.bFixupSCS = bFixupSCS
+				}
+			);
+		}
+		else
+		{
 		// WARNING: for (TPair<UClass*, UClass*> OldToNewClass : InOldToNewClassMap) duplicated below 
 		// to handle reconstructing actors which need to be reinstanced after their owned components 
 		// have been updated:
@@ -2721,25 +2729,7 @@ void FBlueprintCompileReinstancer::ReplaceInstancesOfClass_Inner(const TMap<UCla
 					});
 				}
 
-				// We need to cache the archetype of the objects about to be replaced 
-				// as it will not be possible to get them during this process as it renames the objects
-				// These cached archetypes will not be updated if they were set earlier
-				for (UObject* OldObject : ObjectsToReplace)
-				{
-					if(!IsValid(OldObject))
-					{
-						continue;
-					}
-
-					CacheManager.CacheArchetype(OldObject);
-					CachedArchetypeObjects.Add(OldObject);
-
-					ForEachObjectWithOuter(OldObject, [&CacheManager, &CachedArchetypeObjects](UObject* SubObject)
-					{
-						CacheManager.CacheArchetype(SubObject);
-						CachedArchetypeObjects.Add(SubObject);
-					});
-				}
+				CacheArchetypes(ObjectsToReplace, CacheManager, CachedArchetypeObjects);
 				
 				// Then fix 'real' (non archetype) instances of the class
 				for (int32 OldObjIndex = 0; OldObjIndex < ObjectsToReplace.Num(); ++OldObjIndex)
@@ -2772,7 +2762,7 @@ void FBlueprintCompileReinstancer::ReplaceInstancesOfClass_Inner(const TMap<UCla
 					// WARNING: This loop only handles non-actor objects, actor objects are handled below:
 					if (OldActor == nullptr)
 					{
-						CheckAndSaveOuterPackageToCleanList(OldObject);
+						CheckAndSaveOuterPackageToCleanList(CleanPackageList, OldObject);
 
 						UObject* NewUObject = nullptr;
 						ReplaceObjectHelper(OldObject, OldClass, NewUObject, NewClass, OldToNewInstanceMap, InOldToNewClassMap, OldToNewNameMap, OldObjIndex, ObjectsToReplace, PotentialEditorsForRefreshing, OwnersToRerunConstructionScript, &FDirectAttachChildrenAccessor::Get, bIsComponent, bArchetypesAreUpToDate);
@@ -2788,16 +2778,6 @@ void FBlueprintCompileReinstancer::ReplaceInstancesOfClass_Inner(const TMap<UCla
 			}
 		}
 
-
-		FDelegateHandle OnLevelActorDeletedHandle = GEngine ? GEngine->OnLevelActorDeleted().AddLambda([&OldToNewInstanceMap](AActor* DestroyedActor)
-		{
-			if (UObject** ReplacementObject = OldToNewInstanceMap.Find(DestroyedActor))
-			{
-				AActor* ReplacementActor = CastChecked<AActor>(*ReplacementObject);
-				ReplacementActor->GetWorld()->EditorDestroyActor(ReplacementActor, /*bShouldModifyLevel =*/true);
-			}
-		}) : FDelegateHandle();
-
 		// WARNING: for (TPair<UClass*, UClass*> OldToNewClass : InOldToNewClassMap) duplicated above 
 		// this loop only handles actors - which need to be reconstructed *after* their owned components 
 		// have been reinstanced:
@@ -2854,25 +2834,12 @@ void FBlueprintCompileReinstancer::ReplaceInstancesOfClass_Inner(const TMap<UCla
 					if (OldActor != nullptr)
 					{
 						// Note: For OFPA, each instance may be stored in its own external package.
-						CheckAndSaveOuterPackageToCleanList(OldActor);
+						CheckAndSaveOuterPackageToCleanList(CleanPackageList, OldActor);
 
 						UObject* NewUObject = nullptr;
 						if (OldActor->GetLevel() && OldActor->GetWorld() && !OldActor->GetWorld()->IsCleanedUp())
 						{
-							// Attached actors will be marked dirty as well when the old instance is destroyed, but
-							// since these attachments will get restored as part of replacement, there's no need to
-							// re-save after reinstancing. In general, this applies only to OFPA-enabled levels, as
-							// in that case, attached actors reside in their own external packages and not the level.
-							TArray<AActor*> AttachedActors;
-							OldActor->GetAttachedActors(AttachedActors);
-							for (AActor* AttachedActor : AttachedActors)
-							{
-								if (AttachedActor)
-								{
-									CheckAndSaveOuterPackageToCleanList(AttachedActor);
-								}
-							}
-
+							CheckAndSaveAttachedActorsOuterPackageToCleanList(CleanPackageList, OldActor);
 							ReplaceActorHelper(OldActor, OldClass, NewUObject, NewClass, OldToNewInstanceMap, InOldToNewClassMap, ReinstancedObjectsWeakReferenceMap, ActorAttachmentData, ReplacementActors, bPreserveRootComponent, bSelectionChanged);
 						}
 						else
@@ -2891,6 +2858,7 @@ void FBlueprintCompileReinstancer::ReplaceInstancesOfClass_Inner(const TMap<UCla
 				}
 			}
 		}
+		}
 
 		// Reset any cached archetypes
 		for (UObject* CachedArchetypeObject : CachedArchetypeObjects)
@@ -3381,3 +3349,403 @@ void FBlueprintCompileReinstancer::PreCreateSubObjectsForReinstantiation_Inner(
 		}
 	}
 }
+
+void FBlueprintCompileReinstancer::CacheArchetypes(
+	const TArray<UObject*>& ObjectsToReplace, 
+	FEditorCacheArchetypeManager& CacheManager, 
+	TSet<UObject*>& CachedArchetypeObjects)
+{
+	// We need to cache the archetype of the objects about to be replaced 
+	// as it will not be possible to get them during this process as it renames the objects
+	// These cached archetypes will not be updated if they were set earlier
+	for (UObject* OldObject : ObjectsToReplace)
+	{
+		if(!IsValid(OldObject))
+		{
+			continue;
+		}
+
+		CacheManager.CacheArchetype(OldObject);
+		CachedArchetypeObjects.Add(OldObject);
+
+		ForEachObjectWithOuter(OldObject, [&CacheManager, &CachedArchetypeObjects](UObject* SubObject)
+		{
+			CacheManager.CacheArchetype(SubObject);
+			CachedArchetypeObjects.Add(SubObject);
+		});
+	}
+}
+
+FBlake3Hash FBlueprintCompileReinstancer::CalculateConstructionHashIfNecessary(const UClass* ForClass)
+{
+	TArray<UObject*> Count;
+	GetObjectsToReinstance(ForClass, Count);
+	if(Count.Num() == 0)
+	{
+		return FBlake3Hash();
+	}
+
+	return CalculateConstructionHash(ForClass);
+}
+
+FBlake3Hash FBlueprintCompileReinstancer::CalculateConstructionHash(const UClass* ForClass)
+{
+	check(ForClass && ForClass->ClassDefaultObject);
+
+	// must capture all template data:
+	FUtf8String TemplateData = UE::JsonObjectGraph::Stringify(
+		{ForClass, ForClass->ClassDefaultObject}, 
+		FJsonStringifyOptions(EJsonStringifyFlags::FilterEditorOnlyData));
+	FBlake3 Hash;
+	Hash.Update(TemplateData.GetCharArray().GetData(), TemplateData.GetCharArray().Num());
+	return Hash.Finalize();
+}
+
+void FBlueprintCompileReinstancer::ReinstanceWithVersioningData(const TMap<UClass*, UClass*>& InOldToNewClassMap, const FReinstancingState& ReinstancingState)
+{
+	TMap<UClass*, UClass*> UnchangedClasses;
+	TMap<UClass*, UClass*> NormalReinstancingClasses;
+	TMap<UClass*, UClass*> ActorReinstancingClasses;
+	CategorizeClassesAndInstances(
+		InOldToNewClassMap, 
+		UnchangedClasses, 
+		NormalReinstancingClasses, 
+		ActorReinstancingClasses, 
+		ReinstancingState);
+	RestoreClass(UnchangedClasses, ReinstancingState);
+	ReinstanceNormalObjects(NormalReinstancingClasses, ReinstancingState);
+	ReinstanceActors(ActorReinstancingClasses, ReinstancingState);
+}
+
+void FBlueprintCompileReinstancer::CategorizeClassesAndInstances(
+	const TMap<UClass*, UClass*>& InOldToNewClassMap
+	, TMap<UClass*, UClass*>& UnchangedClasses
+	, TMap<UClass*, UClass*>& NormalReinstancingClasses
+	, TMap<UClass*, UClass*>& ActorReinstancingClasses
+	, const FBlueprintCompileReinstancer::FReinstancingState& ReinstancingState)
+{
+	for (TPair<UClass*, UClass*> OldToNewClass : InOldToNewClassMap)
+	{
+		const UClass* const OldClass = OldToNewClass.Key;
+		const UClass* const NewClass = OldToNewClass.Value;
+
+		const bool bIsComponent = NewClass->IsChildOf<UActorComponent>();
+		// If any of the class changes are of an actor component to scene component or reverse then we will fixup SCS of all actors affected
+		if (bIsComponent && !ReinstancingState.bFixupSCS)
+		{
+			ReinstancingState.bFixupSCS = (NewClass->IsChildOf<USceneComponent>() != OldClass->IsChildOf<USceneComponent>());
+		}
+
+		// if there are no instances, nothing to do:
+		TArray<UObject*>& ObjectsToReplace = ReinstancingState.ObjectsToReplaceScratch;
+		GetObjectsToReinstance(OldClass, ObjectsToReplace);
+		if(ObjectsToReplace.Num() == 0)
+		{
+			continue; // unambiguously no reinstancing necessary
+		}
+		
+		// if layout and default values match, we can reinstance fast:
+		bool bClassMatchesOldClass = false;
+		bool bUseActorReinstancing = false;
+		// if we have a script constructor we would have to check it for equivalence here
+		// but there are many obsticales to script constructors:
+		const FBlake3Hash* ConstructionVersion = ReinstancingState.ConstructionVersioningData.Find(NewClass);
+		if( ConstructionVersion && 
+			FStructUtils::TheSameLayout(OldClass, NewClass) &&
+			CalculateConstructionHash(NewClass) == *ConstructionVersion)
+		{
+			bClassMatchesOldClass = true;
+
+			// we can only do fast reinstancing when ownership of objects is clear,
+			// otherwise disable it - this avoids n^2 archetype lookups for each 
+			// instance, which may be still profitable but not much.. we
+			// still have to pay for deep hierarchies here, which is not ideal
+			// but is only O(n) and hierarchy depth is never going to be >1000
+			// at some point we'll need a flag to identify notionally free
+			// objects (e.g. not created by a ctor, native or otherwise) - could
+			// rely on some kind of 'is editable' flag but we're fighting with perspectives
+			// there, where editability is associated with a property and has no
+			// tracking on the object...:
+			const auto OuteredToRegeneratedObject = [](const UObject* Obj)
+			{
+				UObject* OuterIter = Obj->GetOuter();
+				while(OuterIter)
+				{
+					if(OuterIter->GetClass()->HasAnyClassFlags(CLASS_NewerVersionExists))
+					{
+						return true;
+					}
+					OuterIter = OuterIter->GetOuter();
+				}
+				return false;
+			};
+			for(UObject* Object : ObjectsToReplace)
+			{
+				if(OuteredToRegeneratedObject(Object))
+				{
+					bClassMatchesOldClass = false;
+					break;
+				}
+			}
+		}
+		// actors require late reinstancing, but if reinstancing fast
+		// they should be ok, as finalization happens at the same time
+		if(NewClass->IsChildOf(AActor::StaticClass()))
+		{
+			bUseActorReinstancing = true;
+		}
+
+		if(bClassMatchesOldClass)
+		{
+			UnchangedClasses.Add(OldToNewClass.Key, OldToNewClass.Value);
+		}
+		else if(bUseActorReinstancing)
+		{
+			ActorReinstancingClasses.Add(OldToNewClass.Key, OldToNewClass.Value);
+		}
+		else
+		{
+			NormalReinstancingClasses.Add(OldToNewClass.Key, OldToNewClass.Value);
+		}
+
+	}
+}
+
+void FBlueprintCompileReinstancer::RestoreClass(
+	TMap<UClass*, UClass*>& OldToNewClassMap, 
+	const FBlueprintCompileReinstancer::FReinstancingState& ReinstancingState)
+{
+	// update classes
+	for (TPair<UClass*, UClass*> OldToNewClass : OldToNewClassMap)
+	{
+		const UClass* const OldClass = OldToNewClass.Key;
+		UClass* const NewClass = OldToNewClass.Value;
+		TArray<UObject*>& ObjectsToReplace = ReinstancingState.ObjectsToReplaceScratch;
+		GetObjectsToReinstance(OldClass, ObjectsToReplace);
+
+		for (UObject* Object : ObjectsToReplace)
+		{
+			Object->SetClass(NewClass); // class layouts match, so we can just set the object back to its original class
+		}
+	}
+}
+
+void FBlueprintCompileReinstancer::ReinstanceNormalObjects(
+	TMap<UClass*, UClass*>& OldToNewClassMap,
+	const FBlueprintCompileReinstancer::FReinstancingState& ReinstancingState)
+{
+	FEditorCacheArchetypeManager& CacheManager = FEditorCacheArchetypeManager::Get();
+	TArray<UObject*> EmptyList;
+	const int32 OldObjIndex = INDEX_NONE; // not needed as archetypes are up to date in this codepath
+	const bool bArchetypesAreUpToDate = true;
+
+	for (TPair<UClass*, UClass*> OldToNewClass : OldToNewClassMap)
+	{
+		UClass* const OldClass = OldToNewClass.Key;
+		UClass* const NewClass = OldToNewClass.Value;
+		
+		const bool bIsComponent = NewClass->IsChildOf<UActorComponent>();
+
+		TArray<UObject*>& ObjectsToReplace = ReinstancingState.ObjectsToReplaceScratch;
+		GetObjectsToReinstance(OldClass, ObjectsToReplace);
+		
+		CacheArchetypes(ObjectsToReplace, CacheManager, ReinstancingState.CachedArchetypeObjects);
+
+		for (UObject* OldObject : ObjectsToReplace)
+		{
+			check(IsValid(OldObject)); // is RF_MirroredGarbage in GetObjectsToReinstance not enough?
+			if(!ShouldReinstanceObject(OldObject, ReinstancingState.InOldToNewClassMap, bIsComponent))
+			{
+				continue;
+			}
+
+			CheckAndSaveOuterPackageToCleanList(ReinstancingState.CleanPackageList, OldObject);
+			UObject* NewUObject = nullptr;
+			ReplaceObjectHelper(
+				OldObject, 
+				OldClass, 
+				NewUObject, 
+				NewClass, 
+				ReinstancingState.OldToNewInstanceMap, 
+				ReinstancingState.InOldToNewClassMap, 
+				ReinstancingState.OldToNewNameMap, 
+				OldObjIndex, 
+				EmptyList, 
+				ReinstancingState.PotentialEditorsForRefreshing, 
+				ReinstancingState.OwnersToRerunConstructionScript, 
+				&FDirectAttachChildrenAccessor::Get, 
+				bIsComponent, 
+				bArchetypesAreUpToDate);
+			UpdateObjectBeingDebugged(OldObject, NewUObject);
+			ReinstancingState.ObjectsReplaced.Add(OldObject);
+		}
+	}
+}
+
+void FBlueprintCompileReinstancer::ReinstanceActors(
+	TMap<UClass*, UClass*>& OldToNewClassMap, 
+	const FBlueprintCompileReinstancer::FReinstancingState& ReinstancingState)
+{
+	TArray<UObject*> EmptyList;
+	const int32 OldObjIndex = INDEX_NONE; // not needed as archetypes are up to date in this codepath
+	const bool bArchetypesAreUpToDate = true;
+	const bool bIsComponent = false;
+	const bool bPreserveRootComponent = true;
+
+	for (TPair<UClass*, UClass*> OldToNewClass : OldToNewClassMap)
+	{
+		UClass* OldClass = OldToNewClass.Key;
+		UClass* NewClass = OldToNewClass.Value;
+		check(OldClass && NewClass);
+		TArray<UObject*>& ObjectsToReplace = ReinstancingState.ObjectsToReplaceScratch;
+		GetObjectsToReinstance(OldClass, ObjectsToReplace);
+
+		// store old attachment data before we mess with components, etc:
+		TMap<UObject*, FActorAttachmentData> ActorAttachmentData;
+		for (UObject* OldObject : ObjectsToReplace)
+		{
+			if (AActor* OldActor = Cast<AActor>(OldObject))
+			{
+				ActorAttachmentData.Add(OldObject, FActorAttachmentData(OldActor));
+			}
+		}
+		
+		for (UObject* OldObject : ObjectsToReplace)
+		{
+			check(IsValid(OldObject)); // is RF_MirroredGarbage in GetObjectsToReinstance not enough?
+			AActor* OldActor = CastChecked<AActor>(OldObject);
+			if(!ShouldReinstanceActor(OldActor))
+			{
+				continue;
+			}
+
+			CheckAndSaveOuterPackageToCleanList(ReinstancingState.CleanPackageList, OldObject);
+
+			UObject* NewUObject = nullptr;
+			if (OldActor->GetLevel() && OldActor->GetWorld() && !OldActor->GetWorld()->IsCleanedUp())
+			{
+				CheckAndSaveAttachedActorsOuterPackageToCleanList(ReinstancingState.CleanPackageList, OldActor);
+				ReplaceActorHelper(
+					OldActor, 
+					OldClass, 
+					NewUObject, 
+					NewClass, 
+					ReinstancingState.OldToNewInstanceMap, 
+					ReinstancingState.InOldToNewClassMap, 
+					ReinstancingState.ReinstancedObjectsWeakReferenceMap, 
+					ActorAttachmentData, 
+					ReinstancingState.ReplacementActors, 
+					bPreserveRootComponent, 
+					ReinstancingState.bSelectionChanged);
+			}
+			else
+			{
+				// Actors that are not in a level cannot be reconstructed, reinstance them as normal objects:
+				ReplaceObjectHelper(
+					OldObject, 
+					OldClass,
+					NewUObject,
+					NewClass,
+					ReinstancingState.OldToNewInstanceMap,
+					ReinstancingState.InOldToNewClassMap,
+					ReinstancingState.OldToNewNameMap,
+					OldObjIndex,
+					EmptyList,
+					ReinstancingState.PotentialEditorsForRefreshing,
+					ReinstancingState.OwnersToRerunConstructionScript,
+					&FDirectAttachChildrenAccessor::Get,
+					bIsComponent,
+					bArchetypesAreUpToDate);
+			}
+			UpdateObjectBeingDebugged(OldObject, NewUObject);
+			ReinstancingState.ObjectsReplaced.Add(OldObject);
+		}
+	}
+}
+
+void FBlueprintCompileReinstancer::GetObjectsToReinstance(const UClass* OldClass, TArray<UObject*>& ObjectsToReplace)
+{
+	ObjectsToReplace.Reset();
+	const bool bIncludeDerivedClasses = false;
+	GetObjectsOfClass(OldClass, ObjectsToReplace, bIncludeDerivedClasses, RF_ClassDefaultObject | RF_MirroredGarbage);
+}
+
+void FBlueprintCompileReinstancer::UpdateObjectBeingDebugged(UObject* InOldObject, UObject* InNewObject)
+{
+	if (UBlueprint* OldObjBlueprint = Cast<UBlueprint>(InOldObject->GetClass()->ClassGeneratedBy))
+	{
+		// For now, don't update the object if the outer BP assets don't match (e.g. after a reload). Otherwise, it will
+		// trigger an ensure() in SetObjectBeginDebugged(). This will be replaced with a better solution in a future release.
+		if (OldObjBlueprint == Cast<UBlueprint>(InNewObject->GetClass()->ClassGeneratedBy))
+		{
+			// The old object may already be PendingKill, but we still want to check the current
+			// ptr value for a match. Otherwise, the selection will get cleared after every compile.
+			const UObject* DebugObj = OldObjBlueprint->GetObjectBeingDebugged(EGetObjectOrWorldBeingDebuggedFlags::IgnorePendingKill);
+			if (DebugObj == InOldObject)
+			{
+				OldObjBlueprint->SetObjectBeingDebugged(InNewObject);
+			}
+		}
+	}
+}
+
+void FBlueprintCompileReinstancer::CheckAndSaveOuterPackageToCleanList(TSet<UPackage*>& CleanPackageList, const UObject* InObject)
+{
+	check(InObject);
+	UPackage* ObjectPackage = InObject->GetPackage();
+	if (ObjectPackage && !ObjectPackage->IsDirty() && ObjectPackage != GetTransientPackage())
+	{
+		CleanPackageList.Add(ObjectPackage);
+	}
+}
+
+void FBlueprintCompileReinstancer::CheckAndSaveAttachedActorsOuterPackageToCleanList(TSet<UPackage*>& CleanPackageList, const AActor* Actor)
+{
+	// Attached actors will be marked dirty as well when the old instance is destroyed, but
+	// since these attachments will get restored as part of replacement, there's no need to
+	// re-save after reinstancing. In general, this applies only to OFPA-enabled levels, as
+	// in that case, attached actors reside in their own external packages and not the level.
+	TArray<AActor*> AttachedActors;
+	Actor->GetAttachedActors(AttachedActors);
+	for (AActor* AttachedActor : AttachedActors)
+	{
+		if (AttachedActor)
+		{
+			CheckAndSaveOuterPackageToCleanList(CleanPackageList, AttachedActor);
+		}
+	}
+}
+
+bool FBlueprintCompileReinstancer::ShouldReinstanceObject(
+	const UObject* OldObject, 
+	const TMap<UClass*, UClass*>& InOldToNewClassMap, 
+	const bool bIsComponent)
+{
+	// Skipping any default sub object that outer is going to be replaced
+	// This isn't needed for the actor loop as the only outer for actor is a level
+	if ((OldObject->IsDefaultSubobject() || OldObject->HasAnyFlags(RF_DefaultSubObject)) && 
+		InOldToNewClassMap.Contains(OldObject->GetOuter()->GetClass()))
+	{
+		return false;
+	}
+
+	// Skip archetype instances, EXCEPT for component templates - this component
+	// check is likely not needed, but leaving it for now:
+	if ( !bIsComponent && OldObject->IsTemplate() )
+	{
+		return false;
+	}
+
+	return true;
+}
+
+bool FBlueprintCompileReinstancer::ShouldReinstanceActor(const AActor* OldActor)
+{
+	// Skip archetype instances, EXCEPT for child actor templates
+	const bool bIsChildActorTemplate = OldActor->GetOuter()->IsA<UChildActorComponent>();
+	if (!bIsChildActorTemplate && OldActor->IsTemplate())
+	{
+		return false;
+	}
+	return true;
+}
diff --git a/Engine/Source/Editor/UnrealEd/Public/Kismet2/KismetReinstanceUtilities.h b/Engine/Source/Editor/UnrealEd/Public/Kismet2/KismetReinstanceUtilities.h
--- a/Engine/Source/Editor/UnrealEd/Public/Kismet2/KismetReinstanceUtilities.h
+++ b/Engine/Source/Editor/UnrealEd/Public/Kismet2/KismetReinstanceUtilities.h
@@ -10,6 +10,7 @@
 
 class FReinstanceFinalizer;
 class UBlueprint;
+class FEditorCacheArchetypeManager;
 
 DECLARE_STATS_GROUP(TEXT("Kismet Reinstancer"), STATGROUP_KismetReinstancer, STATCAT_Advanced);
 
@@ -54,6 +55,9 @@ struct FReplaceInstancesOfClassParameters
 	/** Set of objects for which new objects should not be created, useful if client has created new instances themself */
 	const TSet<UObject*>* InstancesThatShouldUseOldClass = nullptr;
 
+	/** Hashes identifying original class's construction version, if there's a match we can skip reconstruction */
+	TMap<TObjectPtr<const UClass>, FBlake3Hash>* ConstructionVersioningData = nullptr;
+
 	/** Set to true if class object has been replaced (non-batch only, legacy) */
 	bool bClassObjectReplaced = false; 
 
@@ -331,4 +335,51 @@ private:
 
 	/** Returns true if A is higher up the class hierarchy  */
 	static UNREALED_API bool ReinstancerOrderingFunction(UClass* A, UClass* B);
+	
+	/** Helper for caching archetypes before objects are reconstructed */
+	static void CacheArchetypes(	
+		const TArray<UObject*>& ObjectsToReplace, 
+		FEditorCacheArchetypeManager& CacheManager, 
+		TSet<UObject*>& CachedArchetypeObjects );
+	/** Returns a hash identifying the result of constructing an instance of the provided Class */
+	static UNREALED_API FBlake3Hash CalculateConstructionHash(const UClass* ForClass);
+	/** Returns the above, but only if there are instances of the class that would require reinstancing */
+	static UNREALED_API FBlake3Hash CalculateConstructionHashIfNecessary(const UClass* ForClass);
+
+	struct FReinstancingState
+	{
+		TArray<UObject*>& ObjectsToReplaceScratch;
+		TSet<UObject*>& CachedArchetypeObjects;
+		TMap<UObject*, UObject*>& OldToNewInstanceMap;
+		const TMap<UClass*, UClass*>& InOldToNewClassMap;
+		TMap<UObject*, FName>& OldToNewNameMap;
+		TArray<UObject*>& PotentialEditorsForRefreshing;
+		TSet<AActor*>& OwnersToRerunConstructionScript;
+		TArray<struct FActorReplacementHelper>& ReplacementActors;
+		TMap<FSoftObjectPath, UObject*>& ReinstancedObjectsWeakReferenceMap;
+		TMap<TObjectPtr<const UClass>, FBlake3Hash>& ConstructionVersioningData;
+		TArray<UObject*>& ObjectsReplaced;
+		TSet<UPackage*>& CleanPackageList;
+		bool& bSelectionChanged;
+		bool& bFixupSCS;
+	};
+
+	static void ReinstanceWithVersioningData(
+		const TMap<UClass*, UClass*>& InOldToNewClassMap, const FReinstancingState& ReinstancingState);
+	static void CategorizeClassesAndInstances(
+		const TMap<UClass*, UClass*>& InOldToNewClassMap
+		, TMap<UClass*, UClass*>& UnchangedClasses
+		, TMap<UClass*, UClass*>& NormalReinstancingClasses
+		, TMap<UClass*, UClass*>& ActiveActorReinstancingClasses
+		, const FReinstancingState& ReinstancingState);
+	static void RestoreClass(TMap<UClass*, UClass*>& OldToNewClassMap, const FReinstancingState& ReinstancingState);
+	static void ReinstanceNormalObjects(TMap<UClass*, UClass*>& OldToNewClassMap, const FReinstancingState& ReinstancingState);
+	static void ReinstanceActors(TMap<UClass*, UClass*>& OldToNewClassMap, const FReinstancingState& ReinstancingState);
+	static void GetObjectsToReinstance(const UClass* OldClass, TArray<UObject*>& ObjectsToReplace);
+	static void UpdateObjectBeingDebugged(UObject* InOldObject, UObject* InNewObject);
+	static void CheckAndSaveOuterPackageToCleanList(TSet<UPackage*>& CleanPackageList, const UObject* InObject);
+	static void CheckAndSaveAttachedActorsOuterPackageToCleanList(TSet<UPackage*>& CleanPackageList, const AActor* Actor);
+	static bool ShouldReinstanceObject(
+		const UObject* OldObject, const TMap<UClass*, UClass*>& InOldToNewClassMap, const bool bIsComponent);
+	static bool ShouldReinstanceActor(const AActor* OldActor);
 };
diff --git a/Engine/Source/Editor/UnrealEd/UnrealEd.Build.cs b/Engine/Source/Editor/UnrealEd/UnrealEd.Build.cs
--- a/Engine/Source/Editor/UnrealEd/UnrealEd.Build.cs
+++ b/Engine/Source/Editor/UnrealEd/UnrealEd.Build.cs
@@ -124,6 +124,7 @@ public class UnrealEd : ModuleRules
 				"ImageWrapper",
 				"ImageWriteQueue",
 				"InputCore",
+				"JsonObjectGraph",
 				"LauncherServices",
 				"MaterialEditor",
 				"MessageLog",
