Refine reinstancing layout equality to only include CDOs and non function subobjects of a class. The calling code checks class layout and function bytecode always changes on load, causing wasted cpu reinstancing level actors during cook

#rnx

diff --git a/Engine/Source/Editor/UnrealEd/Private/Kismet2/KismetReinstanceUtilities.cpp b/Engine/Source/Editor/UnrealEd/Private/Kismet2/KismetReinstanceUtilities.cpp
--- a/Engine/Source/Editor/UnrealEd/Private/Kismet2/KismetReinstanceUtilities.cpp
+++ b/Engine/Source/Editor/UnrealEd/Private/Kismet2/KismetReinstanceUtilities.cpp
@@ -2,6 +2,7 @@
 
 #include "Kismet2/KismetReinstanceUtilities.h"
 #include "Algo/ForEach.h"
+#include "Algo/RemoveIf.h"
 #include "BlueprintCompilationManager.h"
 #include "ComponentInstanceDataCache.h"
 #include "Engine/Blueprint.h"
@@ -3392,10 +3393,23 @@ FBlake3Hash FBlueprintCompileReinstancer::CalculateConstructionHash(const UClass
 {
 	check(ForClass && ForClass->ClassDefaultObject);
 
-	// must capture all template data:
+	// must capture all template data - just grab the subobjects of the class and exclude the ufunctions.
+	// UFunctions should be excluded as on load of an editor asset they will not include bytecode, but
+	// after compile they will so if we include them in the class version we'll always have to reinstance. 
+	// This can be impactful as we're frequently 'noop' reinstnacing level script actors.
+	TArray<UObject*> Objects;
+	constexpr bool bIncludeNestedObjects = false;
+	GetObjectsWithOuter(ForClass, Objects, bIncludeNestedObjects );
+	Objects.SetNum(Algo::RemoveIf(Objects, [](const UObject* Obj)
+	{
+		return Obj->IsA<UFunction>();
+	}));
+	Objects.Add(ForClass->ClassDefaultObject);
+	
 	FUtf8String TemplateData = UE::JsonObjectGraph::Stringify(
-		{ForClass, ForClass->ClassDefaultObject}, 
+		Objects,
 		FJsonStringifyOptions(EJsonStringifyFlags::FilterEditorOnlyData));
+
 	FBlake3 Hash;
 	Hash.Update(TemplateData.GetCharArray().GetData(), TemplateData.GetCharArray().Num());
 	return Hash.Finalize();
