comment typo fix

#rnx

diff --git a/Engine/Source/Editor/UnrealEd/Private/Kismet2/KismetReinstanceUtilities.cpp b/Engine/Source/Editor/UnrealEd/Private/Kismet2/KismetReinstanceUtilities.cpp
--- a/Engine/Source/Editor/UnrealEd/Private/Kismet2/KismetReinstanceUtilities.cpp
+++ b/Engine/Source/Editor/UnrealEd/Private/Kismet2/KismetReinstanceUtilities.cpp
@@ -3462,7 +3462,7 @@ void FBlueprintCompileReinstancer::CategorizeClassesAndInstances(
 		bool bClassMatchesOldClass = false;
 		bool bUseActorReinstancing = false;
 		// if we have a script constructor we would have to check it for equivalence here
-		// but there are many obsticales to script constructors:
+		// but there are many obstacles to script constructors:
 		const FBlake3Hash* ConstructionVersion = ReinstancingState.ConstructionVersioningData.Find(NewClass);
 		if( ConstructionVersion && 
 			FStructUtils::TheSameLayout(OldClass, NewClass) &&
