Support OldToNewTemplates, as it is used when compiling during PIE

#rnx

diff --git a/Engine/Source/Editor/UnrealEd/Private/Kismet2/KismetReinstanceUtilities.cpp b/Engine/Source/Editor/UnrealEd/Private/Kismet2/KismetReinstanceUtilities.cpp
--- a/Engine/Source/Editor/UnrealEd/Private/Kismet2/KismetReinstanceUtilities.cpp
+++ b/Engine/Source/Editor/UnrealEd/Private/Kismet2/KismetReinstanceUtilities.cpp
@@ -2654,8 +2654,17 @@ void FBlueprintCompileReinstancer::ReplaceInstancesOfClass_Inner(const TMap<UCla
 		if(Params.ConstructionVersioningData)
 		{
 			check(bArchetypesAreUpToDate); // not supported, can't be
-			check(Params.OldToNewTemplates == nullptr); // not supported, can't be.. new templates requires new instances
 			check(Params.InstancesThatShouldUseOldClass == nullptr); // not supported, could be.. but why?
+
+			// if an OldToNewTemplates map has been provided, include all the instances in the replacement map
+			// this is used to make reinstancing during PIE more robust:
+			if (Params.OldToNewTemplates)
+			{
+				for(const TPair<UClass*, TMap<UObject*, UObject*>>& Mapping : *Params.OldToNewTemplates)
+				{
+					OldToNewInstanceMap.Append(Mapping.Value);
+				}
+			}
 			ReinstanceWithVersioningData(InOldToNewClassMap, 
 				{
 					.ObjectsToReplaceScratch = ObjectsToReplace,
@@ -3399,7 +3408,7 @@ FBlake3Hash FBlueprintCompileReinstancer::CalculateConstructionHash(const UClass
 	// This can be impactful as we're frequently 'noop' reinstnacing level script actors.
 	TArray<UObject*> Objects;
 	constexpr bool bIncludeNestedObjects = false;
-	GetObjectsWithOuter(ForClass, Objects, bIncludeNestedObjects );
+	GetObjectsWithOuter(ForClass, Objects, bIncludeNestedObjects, RF_Transient, EInternalObjectFlags::Garbage );
 	Objects.SetNum(Algo::RemoveIf(Objects, [](const UObject* Obj)
 	{
 		return Obj->IsA<UFunction>();
