Sort subobjects of class so that we get stable results when comparing classes across compile, this improves stability of comparison for widget blueprints

#rnx

diff --git a/Engine/Source/Editor/UnrealEd/Private/Kismet2/KismetReinstanceUtilities.cpp b/Engine/Source/Editor/UnrealEd/Private/Kismet2/KismetReinstanceUtilities.cpp
--- a/Engine/Source/Editor/UnrealEd/Private/Kismet2/KismetReinstanceUtilities.cpp
+++ b/Engine/Source/Editor/UnrealEd/Private/Kismet2/KismetReinstanceUtilities.cpp
@@ -3413,6 +3413,7 @@ FBlake3Hash FBlueprintCompileReinstancer::CalculateConstructionHash(const UClass
 	{
 		return Obj->IsA<UFunction>();
 	}));
+	Objects.Sort([](const UObject& A, const UObject& B) { return A.GetFName().LexicalLess(B.GetFName()); });
 	Objects.Add(ForClass->ClassDefaultObject);
 	
 	FUtf8String TemplateData = UE::JsonObjectGraph::Stringify(
