[Viewport Toolbar] Add mouse+scroll wheel info to Camera Speed menu tooltip

diff --git a/Engine/Source/Editor/UnrealEd/Private/STransformViewportToolbar.cpp b/Engine/Source/Editor/UnrealEd/Private/STransformViewportToolbar.cpp
--- a/Engine/Source/Editor/UnrealEd/Private/STransformViewportToolbar.cpp
+++ b/Engine/Source/Editor/UnrealEd/Private/STransformViewportToolbar.cpp
@@ -493,7 +493,7 @@ TSharedRef< SWidget > STransformViewportToolBar::MakeTransformToolBar( const TSh
 			SNew(SEditorViewportToolbarMenu)
 				.ParentToolBar(SharedThis(this))
 				.AddMetaData<FTagMetaData>(FTagMetaData(TEXT("CameraSpeedButton")))
-				.ToolTipText(LOCTEXT("CameraSpeed_ToolTip", "Camera Speed"))
+				.ToolTipText(UE::UnrealEd::GetCameraSpeedTooltip())
 				.LabelIcon(FAppStyle::Get().GetBrush("EditorViewport.CamSpeedSetting"))
 				.Label(this, &STransformViewportToolBar::GetCameraSpeedLabel)
 				// Anchor to the right, otherwise the slider in this menu will jitter when the label width changes
diff --git a/Engine/Source/Editor/UnrealEd/Private/ViewportToolbar/UnrealEdViewportToolbar.cpp b/Engine/Source/Editor/UnrealEd/Private/ViewportToolbar/UnrealEdViewportToolbar.cpp
--- a/Engine/Source/Editor/UnrealEd/Private/ViewportToolbar/UnrealEdViewportToolbar.cpp
+++ b/Engine/Source/Editor/UnrealEd/Private/ViewportToolbar/UnrealEdViewportToolbar.cpp
@@ -3317,6 +3317,11 @@ FToolMenuEntry CreateScalabilitySubmenu()
 	);
 }
 
+FText GetCameraSpeedTooltip()
+{
+	return LOCTEXT("CameraSpeedSubMenu_ToolTip", "Set the camera speed.\nShortcut: Hold either mouse button and use the scroll wheel.");
+}
+
 bool ShouldShowViewportRealtimeWarning(const FEditorViewportClient& ViewportClient)
 {
 	return !ViewportClient.IsRealtime() && !ViewportClient.IsRealtimeOverrideSet() && ViewportClient.IsPerspective();
@@ -3698,7 +3703,7 @@ FToolMenuEntry CreateCameraSpeedMenu()
 				FToolMenuEntry& Entry = InDynamicSection.AddSubMenu(
 					"CameraSpeed",
 					LOCTEXT("CameraSpeedSubMenu", "Camera Speed"),
-					LOCTEXT("CameraSpeedSubMenu_ToolTip", "Camera Speed related actions"),
+					UE::UnrealEd::GetCameraSpeedTooltip(),
 					FNewToolMenuDelegate::CreateLambda(
 						[](UToolMenu* InMenu)
 						{
diff --git a/Engine/Source/Editor/UnrealEd/Public/ViewportToolbar/UnrealEdViewportToolbar.h b/Engine/Source/Editor/UnrealEd/Public/ViewportToolbar/UnrealEdViewportToolbar.h
--- a/Engine/Source/Editor/UnrealEd/Public/ViewportToolbar/UnrealEdViewportToolbar.h
+++ b/Engine/Source/Editor/UnrealEd/Public/ViewportToolbar/UnrealEdViewportToolbar.h
@@ -203,6 +203,8 @@ UNREALED_API FToolMenuEntry CreateScreenPercentageSubmenu();
 
 UNREALED_API FToolMenuEntry CreateScalabilitySubmenu();
 
+UNREALED_API FText GetCameraSpeedTooltip();
+
 // Camera Menu Widgets
 TSharedRef<SWidget> CreateCameraMenuWidget(const TSharedRef<SEditorViewport>& InViewport, bool bInShowExposureSettings = false);
 TSharedRef<SWidget> CreateFOVMenuWidget(const TSharedRef<SEditorViewport>& InViewport);
