[Viewport Toolbar] Polish: Write and add tooltips for checkbox submenus in the Snapping viewport toolbar menu

diff --git a/Engine/Source/Editor/UnrealEd/Private/EditorViewportCommands.cpp b/Engine/Source/Editor/UnrealEd/Private/EditorViewportCommands.cpp
--- a/Engine/Source/Editor/UnrealEd/Private/EditorViewportCommands.cpp
+++ b/Engine/Source/Editor/UnrealEd/Private/EditorViewportCommands.cpp
@@ -187,11 +187,11 @@ void FEditorViewportCommands::RegisterCommands()
 	UI_COMMAND( FocusOutlinerToSelection, "Frame Selected in Outliner", "Focuses the selected actors in all open outliners", EUserInterfaceActionType::Button, FInputChord() );
 	UI_COMMAND( FocusOutlinerToContextFolder, "Frame Current Folder in Outliner", "Focuses the current folder (if it's set) in the outliner", EUserInterfaceActionType::Button, FInputChord() );
 
-	UI_COMMAND( LocationGridSnap, "Grid Snap", "Enables or disables snapping to the grid when dragging objects around", EUserInterfaceActionType::ToggleButton, FInputChord() );
-	UI_COMMAND( RotationGridSnap, "Rotation Snap", "Enables or disables snapping objects to a rotation grid", EUserInterfaceActionType::ToggleButton, FInputChord() );
+	UI_COMMAND( LocationGridSnap, "Grid Snap", "Toggle whether objects snap to a grid when you move them.\nUse the submenu to choose the grid size.", EUserInterfaceActionType::ToggleButton, FInputChord() );
+	UI_COMMAND( RotationGridSnap, "Rotation Snap", "Toggle whether objects snap to preset angles when you rotate them.\nUse the submenu to choose the angle increment.", EUserInterfaceActionType::ToggleButton, FInputChord() );
 	UI_COMMAND( Layer2DSnap, "Layer2D Snap", "Enables or disables snapping objects to a 2D layer", EUserInterfaceActionType::ToggleButton, FInputChord());
-	UI_COMMAND( ScaleGridSnap, "Scale Snap", "Enables or disables snapping objects to a scale grid", EUserInterfaceActionType::ToggleButton, FInputChord() );
-	UI_COMMAND( SurfaceSnapping, "Surface Snapping", "If enabled, actors will snap to surfaces in the world when dragging", EUserInterfaceActionType::ToggleButton, FInputChord() );
+	UI_COMMAND( ScaleGridSnap, "Scale Snap", "Toggle whether objects resize by a preset ratio when you scale them.\nUse the submenu to choose the scaling increment.", EUserInterfaceActionType::ToggleButton, FInputChord() );
+	UI_COMMAND( SurfaceSnapping, "Surface Snapping", "Toggle whether objects snap to surfaces when you drag them.\nUse the submenu to set up the snapping behavior.", EUserInterfaceActionType::ToggleButton, FInputChord() );
 	UI_COMMAND( RotateToSurfaceNormal, "Rotate to Surface Normal", "If enabled, actors rotation will match that of their snapping surface", EUserInterfaceActionType::ToggleButton, FInputChord() );
 
 	UI_COMMAND( ToggleAutoExposure, "Auto", "If enabled, enables automatic exposure", EUserInterfaceActionType::ToggleButton, FInputChord() );
