Disable TAA when in landscape mode : landscape brushes don't play well with temporal algorithms

diff --git a/Engine/Source/Editor/UnrealEd/Private/EditorViewportClient.cpp b/Engine/Source/Editor/UnrealEd/Private/EditorViewportClient.cpp
--- a/Engine/Source/Editor/UnrealEd/Private/EditorViewportClient.cpp
+++ b/Engine/Source/Editor/UnrealEd/Private/EditorViewportClient.cpp
@@ -78,6 +78,7 @@
 #include "GlobalRenderResources.h"
 #include "Settings/EditorStyleSettings.h"
 #include "GameFramework/ActorPrimitiveColorHandler.h"
+#include "LandscapeRender.h"
 
 #define LOCTEXT_NAMESPACE "EditorViewportClient"
 
@@ -4263,6 +4264,11 @@ void FEditorViewportClient::Draw(FViewport* InViewport, FCanvas* Canvas)
 	ViewFamily.ExposureSettings = ExposureSettings;
 
 	ViewFamily.LandscapeLODOverride = LandscapeLODOverride;
+	// Landscape brushes don't play well with TAA : 
+	if (GLandscapeEditModeActive)
+	{
+		ViewFamily.EngineShowFlags.SetTemporalAA(false);
+	}
 
 	// Setup the screen percentage and upscaling method for the view family.
 	{
