- Factorize the code to write file package from the class LooseCookedPackageWriter in FFilePackageWriterUtil.
- New package writer FDefaultCookedFilePackageWriter used to write a package to a given filename.

#rnx

diff --git a/Engine/Source/Editor/UnrealEd/Private/Cooker/LooseCookedPackageWriter.cpp b/Engine/Source/Editor/UnrealEd/Private/Cooker/LooseCookedPackageWriter.cpp
--- a/Engine/Source/Editor/UnrealEd/Private/Cooker/LooseCookedPackageWriter.cpp
+++ b/Engine/Source/Editor/UnrealEd/Private/Cooker/LooseCookedPackageWriter.cpp
@@ -2,58 +2,25 @@
 
 #include "Cooker/LooseCookedPackageWriter.h"
 
-#include "AssetRegistry/AssetData.h"
 #include "AssetRegistry/AssetRegistryState.h"
 #include "AssetRegistry/IAssetRegistry.h"
 #include "Async/ParallelFor.h"
-#include "Containers/ContainersFwd.h"
-#include "Containers/Set.h"
-#include "Containers/StringView.h"
 #include "Cooker/AsyncIODelete.h"
 #include "Cooker/CompactBinaryTCP.h"
-#include "Cooker/CookPackageData.h"
 #include "Cooker/CookSandbox.h"
-#include "Cooker/CookTypes.h"
 #include "HAL/FileManager.h"
-#include "HAL/PlatformCrt.h"
-#include "HAL/PlatformFile.h"
 #include "HAL/PlatformFileManager.h"
-#include "HAL/PlatformMisc.h"
-#include "HAL/PlatformString.h"
-#include "Hash/Blake3.h"
 #include "Interfaces/ITargetPlatform.h"
-#include "IO/IoHash.h"
-#include "Logging/LogCategory.h"
-#include "Logging/LogMacros.h"
 #include "LooseFilesCookArtifactReader.h"
-#include "Misc/App.h"
-#include "Misc/AssertionMacros.h"
-#include "Misc/CString.h"
-#include "Misc/EnumClassFlags.h"
 #include "Misc/FileHelper.h"
-#include "Misc/Optional.h"
-#include "Misc/PackageName.h"
-#include "Misc/PackagePath.h"
 #include "Misc/Paths.h"
 #include "Misc/PathViews.h"
-#include "Misc/ScopeLock.h"
-#include "Misc/SecureHash.h"
-#include "Misc/StringBuilder.h"
 #include "PackageStoreOptimizer.h"
-#include "ProfilingDebugging/CpuProfilerTrace.h"
-#include "Serialization/Archive.h"
 #include "Serialization/ArchiveStackTrace.h"
-#include "Serialization/ArrayReader.h"
+#include "Serialization/CompactBinarySerialization.h"
 #include "Serialization/CompactBinaryWriter.h"
+#include "Serialization/FilePackageWriterUtil.h"
 #include "Serialization/LargeMemoryWriter.h"
-#include "Serialization/MemoryWriter.h"
-#include "String/Find.h"
-#include "Tasks/Task.h"
-#include "Templates/RefCounting.h"
-#include "Templates/UnrealTemplate.h"
-#include "Trace/Detail/Channel.h"
-#include "UObject/ObjectVersion.h"
-#include "UObject/Package.h"
 #include "UObject/SavePackage.h"
 
 LLM_DEFINE_TAG(Cooker_PackageStoreManifest);
@@ -89,43 +56,25 @@ void FLooseCookedPackageWriter::BeginPackage(const FBeginPackageInfo& Info)
 	PackageInfo.BulkDataChunks.Reset();
 }
 
-void FLooseCookedPackageWriter::CommitPackageInternal(FPackageWriterRecords::FPackage&& BaseRecord,
-	const FCommitPackageInfo& Info)
+void FLooseCookedPackageWriter::CommitPackageInternal(FPackageWriterRecords::FPackage&& BaseRecord, const FCommitPackageInfo& Info)
 {
-	FRecord& Record = static_cast<FRecord&>(BaseRecord);
-	if (Info.Status == IPackageWriter::ECommitStatus::Success)
 	{
-		AsyncSave(Record, Info);
+		FFilePackageWriterUtil::FRecord& Record = static_cast<FFilePackageWriterUtil::FRecord&>(BaseRecord);
+		FFilePackageWriterUtil::FWritePackageParameters Parameters(Record, Info, &AllPackageHashes, &PackageHashesLock, bProvidePerPackageResults);
+		FFilePackageWriterUtil::WritePackage(Parameters);
 	}
-	if (Info.Status != IPackageWriter::ECommitStatus::Canceled &&
-		Info.Status != IPackageWriter::ECommitStatus::NotCommitted)
+
+	if (Info.Status != IPackageWriter::ECommitStatus::Canceled && Info.Status != IPackageWriter::ECommitStatus::NotCommitted)
 	{
+		FPackageWriterRecords::FPackage& Record = static_cast<FPackageWriterRecords::FPackage&>(BaseRecord);
 		UpdateManifest(Record);
 	}
 }
 
-void FLooseCookedPackageWriter::AsyncSave(FRecord& Record, const FCommitPackageInfo& Info)
-{
-	FCommitContext Context{ Info };
-
-	// The order of these collection calls is important, both for ExportsBuffers (affects the meaning of offsets
-	// to those buffers) and for OutputFiles (affects the calculation of the Hash for the set of PackageData)
-	// The order of ExportsBuffers must match CompleteExportsArchiveForDiff.
-	CollectForSavePackageData(Record, Context);
-	CollectForSaveBulkData(Record, Context);
-	CollectForSaveLinkerAdditionalDataRecords(Record, Context);
-	CollectForSaveAdditionalFileRecords(Record, Context);
-	CollectForSaveExportsFooter(Record, Context);
-	CollectForSaveExportsPackageTrailer(Record, Context);
-	CollectForSaveExportsBuffers(Record, Context);
-
-	AsyncSaveOutputFiles(Record, Context);
-}
-
 void FLooseCookedPackageWriter::CompleteExportsArchiveForDiff(FPackageInfo& Info, FLargeMemoryWriter& ExportsArchive)
 {
 	FPackageWriterRecords::FPackage& BaseRecord = Records.FindRecordChecked(Info.PackageName);
-	FRecord& Record = static_cast<FRecord&>(BaseRecord);
+	FFilePackageWriterUtil::FRecord& Record = static_cast<FFilePackageWriterUtil::FRecord&>(BaseRecord);
 	Record.bCompletedExportsArchiveForDiff = true;
 
 	// Add on all the attachments which are usually added on during Commit. The order must match AsyncSave.
@@ -159,347 +108,18 @@ void FLooseCookedPackageWriter::CompleteExportsArchiveForDiff(FPackageInfo& Info
 	}
 }
 
-
-void FLooseCookedPackageWriter::CollectForSavePackageData(FRecord& Record, FCommitContext& Context)
-{
-	Context.ExportsBuffers.AddDefaulted(Record.Packages.Num());
-	for (FPackageWriterRecords::FWritePackage& Package : Record.Packages)
-	{
-		Context.ExportsBuffers[Package.Info.MultiOutputIndex].Add(FExportBuffer{ Package.Buffer, MoveTemp(Package.Regions) });
-	}
-}
-
-void FLooseCookedPackageWriter::CollectForSaveBulkData(FRecord& Record, FCommitContext& Context)
-{
-	for (FBulkDataRecord& BulkRecord : Record.BulkDatas)
-	{
-		if (BulkRecord.Info.BulkDataType == FBulkDataInfo::AppendToExports)
-		{
-			if (Record.bCompletedExportsArchiveForDiff)
-			{
-				// Already Added in CompleteExportsArchiveForDiff
-				continue;
-			}
-			Context.ExportsBuffers[BulkRecord.Info.MultiOutputIndex].Add(FExportBuffer{ BulkRecord.Buffer, MoveTemp(BulkRecord.Regions) });
-		}
-		else
-		{
-			FWriteFileData& OutputFile = Context.OutputFiles.Emplace_GetRef();
-			OutputFile.Filename = BulkRecord.Info.LooseFilePath;
-			OutputFile.Buffer = FCompositeBuffer(BulkRecord.Buffer);
-			OutputFile.Regions = MoveTemp(BulkRecord.Regions);
-			OutputFile.bIsSidecar = true;
-			OutputFile.bContributeToHash = BulkRecord.Info.MultiOutputIndex == 0; // Only caculate the main package output hash
-			OutputFile.ChunkId = BulkRecord.Info.ChunkId;
-		}
-	}
-}
-void FLooseCookedPackageWriter::CollectForSaveLinkerAdditionalDataRecords(FRecord& Record, FCommitContext& Context)
-{
-	if (Record.bCompletedExportsArchiveForDiff)
-	{
-		// Already Added in CompleteExportsArchiveForDiff
-		return;
-	}
-
-	for (FLinkerAdditionalDataRecord& AdditionalRecord : Record.LinkerAdditionalDatas)
-	{
-		Context.ExportsBuffers[AdditionalRecord.Info.MultiOutputIndex].Add(FExportBuffer{ AdditionalRecord.Buffer, MoveTemp(AdditionalRecord.Regions) });
-	}
-}
-
-void FLooseCookedPackageWriter::CollectForSaveAdditionalFileRecords(FRecord& Record, FCommitContext& Context)
-{
-	for (FAdditionalFileRecord& AdditionalRecord : Record.AdditionalFiles)
-	{
-		FWriteFileData& OutputFile = Context.OutputFiles.Emplace_GetRef();
-		OutputFile.Filename = AdditionalRecord.Info.Filename;
-		OutputFile.Buffer = FCompositeBuffer(AdditionalRecord.Buffer);
-		OutputFile.bIsSidecar = true;
-		OutputFile.bContributeToHash = AdditionalRecord.Info.MultiOutputIndex == 0; // Only calculate the main package output hash
-		OutputFile.ChunkId = AdditionalRecord.Info.ChunkId;
-	}
-}
-
-void FLooseCookedPackageWriter::CollectForSaveExportsFooter(FRecord& Record, FCommitContext& Context)
-{
-	if (Record.bCompletedExportsArchiveForDiff)
-	{
-		// Already Added in CompleteExportsArchiveForDiff
-		return;
-	}
-
-	uint32 FooterData = PACKAGE_FILE_TAG;
-	FSharedBuffer Buffer = FSharedBuffer::Clone(&FooterData, sizeof(FooterData));
-	for (FPackageWriterRecords::FWritePackage& Package : Record.Packages)
-	{
-		Context.ExportsBuffers[Package.Info.MultiOutputIndex].Add(FExportBuffer{ Buffer, TArray<FFileRegion>() });
-	}
-}
-
 int64 FLooseCookedPackageWriter::GetExportsFooterSize()
 {
 	return sizeof(uint32);
 }
 
-void FLooseCookedPackageWriter::CollectForSaveExportsPackageTrailer(FRecord& Record, FCommitContext& Context)
-{
-	if (Record.bCompletedExportsArchiveForDiff)
-	{
-		// Already Added in CompleteExportsArchiveForDiff
-		return;
-	}
-
-	for (FPackageTrailerRecord& PackageTrailer : Record.PackageTrailers)
-	{
-		Context.ExportsBuffers[PackageTrailer.Info.MultiOutputIndex].Add(
-			FExportBuffer{ PackageTrailer.Buffer, TArray<FFileRegion>() });
-	}
-}
-
-void FLooseCookedPackageWriter::CollectForSaveExportsBuffers(FRecord& Record, FCommitContext& Context)
-{
-	check(Context.ExportsBuffers.Num() == Record.Packages.Num());
-	for (FPackageWriterRecords::FWritePackage& Package : Record.Packages)
-	{
-		TArray<FExportBuffer>& ExportsBuffers = Context.ExportsBuffers[Package.Info.MultiOutputIndex];
-		check(ExportsBuffers.Num() > 0);
-
-		// Split the ExportsBuffer into (1) Header and (2) Exports + AllAppendedData
-		int64 HeaderSize = Package.Info.HeaderSize;
-		FExportBuffer& HeaderAndExportsBuffer = ExportsBuffers[0];
-		FSharedBuffer& HeaderAndExportsData = HeaderAndExportsBuffer.Buffer;
-
-		// Header (.uasset/.umap)
-		{
-			FWriteFileData& OutputFile = Context.OutputFiles.Emplace_GetRef();
-			OutputFile.Filename = Package.Info.LooseFilePath;
-			OutputFile.Buffer = FCompositeBuffer(
-				FSharedBuffer::MakeView(HeaderAndExportsData.GetData(), HeaderSize, HeaderAndExportsData));
-			OutputFile.bIsSidecar = false;
-			OutputFile.bContributeToHash = Package.Info.MultiOutputIndex == 0; // Only calculate the main package output hash
-		}
-
-		// Exports + AllAppendedData (.uexp)
-		{
-			FWriteFileData& OutputFile = Context.OutputFiles.Emplace_GetRef();
-			OutputFile.Filename = FPaths::ChangeExtension(Package.Info.LooseFilePath, LexToString(EPackageExtension::Exports));
-			OutputFile.bIsSidecar = false;
-			OutputFile.bContributeToHash = Package.Info.MultiOutputIndex == 0; // Only caculate the main package output hash
-
-			int32 NumBuffers = ExportsBuffers.Num();
-			TArray<FSharedBuffer> BuffersForComposition;
-			BuffersForComposition.Reserve(NumBuffers);
-
-			const uint8* ExportsStart = static_cast<const uint8*>(HeaderAndExportsData.GetData()) + HeaderSize;
-			BuffersForComposition.Add(FSharedBuffer::MakeView(ExportsStart, HeaderAndExportsData.GetSize() - HeaderSize,
-				HeaderAndExportsData));
-			OutputFile.Regions.Append(MoveTemp(HeaderAndExportsBuffer.Regions));
-
-			for (FExportBuffer& ExportsBuffer : TArrayView<FExportBuffer>(ExportsBuffers).Slice(1, NumBuffers - 1))
-			{
-				BuffersForComposition.Add(ExportsBuffer.Buffer);
-				OutputFile.Regions.Append(MoveTemp(ExportsBuffer.Regions));
-			}
-			OutputFile.Buffer = FCompositeBuffer(BuffersForComposition);
-
-			// Adjust regions so they are relative to the start of the uexp file
-			for (FFileRegion& Region : OutputFile.Regions)
-			{
-				Region.Offset -= HeaderSize;
-			}
-		}
-	}
-}
-
 FPackageWriterRecords::FPackage* FLooseCookedPackageWriter::ConstructRecord()
 {
-	return new FRecord();
+	//Watch out : if you change the type of record, you need to go through this file to make sure there is not a static cast that would not work anymore.
+	return new FFilePackageWriterUtil::FRecord();
 }
 
-void FLooseCookedPackageWriter::AsyncSaveOutputFiles(FRecord& Record, FCommitContext& Context)
-{
-	if (!EnumHasAnyFlags(Context.Info.WriteOptions, EWriteOptions::Write | EWriteOptions::ComputeHash))
-	{
-		return;
-	}
-
-	UE::SavePackageUtilities::IncrementOutstandingAsyncWrites();
-
-	TRefCountPtr<FPackageHashes> ThisPackageHashes;
-	TUniquePtr<TPromise<int>> PackageHashesCompletionPromise;
-	
-	if (EnumHasAnyFlags(Context.Info.WriteOptions, EWriteOptions::ComputeHash))
-	{
-		ThisPackageHashes = new FPackageHashes();
-		if (bProvidePerPackageResults)
-		{
-			PackageHashesCompletionPromise.Reset(new TPromise<int>());
-			ThisPackageHashes->CompletionFuture = PackageHashesCompletionPromise->GetFuture();
-		}
-
-		bool bAlreadyExisted = false;
-		{
-			FScopeLock PackageHashesScopeLock(&PackageHashesLock);
-			TRefCountPtr<FPackageHashes>& ExistingPackageHashes = AllPackageHashes.FindOrAdd(Context.Info.PackageName);
-			// This calculation of bAlreadyExisted looks weird but we're finding the _refcount_, not the hashes. So if it gets
-			// constructed, it's not actually assigned a pointer.
-			bAlreadyExisted = ExistingPackageHashes.IsValid();
-			ExistingPackageHashes = ThisPackageHashes;
-		}
-		if (bAlreadyExisted)
-		{
-			UE_LOG(LogSavePackage, Error, TEXT("FLooseCookedPackageWriter encountered the same package twice in a cook! (%s)"), *Context.Info.PackageName.ToString());
-		}
-	}
-
-	UE::Tasks::Launch(TEXT("HashAndWriteLooseCookedFile"),
-		[OutputFiles = MoveTemp(Context.OutputFiles), WriteOptions = Context.Info.WriteOptions,
-		ThisPackageHashes = MoveTemp(ThisPackageHashes), PackageHashesCompletionPromise=MoveTemp(PackageHashesCompletionPromise)]
-		() mutable
-		{
-			FMD5 AccumulatedHash;
-			for (const FWriteFileData& OutputFile : OutputFiles)
-			{
-				OutputFile.HashAndWrite(AccumulatedHash, ThisPackageHashes, WriteOptions);
-			}
-
-			if (EnumHasAnyFlags(WriteOptions, EWriteOptions::ComputeHash))
-			{
-				ThisPackageHashes->PackageHash.Set(AccumulatedHash);
-			}
-
-			if (PackageHashesCompletionPromise)
-			{
-				// Note that setting this Promise might call arbitrary code from anything that subscribed
-				// to ThisPackageHashes->CompletionFuture.Then(). So don't call it inside a lock.
-				PackageHashesCompletionPromise->SetValue(0);
-			}
-
-			// This is used to release the game thread to access the hashes
-			UE::SavePackageUtilities::DecrementOutstandingAsyncWrites();
-		},
-		UE::Tasks::ETaskPriority::BackgroundNormal
-	);
-}
-
-static void WriteToFile(const FString& Filename, const FCompositeBuffer& Buffer)
-{
-	IFileManager& FileManager = IFileManager::Get();
-
-	struct FFailureReason
-	{
-		uint32 LastErrorCode = 0;
-		bool bSizeMatchFailed = false;
-		int64 ExpectedSize = 0;
-		int64 ActualSize = 0;
-		bool bArchiveError = false;
-	};
-	TOptional<FFailureReason> FailureReason;
-
-	for (int32 Tries = 0; Tries < 3; ++Tries)
-	{
-		FArchive* Ar = FileManager.CreateFileWriter(*Filename);
-		if (!Ar)
-		{
-			if (!FailureReason)
-			{
-				FailureReason = FFailureReason{ FPlatformMisc::GetLastError(), false };
-			}
-			continue;
-		}
-
-		int64 DataSize = 0;
-		for (const FSharedBuffer& Segment : Buffer.GetSegments())
-		{
-			int64 SegmentSize = static_cast<int64>(Segment.GetSize());
-			Ar->Serialize(const_cast<void*>(Segment.GetData()), SegmentSize);
-			DataSize += SegmentSize;
-		}
-		bool bArchiveError = Ar->IsError();
-		delete Ar;
-
-		int64 ActualSize = FileManager.FileSize(*Filename);
-		if (ActualSize != DataSize)
-		{
-			if (!FailureReason)
-			{
-				FailureReason = FFailureReason{ 0, true, DataSize, ActualSize, bArchiveError };
-			}
-			FileManager.Delete(*Filename);
-			continue;
-		}
-		return;
-	}
-
-	FString ReasonText;
-	if (FailureReason && FailureReason->bSizeMatchFailed)
-	{
-		ReasonText = FString::Printf(TEXT("Unexpected file size. Tried to write %" INT64_FMT " but resultant size was %" INT64_FMT ".%s")
-			TEXT(" Another operation is modifying the file, or the write operation failed to write completely."),
-			FailureReason->ExpectedSize, FailureReason->ActualSize, FailureReason->bArchiveError ? TEXT(" Ar->Serialize failed.") : TEXT(""));
-	}
-	else if (FailureReason && FailureReason->LastErrorCode != 0)
-	{
-		TCHAR LastErrorText[1024];
-		FPlatformMisc::GetSystemErrorMessage(LastErrorText, UE_ARRAY_COUNT(LastErrorText), FailureReason->LastErrorCode);
-		ReasonText = LastErrorText;
-	}
-	else
-	{
-		ReasonText = TEXT("Unknown failure reason.");
-	}
-	UE_LOG(LogSavePackage, Fatal, TEXT("SavePackage Async write %s failed: %s"), *Filename, *ReasonText);
-}
-
-void FLooseCookedPackageWriter::FWriteFileData::HashAndWrite(FMD5& AccumulatedHash, const TRefCountPtr<FPackageHashes>& PackageHashes, EWriteOptions WriteOptions) const
-{
-	//@todo: FH: Should we calculate the hash of both output, currently only the main package output hash is calculated
-	if (EnumHasAnyFlags(WriteOptions, EWriteOptions::ComputeHash) && bContributeToHash)
-	{
-		for (const FSharedBuffer& Segment : Buffer.GetSegments())
-		{
-			AccumulatedHash.Update(static_cast<const uint8*>(Segment.GetData()), Segment.GetSize());
-		}
-
-		if (ChunkId.IsValid())
-		{
-			FBlake3 ChunkHash;
-			for (const FSharedBuffer& Segment : Buffer.GetSegments())
-			{
-				ChunkHash.Update(static_cast<const uint8*>(Segment.GetData()), Segment.GetSize());
-			}
-			FIoHash FinalHash(ChunkHash.Finalize());
-			PackageHashes->ChunkHashes.Add(ChunkId, FinalHash);
-		}
-	}
-
-	if ((bIsSidecar && EnumHasAnyFlags(WriteOptions, EWriteOptions::WriteSidecars)) ||
-		(!bIsSidecar && EnumHasAnyFlags(WriteOptions, EWriteOptions::WritePackage)))
-	{
-		const FString* WriteFilename = &Filename;
-		FString FilenameBuffer;
-		if (EnumHasAnyFlags(WriteOptions, EWriteOptions::SaveForDiff))
-		{
-			FilenameBuffer = FPaths::Combine(FPaths::GetPath(Filename),
-				FPaths::GetBaseFilename(Filename) + TEXT("_ForDiff") + FPaths::GetExtension(Filename, true));
-			WriteFilename = &FilenameBuffer;
-		}
-		WriteToFile(*WriteFilename, Buffer);
-
-		if (Regions.Num() > 0)
-		{
-			TArray<uint8> Memory;
-			FMemoryWriter Ar(Memory);
-			FFileRegion::SerializeFileRegions(Ar, const_cast<TArray<FFileRegion>&>(Regions));
-
-			WriteToFile(*WriteFilename + FFileRegion::RegionsFileExtension,
-				FCompositeBuffer(FSharedBuffer::MakeView(Memory.GetData(), Memory.Num())));
-		}
-	}
-}
-
-void FLooseCookedPackageWriter::UpdateManifest(FRecord& Record)
+void FLooseCookedPackageWriter::UpdateManifest(const FPackageWriterRecords::FPackage& Record)
 {
 	LLM_SCOPE_BYTAG(Cooker_PackageStoreManifest);
 	FScopeLock Lock(&OplogLock);
@@ -641,6 +261,11 @@ bool FLooseCookedPackageWriter::ReadOplogEntry(FOplogPackageInfo& PackageInfo, c
 	return true;
 }
 
+TMap<FName, TRefCountPtr<FPackageHashes>>& FLooseCookedPackageWriter::GetPackageHashes()
+{
+	return AllPackageHashes;
+}
+
 void FLooseCookedPackageWriter::BeginCook(const FCookInfo& Info)
 {
 	if (!Info.bWorkerOnSharedSandbox)
diff --git a/Engine/Source/Editor/UnrealEd/Private/Cooker/LooseCookedPackageWriter.h b/Engine/Source/Editor/UnrealEd/Private/Cooker/LooseCookedPackageWriter.h
--- a/Engine/Source/Editor/UnrealEd/Private/Cooker/LooseCookedPackageWriter.h
+++ b/Engine/Source/Editor/UnrealEd/Private/Cooker/LooseCookedPackageWriter.h
@@ -2,25 +2,10 @@
 
 #pragma once
 
-#include "Containers/Array.h"
-#include "Containers/ArrayView.h"
-#include "Containers/Map.h"
-#include "Containers/StringFwd.h"
-#include "Containers/UnrealString.h"
-#include "HAL/CriticalSection.h"
-#include "HAL/Platform.h"
-#include "IO/IoDispatcher.h"
-#include "Memory/CompositeBuffer.h"
-#include "Memory/SharedBuffer.h"
-#include "Misc/DateTime.h"
+#include "IO/IoChunkId.h"
 #include "Misc/PackagePath.h"
-#include "Serialization/CompactBinary.h"
-#include "Serialization/FileRegions.h"
-#include "Serialization/PackageWriter.h"
 #include "Serialization/PackageWriterToSharedBuffer.h"
 #include "Templates/SharedPointer.h"
-#include "Templates/UniquePtr.h"
-#include "UObject/NameTypes.h"
 
 class FAssetRegistryState;
 class FAsyncIODelete;
@@ -80,51 +65,13 @@ public:
 	virtual void CompleteExportsArchiveForDiff(FPackageInfo& Info, FLargeMemoryWriter& ExportsArchive) override;
 	virtual EPackageWriterResult BeginCacheForCookedPlatformData(FBeginCacheForCookedPlatformDataInfo& Info) override;
 
-	virtual void CommitPackageInternal(FPackageWriterRecords::FPackage&& BaseRecord,
-		const FCommitPackageInfo& Info) override;
+	virtual void CommitPackageInternal(FPackageWriterRecords::FPackage&& BaseRecord, const FCommitPackageInfo& Info) override;
 	virtual FPackageWriterRecords::FPackage* ConstructRecord() override;
 
 	static EPackageExtension BulkDataTypeToExtension(FBulkDataInfo::EType BulkDataType);
 
 private:
 
-	/** Version of the superclass's per-package record that includes our class-specific data. */
-	struct FRecord : public FPackageWriterRecords::FPackage
-	{
-		bool bCompletedExportsArchiveForDiff = false;
-	};
-
-	/** Buffers that are combined into the HeaderAndExports file (which is then split into .uasset + .uexp or .uoasset + .uoexp). */
-	struct FExportBuffer
-	{
-		FSharedBuffer Buffer;
-		TArray<FFileRegion> Regions;
-	};
-
-	/**
-	 * The data needed to asynchronously write one of the files (.uasset, .uexp, .ubulk, any optional and any additional),
-	 * without reference back to other data on this writer.
-	 */
-	struct FWriteFileData
-	{
-		FString Filename;
-		FCompositeBuffer Buffer;
-		TArray<FFileRegion> Regions;
-		bool bIsSidecar;
-		bool bContributeToHash = true;
-		FIoChunkId ChunkId = FIoChunkId::InvalidChunkId;
-
-		void HashAndWrite(FMD5& AccumulatedHash, const TRefCountPtr<FPackageHashes>& PackageHashes, EWriteOptions WriteOptions) const;
-	};
-
-	/** Stack data for the helper functions of CommitPackageInternal. */
-	struct FCommitContext
-	{
-		const FCommitPackageInfo& Info;
-		TArray<TArray<FExportBuffer>> ExportsBuffers;
-		TArray<FWriteFileData> OutputFiles;
-	};
-
 	struct FOplogChunkInfo
 	{
 		FString RelativeFileName;
@@ -145,24 +92,12 @@ private:
 	void FindAndDeleteCookedFilesForPackages(TConstArrayView<FName> PackageNames);
 
 	void RemoveCookedPackagesByPackageName(TArrayView<const FName> PackageNamesToRemove, bool bRemoveRecords);
-	void AsyncSave(FRecord& Record, const FCommitPackageInfo& Info);
-
-	void CollectForSavePackageData(FRecord& Record, FCommitContext& Context);
-	void CollectForSaveBulkData(FRecord& Record, FCommitContext& Context);
-	void CollectForSaveLinkerAdditionalDataRecords(FRecord& Record, FCommitContext& Context);
-	void CollectForSaveAdditionalFileRecords(FRecord& Record, FCommitContext& Context);
-	void CollectForSaveExportsFooter(FRecord& Record, FCommitContext& Context);
-	void CollectForSaveExportsPackageTrailer(FRecord& Record, FCommitContext& Context);
-	void CollectForSaveExportsBuffers(FRecord& Record, FCommitContext& Context);
-	void AsyncSaveOutputFiles(FRecord& Record, FCommitContext& Context);
-	void UpdateManifest(FRecord& Record);
+	
+	void UpdateManifest(const FPackageWriterRecords::FPackage& Record);
 	void WriteOplogEntry(FCbWriter& Writer, const FOplogPackageInfo& PackageInfo);
 	bool ReadOplogEntry(FOplogPackageInfo& PackageInfo, const FCbFieldView& Field);
 
-	TMap<FName, TRefCountPtr<FPackageHashes>>& GetPackageHashes() override
-	{
-		return AllPackageHashes;
-	}
+	TMap<FName, TRefCountPtr<FPackageHashes>>& GetPackageHashes() override;
 
 	TSharedRef<FLooseFilesCookArtifactReader> CookArtifactReader;
 
diff --git a/Engine/Source/Editor/UnrealEd/Private/Serialization/DefaultCookedFilePackageWriter.cpp b/Engine/Source/Editor/UnrealEd/Private/Serialization/DefaultCookedFilePackageWriter.cpp
new file mode 100644
--- /dev/null
+++ b/Engine/Source/Editor/UnrealEd/Private/Serialization/DefaultCookedFilePackageWriter.cpp
@@ -0,0 +1,46 @@
+// Copyright Epic Games, Inc. All Rights Reserved.
+
+#include "Serialization/DefaultCookedFilePackageWriter.h"
+
+#include "AssetRegistry/AssetRegistryState.h"
+#include "Misc/Paths.h"
+#include "Serialization/FilePackageWriterUtil.h"
+
+FDefaultCookedFilePackageWriter::FDefaultCookedFilePackageWriter(const FString& BaseFilename)
+	: TPackageWriterToSharedBuffer<IPackageWriter>()
+	, BaseFilename(BaseFilename)
+{}
+
+void FDefaultCookedFilePackageWriter::WritePackageData(const IPackageWriter::FPackageInfo& Info, FLargeMemoryWriter& ExportsArchive, const TArray<FFileRegion>& FileRegions)
+{
+	//This is working because Super::WritePackageData will store a copy of InfoCopy. But it's sketchy!
+	IPackageWriter::FPackageInfo InfoCopy = Info;
+
+	FString Extension = FPaths::GetExtension(Info.LooseFilePath, false);
+	InfoCopy.LooseFilePath = FPaths::SetExtension(BaseFilename, Extension);
+	
+	Super::WritePackageData(InfoCopy, ExportsArchive, FileRegions);
+}
+
+void FDefaultCookedFilePackageWriter::WriteBulkData(const IPackageWriter::FBulkDataInfo& Info, const FIoBuffer& BulkData, const TArray<FFileRegion>& FileRegions)
+{
+	//This is working because Super::WriteBulkData will store a copy of InfoCopy. But it's sketchy!
+	IPackageWriter::FBulkDataInfo InfoCopy = Info;
+
+	FString Extension = FPaths::GetExtension(Info.LooseFilePath, false);
+	InfoCopy.LooseFilePath = FPaths::SetExtension(BaseFilename, Extension);
+
+	Super::WriteBulkData(InfoCopy, BulkData, FileRegions);
+}
+
+void FDefaultCookedFilePackageWriter::CommitPackageInternal(FPackageWriterRecords::FPackage&& BaseRecord, const FCommitPackageInfo& Info)
+{
+	FFilePackageWriterUtil::FRecord& Record = static_cast<FFilePackageWriterUtil::FRecord&>(BaseRecord);
+	FFilePackageWriterUtil::FWritePackageParameters Parameters(Record, Info, nullptr, nullptr, false);
+	FFilePackageWriterUtil::WritePackage(Parameters);
+}
+
+FPackageWriterRecords::FPackage* FDefaultCookedFilePackageWriter::ConstructRecord()
+{
+	return new FFilePackageWriterUtil::FRecord();
+}
diff --git a/Engine/Source/Editor/UnrealEd/Private/Serialization/FilePackageWriterUtil.cpp b/Engine/Source/Editor/UnrealEd/Private/Serialization/FilePackageWriterUtil.cpp
new file mode 100644
--- /dev/null
+++ b/Engine/Source/Editor/UnrealEd/Private/Serialization/FilePackageWriterUtil.cpp
@@ -0,0 +1,395 @@
+// Copyright Epic Games, Inc. All Rights Reserved.
+
+#include "Serialization/FilePackageWriterUtil.h"
+
+#include "HAL/FileManager.h"
+#include "Misc/PackagePath.h"
+#include "Misc/Paths.h"
+#include "Serialization/MemoryWriter.h"
+#include "UObject/SavePackage.h"
+
+namespace FilePackageWriterUtil
+{
+	static void WriteToFile(const FString& Filename, const FCompositeBuffer& Buffer)
+	{
+		IFileManager& FileManager = IFileManager::Get();
+
+		struct FFailureReason
+		{
+			uint32 LastErrorCode = 0;
+			bool bSizeMatchFailed = false;
+			int64 ExpectedSize = 0;
+			int64 ActualSize = 0;
+			bool bArchiveError = false;
+		};
+		TOptional<FFailureReason> FailureReason;
+
+		for (int32 Tries = 0; Tries < 3; ++Tries)
+		{
+			FArchive* Ar = FileManager.CreateFileWriter(*Filename);
+			if (!Ar)
+			{
+				if (!FailureReason)
+				{
+					FailureReason = FFailureReason{ FPlatformMisc::GetLastError(), false };
+				}
+				continue;
+			}
+
+			int64 DataSize = 0;
+			for (const FSharedBuffer& Segment : Buffer.GetSegments())
+			{
+				int64 SegmentSize = static_cast<int64>(Segment.GetSize());
+				Ar->Serialize(const_cast<void*>(Segment.GetData()), SegmentSize);
+				DataSize += SegmentSize;
+			}
+			bool bArchiveError = Ar->IsError();
+			delete Ar;
+
+			int64 ActualSize = FileManager.FileSize(*Filename);
+			if (ActualSize != DataSize)
+			{
+				if (!FailureReason)
+				{
+					FailureReason = FFailureReason{ 0, true, DataSize, ActualSize, bArchiveError };
+				}
+				FileManager.Delete(*Filename);
+				continue;
+			}
+			return;
+		}
+
+		FString ReasonText;
+		if (FailureReason && FailureReason->bSizeMatchFailed)
+		{
+			ReasonText = FString::Printf(TEXT("Unexpected file size. Tried to write %" INT64_FMT " but resultant size was %" INT64_FMT ".%s")
+				TEXT(" Another operation is modifying the file, or the write operation failed to write completely."),
+				FailureReason->ExpectedSize, FailureReason->ActualSize, FailureReason->bArchiveError ? TEXT(" Ar->Serialize failed.") : TEXT(""));
+		}
+		else if (FailureReason && FailureReason->LastErrorCode != 0)
+		{
+			TCHAR LastErrorText[1024];
+			FPlatformMisc::GetSystemErrorMessage(LastErrorText, UE_ARRAY_COUNT(LastErrorText), FailureReason->LastErrorCode);
+			ReasonText = LastErrorText;
+		}
+		else
+		{
+			ReasonText = TEXT("Unknown failure reason.");
+		}
+		UE_LOG(LogSavePackage, Fatal, TEXT("SavePackage Async write %s failed: %s"), *Filename, *ReasonText);
+	}
+}
+
+FFilePackageWriterUtil::FWritePackageParameters::FWritePackageParameters(FRecord& InRecord,
+	const IPackageWriter::FCommitPackageInfo& InInfo,
+	TMap<FName, TRefCountPtr<FPackageHashes>>* InAllPackageHashes,
+	FCriticalSection* InPackageHashesLock,
+	bool bInProvidePerPackageResult)
+	: Record(InRecord)
+	, Info(InInfo)
+	, AllPackageHashes(InAllPackageHashes)
+	, PackageHashesLock(InPackageHashesLock)
+	, bProvidePerPackageResult(bInProvidePerPackageResult)
+{}
+
+void FFilePackageWriterUtil::FWriteFileData::HashAndWrite(FMD5& AccumulatedHash, const TRefCountPtr<FPackageHashes>& PackageHashes, IPackageWriter::EWriteOptions WriteOptions) const
+{
+	//@todo: FH: Should we calculate the hash of both output, currently only the main package output hash is calculated
+	if (EnumHasAnyFlags(WriteOptions, IPackageWriter::EWriteOptions::ComputeHash) && bContributeToHash)
+	{
+		for (const FSharedBuffer& Segment : Buffer.GetSegments())
+		{
+			AccumulatedHash.Update(static_cast<const uint8*>(Segment.GetData()), Segment.GetSize());
+		}
+
+		if (ChunkId.IsValid())
+		{
+			FBlake3 ChunkHash;
+			for (const FSharedBuffer& Segment : Buffer.GetSegments())
+			{
+				ChunkHash.Update(static_cast<const uint8*>(Segment.GetData()), Segment.GetSize());
+			}
+			FIoHash FinalHash(ChunkHash.Finalize());
+			PackageHashes->ChunkHashes.Add(ChunkId, FinalHash);
+		}
+	}
+
+	if ((bIsSidecar && EnumHasAnyFlags(WriteOptions, IPackageWriter::EWriteOptions::WriteSidecars)) ||
+		(!bIsSidecar && EnumHasAnyFlags(WriteOptions, IPackageWriter::EWriteOptions::WritePackage)))
+	{
+		const FString* WriteFilename = &Filename;
+		FString FilenameBuffer;
+		if (EnumHasAnyFlags(WriteOptions, IPackageWriter::EWriteOptions::SaveForDiff))
+		{
+			FilenameBuffer = FPaths::Combine(FPaths::GetPath(Filename),
+				FPaths::GetBaseFilename(Filename) + TEXT("_ForDiff") + FPaths::GetExtension(Filename, true));
+			WriteFilename = &FilenameBuffer;
+		}
+		FilePackageWriterUtil::WriteToFile(*WriteFilename, Buffer);
+
+		if (Regions.Num() > 0)
+		{
+			TArray<uint8> Memory;
+			FMemoryWriter Ar(Memory);
+			FFileRegion::SerializeFileRegions(Ar, const_cast<TArray<FFileRegion>&>(Regions));
+
+			FilePackageWriterUtil::WriteToFile(*WriteFilename + FFileRegion::RegionsFileExtension,
+				FCompositeBuffer(FSharedBuffer::MakeView(Memory.GetData(), Memory.Num())));
+		}
+	}
+}
+
+void FFilePackageWriterUtil::WritePackage(FWritePackageParameters& Parameters)
+{
+	if (Parameters.Info.Status == IPackageWriter::ECommitStatus::Success)
+	{
+		AsyncSave(Parameters);
+	}
+}
+
+void FFilePackageWriterUtil::AsyncSave(FWritePackageParameters& Parameters)
+{
+	FCommitContext Context{ Parameters.Info };
+
+	// The order of these collection calls is important, both for ExportsBuffers (affects the meaning of offsets
+	// to those buffers) and for OutputFiles (affects the calculation of the Hash for the set of PackageData)
+	// The order of ExportsBuffers must match CompleteExportsArchiveForDiff.
+	CollectForSavePackageData(Parameters.Record, Context);
+	CollectForSaveBulkData(Parameters.Record, Context);
+	CollectForSaveLinkerAdditionalDataRecords(Parameters.Record, Context);
+	CollectForSaveAdditionalFileRecords(Parameters.Record, Context);
+	CollectForSaveExportsFooter(Parameters.Record, Context);
+	CollectForSaveExportsPackageTrailer(Parameters.Record, Context);
+	CollectForSaveExportsBuffers(Parameters.Record, Context);
+
+	AsyncSaveOutputFiles(Context, Parameters.AllPackageHashes, Parameters.PackageHashesLock, Parameters.bProvidePerPackageResult);
+}
+
+void FFilePackageWriterUtil::CollectForSavePackageData(FRecord& Record, FCommitContext& Context)
+{
+	Context.ExportsBuffers.AddDefaulted(Record.Packages.Num());
+	for (FPackageWriterRecords::FWritePackage& Package : Record.Packages)
+	{
+		Context.ExportsBuffers[Package.Info.MultiOutputIndex].Add(FExportBuffer{ Package.Buffer, MoveTemp(Package.Regions) });
+	}
+}
+
+void FFilePackageWriterUtil::CollectForSaveBulkData(FRecord& Record, FCommitContext& Context)
+{
+	for (FPackageWriterRecords::FBulkData& BulkRecord : Record.BulkDatas)
+	{
+		if (BulkRecord.Info.BulkDataType == IPackageWriter::FBulkDataInfo::AppendToExports)
+		{
+			if (Record.bCompletedExportsArchiveForDiff)
+			{
+				// Already Added in CompleteExportsArchiveForDiff
+				continue;
+			}
+			Context.ExportsBuffers[BulkRecord.Info.MultiOutputIndex].Add(FExportBuffer{ BulkRecord.Buffer, MoveTemp(BulkRecord.Regions) });
+		}
+		else
+		{
+			FWriteFileData& OutputFile = Context.OutputFiles.Emplace_GetRef();
+			OutputFile.Filename = BulkRecord.Info.LooseFilePath;
+			OutputFile.Buffer = FCompositeBuffer(BulkRecord.Buffer);
+			OutputFile.Regions = MoveTemp(BulkRecord.Regions);
+			OutputFile.bIsSidecar = true;
+			OutputFile.bContributeToHash = BulkRecord.Info.MultiOutputIndex == 0; // Only caculate the main package output hash
+			OutputFile.ChunkId = BulkRecord.Info.ChunkId;
+		}
+	}
+}
+
+void FFilePackageWriterUtil::CollectForSaveLinkerAdditionalDataRecords(FRecord& Record, FCommitContext& Context)
+{
+	if (Record.bCompletedExportsArchiveForDiff)
+	{
+		// Already Added in CompleteExportsArchiveForDiff
+		return;
+	}
+
+	for (FPackageWriterRecords::FLinkerAdditionalData& AdditionalRecord : Record.LinkerAdditionalDatas)
+	{
+		Context.ExportsBuffers[AdditionalRecord.Info.MultiOutputIndex].Add(FExportBuffer{ AdditionalRecord.Buffer, MoveTemp(AdditionalRecord.Regions) });
+	}
+}
+
+void FFilePackageWriterUtil::CollectForSaveAdditionalFileRecords(FRecord& Record, FCommitContext& Context)
+{
+	for (FPackageWriterRecords::FAdditionalFile& AdditionalRecord : Record.AdditionalFiles)
+	{
+		FWriteFileData& OutputFile = Context.OutputFiles.Emplace_GetRef();
+		OutputFile.Filename = AdditionalRecord.Info.Filename;
+		OutputFile.Buffer = FCompositeBuffer(AdditionalRecord.Buffer);
+		OutputFile.bIsSidecar = true;
+		OutputFile.bContributeToHash = AdditionalRecord.Info.MultiOutputIndex == 0; // Only calculate the main package output hash
+		OutputFile.ChunkId = AdditionalRecord.Info.ChunkId;
+	}
+}
+
+void FFilePackageWriterUtil::CollectForSaveExportsFooter(FRecord& Record, FCommitContext& Context)
+{
+	if (Record.bCompletedExportsArchiveForDiff)
+	{
+		// Already Added in CompleteExportsArchiveForDiff
+		return;
+	}
+
+	uint32 FooterData = PACKAGE_FILE_TAG;
+	FSharedBuffer Buffer = FSharedBuffer::Clone(&FooterData, sizeof(FooterData));
+	for (FPackageWriterRecords::FWritePackage& Package : Record.Packages)
+	{
+		Context.ExportsBuffers[Package.Info.MultiOutputIndex].Add(FExportBuffer{ Buffer, TArray<FFileRegion>() });
+	}
+}
+
+void FFilePackageWriterUtil::CollectForSaveExportsPackageTrailer(FRecord& Record, FCommitContext& Context)
+{
+	if (Record.bCompletedExportsArchiveForDiff)
+	{
+		// Already Added in CompleteExportsArchiveForDiff
+		return;
+	}
+
+	for (FPackageWriterRecords::FPackageTrailer& PackageTrailer : Record.PackageTrailers)
+	{
+		Context.ExportsBuffers[PackageTrailer.Info.MultiOutputIndex].Add(
+			FExportBuffer{ PackageTrailer.Buffer, TArray<FFileRegion>() });
+	}
+}
+
+void FFilePackageWriterUtil::CollectForSaveExportsBuffers(FRecord& Record, FCommitContext& Context)
+{
+	check(Context.ExportsBuffers.Num() == Record.Packages.Num());
+	for (FPackageWriterRecords::FWritePackage& Package : Record.Packages)
+	{
+		TArray<FExportBuffer>& ExportsBuffers = Context.ExportsBuffers[Package.Info.MultiOutputIndex];
+		check(ExportsBuffers.Num() > 0);
+
+		// Split the ExportsBuffer into (1) Header and (2) Exports + AllAppendedData
+		int64 HeaderSize = Package.Info.HeaderSize;
+		FExportBuffer& HeaderAndExportsBuffer = ExportsBuffers[0];
+		FSharedBuffer& HeaderAndExportsData = HeaderAndExportsBuffer.Buffer;
+
+	
+		// Header (.uasset/.umap)
+		{
+			FWriteFileData& OutputFile = Context.OutputFiles.Emplace_GetRef();
+			OutputFile.Filename = Package.Info.LooseFilePath;
+			OutputFile.Buffer = FCompositeBuffer(
+				FSharedBuffer::MakeView(HeaderAndExportsData.GetData(), HeaderSize, HeaderAndExportsData));
+			OutputFile.bIsSidecar = false;
+			OutputFile.bContributeToHash = Package.Info.MultiOutputIndex == 0; // Only calculate the main package output hash
+		}
+
+		// Exports + AllAppendedData (.uexp)
+		{
+			FWriteFileData& OutputFile = Context.OutputFiles.Emplace_GetRef();
+			OutputFile.Filename = FPaths::ChangeExtension(Package.Info.LooseFilePath, LexToString(EPackageExtension::Exports));
+			OutputFile.bIsSidecar = false;
+			OutputFile.bContributeToHash = Package.Info.MultiOutputIndex == 0; // Only calculate the main package output hash
+
+			int32 NumBuffers = ExportsBuffers.Num();
+			TArray<FSharedBuffer> BuffersForComposition;
+			BuffersForComposition.Reserve(NumBuffers);
+
+			const uint8* ExportsStart = static_cast<const uint8*>(HeaderAndExportsData.GetData()) + HeaderSize;
+			BuffersForComposition.Add(FSharedBuffer::MakeView(ExportsStart, HeaderAndExportsData.GetSize() - HeaderSize,
+				HeaderAndExportsData));
+			OutputFile.Regions.Append(MoveTemp(HeaderAndExportsBuffer.Regions));
+
+			for (FExportBuffer& ExportsBuffer : TArrayView<FExportBuffer>(ExportsBuffers).Slice(1, NumBuffers - 1))
+			{
+				BuffersForComposition.Add(ExportsBuffer.Buffer);
+				OutputFile.Regions.Append(MoveTemp(ExportsBuffer.Regions));
+			}
+			OutputFile.Buffer = FCompositeBuffer(BuffersForComposition);
+
+			// Adjust regions so they are relative to the start of the uexp file
+			for (FFileRegion& Region : OutputFile.Regions)
+			{
+				Region.Offset -= HeaderSize;
+			}
+		}
+	}
+}
+
+void FFilePackageWriterUtil::AsyncSaveOutputFiles(FCommitContext& Context, TMap<FName, TRefCountPtr<FPackageHashes>>* AllPackageHashes, FCriticalSection* PackageHashesLock, bool bProvidePerPackageResult)
+{
+	if (bProvidePerPackageResult && !AllPackageHashes)
+	{
+		UE_LOG(LogSavePackage, Error, TEXT("FFilePackageWriterUtil::AsyncSaveOutputFiles : if bProvidePerPackageResult is true then AllPackageHashes can't be null."));
+		return;
+	}
+
+	if (AllPackageHashes && !PackageHashesLock)
+	{
+		UE_LOG(LogSavePackage, Error, TEXT("FFilePackageWriterUtil::AsyncSaveOutputFiles : if AllPackageHashes is provided, then PackageHashesLock can't be null."));
+		return;
+	}
+
+	if (!EnumHasAnyFlags(Context.Info.WriteOptions, IPackageWriter::EWriteOptions::Write | IPackageWriter::EWriteOptions::ComputeHash))
+	{
+		return;
+	}
+
+	UE::SavePackageUtilities::IncrementOutstandingAsyncWrites();
+
+	TRefCountPtr<FPackageHashes> ThisPackageHashes;
+	TUniquePtr<TPromise<int>> PackageHashesCompletionPromise;
+
+	if (EnumHasAnyFlags(Context.Info.WriteOptions, IPackageWriter::EWriteOptions::ComputeHash))
+	{
+		ThisPackageHashes = new FPackageHashes();
+		
+		if (bProvidePerPackageResult)
+		{
+			PackageHashesCompletionPromise.Reset(new TPromise<int>());
+			ThisPackageHashes->CompletionFuture = PackageHashesCompletionPromise->GetFuture();
+		}
+
+		bool bAlreadyExisted = false;
+		if(AllPackageHashes)
+		{
+			FScopeLock PackageHashesScopeLock(PackageHashesLock);
+			TRefCountPtr<FPackageHashes>& ExistingPackageHashes = AllPackageHashes->FindOrAdd(Context.Info.PackageName);
+			// This calculation of bAlreadyExisted looks weird but we're finding the _refcount_, not the hashes. So if it gets
+			// constructed, it's not actually assigned a pointer.
+			bAlreadyExisted = ExistingPackageHashes.IsValid();
+			ExistingPackageHashes = ThisPackageHashes;
+		}
+		if (bAlreadyExisted)
+		{
+			UE_LOG(LogSavePackage, Error, TEXT("FCookedFilePackageWriter encountered the same package twice in a cook! (%s)"), *Context.Info.PackageName.ToString());
+		}
+	}
+
+	UE::Tasks::Launch(TEXT("HashAndWriteCookedFile"),
+		[OutputFiles = MoveTemp(Context.OutputFiles), WriteOptions = Context.Info.WriteOptions,
+		ThisPackageHashes = MoveTemp(ThisPackageHashes), PackageHashesCompletionPromise = MoveTemp(PackageHashesCompletionPromise)]
+		() mutable
+		{
+			FMD5 AccumulatedHash;
+			for (const FWriteFileData& OutputFile : OutputFiles)
+			{
+				OutputFile.HashAndWrite(AccumulatedHash, ThisPackageHashes, WriteOptions);
+			}
+
+			if (EnumHasAnyFlags(WriteOptions, IPackageWriter::EWriteOptions::ComputeHash))
+			{
+				ThisPackageHashes->PackageHash.Set(AccumulatedHash);
+			}
+
+			if (PackageHashesCompletionPromise)
+			{
+				// Note that setting this Promise might call arbitrary code from anything that subscribed
+				// to ThisPackageHashes->CompletionFuture.Then(). So don't call it inside a lock.
+				PackageHashesCompletionPromise->SetValue(0);
+			}
+
+			// This is used to release the game thread to access the hashes
+			UE::SavePackageUtilities::DecrementOutstandingAsyncWrites();
+		},
+		UE::Tasks::ETaskPriority::BackgroundNormal
+	);
+}
\ No newline at end of file
diff --git a/Engine/Source/Editor/UnrealEd/Public/Serialization/DefaultCookedFilePackageWriter.h b/Engine/Source/Editor/UnrealEd/Public/Serialization/DefaultCookedFilePackageWriter.h
new file mode 100644
--- /dev/null
+++ b/Engine/Source/Editor/UnrealEd/Public/Serialization/DefaultCookedFilePackageWriter.h
@@ -0,0 +1,29 @@
+// Copyright Epic Games, Inc. All Rights Reserved.
+
+#pragma once
+
+#include "Serialization/PackageWriterToSharedBuffer.h"
+
+/**
+ * Package writer to write a cooked file to a given directory with a given name.
+ * BaseFilename is the filename of the package to write. The extension is ignored.
+ * Then extension of each file is specified to match Unreal's loading expectations.
+ */
+class FDefaultCookedFilePackageWriter : public TPackageWriterToSharedBuffer<IPackageWriter>
+{
+public:
+	using Super = TPackageWriterToSharedBuffer<IPackageWriter>;
+
+	FDefaultCookedFilePackageWriter(const FString& BaseFilename);
+
+	void WritePackageData(const IPackageWriter::FPackageInfo& Info, FLargeMemoryWriter& ExportsArchive, const TArray<FFileRegion>& FileRegions) override;
+	void WriteBulkData(const IPackageWriter::FBulkDataInfo& Info, const FIoBuffer& BulkData, const TArray<FFileRegion>& FileRegions) override;
+	
+protected:
+	void CommitPackageInternal(FPackageWriterRecords::FPackage&& BaseRecord, const FCommitPackageInfo& Info) override;
+
+	FPackageWriterRecords::FPackage* ConstructRecord() override;
+
+private:
+	FString BaseFilename;
+};
diff --git a/Engine/Source/Editor/UnrealEd/Public/Serialization/FilePackageWriterUtil.h b/Engine/Source/Editor/UnrealEd/Public/Serialization/FilePackageWriterUtil.h
new file mode 100644
--- /dev/null
+++ b/Engine/Source/Editor/UnrealEd/Public/Serialization/FilePackageWriterUtil.h
@@ -0,0 +1,82 @@
+// Copyright Epic Games, Inc. All Rights Reserved.
+
+#pragma once
+
+#include "IO/IoChunkId.h"
+#include "Memory/CompositeBuffer.h"
+#include "Memory/SharedBuffer.h"
+#include "Serialization/PackageWriterToSharedBuffer.h"
+
+/** 
+ * Class containing functions to save cooked packages in separate .uasset,.uexp,.ubulk files.
+ */
+class FFilePackageWriterUtil
+{
+public:
+
+	/** Version of the superclass's per-package record that includes our class-specific data. */
+	struct FRecord : public FPackageWriterRecords::FPackage
+	{
+		bool bCompletedExportsArchiveForDiff = false;
+	};
+
+	struct FWritePackageParameters
+	{
+		FRecord& Record;
+		const IPackageWriter::FCommitPackageInfo& Info;
+		TMap<FName, TRefCountPtr<FPackageHashes>>* AllPackageHashes;
+		FCriticalSection* PackageHashesLock;
+		bool bProvidePerPackageResult;
+
+		FWritePackageParameters(FRecord& InRecord,
+			const IPackageWriter::FCommitPackageInfo& InInfo,
+			TMap<FName, TRefCountPtr<FPackageHashes>>* InAllPackageHashes,
+			FCriticalSection* InPackageHashesLock,
+			bool bInProvidePerPackageResult);
+	};
+
+	static void WritePackage(FWritePackageParameters& Parameters);
+
+private:
+
+	/** Buffers that are combined into the HeaderAndExports file (which is then split into .uasset + .uexp or .uoasset + .uoexp). */
+	struct FExportBuffer
+	{
+		FSharedBuffer Buffer;
+		TArray<FFileRegion> Regions;
+	};
+
+	/**
+	 * The data needed to asynchronously write one of the files (.uasset, .uexp, .ubulk, any optional and any additional),
+	 * without reference back to other data on this writer.
+	 */
+	struct FWriteFileData
+	{
+		FString Filename;
+		FCompositeBuffer Buffer;
+		TArray<FFileRegion> Regions;
+		bool bIsSidecar;
+		bool bContributeToHash = true;
+		FIoChunkId ChunkId = FIoChunkId::InvalidChunkId;
+
+		void HashAndWrite(FMD5& AccumulatedHash, const TRefCountPtr<FPackageHashes>& PackageHashes, IPackageWriter::EWriteOptions WriteOptions) const;
+	};
+
+	/** Stack data for the helper functions of CommitPackageInternal. */
+	struct FCommitContext
+	{
+		const IPackageWriter::FCommitPackageInfo& Info;
+		TArray<TArray<FExportBuffer>> ExportsBuffers;
+		TArray<FWriteFileData> OutputFiles;
+	};
+
+	static void AsyncSave(FWritePackageParameters& Parameters);
+	static void CollectForSavePackageData(FRecord& Record, FCommitContext& Context);
+	static void CollectForSaveBulkData(FRecord& Record, FCommitContext& Context);
+	static void CollectForSaveLinkerAdditionalDataRecords(FRecord& Record, FCommitContext& Context);
+	static void CollectForSaveAdditionalFileRecords(FRecord& Record, FCommitContext& Context);
+	static void CollectForSaveExportsFooter(FRecord& Record, FCommitContext& Context);
+	static void CollectForSaveExportsPackageTrailer(FRecord& Record, FCommitContext& Context);
+	static void CollectForSaveExportsBuffers(FRecord& Record, FCommitContext& Context);
+	static void AsyncSaveOutputFiles(FCommitContext& Context, TMap<FName, TRefCountPtr<FPackageHashes>>* AllPackageHashes, FCriticalSection* PackageHashesLock, bool bProvidePerPackageResult);
+};
