[Sequencer] Limited Selectability: Keep selected Sequencer selection when enabling tool (regression)

diff --git a/Engine/Source/Editor/UnrealEd/Private/EditorViewportSelectability.cpp b/Engine/Source/Editor/UnrealEd/Private/EditorViewportSelectability.cpp
--- a/Engine/Source/Editor/UnrealEd/Private/EditorViewportSelectability.cpp
+++ b/Engine/Source/Editor/UnrealEd/Private/EditorViewportSelectability.cpp
@@ -235,7 +235,7 @@ void FEditorViewportSelectability::DeselectNonSelectableActors()
 
 	UWorld* const World = OnGetWorld.IsBound() ? OnGetWorld.Execute() : nullptr;
 
-	SelectActorsByPredicate(World, false, false
+	SelectActorsByPredicate(World, /*bInSelect=*/true, /*bInClearSelection=*/true
 		, [this](AActor* const InActor) -> bool
 		{
 			return IsObjectSelectableInViewport(InActor);
