Always include CDO refs for replacement when doing fast reinstancing. This ensures that instances do not end up with stale references to subobjects of CDOs after compilation, even if reconstruction is skipped.
#rnx

diff --git a/Engine/Source/Editor/Kismet/Private/BlueprintCompilationManager.cpp b/Engine/Source/Editor/Kismet/Private/BlueprintCompilationManager.cpp
--- a/Engine/Source/Editor/Kismet/Private/BlueprintCompilationManager.cpp
+++ b/Engine/Source/Editor/Kismet/Private/BlueprintCompilationManager.cpp
@@ -101,6 +101,40 @@
 #define VERIFY_NO_STALE_CLASS_REFERENCES 0
 #define VERIFY_NO_BAD_SKELETON_REFERENCES 0
 
+namespace UE::Kismet::BlueprintCompilationManager::Private
+{
+	namespace ConsoleVariables
+	{
+		/** Flag to use the new FBlueprintCompileReinstancer::MoveDependentSkelToReinst and avoid problematic recursion during reinstancing */
+		static bool bEnableSkelReinstUpdate = true;
+		static FAutoConsoleVariableRef CVarEnableSkelReinstUpdate(
+			TEXT("BP.bEnableSkelReinstUpdate"), bEnableSkelReinstUpdate,
+			TEXT("If true the Reinstancing of SKEL classes will use the new FBlueprintCompileReinstancer::MoveDependentSkelToReinst(o(n)) instead of the old MoveSkelCDOAside (o(n^2))"),
+			ECVF_Default);
+
+		/** Flag to disable faster compiles for individual blueprints if they have no function signature changes */
+		static bool bForceAllDependenciesToRecompile = false;
+		static FAutoConsoleVariableRef CVarForceAllDependenciesToRecompile(
+			TEXT("BP.bForceAllDependenciesToRecompile"), bForceAllDependenciesToRecompile,
+			TEXT("If true all dependencies will be bytecode-compiled even when all referenced functions have no signature changes. Intended for compiler development/debugging purposes."),
+			ECVF_Default);
+
+		/** 
+		 * Flag to disable skipping of reinstancing logic when a class is not observabling changed
+		 * by compilation. This can only work when compiling in place (e.g. reusing UClasses) - so
+		 * can't be done for asset reload, live coding, or verse 
+		 */
+		static bool bAllowFastReinstancing = true;
+		static FAutoConsoleVariableRef CVarAllowFastReinstancing(
+			TEXT("BP.bAllowFastReinstancing"), bAllowFastReinstancing,
+			TEXT("True by default, this flag enables a fast path in the logic that updates class \
+				objects at runtime (Reinstancing). When this flag is set we check class layout and \
+				construction data for a match, if the layout and construction data match then we \
+				don't replace instances of the class."),
+			ECVF_Default);
+	}
+}
+
 struct FReinstancingJob;
 struct FSkeletonFixupData;
 struct FCompilerData;
@@ -283,7 +317,11 @@
 	const bool bBatchCompile				= (Request.CompileOptions & EBlueprintCompileOptions::BatchCompile						) != EBlueprintCompileOptions::None;
 	const bool bSkipReinstancing			= (Request.CompileOptions & EBlueprintCompileOptions::SkipReinstancing					) != EBlueprintCompileOptions::None;
 	const bool bSkipSaving					= (Request.CompileOptions & EBlueprintCompileOptions::SkipSave							) != EBlueprintCompileOptions::None;
-	const bool bFindAndReplaceCDOReferences	= (Request.CompileOptions & EBlueprintCompileOptions::IncludeCDOInReferenceReplacement	) != EBlueprintCompileOptions::None;
+	const bool bFindAndReplaceCDOReferences	= (Request.CompileOptions & EBlueprintCompileOptions::IncludeCDOInReferenceReplacement	) != EBlueprintCompileOptions::None ||
+		// when fast reinstancing we want to always include CDO References as instances have a habit of caching object
+		// references from the CDO. This works fine when they are reconstructed any time the CDO is regenerated
+		// but fast reinstancing attempts to avoid expensive object recreation/reference replacment:
+		(UE::Kismet::BlueprintCompilationManager::Private::ConsoleVariables::bAllowFastReinstancing); 
 
 	ensure(!bIsRegeneratingOnLoad); // unexpected code path, compile on load handled with different function call
 	ensure(!bSkipReinstancing); // This is an internal option, should not go through CompileSynchronouslyImpl
@@ -569,40 +607,6 @@
 {
 }
 
-namespace UE::Kismet::BlueprintCompilationManager::Private
-{
-	namespace ConsoleVariables
-	{
-		/** Flag to use the new FBlueprintCompileReinstancer::MoveDependentSkelToReinst and avoid problematic recursion during reinstancing */
-		static bool bEnableSkelReinstUpdate = true;
-		static FAutoConsoleVariableRef CVarEnableSkelReinstUpdate(
-			TEXT("BP.bEnableSkelReinstUpdate"), bEnableSkelReinstUpdate,
-			TEXT("If true the Reinstancing of SKEL classes will use the new FBlueprintCompileReinstancer::MoveDependentSkelToReinst(o(n)) instead of the old MoveSkelCDOAside (o(n^2))"),
-			ECVF_Default);
-
-		/** Flag to disable faster compiles for individual blueprints if they have no function signature changes */
-		static bool bForceAllDependenciesToRecompile = false;
-		static FAutoConsoleVariableRef CVarForceAllDependenciesToRecompile(
-			TEXT("BP.bForceAllDependenciesToRecompile"), bForceAllDependenciesToRecompile,
-			TEXT("If true all dependencies will be bytecode-compiled even when all referenced functions have no signature changes. Intended for compiler development/debugging purposes."),
-			ECVF_Default);
-
-		/** 
-		 * Flag to disable skipping of reinstancing logic when a class is not observabling changed
-		 * by compilation. This can only work when compiling in place (e.g. reusing UClasses) - so
-		 * can't be done for asset reload, live coding, or verse 
-		 */
-		static bool bAllowFastReinstancing = true;
-		static FAutoConsoleVariableRef CVarAllowFastReinstancing(
-			TEXT("BP.bAllowFastReinstancing"), bAllowFastReinstancing,
-			TEXT("True by default, this flag enables a fast path in the logic that updates class \
-				objects at runtime (Reinstancing). When this flag is set we check class layout and \
-				construction data for a match, if the layout and construction data match then we \
-				don't replace instances of the class."),
-			ECVF_Default);
-	}
-}
-
 UE_TRACE_EVENT_BEGIN(Cpu, FlushCompilationQueue, NoSync)
 	UE_TRACE_EVENT_FIELD(UE::Trace::WideString, BlueprintPath)
 	UE_TRACE_EVENT_FIELD(int32, TotalBlueprints)
diff --git a/Engine/Source/Editor/UnrealEd/Private/Kismet2/KismetReinstanceUtilities.cpp b/Engine/Source/Editor/UnrealEd/Private/Kismet2/KismetReinstanceUtilities.cpp
--- a/Engine/Source/Editor/UnrealEd/Private/Kismet2/KismetReinstanceUtilities.cpp
+++ b/Engine/Source/Editor/UnrealEd/Private/Kismet2/KismetReinstanceUtilities.cpp
@@ -72,6 +72,56 @@ namespace UE::ReinstanceUtils
 	const EObjectFlags FlagMask = RF_Public | RF_ArchetypeObject | RF_Transactional | RF_Transient | RF_TextExportTransient | RF_InheritableComponentTemplate | RF_Standalone; //TODO: what about RF_RootSet?
 }
 
+namespace UE::Private
+{
+
+// The class for finding and replacing weak references.
+// We can't relay on "standard" weak references replacement as
+// it depends on FSoftObjectPath::ResolveObject, which
+// tries to find the object with the stored path. It is
+// impossible, cause above we deleted old actors (after
+// spawning new ones), so during objects traverse we have to
+// find FSoftObjectPath with the raw given path taken
+// before deletion of old actors and fix them.
+class FReferenceReplace : public FArchiveReplaceObjectAndStructPropertyRef<UObject>
+{
+public:
+	FReferenceReplace(UObject* InSearchObject, const TMap<UObject*, UObject*>& InReplacementMap, const TMap<FSoftObjectPath, UObject*>& InWeakReferencesMap)
+		: FArchiveReplaceObjectAndStructPropertyRef<UObject>(InSearchObject, InReplacementMap, EArchiveReplaceObjectFlags::DelayStart), WeakReferencesMap(InWeakReferencesMap)
+	{
+		SerializeSearchObject();
+	}
+
+	FReferenceReplace(UObject* InSearchObject, const TMap<UObject*, UObject*>& InReplacementMap, const TMap<FSoftObjectPath, UObject*>& InWeakReferencesMap, EArchiveReplaceObjectFlags Flags)
+		: FArchiveReplaceObjectAndStructPropertyRef<UObject>(InSearchObject, InReplacementMap, EArchiveReplaceObjectFlags::DelayStart), WeakReferencesMap(InWeakReferencesMap)
+	{
+		if (!(Flags & EArchiveReplaceObjectFlags::DelayStart))
+		{
+			SerializeSearchObject();
+		}
+	}
+
+	FArchive& operator<<(FSoftObjectPath& Ref) override
+	{
+		const UObject*const* PtrToObjPtr = WeakReferencesMap.Find(Ref);
+
+		if (PtrToObjPtr != nullptr)
+		{
+			Ref = *PtrToObjPtr;
+		}
+
+		return *this;
+	}
+
+	FArchive& operator<<(FSoftObjectPtr& Ref) override
+	{
+		return operator<<(Ref.GetUniqueID());
+	}
+
+private:
+	const TMap<FSoftObjectPath, UObject*>& WeakReferencesMap;
+};
+
 struct FReplaceReferenceHelper
 {
 	static void ValidateReplacementMappings(const TMap<UObject*, UObject*>& OldToNewInstanceMap)
@@ -232,54 +282,7 @@ struct FReplaceReferenceHelper
 				// delete components owned by new actors)
 				if (!ObjectsToReplace.Contains(Obj))
 				{
-					// The class for finding and replacing weak references.
-					// We can't relay on "standard" weak references replacement as
-					// it depends on FSoftObjectPath::ResolveObject, which
-					// tries to find the object with the stored path. It is
-					// impossible, cause above we deleted old actors (after
-					// spawning new ones), so during objects traverse we have to
-					// find FSoftObjectPath with the raw given path taken
-					// before deletion of old actors and fix them.
-					class ReferenceReplace : public FArchiveReplaceObjectAndStructPropertyRef<UObject>
-					{
-					public:
-						ReferenceReplace(UObject* InSearchObject, const TMap<UObject*, UObject*>& InReplacementMap, const TMap<FSoftObjectPath, UObject*>& InWeakReferencesMap)
-							: FArchiveReplaceObjectAndStructPropertyRef<UObject>(InSearchObject, InReplacementMap, EArchiveReplaceObjectFlags::DelayStart), WeakReferencesMap(InWeakReferencesMap)
-						{
-							SerializeSearchObject();
-						}
-
-						ReferenceReplace(UObject* InSearchObject, const TMap<UObject*, UObject*>& InReplacementMap, const TMap<FSoftObjectPath, UObject*>& InWeakReferencesMap, EArchiveReplaceObjectFlags Flags)
-							: FArchiveReplaceObjectAndStructPropertyRef<UObject>(InSearchObject, InReplacementMap, EArchiveReplaceObjectFlags::DelayStart), WeakReferencesMap(InWeakReferencesMap)
-						{
-							if (!(Flags & EArchiveReplaceObjectFlags::DelayStart))
-							{
-								SerializeSearchObject();
-							}
-						}
-
-						FArchive& operator<<(FSoftObjectPath& Ref) override
-						{
-							const UObject*const* PtrToObjPtr = WeakReferencesMap.Find(Ref);
-
-							if (PtrToObjPtr != nullptr)
-							{
-								Ref = *PtrToObjPtr;
-							}
-
-							return *this;
-						}
-
-						FArchive& operator<<(FSoftObjectPtr& Ref) override
-						{
-							return operator<<(Ref.GetUniqueID());
-						}
-
-					private:
-						const TMap<FSoftObjectPath, UObject*>& WeakReferencesMap;
-					};
-
-					ReferenceReplace ReplaceAr(Obj, OldToNewInstanceMap, ReinstancedObjectsWeakReferenceMap);
+					FReferenceReplace ReplaceAr(Obj, OldToNewInstanceMap, ReinstancedObjectsWeakReferenceMap);
 				}
 			}
 		}
@@ -464,6 +467,7 @@ struct FArchetypeReinstanceHelper
 		return OutName;
 	}
 };
+}
 
 /////////////////////////////////////////////////////////////////////////////////
 // FBlueprintCompileReinstancer
@@ -860,14 +864,14 @@ void FBlueprintCompileReinstancer::FinalizeFastReinstancing(TArray<UObject*>& Ob
 	TArray<UObject*> SourceObjects;
 	TMap<UObject*, UObject*> OldToNewInstanceMap;
 	TMap<FSoftObjectPath, UObject*> ReinstancedObjectsWeakReferenceMap;
-	FReplaceReferenceHelper::IncludeCDO(DuplicatedClass, ClassToReinstance, OldToNewInstanceMap, SourceObjects, OriginalCDO);
+	UE::Private::FReplaceReferenceHelper::IncludeCDO(DuplicatedClass, ClassToReinstance, OldToNewInstanceMap, SourceObjects, OriginalCDO);
 
 	if (IsClassObjectReplaced())
 	{
-		FReplaceReferenceHelper::IncludeClass(DuplicatedClass, ClassToReinstance, OldToNewInstanceMap, SourceObjects, ObjectsToReplace);
+		UE::Private::FReplaceReferenceHelper::IncludeClass(DuplicatedClass, ClassToReinstance, OldToNewInstanceMap, SourceObjects, ObjectsToReplace);
 	}
 
-	FReplaceReferenceHelper::FindAndReplaceReferences(SourceObjects, &ObjectsThatShouldUseOldStuff, ObjectsToReplace, OldToNewInstanceMap, ReinstancedObjectsWeakReferenceMap);
+	UE::Private::FReplaceReferenceHelper::FindAndReplaceReferences(SourceObjects, &ObjectsThatShouldUseOldStuff, ObjectsToReplace, OldToNewInstanceMap, ReinstancedObjectsWeakReferenceMap);
 
 	if (ClassToReinstance->IsChildOf<UActorComponent>())
 	{
@@ -2222,10 +2226,10 @@ static void ReplaceObjectHelper(UObject*& OldObject, UClass* OldClass, UObject*&
 		{
 			// Gather all component templates from the current archetype to the farthest antecedent inherited template(s).
 			TArray<UObject*> OldArchetypeObjects;
-			FArchetypeReinstanceHelper::GetArchetypeObjects(OldObject, OldArchetypeObjects, RF_InheritableComponentTemplate);
+			UE::Private::FArchetypeReinstanceHelper::GetArchetypeObjects(OldObject, OldArchetypeObjects, RF_InheritableComponentTemplate);
 
 			// Find a unique object name that does not conflict with anything in the scope of all outers in the template chain.
-			const FString OldArchetypeName = FArchetypeReinstanceHelper::FindUniqueArchetypeObjectName(OldArchetypeObjects).ToString();
+			const FString OldArchetypeName = UE::Private::FArchetypeReinstanceHelper::FindUniqueArchetypeObjectName(OldArchetypeObjects).ToString();
 
 			for (UObject* OldArchetypeObject : OldArchetypeObjects)
 			{
@@ -2910,7 +2914,7 @@ void FBlueprintCompileReinstancer::ReplaceInstancesOfClass_Inner(const TMap<UCla
 		// Always map old to new instances of CDOs along with any owned subobject(s). This allows delegates to be
 		// notified that these instances have been replaced. However, we don't proactively find and replace those
 		// references ourselves unless input parameters have explicitly configured this path to do so (see below).
-		FReplaceReferenceHelper::IncludeCDO(OldClass, NewClass, OldToNewInstanceMap, OldCDOSourceObjects, InOriginalCDO, Params.OldToNewTemplates);
+		UE::Private::FReplaceReferenceHelper::IncludeCDO(OldClass, NewClass, OldToNewInstanceMap, OldCDOSourceObjects, InOriginalCDO, Params.OldToNewTemplates);
 		if (bReplaceReferencesToOldCDOs)
 		{
 			// This means we'll proactively find and replace references to old CDOs and any owned subobject(s). It
@@ -2921,7 +2925,7 @@ void FBlueprintCompileReinstancer::ReplaceInstancesOfClass_Inner(const TMap<UCla
 			// This is part of the legacy reload path; it is only enabled if we're also replacing references to old CDOs.
 			if (bClassObjectReplaced)
 			{
-				FReplaceReferenceHelper::IncludeClass(OldClass, NewClass, OldToNewInstanceMap, SourceObjects, ObjectsReplaced);
+				UE::Private::FReplaceReferenceHelper::IncludeClass(OldClass, NewClass, OldToNewInstanceMap, SourceObjects, ObjectsReplaced);
 			}
 		}
 	}
@@ -2977,7 +2981,13 @@ void FBlueprintCompileReinstancer::ReplaceInstancesOfClass_Inner(const TMap<UCla
 	}
 
 	//FReplaceReferenceHelper::ValidateReplacementMappings(OldToNewInstanceMap);
-	FReplaceReferenceHelper::FindAndReplaceReferences(SourceObjects, ObjectsThatShouldUseOldStuff, ObjectsReplaced, OldToNewInstanceMap, ReinstancedObjectsWeakReferenceMap);
+	// Don't run expensive FindAndReplaceReferences when no objects were replaced if we have versioning data.
+	// SourceObjects will always be non zero when using this code path but we can't afford to do global
+	// reference replacement at all times:
+	if(!Params.ConstructionVersioningData || ObjectsReplaced.Num() > 0 ) 
+	{
+		UE::Private::FReplaceReferenceHelper::FindAndReplaceReferences(SourceObjects, ObjectsThatShouldUseOldStuff, ObjectsReplaced, OldToNewInstanceMap, ReinstancedObjectsWeakReferenceMap);
+	}
 	
 	for (UObject* Obj : ObjectsReplaced)
 	{
@@ -3239,7 +3249,7 @@ void FBlueprintCompileReinstancer::PreCreateSubObjectsForReinstantiation(
 	if (!FOverridableManager::Get().IsEnabled(*OldObject) || !OldObject->HasAnyFlags(RF_NeedPostLoad))
 	{
 		OldInstancedSubObjectsPtr = &OldInstancedSubObjects;
-		FReplaceReferenceHelper::GetOwnedSubobjectsRecursive(OldObject, OldInstancedSubObjects);
+		UE::Private::FReplaceReferenceHelper::GetOwnedSubobjectsRecursive(OldObject, OldInstancedSubObjects);
 	}
 
 	// Add the mapping from the old to the new object exists...
@@ -3550,6 +3560,12 @@ void FBlueprintCompileReinstancer::RestoreClass(
 		for (UObject* Object : ObjectsToReplace)
 		{
 			Object->SetClass(NewClass); // class layouts match, so we can just set the object back to its original class
+			
+			// run reference replacement on the object, in case it has stored references to regenerated archetype data:
+			UE::Private::FReferenceReplace ReplaceAr(
+				Object, 
+				ReinstancingState.OldToNewInstanceMap,
+				ReinstancingState.ReinstancedObjectsWeakReferenceMap);
 		}
 	}
 }
