Remove Material entry from the RT debug menu (only used by Substrate's RT debug).

Substrate RT debug mode now falls back onto the 'other' bucket and can be made visible with r.RayTracing.Visualize.Other 1.

diff --git a/Engine/Source/Editor/UnrealEd/Private/RayTracingDebugVisualizationMenuCommands.cpp b/Engine/Source/Editor/UnrealEd/Private/RayTracingDebugVisualizationMenuCommands.cpp
--- a/Engine/Source/Editor/UnrealEd/Private/RayTracingDebugVisualizationMenuCommands.cpp
+++ b/Engine/Source/Editor/UnrealEd/Private/RayTracingDebugVisualizationMenuCommands.cpp
@@ -72,10 +72,6 @@ void FRayTracingDebugVisualizationMenuCommands::BuildCommandMap()
 			Record.Type = FVisualizationType::Performance;
 			break;
 
-		case FRayTracingVisualizationData::FModeType::Material:
-			Record.Type = FVisualizationType::Material;
-			break;
-
 		case FRayTracingVisualizationData::FModeType::Other:
 			Record.Type = FVisualizationType::Other;
 			break;
@@ -108,12 +104,6 @@ void FRayTracingDebugVisualizationMenuCommands::BuildVisualisationSubMenu(FMenuB
 			Menu.EndSection();
 		}
 
-		{
-			Menu.BeginSection("LevelViewportRayTracingVisualizationMaterial", LOCTEXT("RayTracingVisualizationMaterial", "Material"));
-			Commands.AddCommandTypeToMenu(Menu, FVisualizationType::Material);
-			Menu.EndSection();
-		}
-
 		if (bShowOther)
 		{
 			Menu.BeginSection("LevelViewportRayTracingVisualizationOther", LOCTEXT("RayTracingVisualizationOther", "Other"));
diff --git a/Engine/Source/Runtime/Engine/Private/RayTracingVisualizationData.cpp b/Engine/Source/Runtime/Engine/Private/RayTracingVisualizationData.cpp
--- a/Engine/Source/Runtime/Engine/Private/RayTracingVisualizationData.cpp
+++ b/Engine/Source/Runtime/Engine/Private/RayTracingVisualizationData.cpp
@@ -76,7 +76,7 @@ FRayTracingVisualizationData::FRayTracingVisualizationData()
 
 		if (Substrate::IsSubstrateEnabled())
 		{
-			AddVisualizationMode(TEXT("SubstrateMaterialProperties"), LOCTEXT("SubstrateMaterialProperties", "Substrate Material Properties"), FModeType::Material, RAY_TRACING_DEBUG_VIZ_SUBSTRATE_DATA, true);
+			AddVisualizationMode(TEXT("SubstrateMaterialProperties"), LOCTEXT("SubstrateMaterialProperties", "Substrate Material Properties"), FModeType::Other, RAY_TRACING_DEBUG_VIZ_SUBSTRATE_DATA, true);
 		}
 	}
 }
diff --git a/Engine/Source/Runtime/Engine/Public/RayTracingVisualizationData.h b/Engine/Source/Runtime/Engine/Public/RayTracingVisualizationData.h
--- a/Engine/Source/Runtime/Engine/Public/RayTracingVisualizationData.h
+++ b/Engine/Source/Runtime/Engine/Public/RayTracingVisualizationData.h
@@ -12,7 +12,6 @@ public:
 		Overview,
 		Standard,
 		Performance,
-		Material,
 		Other,
 	};
 
