World Partition - Builders: Wait for the AR scan completion before calling PreWorldInitialization()
* Will prevent rare file operation issues in distributed HLOD builds, where the AR could be scanning an actor package while the code in PreWorldInitialization() tries to delete it

diff --git a/Engine/Source/Editor/UnrealEd/Private/WorldPartition/WorldPartitionBuilder.cpp b/Engine/Source/Editor/UnrealEd/Private/WorldPartition/WorldPartitionBuilder.cpp
--- a/Engine/Source/Editor/UnrealEd/Private/WorldPartition/WorldPartitionBuilder.cpp
+++ b/Engine/Source/Editor/UnrealEd/Private/WorldPartition/WorldPartitionBuilder.cpp
@@ -89,6 +89,10 @@ bool UWorldPartitionBuilder::RunBuilder(UWorld* World)
 		return true;
 	}
 
+	// Wait for the asset registry to finish scanning
+	IAssetRegistry& AssetRegistry = FModuleManager::GetModuleChecked<FAssetRegistryModule>("AssetRegistry").Get();
+	AssetRegistry.WaitForCompletion();
+
 	FPackageSourceControlHelper SCCHelper;
 
 	// Perform builder pre world initialisation
@@ -101,9 +105,6 @@ bool UWorldPartitionBuilder::RunBuilder(UWorld* World)
 	bool bResult = true;
 	// Setup the world
 	{
-		IAssetRegistry& AssetRegistry = FModuleManager::GetModuleChecked<FAssetRegistryModule>("AssetRegistry").Get();
-		AssetRegistry.WaitForCompletion();
-
 		TUniquePtr<FScopedEditorWorld> EditorWorld;
 		if (World->bIsWorldInitialized)
 		{
