Update Viewport Orbit Gizmo to optionally display Forward, Right, Up coordinate display names

diff --git a/Engine/Source/Editor/UnrealEd/Private/EditorViewportClient.cpp b/Engine/Source/Editor/UnrealEd/Private/EditorViewportClient.cpp
--- a/Engine/Source/Editor/UnrealEd/Private/EditorViewportClient.cpp
+++ b/Engine/Source/Editor/UnrealEd/Private/EditorViewportClient.cpp
@@ -3592,6 +3592,39 @@ void FEditorViewportClient::DrawAxes(FViewport* InViewport, FCanvas* Canvas, con
 	FIntPoint AxisEnd;
 	FCanvasLineItem LineItem;
 	FCanvasTextItem TextItem( FVector2D::ZeroVector, FText::GetEmpty(), Font, FLinearColor::White );
+	
+	const TStaticArray<FText, 3> AxisNames = []()
+	{
+		static bool bUseForwardRightUpDisplayNames = []()
+		{
+			if (IConsoleVariable* UseForwardRightUpDisplayNamesVariable = IConsoleManager::Get().FindConsoleVariable(TEXT("Editor.UseForwardRightUpDisplayNames")))
+			{
+				return UseForwardRightUpDisplayNamesVariable->GetBool();
+			}
+			return false;
+		}();
+
+		if (bUseForwardRightUpDisplayNames)
+		{
+			return TStaticArray<FText, 3>({LOCTEXT("ForwardAxis","F"), LOCTEXT("RightAxis","R"), LOCTEXT("UpAxis","U")});
+		}
+		else
+		{
+			return TStaticArray<FText, 3>({LOCTEXT("XAxis","X"), LOCTEXT("YAxis","Y"), LOCTEXT("ZAxis","Z")});
+		}
+	}();
+	auto GetAxisName = [&AxisNames](EAxisList::Type InAxis) -> FText
+	{
+		switch (InAxis)
+		{
+		case EAxisList::X: return AxisNames[0];
+		case EAxisList::Y: return AxisNames[1];
+		case EAxisList::Z: return AxisNames[2];
+		default:
+			return FText();
+		}
+	};
+	
 	if( ( InAxis & EAxisList::X ) ==  EAxisList::X)
 	{
 		AxisVec = AxisSize * ViewTM.InverseTransformVector( FVector(1,0,0) );
@@ -3599,7 +3632,7 @@ void FEditorViewportClient::DrawAxes(FViewport* InViewport, FCanvas* Canvas, con
 		LineItem.SetColor( FLinearColor::Red );
 		TextItem.SetColor( FLinearColor::Red );
 		LineItem.Draw( Canvas, AxisOrigin, AxisEnd );
-		TextItem.Text = LOCTEXT("XAxis","X");
+		TextItem.Text = GetAxisName(EAxisList::X);
 		TextItem.Draw( Canvas, FVector2D(AxisEnd.X + 2, AxisEnd.Y - 0.5*YL) );
 	}
 
@@ -3610,7 +3643,7 @@ void FEditorViewportClient::DrawAxes(FViewport* InViewport, FCanvas* Canvas, con
 		LineItem.SetColor( FLinearColor::Green );
 		TextItem.SetColor( FLinearColor::Green );
 		LineItem.Draw( Canvas, AxisOrigin, AxisEnd );
-		TextItem.Text = LOCTEXT("YAxis","Y");
+		TextItem.Text = GetAxisName(EAxisList::Y);
 		TextItem.Draw( Canvas, FVector2D(AxisEnd.X + 2, AxisEnd.Y - 0.5*YL) );
 
 	}
@@ -3622,7 +3655,7 @@ void FEditorViewportClient::DrawAxes(FViewport* InViewport, FCanvas* Canvas, con
 		LineItem.SetColor( FLinearColor::Blue );
 		TextItem.SetColor( FLinearColor::Blue );
 		LineItem.Draw( Canvas, AxisOrigin, AxisEnd );
-		TextItem.Text = LOCTEXT("ZAxis","Z");
+		TextItem.Text = GetAxisName(EAxisList::Z);
 		TextItem.Draw( Canvas, FVector2D(AxisEnd.X + 2, AxisEnd.Y - 0.5*YL) );
 	}
 }
