CookIndeterminism: AssetRegistry: Delete from the runtime AssetRegistry Assets that were present in a package when the package was last saved in editor, but are not present in the cooked package. This also fixes a difference in behavior between SPCook and MPCook; packages cooked on remote CookWorkers were not updating the package flags for those no-longer-existing assets, but packages cooked on director was, which created indeterministic results. Refactor UpdateAssetRegistryData to share as much code as possible between the CookDirector path and the CookWorker path.
#rnx

diff --git a/Engine/Source/Editor/UnrealEd/Private/Commandlets/AssetRegistryGenerator.cpp b/Engine/Source/Editor/UnrealEd/Private/Commandlets/AssetRegistryGenerator.cpp
--- a/Engine/Source/Editor/UnrealEd/Private/Commandlets/AssetRegistryGenerator.cpp
+++ b/Engine/Source/Editor/UnrealEd/Private/Commandlets/AssetRegistryGenerator.cpp
@@ -2657,25 +2657,71 @@
 	TOptional<TArray<FAssetDependency>>&& OverridePackageDependencies, UCookOnTheFlyServer& COTFS)
 {
 	LLM_SCOPE_BYTAG(Cooker_GeneratedAssetRegistry);
-	PreviousPackagesToUpdate.Remove(PackageName);
+	UE::Cook::FAssetRegistryPackageMessage Message = RecordUpdateAssetRegistryData(PackageName, TargetPlatform,
+		CookResult, SavePackageResult, MoveTemp(AssetDatasFromSave), MoveTemp(OverrideAssetPackageData),
+		MoveTemp(OverridePackageDependencies));
+	StoreUpdateAssetRegistryData(PackageName, MoveTemp(Message));
+}
 
-	uint32 NewPackageFlags = 0;
-	bool bSaveSucceeded = CookResult == UE::Cook::ECookResult::Succeeded;
+UE::Cook::FAssetRegistryPackageMessage FAssetRegistryGenerator::RecordUpdateAssetRegistryData(FName PackageName,
+	const ITargetPlatform* TargetPlatform, UE::Cook::ECookResult CookResult,
+	FSavePackageResultStruct* SavePackageResult,
+	TOptional<TArray<FAssetData>>&& AssetDatasFromSave,
+	TOptional<FAssetPackageData>&& OverrideAssetPackageData,
+	TOptional<TArray<FAssetDependency>>&& OverridePackageDependencies)
+{
+	UE::Cook::FAssetRegistryPackageMessage Message;
 
-	// Copy latest data for all Assets in the package into the cooked registry. This should be done even
-	// if not successful so that editor-only packages are recorded as well. When the AssetDatas were not
-	// calculated by SavePackage, copy them instead from the on-disk AssetRegistry.
+	// When the AssetDatas were not calculated by SavePackage, copy them instead from the on-disk AssetRegistry.
 	if (!AssetDatasFromSave)
 	{
 		AssetDatasFromSave.Emplace();
-		AssetRegistry.GetAssetsByPackageName(PackageName, *AssetDatasFromSave, true /* bIncludeOnlyDiskAssets */,
-			false /* SkipARFilteredAssets */);
+		IAssetRegistry::GetChecked().GetAssetsByPackageName(PackageName, *AssetDatasFromSave,
+			true /* bIncludeOnlyDiskAssets */, false /* SkipARFilteredAssets */);
 	}
-	for (FAssetData& AssetData : *AssetDatasFromSave)
+
+	bool bSaveSucceeded = CookResult == UE::Cook::ECookResult::Succeeded;
+	if (bSaveSucceeded)
 	{
+		check(SavePackageResult);
+
+		// Set the PackageFlags to the recorded value from SavePackage
+		Message.PackageFlags = SavePackageResult->SerializedPackageFlags;
+		Message.DiskSize = SavePackageResult->TotalFileSize;
+	}
+	else
+	{
+		// Set the package flags to zero to indicate that the package failed to save
+		Message.PackageFlags = 0;
+		// Set DiskSize (previous value was disksize in the WorkspaceDomain) to a negative number to indicate the
+		// cooked file does not exist. The magnitude of the negative value indicates CookResult.
+		Message.DiskSize = CookResultToDiskSize(CookResult);
+	}
+
+	Message.PackageName = PackageName;
+	Message.TargetPlatform = TargetPlatform;
+	Message.OverrideAssetPackageData = MoveTemp(OverrideAssetPackageData);
+	Message.OverridePackageDependencies = MoveTemp(OverridePackageDependencies);
+	Message.AssetDatas = MoveTemp(*AssetDatasFromSave);
+
+	return Message;
+}
+
+void FAssetRegistryGenerator::StoreUpdateAssetRegistryData(FName PackageName, UE::Cook::FAssetRegistryPackageMessage&& Message)
+{
+	PreviousPackagesToUpdate.Remove(PackageName);
+
+	// Copy latest data for all Assets in the package into the cooked registry. This should be done even
+	// if not successful so that editor-only packages are recorded as well.
+	for (FAssetData& AssetData : Message.AssetDatas)
+	{
+		AssetData.PackageFlags = Message.PackageFlags;
 		State.UpdateAssetData(MoveTemp(AssetData), true /* bCreateIfNotExists */);
 	}
+	bool bCookSucceeded = Message.DiskSize >= 0;
+	RemovePackageAssetsThatDoNotExistInCookedPackage(PackageName, Message.AssetDatas, bCookSucceeded);
 
+	// Copy AssetPackageData and Dependencies
 	FAssetPackageData* AssetPackageData = GetAssetPackageData(PackageName);
 	if (!bClonedGlobalAssetRegistry)
 	{
@@ -2689,41 +2735,79 @@
 		AssetRegistry.GetDependencies(PackageName, GlobalDependencies, UE::AssetRegistry::EDependencyCategory::All);
 		State.SetDependencies(FAssetIdentifier(PackageName), GlobalDependencies, UE::AssetRegistry::EDependencyCategory::All);
 	}
-	if (OverrideAssetPackageData)
+	if (Message.OverrideAssetPackageData)
 	{
 		// Generated packages pass in their own AssetPackageData
-		*AssetPackageData = MoveTemp(*OverrideAssetPackageData);
+		*AssetPackageData = MoveTemp(*Message.OverrideAssetPackageData);
 	}
-	if (OverridePackageDependencies)
+	AssetPackageData->DiskSize = Message.DiskSize;
+	// AssetPackageData->CookedHash is assigned during CookByTheBookFinished
+
+	if (Message.OverridePackageDependencies)
 	{
 		// Generated and Generator packages pass in their own dependencies
-		SetOverridePackageDependencies(PackageName, *OverridePackageDependencies);
+		SetOverridePackageDependencies(PackageName, *Message.OverridePackageDependencies);
 	}
+}
 
-	if (bSaveSucceeded)
+void FAssetRegistryGenerator::RemovePackageAssetsThatDoNotExistInCookedPackage(FName PackageName,
+	const TArray<FAssetData>& AssetDatasFromSave, bool bCookSucceeded)
+{
+	TArray<FSoftObjectPath> AssetsToRemove;
+	// The usual number of assets is 1. For that case, and all the way up to say 5, it is faster to make
+	// n^2 cheap compares than to construct a TSet and make n hash lookups.
+	constexpr int32 MaxSizeForLinearCompare = 5;
+	if (AssetDatasFromSave.Num() <= MaxSizeForLinearCompare)
 	{
-		check(SavePackageResult);
-
-		// Set the PackageFlags to the recorded value from SavePackage
-		NewPackageFlags = SavePackageResult->SerializedPackageFlags;
-
-		AssetPackageData->DiskSize = SavePackageResult->TotalFileSize;
-
-		// The CookedHash is assigned during CookByTheBookFinished		
+		State.EnumerateAssetsByPackageName(PackageName,
+			[&AssetsToRemove, &AssetDatasFromSave](const FAssetData* ExistingAsset)
+			{
+				FSoftObjectPath ExistingObjectPath = ExistingAsset->GetSoftObjectPath();
+				if (!AssetDatasFromSave.ContainsByPredicate([&ExistingObjectPath](const FAssetData& AssetFromSave)
+					{
+						return ExistingObjectPath == AssetFromSave.GetSoftObjectPath();
+					}))
+				{
+					AssetsToRemove.Add(ExistingObjectPath);
+				}
+				return true;
+			});
 	}
 	else
 	{
-		// Set the package flags to zero to indicate that the package failed to save
-		NewPackageFlags = 0;
-
-		// Set DiskSize (previous value was disksize in the WorkspaceDomain) to a negative number to indicate the
-		// cooked file does not exist. The magnitude of the negative value indicates CookResult.
-		AssetPackageData->DiskSize = CookResultToDiskSize(CookResult);
+		TSet<FSoftObjectPath> SetOfAssetsFromSave;
+		for (const FAssetData& AssetData : AssetDatasFromSave)
+		{
+			SetOfAssetsFromSave.Add(AssetData.GetSoftObjectPath());
+		}
+		State.EnumerateAssetsByPackageName(PackageName,
+			[&AssetsToRemove, &SetOfAssetsFromSave](const FAssetData* ExistingAsset)
+			{
+				FSoftObjectPath ExistingObjectPath = ExistingAsset->GetSoftObjectPath();
+				if (!SetOfAssetsFromSave.Contains(ExistingObjectPath))
+				{
+					AssetsToRemove.Add(ExistingObjectPath);
+				}
+				return true;
+			});
 	}
-
-	const bool bUpdated = UpdateAssetPackageFlags(PackageName, NewPackageFlags);
-	UE_CLOG(!bUpdated && bSaveSucceeded && !COTFS.bSkipSave, LogAssetRegistryGenerator, Warning,
-		TEXT("Trying to update asset package flags in package '%s' that does not exist"), *PackageName.ToString());
+	for (const FSoftObjectPath& AssetToRemove : AssetsToRemove)
+	{
+		bool bUnusedRemovedAssetData;
+		bool bUnusedRemovedPackageData;
+		if (bCookSucceeded)
+		{
+			// Give a log message to diagnose issues in case deleting the Asset is unexpected.
+			// But the log fires frequently in some cases of __ExternalObject__ packages that are not saved; suppress
+			// it when the package is not saved since deletion of assets in the package is unimportant for the
+			// non-runtime packages.
+			UE_LOG(LogAssetRegistryGenerator, Log,
+				TEXT("Removing Asset %s from the runtime AssetRegistry; it does not exist in the cooked package."),
+				*WriteToString<256>(AssetToRemove));
+		}
+		State.RemoveAssetData(AssetToRemove, false /* bRemoveDependencyData */,
+			bUnusedRemovedAssetData, bUnusedRemovedPackageData);
+	}
 }
 
 void FAssetRegistryGenerator::SetOverridePackageDependencies(FName PackageName, TConstArrayView<FAssetDependency> OverridePackageDependencies)
@@ -2743,24 +2827,8 @@
 	LLM_SCOPE_BYTAG(Cooker_GeneratedAssetRegistry);
 	const FName PackageName = Context.GetPackageName();
 	check(!PackageName.IsNone());
-	PreviousPackagesToUpdate.Remove(PackageName);
 
-	for (FAssetData& AssetData : Message.AssetDatas)
-	{
-		AssetData.PackageFlags = Message.PackageFlags;
-		State.UpdateAssetData(MoveTemp(AssetData), true /* bCreateIfNotExists */);
-	}
-	FAssetPackageData* AssetPackageData = GetAssetPackageData(PackageName);
-	if (Message.OverrideAssetPackageData)
-	{
-		*AssetPackageData = MoveTemp(*Message.OverrideAssetPackageData);
-	}
-	AssetPackageData->DiskSize = Message.DiskSize;
-
-	if (Message.OverridePackageDependencies)
-	{
-		SetOverridePackageDependencies(PackageName, *Message.OverridePackageDependencies);
-	}
+	StoreUpdateAssetRegistryData(PackageName, MoveTemp(Message));
 }
 
 namespace UE::Cook
@@ -2777,40 +2845,9 @@
 	TOptional<TArray<FAssetData>>&& AssetDatasFromSave, TOptional<FAssetPackageData>&& OverrideAssetPackageData,
 	TOptional<TArray<FAssetDependency>>&& OverridePackageDependencies, UCookOnTheFlyServer& COTFS)
 {
-	uint32 NewPackageFlags = 0;
-	int64 DiskSize = -1;
-	bool bSaveSucceeded = CookResult == UE::Cook::ECookResult::Succeeded;
-	if (bSaveSucceeded)
-	{
-		check(SavePackageResult);
-		NewPackageFlags = SavePackageResult->SerializedPackageFlags;
-		DiskSize = SavePackageResult->TotalFileSize;
-	}
-	else
-	{
-		// Set the package flags to zero to indicate that the package failed to save
-		NewPackageFlags = 0;
-		// Set DiskSize (previous value was disksize in the WorkspaceDomain) to a negative number to indicate the
-		// cooked file does not exist. The magnitude of the negative value indicates CookResult.
-		DiskSize = CookResultToDiskSize(CookResult);
-	}
+	FAssetRegistryPackageMessage Message = FAssetRegistryGenerator::RecordUpdateAssetRegistryData(
+		PackageName, TargetPlatform, CookResult, SavePackageResult,
+		MoveTemp(AssetDatasFromSave), MoveTemp(OverrideAssetPackageData), MoveTemp(OverridePackageDependencies));
 
-	FAssetRegistryPackageMessage Message;
-	Message.PackageName = PackageName;
-	Message.TargetPlatform = TargetPlatform;
-	Message.PackageFlags = NewPackageFlags;
-	Message.DiskSize = DiskSize;
-	Message.OverrideAssetPackageData = MoveTemp(OverrideAssetPackageData);
-	Message.OverridePackageDependencies = MoveTemp(OverridePackageDependencies);
-
-	// Add to the message all the AssetDatas in the package
-	if (!AssetDatasFromSave)
-	{
-		AssetDatasFromSave.Emplace();
-		IAssetRegistry::Get()->GetAssetsByPackageName(PackageName, *AssetDatasFromSave,
-			true /* bIncludeOnlyDiskAssets */, false /* SkipARFilteredAssets */);
-	}
-	Message.AssetDatas = MoveTemp(*AssetDatasFromSave);
-
 	// Send the message to the director
 	FCbWriter Writer;
diff --git a/Engine/Source/Editor/UnrealEd/Public/Commandlets/AssetRegistryGenerator.h b/Engine/Source/Editor/UnrealEd/Public/Commandlets/AssetRegistryGenerator.h
--- a/Engine/Source/Editor/UnrealEd/Public/Commandlets/AssetRegistryGenerator.h
+++ b/Engine/Source/Editor/UnrealEd/Public/Commandlets/AssetRegistryGenerator.h
@@ -27,6 +27,7 @@ struct FChunkDependencyTreeNode;
 struct FSoftObjectPath;
 namespace UE::Cook { class FAssetRegistryMPCollector; }
 namespace UE::Cook { class FAssetRegistryPackageMessage; }
+namespace UE::Cook { class FAssetRegistryReporterRemote; }
 namespace UE::Cook { class FCookSandbox; }
 namespace UE::Cook { class FCookWorkerClient; }
 namespace UE::Cook { struct FPackageData; }
@@ -283,6 +284,16 @@ private:
 	 */
 	void UpdateKeptPackages();
 
+	static UE::Cook::FAssetRegistryPackageMessage RecordUpdateAssetRegistryData(FName PackageName,
+		const ITargetPlatform* TargetPlatform, UE::Cook::ECookResult CookResult,
+		FSavePackageResultStruct* SavePackageResult,
+		TOptional<TArray<FAssetData>>&& AssetDatasFromSave,
+		TOptional<FAssetPackageData>&& OverrideAssetPackageData,
+		TOptional<TArray<FAssetDependency>>&& OverridePackageDependencies);
+	void StoreUpdateAssetRegistryData(FName PackageName, UE::Cook::FAssetRegistryPackageMessage&& Message);
+	/** Delete Assets we copied from the global AssetRegistry that no longer exist in the cook save of the package. */
+	void RemovePackageAssetsThatDoNotExistInCookedPackage(FName PackageName,
+		const TArray<FAssetData>& AssetDatasFromSave, bool bCookSucceeded);
 	void SetOverridePackageDependencies(FName PackageName, TConstArrayView<FAssetDependency> OverridePackageDependencies);
 
 	bool ComputePackageDifferences_IsPackageFileUnchanged(const FComputeDifferenceOptions& Options, FName PackageName,
@@ -439,6 +450,8 @@ private:
 
 	/** If InState records PackageName is generated, return the name of the Generator, otherwise return NAME_None. */
 	static FName GetGeneratorPackage(FName PackageName, const FAssetRegistryState& InState);
+
+	friend class UE::Cook::FAssetRegistryReporterRemote;
 };
 
 namespace UE::Cook
