World Bookmark - Made the World property read-only

diff --git a/Engine/Source/Editor/UnrealEd/Public/EditorState/WorldEditorState.h b/Engine/Source/Editor/UnrealEd/Public/EditorState/WorldEditorState.h
--- a/Engine/Source/Editor/UnrealEd/Public/EditorState/WorldEditorState.h
+++ b/Engine/Source/Editor/UnrealEd/Public/EditorState/WorldEditorState.h
@@ -22,7 +22,7 @@ private:
 	//~ End UEditorState Interface
 
 	// World associated with the bookmark
-	UPROPERTY(EditAnywhere, Category = General)
+	UPROPERTY(VisibleAnywhere, Category = General)
 	TSoftObjectPtr<UWorld> World;
 };
 
