Fixed crashes occurring when undoing operations on a material created by dragging and dropping a texture onto a mesh. The created material was missing the Transactional flag, causing undo operations to not fully restore the original state.

[RN] Fixed crashes occurring when undoing operations on a material created by dragging and dropping a texture onto a mesh. The created material was missing the Transactional flag, causing undo operations to not fully restore the original state.

diff --git a/Engine/Source/Editor/UnrealEd/Private/LevelEditorViewport.cpp b/Engine/Source/Editor/UnrealEd/Private/LevelEditorViewport.cpp
--- a/Engine/Source/Editor/UnrealEd/Private/LevelEditorViewport.cpp
+++ b/Engine/Source/Editor/UnrealEd/Private/LevelEditorViewport.cpp
@@ -812,7 +812,7 @@ UObject* FLevelEditorViewportClient::GetOrCreateMaterialFromTexture(UTexture* Un
 		UMaterialFactoryNew* MaterialFactory = NewObject<UMaterialFactoryNew>();
 
 		UMaterial* UnrealMaterial = (UMaterial*)MaterialFactory->FactoryCreateNew(
-			UMaterial::StaticClass(), Package, *MaterialFullName, RF_Standalone | RF_Public, NULL, GWarn);
+			UMaterial::StaticClass(), Package, *MaterialFullName, RF_Standalone | RF_Public | RF_Transactional, NULL, GWarn);
 
 		if (UnrealMaterial == nullptr)
 		{
