Wrap problematic checks with #if DO_CHECK
#rnx

diff --git a/Engine/Plugins/Animation/PoseSearch/Source/Runtime/Private/PoseSearchAssetIndexer.cpp b/Engine/Plugins/Animation/PoseSearch/Source/Runtime/Private/PoseSearchAssetIndexer.cpp
--- a/Engine/Plugins/Animation/PoseSearch/Source/Runtime/Private/PoseSearchAssetIndexer.cpp
+++ b/Engine/Plugins/Animation/PoseSearch/Source/Runtime/Private/PoseSearchAssetIndexer.cpp
@@ -846,7 +846,9 @@ const UPoseSearchSchema* FAssetIndexer::GetSchema() const
 #if WITH_EDITOR
 float FAssetIndexer::CalculatePermutationTimeOffset() const
 {
+#if DO_CHECK
 	check(Schema.PermutationsSampleRate > 0 && SearchIndexAsset.IsInitialized());
+#endif
 	const float PermutationTimeOffset = Schema.PermutationsTimeOffset + SearchIndexAsset.GetPermutationIdx() / float(Schema.PermutationsSampleRate);
 	return PermutationTimeOffset;
 }
diff --git a/Engine/Source/Editor/UnrealEd/Private/EditorDomain/EditorBuildInputResolver.cpp b/Engine/Source/Editor/UnrealEd/Private/EditorDomain/EditorBuildInputResolver.cpp
--- a/Engine/Source/Editor/UnrealEd/Private/EditorDomain/EditorBuildInputResolver.cpp
+++ b/Engine/Source/Editor/UnrealEd/Private/EditorDomain/EditorBuildInputResolver.cpp
@@ -216,7 +216,9 @@ TResolveInputRequest<BulkPayloadType, CallbackPayloadType, CallbackType>::TResol
 template <class BulkPayloadType, class CallbackPayloadType, class CallbackType>
 void TResolveInputRequest<BulkPayloadType, CallbackPayloadType, CallbackType>::Cancel()
 {
+#if DO_CHECK
 	check(bCreationComplete);
+#endif
 	if (bCallbackCalled.exchange(true, std::memory_order_relaxed) == false)
 	{
 		Owner.End(this, [this]
@@ -234,7 +236,9 @@ void TResolveInputRequest<BulkPayloadType, CallbackPayloadType, CallbackType>::C
 template <class BulkPayloadType, class CallbackPayloadType, class CallbackType>
 void TResolveInputRequest<BulkPayloadType, CallbackPayloadType, CallbackType>::Wait()
 {
+#if DO_CHECK
 	check(bCreationComplete);
+#endif
 	if (!bCallbackCalled.load(std::memory_order_relaxed))
 	{
 		for (FPayloadData& PayloadData : PayloadDatas)
diff --git a/Engine/Source/Runtime/CoreUObject/Private/Serialization/Formatters/JsonArchiveInputFormatter.cpp b/Engine/Source/Runtime/CoreUObject/Private/Serialization/Formatters/JsonArchiveInputFormatter.cpp
--- a/Engine/Source/Runtime/CoreUObject/Private/Serialization/Formatters/JsonArchiveInputFormatter.cpp
+++ b/Engine/Source/Runtime/CoreUObject/Private/Serialization/Formatters/JsonArchiveInputFormatter.cpp
@@ -425,7 +425,9 @@ void FJsonArchiveInputFormatter::Serialize(FString& Value)
 	bool bSuccess =
 #endif
 		ValueStack.Top()->TryGetString(Value);
+#if DO_CHECK
 	check(bSuccess || ValueStack.Top()->IsNull());
+#endif
 	Value.RemoveFromStart(TEXT("String:"));
 }
 
diff --git a/Engine/Source/Runtime/RenderCore/Private/PipelineCacheUtilities.cpp b/Engine/Source/Runtime/RenderCore/Private/PipelineCacheUtilities.cpp
--- a/Engine/Source/Runtime/RenderCore/Private/PipelineCacheUtilities.cpp
+++ b/Engine/Source/Runtime/RenderCore/Private/PipelineCacheUtilities.cpp
@@ -522,7 +522,9 @@ bool UE::PipelineCacheUtilities::SaveStablePipelineCacheFile(const FString& Outp
 		int64 TotalNumberOfStableShaderKeys = StableShaderKeyIndexTable.Num();
 		for (const FPermsPerPSO& Item : StableResults)
 		{
+#if DO_CHECK
 			check(UE::PipelineCacheUtilities::Private::SanityCheckActiveSlots(Item));
+#endif
 			UE::PipelineCacheUtilities::Private::SaveActiveSlots(MemWriter, Item);
 
 			FPipelineCacheFileFormatPSO NewPSO = *Item.PSO;
@@ -753,8 +755,9 @@ bool UE::PipelineCacheUtilities::LoadStablePipelineCacheFile(const FString& File
 		OriginalPSOs[PermutationGroupIdx].BindCount = ReadVarIntFromArchive(MemReader);
 
 		Item.PSO = &OriginalPSOs[PermutationGroupIdx];
-
+#if DO_CHECK
 		check(UE::PipelineCacheUtilities::Private::SanityCheckActiveSlots(Item));
+#endif
 
 		int64 NumPermutations = ReadVarIntFromArchive(MemReader);
 		Item.Permutations.Reserve(NumPermutations);
