Replace explicit r.Substrate access with Substrate::IsSubstrateEnabled() function.

This is now possible as Substrate::IsSubstrateEnabled() is declared/defined in RenderCore, which is accessible to all needed modules.

diff --git a/Engine/Source/Editor/UnrealEd/Private/Lightmass/Lightmass.cpp b/Engine/Source/Editor/UnrealEd/Private/Lightmass/Lightmass.cpp
--- a/Engine/Source/Editor/UnrealEd/Private/Lightmass/Lightmass.cpp
+++ b/Engine/Source/Editor/UnrealEd/Private/Lightmass/Lightmass.cpp
@@ -59,8 +59,6 @@
 
 extern FSwarmDebugOptions GSwarmDebugOptions;
 
-bool Lightmass_IsSubstrateEnabled();
-
 //@todo_ow: Potentially move to core for sharing ?
 class FObjectMemoryWriter : public  FMemoryWriter 
 {
@@ -1441,7 +1439,7 @@ void FLightmassExporter::GetMaterialHash(const UMaterialInterface* Material, FSH
 		}
 	}
 
-	if (Lightmass_IsSubstrateEnabled())
+	if (Substrate::IsSubstrateEnabled())
 	{
 		uint32 LightmassSubstrateVersion = 0XFA1081D0; // This can be change when the code/logic for converting material to Substrate for lightmap has changed.
 		HashState.Update((const uint8*)&LightmassSubstrateVersion, sizeof(LightmassSubstrateVersion));
diff --git a/Engine/Source/Editor/UnrealEd/Private/Lightmass/LightmassRender.cpp b/Engine/Source/Editor/UnrealEd/Private/Lightmass/LightmassRender.cpp
--- a/Engine/Source/Editor/UnrealEd/Private/Lightmass/LightmassRender.cpp
+++ b/Engine/Source/Editor/UnrealEd/Private/Lightmass/LightmassRender.cpp
@@ -32,12 +32,6 @@
 
 DEFINE_LOG_CATEGORY_STATIC(LogLightmassRender, Error, All);
 
-extern bool Lightmass_IsSubstrateEnabled()
-{
-	static const auto CVar = IConsoleManager::Get().FindTConsoleVariableDataInt(TEXT("r.Substrate"));
-	return CVar && CVar->GetValueOnAnyThread() > 0;
-}
-
 // FLightmassMaterialCompiler - A proxy compiler that overrides various compiler functions for potential problem expressions.
 struct FLightmassMaterialCompiler : public FProxyMaterialCompiler
 {
@@ -260,7 +254,7 @@ public:
 	{
 		if (InMaterialInterface)
 		{
-			bSubstrateEnabled = Lightmass_IsSubstrateEnabled();
+			bSubstrateEnabled = Substrate::IsSubstrateEnabled();
 			MaterialInterface = InMaterialInterface;
 			Material = MaterialInterface ? MaterialInterface->GetMaterial() : NULL;
 			PropertyToCompile = InPropertyToCompile;
@@ -638,7 +632,7 @@ public:
 	{ 
 		// Substrate needs the real material shading model since the expressions access GetMaterialShadingModels() through the compiler to generate the substrate operators, 
 		// and we do not want unlit materials when it is a Slab or a ShadingModel node.
-		return Lightmass_IsSubstrateEnabled() ? Material->GetShadingModels() : MSM_Unlit; 
+		return Substrate::IsSubstrateEnabled() ? Material->GetShadingModels() : MSM_Unlit; 
 	} 
 	/**
 	 * Should shaders compiled for this material be saved to disk?
@@ -1080,7 +1074,7 @@ void FLightmassMaterialRenderer::BeginGenerateMaterialData(
 	if (BaseMaterial)
 	{
 		check(!MaterialExportData.Contains(InMaterial));
-		const bool bSubstrateEnabled = Lightmass_IsSubstrateEnabled();
+		const bool bSubstrateEnabled = Substrate::IsSubstrateEnabled();
 
 		FMaterialExportDataEntry& MaterialData = MaterialExportData.Add(InMaterial, FMaterialExportDataEntry(ChannelName));
 
@@ -1141,7 +1135,7 @@ bool FLightmassMaterialRenderer::GenerateMaterialData(
 	check(BaseMaterial);
 
 	EBlendMode BlendMode = InMaterial.GetBlendMode();
-	const bool bSubstrateEnabled = Lightmass_IsSubstrateEnabled();
+	const bool bSubstrateEnabled = Substrate::IsSubstrateEnabled();
 
 	FMaterialShadingModelField ShadingModels = InMaterial.GetShadingModels();
  	if (!bSubstrateEnabled &&		// Shading models are irrelevant when using Substrate
diff --git a/Engine/Source/Runtime/Engine/Private/Materials/MaterialAttributeDefinitionMap.cpp b/Engine/Source/Runtime/Engine/Private/Materials/MaterialAttributeDefinitionMap.cpp
--- a/Engine/Source/Runtime/Engine/Private/Materials/MaterialAttributeDefinitionMap.cpp
+++ b/Engine/Source/Runtime/Engine/Private/Materials/MaterialAttributeDefinitionMap.cpp
@@ -17,6 +17,7 @@
 #include "Materials/MaterialExpressionFirstPersonOutput.h"
 #include "Serialization/ShaderKeyGenerator.h"
 #include "SubstrateDefinitions.h"
+#include "RenderUtils.h"
 
 #define LOCTEXT_NAMESPACE "MaterialShared"
 
@@ -434,8 +435,7 @@ FText FMaterialAttributeDefinitionMap::GetAttributeOverrideForMaterial(const FGu
 	TArray<TKeyValuePair<EMaterialShadingModel, FString>> CustomPinNames;
 	EMaterialProperty Property = GMaterialPropertyAttributesMap.Find(AttributeID)->Property;
 
-	static const IConsoleVariable* SubstrateCVar = IConsoleManager::Get().FindConsoleVariable(TEXT("r.substrate"));
-	const bool bSubstrateEnabled = (SubstrateCVar && SubstrateCVar->GetInt() != 0);
+	const bool bSubstrateEnabled = Substrate::IsSubstrateEnabled();
 
 	switch (Property)
 	{
diff --git a/Engine/Source/Runtime/Engine/Private/UnrealEngine.cpp b/Engine/Source/Runtime/Engine/Private/UnrealEngine.cpp
--- a/Engine/Source/Runtime/Engine/Private/UnrealEngine.cpp
+++ b/Engine/Source/Runtime/Engine/Private/UnrealEngine.cpp
@@ -3527,8 +3527,7 @@ void UEngine::ConditionallyLoadPreIntegratedSkinBRDFTexture()
 
 void UEngine::LoadLTCTextures()
 {
-	static const auto CVar = IConsoleManager::Get().FindTConsoleVariableDataInt(TEXT("r.Substrate"));
-	const bool bSubstrateEnabled = CVar && CVar->GetValueOnAnyThread() > 0;
+	const bool bSubstrateEnabled = Substrate::IsSubstrateEnabled();
 
 	const bool bGGX = true;
 	const bool bSheen = bSubstrateEnabled;
@@ -3555,8 +3554,7 @@ void UEngine::LoadEnergyTextures()
 	const bool bPathTracingEnabled = false;
 	#endif
 
-	static const auto CVarSubstrate = IConsoleManager::Get().FindTConsoleVariableDataInt(TEXT("r.Substrate"));
-	const bool bSubstrateEnabled = CVarSubstrate && CVarSubstrate->GetValueOnAnyThread() > 0;
+	const bool bSubstrateEnabled = Substrate::IsSubstrateEnabled();
 
 	static const auto CVarEnergyConservation = IConsoleManager::Get().FindTConsoleVariableDataInt(TEXT("r.Material.EnergyConservation"));
 	const bool bEnergyConservationEnabled = CVarEnergyConservation && CVarEnergyConservation->GetValueOnAnyThread() > 0;
