Change the editor save button to always check package dirtyness before saving packages to avoid saving non-dirty packages.
- Historically, this button forced a save for the current level, even if nothing changed.
- When external actors support was introduced, this behavior changed to check package dirtyness only if the level used external actors.
- This change makes the behavior consistent across level types.

This change is required to fix the cased where a user would add external entities to a standard level (not using external actors), which would ask the user to checkout the level package and all external entity files.

diff --git a/Engine/Source/Editor/UnrealEd/Private/FileHelpers.cpp b/Engine/Source/Editor/UnrealEd/Private/FileHelpers.cpp
--- a/Engine/Source/Editor/UnrealEd/Private/FileHelpers.cpp
+++ b/Engine/Source/Editor/UnrealEd/Private/FileHelpers.cpp
@@ -4156,15 +4156,12 @@ bool FEditorFileUtils::SaveCurrentLevel()
 	ULevel* Level = GWorld->GetCurrentLevel();
 	if (Level)
 	{
-		// Check dirtiness if the level is using external objects, no need to save it needlessly
-		bool bCheckDirty = Level->IsUsingExternalObjects();
-
 		TArray<UPackage*> PackagesToSave;
 		
 		UPackage* LevelPackage = Level->GetPackage();
 
 		// Get Packages to save
-		if (!bCheckDirty || LevelPackage->IsDirty() || LevelPackage->HasAnyPackageFlags(PKG_NewlyCreated))
+		if (LevelPackage->IsDirty() || LevelPackage->HasAnyPackageFlags(PKG_NewlyCreated))
 		{
 			PackagesToSave.Add(LevelPackage);
 		}
@@ -4175,7 +4172,7 @@ bool FEditorFileUtils::SaveCurrentLevel()
 		{
 			if (FPackageName::IsValidLongPackageName(ExternalPackage->GetName()))
 			{
-				if (!bCheckDirty || ExternalPackage->IsDirty() || UPackage::IsEmptyPackage(ExternalPackage))
+				if (ExternalPackage->IsDirty() || ExternalPackage->HasAnyPackageFlags(PKG_NewlyCreated) || UPackage::IsEmptyPackage(ExternalPackage))
 				{
 					PackagesToSave.Add(ExternalPackage);
 				}
@@ -4192,7 +4189,7 @@ bool FEditorFileUtils::SaveCurrentLevel()
 			const bool bPromptUserToSave = false;
 			const bool bFastSave = false;
 			const bool bCanBeDeclined = false;
-			bCheckDirty = false; // force the flag back to false because we already checked conditions to add to PackagesToSave. Some Packages like newly created packages might not be dirty and we still want to save them.
+			const bool bCheckDirty = false; // force the flag back to false because we already checked conditions to add to PackagesToSave. Some Packages like newly created packages might not be dirty and we still want to save them.
 			bReturnCode &= InternalSavePackages(PackagesToSave, bPromptUserToSave, bFastSave, bCanBeDeclined, bCheckDirty);
 		}
 	}
