[ContentBrowser] Drag and Drop widget restyle

diff --git a/Engine/Source/Editor/ContentBrowser/Private/DragDropHandler.cpp b/Engine/Source/Editor/ContentBrowser/Private/DragDropHandler.cpp
--- a/Engine/Source/Editor/ContentBrowser/Private/DragDropHandler.cpp
+++ b/Engine/Source/Editor/ContentBrowser/Private/DragDropHandler.cpp
@@ -166,10 +166,15 @@ bool ValidateGenericDragEvent(const FContentBrowserItem& InItem, const FDragDrop
 		{
 			int32 NumDraggedItems = ContentDragDropOp->GetDraggedItems().Num();
 			int32 NumCanMoveOrCopy = 0;
+			bool bCanCopyOnce = false;
+			bool bCanMoveOnce = false;
 			for (const FContentBrowserItem& DraggedItem : ContentDragDropOp->GetDraggedItems())
 			{
-				const bool bCanMoveOrCopy = DraggedItem.CanMove(InItem.GetVirtualPath()) || DraggedItem.CanCopy(InItem.GetVirtualPath());
-				if (bCanMoveOrCopy)
+				const bool bCanMove = DraggedItem.CanMove(InItem.GetVirtualPath());
+				const bool bCanCopy = DraggedItem.CanCopy(InItem.GetVirtualPath());
+				bCanCopyOnce |= bCanCopy;
+				bCanMoveOnce |= bCanMove;
+				if (bCanMove || bCanCopy)
 				{
 					++NumCanMoveOrCopy;
 				}
@@ -181,10 +186,24 @@ bool ValidateGenericDragEvent(const FContentBrowserItem& InItem, const FDragDrop
 			}
 			else
 			{
+				FText MoveOrCopyString = FText::GetEmpty();
+				if (bCanCopyOnce && bCanMoveOnce)
+				{
+					MoveOrCopyString = LOCTEXT("DragAndDropMoveOrCopy", "Move or copy");
+				}
+				else if (bCanCopyOnce)
+				{
+					MoveOrCopyString = LOCTEXT("DragAndDropCopy", "Copy");
+				}
+				else
+				{
+					MoveOrCopyString = LOCTEXT("DragAndDropMove", "Move");
+				}
+
 				const FText FirstItemText = ContentDragDropOp->GetDraggedItems()[0].GetDisplayName();
 				const FText MoveOrCopyText = (NumCanMoveOrCopy > 1)
-					? FText::Format(LOCTEXT("OnDragAssetsOverFolder_MultipleItems", "Move or copy '{0}' and {1} {1}|plural(one=other,other=others)"), FirstItemText, NumDraggedItems - 1)
-					: FText::Format(LOCTEXT("OnDragAssetsOverFolder_SingularItems", "Move or copy '{0}'"), FirstItemText);
+					? FText::Format(LOCTEXT("OnDragAssetsOverFolder_MultipleItems", "{0} '{1}' and {2} {2}|plural(one=other,other=others)"), MoveOrCopyString, FirstItemText, NumDraggedItems - 1)
+					: FText::Format(LOCTEXT("OnDragAssetsOverFolder_SingularItems", "{0} '{1}'"), MoveOrCopyString, FirstItemText);
 
 				if (NumCanMoveOrCopy < NumDraggedItems)
 				{
diff --git a/Engine/Source/Editor/EditorStyle/Private/StarshipStyle.cpp b/Engine/Source/Editor/EditorStyle/Private/StarshipStyle.cpp
--- a/Engine/Source/Editor/EditorStyle/Private/StarshipStyle.cpp
+++ b/Engine/Source/Editor/EditorStyle/Private/StarshipStyle.cpp
@@ -4193,8 +4193,8 @@ void FStarshipEditorStyle::FStyle::SetupGraphEditorStyles()
 		Set("Graph.Pin.Dummy", new IMAGE_BRUSH("/Graph/Pin_dummy", FVector2D(15, 11)));
 
 		Set( "Graph.ConnectorFeedback.Border", new BOX_BRUSH( "Old/Menu_Background", FMargin(8.0f/64.0f) ) );
-		Set( "Graph.ConnectorFeedback.OK", new IMAGE_BRUSH( "Old/Graph/Feedback_OK", Icon16x16 ) );
-		Set( "Graph.ConnectorFeedback.OKWarn", new IMAGE_BRUSH( "Old/Graph/Feedback_OKWarn", Icon16x16 ) );
+		Set( "Graph.ConnectorFeedback.OK", new CORE_IMAGE_BRUSH_SVG("Starship/Common/check-circle", Icon16x16, FStyleColors::AccentGreen));
+		Set( "Graph.ConnectorFeedback.OKWarn", new IMAGE_BRUSH_SVG("Starship/Common/check-circle", Icon16x16, FStyleColors::AccentYellow));
 		Set( "Graph.ConnectorFeedback.Error", new IMAGE_BRUSH( "Old/Graph/Feedback_Error", Icon16x16 ) );
 		Set( "Graph.ConnectorFeedback.NewNode", new IMAGE_BRUSH( "Old/Graph/Feedback_NewNode", Icon16x16 ) );
 		Set( "Graph.ConnectorFeedback.ViaCast", new IMAGE_BRUSH( "Old/Graph/Feedback_ConnectViaCast", Icon16x16 ) );
@@ -6819,6 +6819,8 @@ void FStarshipEditorStyle::FStyle::SetupContentBrowserStyle()
 
 		Set( "ContentBrowser.ContentDirty", new IMAGE_BRUSH_SVG("Starship/Common/DirtyBadge", Icon16x16) );
 		Set( "ContentBrowser.AssetDragDropTooltipBackground", new BOX_BRUSH( "Old/Menu_Background", FMargin(8.0f/64.0f) ) );
+		Set( "ContentBrowser.DragDropBackground", new FSlateRoundedBoxBrush(FStyleColors::Input, 0.f, FStyleColors::InputOutline, 1.f));
+		Set( "ContentBrowser.ThumbnailDragDropBackground", new FSlateRoundedBoxBrush(FStyleColors::Transparent, 0.f, FStyleColors::InputOutline, 1.f));
 		Set( "ContentBrowser.CollectionTreeDragDropBorder", new BOX_BRUSH( "Old/Window/ViewportDebugBorder", 0.8f ) );
 		Set( "ContentBrowser.PopupMessageIcon", new IMAGE_BRUSH( "Icons/alert", Icon32x32) );
 		Set( "ContentBrowser.NewFolderIcon", new CORE_IMAGE_BRUSH_SVG("Starship/Common/folder-plus", Icon16x16 ) );
diff --git a/Engine/Source/Editor/UnrealEd/Private/DragAndDrop/AssetDragDropOp.cpp b/Engine/Source/Editor/UnrealEd/Private/DragAndDrop/AssetDragDropOp.cpp
--- a/Engine/Source/Editor/UnrealEd/Private/DragAndDrop/AssetDragDropOp.cpp
+++ b/Engine/Source/Editor/UnrealEd/Private/DragAndDrop/AssetDragDropOp.cpp
@@ -90,6 +90,7 @@ TSharedPtr<SWidget> FAssetDragDropOp::GetDefaultDecorator() const
 		{
 			FAssetThumbnailConfig AssetThumbnailConfig;
 			AssetThumbnailConfig.bShowAssetChip = true;
+			AssetThumbnailConfig.AssetChipBorderImageOverride = FAppStyle::GetBrush("ContentBrowser.ThumbnailDragDropBackground");
 			ThumbnailWidget = AssetThumbnail->MakeThumbnailWidget(AssetThumbnailConfig);
 		}
 		else
@@ -127,6 +128,79 @@ TSharedPtr<SWidget> FAssetDragDropOp::GetDefaultDecorator() const
 		}
 	}
 
+	if (UE::Editor::ContentBrowser::IsNewStyleEnabled())
+	{
+		return SNew(SBorder)
+				.BorderImage(FAppStyle::GetBrush("ContentBrowser.DragDropBackground"))
+				.Padding(0.f)
+				[
+					SNew(SHorizontalBox)
+
+					// Left slot is for the thumbnail
+					+SHorizontalBox::Slot()
+					.AutoWidth()
+					.HAlign(HAlign_Left)
+					[
+						SNew(SBox)
+						.WidthOverride(ThumbnailSize)
+						.HeightOverride(ThumbnailSize)
+						[
+							SNew(SOverlay)
+
+							+SOverlay::Slot()
+							[
+								ThumbnailWidget.ToSharedRef()
+							]
+
+							+SOverlay::Slot()
+							.HAlign(HAlign_Center)
+							.VAlign(VAlign_Top)
+							.Padding(FMargin(0, 4, 0, 0))
+							[
+								SNew(SBorder)
+								.BorderImage(FAppStyle::GetBrush("Menu.Background"))
+								.Visibility(TotalCount > 1 ? EVisibility::Visible : EVisibility::Collapsed)
+								.Content()
+								[
+									SNew(STextBlock)
+									.Text(FText::AsNumber(TotalCount))
+								]
+							]
+						]
+					]
+
+					// Right slot is for optional tooltip
+					+SHorizontalBox::Slot()
+					.Padding(8.f, 0.f, 0.f, 0.f)
+					.AutoWidth()
+					.VAlign(VAlign_Center)
+					[
+						SNew(SBox)
+						.MinDesiredWidth(80.0f)
+						[
+							SNew(SHorizontalBox)
+
+							+SHorizontalBox::Slot()
+							.AutoWidth()
+							.Padding(0.f, 0.f, 4.0f, 0.f)
+							.VAlign(VAlign_Center)
+							[
+								SNew(SImage) 
+								.Image(this, &FAssetDragDropOp::GetIcon)
+							]
+
+							+SHorizontalBox::Slot()
+							.AutoWidth()
+							.Padding(0,0,8,0)
+							.VAlign(VAlign_Center)
+							[
+								SNew(STextBlock) 
+								.Text(this, &FAssetDragDropOp::GetDecoratorText)
+							]
+						]
+					]
+				];
+	}
 	return 
 		SNew(SBorder)
 		.BorderImage(FAppStyle::GetBrush("ContentBrowser.AssetDragDropTooltipBackground"))
