When cooking with -DlcReevaluateUncookedAssets dependencies of packages that have already been cooked will be evaluated so potential platform dependencies that were not previously cooked have a chance to be cooked.
This is useful when cooking a cooked editor as DLC based upon a cooked client. Where when cooking the client editor only dependencies would not have cooked before.

Fix a bug where packages were only requested if they had already been cooked rather than if they were not.

diff --git a/Engine/Source/Editor/UnrealEd/Private/Cooker/CookRequestCluster.cpp b/Engine/Source/Editor/UnrealEd/Private/Cooker/CookRequestCluster.cpp
--- a/Engine/Source/Editor/UnrealEd/Private/Cooker/CookRequestCluster.cpp
+++ b/Engine/Source/Editor/UnrealEd/Private/Cooker/CookRequestCluster.cpp
@@ -218,7 +218,11 @@ FRequestCluster::FRequestCluster(UCookOnTheFlyServer& InCOTFS, TRingBuffer<FDisc
 			EUrgency Urgency = Discovery->Urgency;
 			DiscoveryQueue.PopFrontValue();
 			Discovery = nullptr;
-			if (!PackageData.IsInProgress() && PackageData.GetPlatformsNeedingCommitNum(EReachability::Runtime) == 0)
+			if (PackageData.IsInProgress())
+			{
+				PackageData.RaiseUrgency(Urgency, ESendFlags::QueueAddAndRemove);
+			}
+			else if (PackageData.GetPlatformsNeedingCommitNum(EReachability::Runtime) > 0 || !PackageData.AreAllReachablePlatformsVisitedByCluster(EReachability::Runtime))
 			{
 				PackageData.SendToState(EPackageState::Request, ESendFlags::QueueRemove,
 					EStateChangeReason::RequestCluster);
@@ -226,10 +230,6 @@ FRequestCluster::FRequestCluster(UCookOnTheFlyServer& InCOTFS, TRingBuffer<FDisc
 				FVertexData& Vertex = FindOrAddVertex(PackageData);
 				SetOwnedByCluster(Vertex, true /* bOwnedByCluster */, false /* bNeedsStateChange */);
 			}
-			else if (PackageData.IsInProgress())
-			{
-				PackageData.RaiseUrgency(Urgency, ESendFlags::QueueAddAndRemove);
-			}
 			continue;
 		}
 
@@ -471,7 +471,7 @@ void FRequestCluster::FetchPackageNames(const FCookerTimer& CookerTimer, bool& b
 					PackageData->SetUrgency(EUrgency::Blocking, ESendFlags::QueueAddAndRemove);
 				}
 			}
-			else if (PackageData->GetPlatformsNeedingCommitNum(EReachability::Runtime) > 0)
+			else if (PackageData->GetPlatformsNeedingCommitNum(EReachability::Runtime) > 0 || !PackageData->AreAllReachablePlatformsVisitedByCluster(EReachability::Runtime))
 			{
 				// If it's missing cookable platforms and not in progress we need to add it to the cluster for cooking
 				FVertexData& Vertex = FindOrAddVertex(*PackageData);
@@ -2780,4 +2780,4 @@ static TArray<T> FindRootsFromLeafToRootOrderList(TConstArrayView<T> LeafToRootO
 	return Roots;
 }
 
-} // namespace UE::Cook
+} // namespace UE::Cook
\ No newline at end of file
