Reenable fast reinstancing, with an added customization for anim blueprints to reinitialize their hidden backpointers

#rnx

diff --git a/Engine/Source/Editor/Kismet/Private/BlueprintCompilationManager.cpp b/Engine/Source/Editor/Kismet/Private/BlueprintCompilationManager.cpp
--- a/Engine/Source/Editor/Kismet/Private/BlueprintCompilationManager.cpp
+++ b/Engine/Source/Editor/Kismet/Private/BlueprintCompilationManager.cpp
@@ -124,7 +124,7 @@ namespace UE::Kismet::BlueprintCompilationManager::Private
 		 * by compilation. This can only work when compiling in place (e.g. reusing UClasses) - so
 		 * can't be done for asset reload, live coding, or verse 
 		 */
-		static bool bAllowFastReinstancing = false; // Disabled while a crash when reinstancing animation blueprints is investigated
+		static bool bAllowFastReinstancing = true;
 		static FAutoConsoleVariableRef CVarAllowFastReinstancing(
 			TEXT("BP.bAllowFastReinstancing"), bAllowFastReinstancing,
 			TEXT("True by default, this flag enables a fast path in the logic that updates class \
diff --git a/Engine/Source/Editor/UnrealEd/Private/Kismet2/KismetReinstanceUtilities.cpp b/Engine/Source/Editor/UnrealEd/Private/Kismet2/KismetReinstanceUtilities.cpp
--- a/Engine/Source/Editor/UnrealEd/Private/Kismet2/KismetReinstanceUtilities.cpp
+++ b/Engine/Source/Editor/UnrealEd/Private/Kismet2/KismetReinstanceUtilities.cpp
@@ -3484,6 +3484,8 @@ void FBlueprintCompileReinstancer::CategorizeClassesAndInstances(
 		// if we have a script constructor we would have to check it for equivalence here
 		// but there are many obstacles to script constructors:
 		const FBlake3Hash* ConstructionVersion = ReinstancingState.ConstructionVersioningData.Find(NewClass);
+		// anim blueprints have uncontrolled backpointers to CDO data and therefore instances must
+		// be reconstructed after compilation:
 		if( ConstructionVersion && 
 			FStructUtils::TheSameLayout(OldClass, NewClass) &&
 			CalculateConstructionHash(NewClass) == *ConstructionVersion)
@@ -3557,6 +3559,8 @@ void FBlueprintCompileReinstancer::RestoreClass(
 		TArray<UObject*>& ObjectsToReplace = ReinstancingState.ObjectsToReplaceScratch;
 		GetObjectsToReinstance(OldClass, ObjectsToReplace);
 
+		const UBlueprintGeneratedClass* BPGC = Cast<UBlueprintGeneratedClass>(NewClass);
+
 		for (UObject* Object : ObjectsToReplace)
 		{
 			Object->SetClass(NewClass); // class layouts match, so we can just set the object back to its original class
@@ -3566,6 +3570,7 @@ void FBlueprintCompileReinstancer::RestoreClass(
 				Object, 
 				ReinstancingState.OldToNewInstanceMap,
 				ReinstancingState.ReinstancedObjectsWeakReferenceMap);
+			BPGC->ReinitializeObjectAfterCompile(Object);
 		}
 	}
 }
diff --git a/Engine/Source/Runtime/Engine/Classes/Animation/AnimBlueprintGeneratedClass.h b/Engine/Source/Runtime/Engine/Classes/Animation/AnimBlueprintGeneratedClass.h
--- a/Engine/Source/Runtime/Engine/Classes/Animation/AnimBlueprintGeneratedClass.h
+++ b/Engine/Source/Runtime/Engine/Classes/Animation/AnimBlueprintGeneratedClass.h
@@ -496,7 +496,7 @@ private:
 
 	// Called by the compiler to make sure that data tables are initialized. This is needed to patch the sparse class
 	// data for child anim BP overrides 
-	ENGINE_API void InitializeAnimNodeData(UObject* DefaultObject, bool bForce);
+	ENGINE_API void InitializeAnimNodeData(UObject* DefaultObject, bool bForce) const;
 
 #if WITH_EDITORONLY_DATA
 	// Verify that the serialized NodeTypeMap can be used with the current set of native node data layouts
@@ -632,6 +632,7 @@ public:
 	static ENGINE_API FName GetMutablesStructName();
 	
 #if WITH_EDITOR
+	void HandleReinitializeObjectAfterCompile(UObject* Object) const override;
 	ENGINE_API virtual void PrepareToConformSparseClassData(UScriptStruct* SparseClassDataArchetypeStruct) override;
 	ENGINE_API virtual void ConformSparseClassData(UObject* Object) override;
 #endif
diff --git a/Engine/Source/Runtime/Engine/Classes/Engine/BlueprintGeneratedClass.h b/Engine/Source/Runtime/Engine/Classes/Engine/BlueprintGeneratedClass.h
--- a/Engine/Source/Runtime/Engine/Classes/Engine/BlueprintGeneratedClass.h
+++ b/Engine/Source/Runtime/Engine/Classes/Engine/BlueprintGeneratedClass.h
@@ -551,6 +551,10 @@ public:
 	ENGINE_API void InitializeFieldNotifies();
 
 #if WITH_EDITOR
+	/** Customization point for fixing up backpointers to the class or CDO after compile for instances that can be reused (e.g. no layout change) */
+	ENGINE_API void ReinitializeObjectAfterCompile(UObject* Object) const;
+	virtual void HandleReinitializeObjectAfterCompile(UObject* Object) const {}
+
 	/**
 	 * Called during serialization to allow the class to stash any sparse class data that needs to 
 	 * be conformed against the archetype once the CDO is available.
diff --git a/Engine/Source/Runtime/Engine/Private/Animation/AnimBlueprintGeneratedClass.cpp b/Engine/Source/Runtime/Engine/Private/Animation/AnimBlueprintGeneratedClass.cpp
--- a/Engine/Source/Runtime/Engine/Private/Animation/AnimBlueprintGeneratedClass.cpp
+++ b/Engine/Source/Runtime/Engine/Private/Animation/AnimBlueprintGeneratedClass.cpp
@@ -577,6 +577,11 @@ uint8* UAnimBlueprintGeneratedClass::GetPersistentUberGraphFrame(UObject* Obj, U
 }
 
 #if WITH_EDITOR
+void UAnimBlueprintGeneratedClass::HandleReinitializeObjectAfterCompile(UObject* Object) const
+{
+	InitializeAnimNodeData(Object, true);
+}
+
 void UAnimBlueprintGeneratedClass::PrepareToConformSparseClassData(UScriptStruct* SparseClassDataArchetypeStruct)
 {
 	// ABPGC sparse data is generated on compile, so there is nothing to conform
@@ -720,7 +725,7 @@ void UAnimBlueprintGeneratedClass::GenerateAnimationBlueprintFunctions()
 	}
 }
 
-void UAnimBlueprintGeneratedClass::InitializeAnimNodeData(UObject* DefaultObject, bool bForce)
+void UAnimBlueprintGeneratedClass::InitializeAnimNodeData(UObject* DefaultObject, bool bForce) const
 {
 #if WITH_EDITOR
 	// In editor, skip this work if properties and node data are mismatched. This will be rectified by compile-on-load
diff --git a/Engine/Source/Runtime/Engine/Private/BlueprintGeneratedClass.cpp b/Engine/Source/Runtime/Engine/Private/BlueprintGeneratedClass.cpp
--- a/Engine/Source/Runtime/Engine/Private/BlueprintGeneratedClass.cpp
+++ b/Engine/Source/Runtime/Engine/Private/BlueprintGeneratedClass.cpp
@@ -935,6 +935,11 @@ void UBlueprintGeneratedClass::PostLoadDefaultObject(UObject* Object)
 }
 
 #if WITH_EDITOR
+void UBlueprintGeneratedClass::ReinitializeObjectAfterCompile(UObject* Object) const
+{
+	HandleReinitializeObjectAfterCompile(Object);
+}
+
 void UBlueprintGeneratedClass::PrepareToConformSparseClassData(UScriptStruct* SparseClassDataArchetypeStruct)
 {
 	checkf(SparseClassDataPendingConformStruct.IsExplicitlyNull() && SparseClassDataPendingConform == nullptr, TEXT("PrepareToConformSparseClassData was called while data was already pending conform!"));
