Add option to rotate Top/Bottom orthographic editor viewports

Use CVar: Editor.TopBottomViewRotate

diff --git a/Engine/Source/Editor/UnrealEd/Private/EditorViewportClient.cpp b/Engine/Source/Editor/UnrealEd/Private/EditorViewportClient.cpp
--- a/Engine/Source/Editor/UnrealEd/Private/EditorViewportClient.cpp
+++ b/Engine/Source/Editor/UnrealEd/Private/EditorViewportClient.cpp
@@ -110,6 +110,12 @@ static TAutoConsoleVariable<float> CVarOrthoEditorDebugClipPlaneScale(
 	TEXT("Other light artefacts may appear when this value changes, this is unavoidable for now.\n"),
 	ECVF_RenderThreadSafe);
 
+static TAutoConsoleVariable<bool> CVarOrthoTopBottomViewRotated(
+	TEXT("Editor.TopBottomViewRotate"),
+	false,
+	TEXT("Rotates the top and bottom ortho views for +X to map to screenspace up and +Y to map to screenspace right"),
+	ECVF_Default);
+
 static bool GetDefaultLowDPIPreviewValue()
 {
 	static auto CVarEditorViewportHighDPIPtr = IConsoleManager::Get().FindConsoleVariable(TEXT("r.Editor.Viewport.HighDPI"));
@@ -1165,11 +1171,22 @@ FSceneView* FEditorViewportClient::CalcSceneView(FSceneViewFamily* ViewFamily, c
 
 			if (EffectiveViewportType == LVT_OrthoXY)
 			{
-				ViewInitOptions.ViewRotationMatrix = FMatrix(
-					FPlane(1, 0, 0, 0),
-					FPlane(0, -1, 0, 0),
-					FPlane(0, 0, -1, 0),
-					FPlane(0, 0, 0, 1));
+				if (!CVarOrthoTopBottomViewRotated.GetValueOnAnyThread())
+				{
+					ViewInitOptions.ViewRotationMatrix = FMatrix(
+						FPlane(1, 0, 0, 0),
+						FPlane(0, -1, 0, 0),
+						FPlane(0, 0, -1, 0),
+						FPlane(0, 0, 0, 1));
+				}
+				else
+				{
+					ViewInitOptions.ViewRotationMatrix = FMatrix(
+						FPlane(0, 1, 0, 0),
+						FPlane(1, 0, 0, 0),
+						FPlane(0, 0, 1, 0),
+						FPlane(0, 0, 0, 1));
+				}
 			}
 			else if (EffectiveViewportType == LVT_OrthoXZ)
 			{
@@ -1189,11 +1206,22 @@ FSceneView* FEditorViewportClient::CalcSceneView(FSceneViewFamily* ViewFamily, c
 			}
 			else if (EffectiveViewportType == LVT_OrthoNegativeXY)
 			{
-				ViewInitOptions.ViewRotationMatrix = FMatrix(
-					FPlane(-1, 0, 0, 0),
-					FPlane(0, -1, 0, 0),
-					FPlane(0, 0, 1, 0),
-					FPlane(0, 0, 0, 1));
+				if (!CVarOrthoTopBottomViewRotated.GetValueOnAnyThread())
+				{
+					ViewInitOptions.ViewRotationMatrix = FMatrix(
+						FPlane(-1, 0, 0, 0),
+						FPlane(0, -1, 0, 0),
+						FPlane(0, 0, -1, 0),
+						FPlane(0, 0, 0, 1));
+				}
+				else
+				{
+					ViewInitOptions.ViewRotationMatrix = FMatrix(
+						FPlane(0, -1, 0, 0),
+						FPlane(-1, 0, 0, 0),
+						FPlane(0, 0, 1, 0),
+						FPlane(0, 0, 0, 1));
+				}
 			}
 			else if (EffectiveViewportType == LVT_OrthoNegativeXZ)
 			{
@@ -4422,7 +4450,11 @@ void FEditorViewportClient::Draw(FViewport* InViewport, FCanvas* Canvas)
 		{
 		case LVT_OrthoXY:
 			{
-				const FRotator XYRot(-90.0f, -90.0f, 0.0f);
+				const FRotator XYRot = []()
+				{
+					if (!CVarOrthoTopBottomViewRotated.GetValueOnAnyThread()) { return FRotator(-90.0f, -90.0f, 0.0f); }
+					else { return FRotator(-90.0f, -90.0f, 90.0f); }
+				}();
 				DrawAxes(Viewport, Canvas, &XYRot, EAxisList::XY);
 				if (View)
 				{
@@ -4452,7 +4484,11 @@ void FEditorViewportClient::Draw(FViewport* InViewport, FCanvas* Canvas)
 			}
 		case LVT_OrthoNegativeXY:
 			{
-				const FRotator XYRot(90.0f, 90.0f, 0.0f);
+				const FRotator XYRot = []()
+				{
+					if (!CVarOrthoTopBottomViewRotated.GetValueOnAnyThread()) { return FRotator(90.0f, 90.0f, 0.0f); }
+					else { return FRotator(90.0f, 90.0f, 90.0f); }
+				}();
 				DrawAxes(Viewport, Canvas, &XYRot, EAxisList::XY);
 				if (View)
 				{
@@ -5466,7 +5502,8 @@ void FEditorViewportClient::ShowWidget(const bool bShow)
 
 void FEditorViewportClient::MoveViewportCamera(const FVector& InDrag, const FRotator& InRot, bool bDollyCamera)
 {
-	switch( GetViewportType() )
+	ELevelViewportType EffectiveViewportType = GetViewportType();
+	switch( EffectiveViewportType )
 	{
 	case LVT_OrthoXY:
 	case LVT_OrthoXZ:
@@ -5475,6 +5512,21 @@ void FEditorViewportClient::MoveViewportCamera(const FVector& InDrag, const FRot
 	case LVT_OrthoNegativeXZ:
 	case LVT_OrthoNegativeYZ:
 		{
+			auto MapDragToOrthoViewLocation = [](const FVector& InDrag, ELevelViewportType ViewportType)
+			{
+				if (CVarOrthoTopBottomViewRotated.GetValueOnAnyThread())
+				{
+					if (ViewportType == LVT_OrthoXY)
+					{
+						return FVector(-InDrag.Y, InDrag.X, InDrag.Z);
+					}
+					else if (ViewportType == LVT_OrthoNegativeXY)
+					{
+						return FVector(InDrag.Y, InDrag.X, InDrag.Z);
+					}
+				}
+				return InDrag;
+			};
 			const bool LeftMouseButtonDown = Viewport->KeyState(EKeys::LeftMouseButton);
 			const bool RightMouseButtonDown = Viewport->KeyState(EKeys::RightMouseButton);
 			const bool bIsUsingTrackpad = FSlateApplication::Get().IsUsingTrackpad();
@@ -5486,7 +5538,8 @@ void FEditorViewportClient::MoveViewportCamera(const FVector& InDrag, const FRot
 			}
 			else
 			{
-				SetViewLocation( GetViewLocation() + InDrag );
+				const FVector ViewLocationAdjustment = MapDragToOrthoViewLocation(InDrag, EffectiveViewportType);
+				SetViewLocation( GetViewLocation() + ViewLocationAdjustment );
 			}
 
 			// Update any linked orthographic viewports.
