Fix shadowing warning on MacArm64 compile target

diff --git a/Engine/Source/Editor/UnrealEd/Private/EditorViewportClient.cpp b/Engine/Source/Editor/UnrealEd/Private/EditorViewportClient.cpp
--- a/Engine/Source/Editor/UnrealEd/Private/EditorViewportClient.cpp
+++ b/Engine/Source/Editor/UnrealEd/Private/EditorViewportClient.cpp
@@ -5512,15 +5512,15 @@ void FEditorViewportClient::MoveViewportCamera(const FVector& InDrag, const FRot
 	case LVT_OrthoNegativeXZ:
 	case LVT_OrthoNegativeYZ:
 		{
-			auto MapDragToOrthoViewLocation = [](const FVector& InDrag, ELevelViewportType ViewportType)
+			auto MapDragToOrthoViewLocation = [](const FVector& InDrag, ELevelViewportType ViewportType_)
 			{
 				if (CVarOrthoTopBottomViewRotated.GetValueOnAnyThread())
 				{
-					if (ViewportType == LVT_OrthoXY)
+					if (ViewportType_ == LVT_OrthoXY)
 					{
 						return FVector(-InDrag.Y, InDrag.X, InDrag.Z);
 					}
-					else if (ViewportType == LVT_OrthoNegativeXY)
+					else if (ViewportType_ == LVT_OrthoNegativeXY)
 					{
 						return FVector(InDrag.Y, InDrag.X, InDrag.Z);
 					}
