Cooker: Implement use of FCookLoadScope(ECookLoadType::EditorOnly) to specify that a package loaded during startup should not be added to the cook.
Remove FAssetRegistryGenerator::StartupPackages; it is unnecessary copy of UCookOnTheFlyServer's StartupPackages.
Remove FAssetRegistryGenerator::GetChunkAssignments; it was unused before this change and removing it saves the need to update it.
Change ECookLoadType::Unexpected to ECookLoadType::Unspecified to better fit with the public meaning of the value rather than the internal interpretation.
Move ECookLoadType and FCookLoadScope out of #if WITH_EDITOR to reduce the need for boilerplate in the runtime classes that use them.
Remove FPackageTracker::MarkLoadedPackagesAsNew; it is currently not needed because we never reset the cook in cookbythebook, and removing it saves the need to make it work with the markup of StartupPackages.
#rnx

diff --git a/Engine/Source/Editor/UnrealEd/Private/Commandlets/AssetRegistryGenerator.cpp b/Engine/Source/Editor/UnrealEd/Private/Commandlets/AssetRegistryGenerator.cpp
--- a/Engine/Source/Editor/UnrealEd/Private/Commandlets/AssetRegistryGenerator.cpp
+++ b/Engine/Source/Editor/UnrealEd/Private/Commandlets/AssetRegistryGenerator.cpp
@@ -278,7 +278,7 @@ int64 FAssetRegistryGenerator::GetMaxChunkSizePerPlatform(const ITargetPlatform*
 	return -1;
 }
 
-TArray<int32> FAssetRegistryGenerator::GetExistingPackageChunkAssignments(FName PackageFName)
+TArray<int32> FAssetRegistryGenerator::GetExistingPackageChunkAssignments(FName PackageFName, const TSet<FName>& StartupPackages)
 {
 	TArray<int32> ExistingChunkIDs;
 	int32 PackageFNameHash = GetTypeHash(PackageFName);
@@ -843,7 +843,8 @@ bool FAssetRegistryGenerator::GenerateStreamingInstallManifest(int64 InOverrideC
 }
 
 void FAssetRegistryGenerator::CalculateChunkIdsAndAssignToManifest(const FName& PackageFName, const FString& PackagePathName,
-	const FString& SandboxFilename, const FString& LastLoadedMapName, UE::Cook::FCookSandbox& InSandboxFile)
+	const FString& SandboxFilename, const FString& LastLoadedMapName, UE::Cook::FCookSandbox& InSandboxFile,
+	const TSet<FName>& StartupPackages)
 {
 	TArray<int32> TargetChunks;
 	TArray<int32> ExistingChunkIDs;
@@ -877,7 +878,7 @@ void FAssetRegistryGenerator::CalculateChunkIdsAndAssignToManifest(const FName&
 			}
 
 			PackageChunkIDs = GetExplicitChunkIDs(PackageNameThatDefinesChunks);
-			ExistingChunkIDs = GetExistingPackageChunkAssignments(PackageNameThatDefinesChunks);
+			ExistingChunkIDs = GetExistingPackageChunkAssignments(PackageNameThatDefinesChunks, StartupPackages);
 			PackageChunkIDs.Append(ExistingChunkIDs);
 			AssetManager.GetPackageChunkIds(PackageNameThatDefinesChunks, TargetPlatform, PackageChunkIDs, TargetChunks);
 		}
@@ -1175,10 +1176,9 @@ void FAssetRegistryGenerator::UpdateCollectionAssetData()
 	}
 }
 
-void FAssetRegistryGenerator::Initialize(const TArray<FName> &InStartupPackages, bool bInitializeFromExisting)
+void FAssetRegistryGenerator::Initialize(bool bInitializeFromExisting)
 {
 	LLM_SCOPE_BYTAG(Cooker_GeneratedAssetRegistry);
-	StartupPackages.Append(InStartupPackages);
 
 	FAssetRegistrySerializationOptions SaveOptions;
 
@@ -1607,7 +1607,7 @@ void FAssetRegistryGenerator::ComputePackageRemovals(const FAssetRegistryState&
 
 void FAssetRegistryGenerator::FinalizeChunkIDs(const TSet<FName>& InCookedPackages,
 	const TSet<FName>& InDevelopmentOnlyPackages, UE::Cook::FCookSandbox& InSandboxFile,
-	bool bGenerateStreamingInstallManifest)
+	bool bGenerateStreamingInstallManifest, const TSet<FName>& StartupPackages)
 {
 	LLM_SCOPE_BYTAG(Cooker_GeneratedAssetRegistry);
 	// bGenerateNoChunks overrides bGenerateStreamingInstallManifest overrides ShouldPlatformGenerateStreamingInstallManifest
@@ -1706,7 +1706,8 @@ void FAssetRegistryGenerator::FinalizeChunkIDs(const TSet<FName>& InCookedPackag
 		const FString& SandboxFilename = AllCookedPackageSet.FindChecked(PackageFName);
 		const FString PackagePathName = PackageFName.ToString();
 
-		CalculateChunkIdsAndAssignToManifest(PackageFName, PackagePathName, SandboxFilename, FString(), InSandboxFile);
+		CalculateChunkIdsAndAssignToManifest(PackageFName, PackagePathName, SandboxFilename, FString(), InSandboxFile,
+			StartupPackages);
 	}
 	// anything that remains in the UnAssignedPackageSet is put in chunk0 by FixupPackageDependenciesForChunks
 
@@ -2618,39 +2619,6 @@ int32 FAssetRegistryGenerator::GetPakchunkIndex(int32 ChunkId) const
 	return ChunkId;
 }
 
-void FAssetRegistryGenerator::GetChunkAssignments(TArray<TSet<FName>>& OutAssignments) const
-{
-	if (ChunkManifests.Num())
-	{
-		// chunk 0 is special as it also contains startup packages
-		TSet<FName> PackagesInChunk0;
-		PackagesInChunk0.Append(StartupPackages);
-
-		if (ChunkManifests[0])
-		{
-			for (const TPair<FName,FString>& Pair: *ChunkManifests[0])
-			{
-				PackagesInChunk0.Add(Pair.Key);
-			}
-		}
-		OutAssignments.Add(PackagesInChunk0);
-
-		for (uint32 ChunkIndex = 1, MaxChunk = ChunkManifests.Num(); ChunkIndex < MaxChunk; ++ChunkIndex)
-		{
-			TSet<FName> PackagesInChunk;
-			if (ChunkManifests[ChunkIndex])
-			{
-				for (const TPair<FName,FString>& Pair : *ChunkManifests[ChunkIndex])
-				{
-					PackagesInChunk.Add(Pair.Key);
-				}
-			}
-
-			OutAssignments.Add(PackagesInChunk);
-		}
-	}
-}
-
 void FAssetRegistryGenerator::UpdateAssetRegistryData(FName PackageName, const UPackage* Package,
 	UE::Cook::ECookResult CookResult, FSavePackageResultStruct* SavePackageResult,
 	TOptional<TArray<FAssetData>>&& AssetDatasFromSave, TOptional<FAssetPackageData>&& OverrideAssetPackageData,
diff --git a/Engine/Source/Editor/UnrealEd/Private/CookOnTheFlyServer.cpp b/Engine/Source/Editor/UnrealEd/Private/CookOnTheFlyServer.cpp
--- a/Engine/Source/Editor/UnrealEd/Private/CookOnTheFlyServer.cpp
+++ b/Engine/Source/Editor/UnrealEd/Private/CookOnTheFlyServer.cpp
@@ -4545,6 +4545,22 @@ void UCookOnTheFlyServer::ProcessUnsolicitedPackages(TArray<FName>* OutDiscovere
 		return;
 	}
 
+	auto AddToOutDiscovered = [OutDiscoveredPackageNames, OutInstigators](FPackageData* PackageData, const FInstigator& Instigator)
+		{
+			if (OutDiscoveredPackageNames)
+			{
+				if (!PackageData->IsInProgress())
+				{
+					FInstigator& Existing = OutInstigators->FindOrAdd(PackageData->GetPackageName());
+					if (Existing.Category == EInstigator::InvalidCategory)
+					{
+						OutDiscoveredPackageNames->Add(PackageData->GetPackageName());
+						Existing = Instigator;
+					}
+				}
+			}
+		};
+
 	TMap<FName, FInstigator> NewPackages = PackageTracker->GetNewPackages();
 
 	for (TPair<FName, FInstigator>& PackageWithInstigator : NewPackages)
@@ -4562,19 +4578,6 @@ void UCookOnTheFlyServer::ProcessUnsolicitedPackages(TArray<FName>* OutDiscovere
 			continue; // Generated pacakges are queued separately (with the correct instigator) in QueueGeneratedPackages
 		}
 
-		if (OutDiscoveredPackageNames)
-		{
-			if (!PackageData->IsInProgress())
-			{
-				FInstigator& Existing = OutInstigators->FindOrAdd(PackageData->GetPackageName());
-				if (Existing.Category == EInstigator::InvalidCategory)
-				{
-					OutDiscoveredPackageNames->Add(PackageData->GetPackageName());
-					Existing = Instigator;
-				}
-			}
-		}
-
 		if (Instigator.Category == EInstigator::EditorOnlyLoad)
 		{
 			// This load was expected so we do not need to add a hidden dependency for it.
@@ -4582,13 +4585,18 @@ void UCookOnTheFlyServer::ProcessUnsolicitedPackages(TArray<FName>* OutDiscovere
 			if (!bSkipOnlyEditorOnly)
 			{
 				QueueDiscoveredPackage(*PackageData, FInstigator(Instigator), EDiscoveredPlatformSet::CopyFromInstigator);
+				AddToOutDiscovered(PackageData, Instigator);
 			}
+			continue;
 		}
-		else if (bSkipOnlyEditorOnly &&
-					(Instigator.Category == EInstigator::ForceExplorableSaveTimeSoftDependency ||
-						(PackageTracker->NeverCookPackageList.Contains(PackageData->GetPackageName()) &&
-							INDEX_NONE != UE::String::FindFirst(WriteToString<256>(PackageData->GetPackageName()),
-								ULevel::GetExternalActorsFolderName(), ESearchCase::IgnoreCase))))
+
+		AddToOutDiscovered(PackageData, Instigator);
+
+		if (bSkipOnlyEditorOnly &&
+			(Instigator.Category == EInstigator::ForceExplorableSaveTimeSoftDependency ||
+				(PackageTracker->NeverCookPackageList.Contains(PackageData->GetPackageName()) &&
+					INDEX_NONE != UE::String::FindFirst(WriteToString<256>(PackageData->GetPackageName()),
+						ULevel::GetExternalActorsFolderName(), ESearchCase::IgnoreCase))))
 		{
 			// ONLYEDITORONLY_TODO: WorldPartition should mark up these loads with EInstigator::ForceExplorableSaveTimeSoftDependency rather than
 			// needing to use a naming convention. We should also mark them once, during Save, rather than marking them
@@ -5684,16 +5692,11 @@ void UCookOnTheFlyServer::Initialize( ECookMode::Type DesiredCookMode, ECookInit
 	PollNextTimeIdleSeconds = MAX_flt;
 
 	CurrentAsyncCacheForType = MaxAsyncCacheForType;
+	TSet<FName> StartupPackages;
+	PackageTracker->InitializeTracking(StartupPackages);
 	if (IsCookByTheBookMode())
 	{
-		for (TObjectIterator<UPackage> It; It; ++It)
-		{
-			if ((*It) != GetTransientPackage())
-			{
-				CookByTheBookOptions->StartupPackages.Add(It->GetFName());
-				UE_LOG(LogCook, Verbose, TEXT("Cooker startup package %s"), *It->GetName());
-			}
-		}
+		CookByTheBookOptions->StartupPackages = MoveTemp(StartupPackages);
 	}
 
 	IdleStatus = EIdleStatus::Done;
@@ -7794,7 +7797,7 @@ void UCookOnTheFlyServer::RefreshPlatformAssetRegistries(const TArrayView<const
 				bInitializeFromExisting = false;
 			}
 
-			PlatformData->RegistryGenerator->Initialize(CookByTheBookOptions->StartupPackages, bInitializeFromExisting);
+			PlatformData->RegistryGenerator->Initialize(bInitializeFromExisting);
 		}
 	}
 }
@@ -9471,7 +9474,7 @@ void UCookOnTheFlyServer::CookByTheBookFinishedInternal()
 			{
 				UE_SCOPED_HIERARCHICAL_COOKTIMER(BuildChunkManifest);
 				Generator.FinalizeChunkIDs(CookedPackageNames, IgnorePackageNames, *SandboxFile,
-					CookByTheBookOptions->bGenerateStreamingInstallManifests);
+					CookByTheBookOptions->bGenerateStreamingInstallManifests, CookByTheBookOptions->StartupPackages);
 			}
 			{
 				UE_SCOPED_HIERARCHICAL_COOKTIMER(SaveManifests);
@@ -9998,8 +10001,6 @@ void UCookOnTheFlyServer::ResetCook(TConstArrayView<TPair<const ITargetPlatform*
 			PackageTracker->UnsolicitedCookedPackages.GetPackagesForPlatformAndRemove(TargetPlatform, PackageNames);
 		}
 	}
-
-	PackageTracker->MarkLoadedPackagesAsNew();
 }
 
 void UCookOnTheFlyServer::ClearCachedCookedPlatformDataForPlatform(const ITargetPlatform* TargetPlatform)
@@ -11190,6 +11191,12 @@ void UCookOnTheFlyServer::GenerateInitialRequests(FBeginCookContext& BeginContex
 
 	// Add soft/hard startup references after collecting requested files and handling empty requests
 	FlushAsyncLoading();
+	// TODO: When we add a way for projects to change to NoStartupPackages by default, we will need to allow packages
+	// that have ECookLoadType::UsedInGame even if the project has opted out of StartupPackages. So we will need to
+	// call ProcessUnsolicitedPackages here even if the project has opted out of StartupPackages, and inside it
+	// skip the StartupPackages but keep the StartupPackageCookLoadScope packages.
+	// We will also need a separate container for StartupCookLoadScopeSoftObjectPackages, so we
+	// can add the SoftObjectPaths required by those packages.
 	if (!CookByTheBookOptions->bSkipHardReferences && !EnumHasAnyFlags(CookOptions, ECookByTheBookOptions::NoStartupPackages))
 	{
 		ProcessUnsolicitedPackages(&FilesInPath, &FilesInPathInstigators);
@@ -12470,6 +12477,7 @@ void UCookOnTheFlyServer::OnDiscoveredPackageDebug(FName PackageName, const UE::
 	switch (Instigator.Category)
 	{
 	case EInstigator::StartupPackage: [[fallthrough]];
+	case EInstigator::StartupPackageCookLoadScope: [[fallthrough]];
 	case EInstigator::GeneratedPackage: [[fallthrough]];
 	case EInstigator::ForceExplorableSaveTimeSoftDependency:
 	case EInstigator::BuildDependency:
diff --git a/Engine/Source/Editor/UnrealEd/Private/Cooker/CookRequestCluster.cpp b/Engine/Source/Editor/UnrealEd/Private/Cooker/CookRequestCluster.cpp
--- a/Engine/Source/Editor/UnrealEd/Private/Cooker/CookRequestCluster.cpp
+++ b/Engine/Source/Editor/UnrealEd/Private/Cooker/CookRequestCluster.cpp
@@ -236,6 +236,7 @@ FRequestCluster::FRequestCluster(UCookOnTheFlyServer& InCOTFS, TRingBuffer<FDisc
 		// Startup packages and Generated packages are expected discovery types and do not need to add hidden
 		// dependencies.
 		if (Discovery->Instigator.Category != EInstigator::StartupPackage &&
+			Discovery->Instigator.Category != EInstigator::StartupPackageCookLoadScope &&
 			Discovery->Instigator.Category != EInstigator::GeneratedPackage)
 		{
 			// For unsolicited packages, we need to check load-reachability to decide whether the load was expected.
diff --git a/Engine/Source/Editor/UnrealEd/Private/Cooker/CookTypes.h b/Engine/Source/Editor/UnrealEd/Private/Cooker/CookTypes.h
--- a/Engine/Source/Editor/UnrealEd/Private/Cooker/CookTypes.h
+++ b/Engine/Source/Editor/UnrealEd/Private/Cooker/CookTypes.h
@@ -515,7 +515,7 @@ struct FCookByTheBookOptions
 {
 public:
 	// Process-lifetime variables
-	TArray<FName>					StartupPackages;
+	TSet<FName>						StartupPackages;
 
 	// Session-lifetime variables
 	/**
diff --git a/Engine/Source/Editor/UnrealEd/Private/Cooker/PackageTracker.cpp b/Engine/Source/Editor/UnrealEd/Private/Cooker/PackageTracker.cpp
--- a/Engine/Source/Editor/UnrealEd/Private/Cooker/PackageTracker.cpp
+++ b/Engine/Source/Editor/UnrealEd/Private/Cooker/PackageTracker.cpp
@@ -55,27 +55,35 @@ FPackageTracker::FPackageTracker(UCookOnTheFlyServer& InCOTFS)
 
 FPackageTracker::~FPackageTracker()
 {
-	if (bTrackingInitialized)
+	Unsubscribe();
+}
+
+void FPackageTracker::Unsubscribe()
+{
+	if (bSubscribed)
 	{
+		bSubscribed = false;
 		GUObjectArray.RemoveUObjectDeleteListener(this);
 		GUObjectArray.RemoveUObjectCreateListener(this);
 	}
 }
 
-void FPackageTracker::InitializeTracking()
+void FPackageTracker::InitializeTracking(TSet<FName>& OutStartupPackages)
 {
 	check(!bTrackingInitialized);
 
 	LLM_SCOPE_BYTAG(Cooker);
 
 	FWriteScopeLock ScopeLock(Lock);
+	OutStartupPackages.Empty();
 	for (TObjectIterator<UPackage> It; It; ++It)
 	{
 		UPackage* Package = *It;
 
-		if (Package->GetOuter() == nullptr)
+		if (Package->GetOuter() == nullptr && Package != GetTransientPackage())
 		{
 			LoadedPackages.Add(Package);
+			OutStartupPackages.Add(Package->GetFName());
 		}
 	}
 
@@ -87,17 +95,44 @@ void FPackageTracker::InitializeTracking()
 
 	GUObjectArray.AddUObjectDeleteListener(this);
 	GUObjectArray.AddUObjectCreateListener(this);
+	bSubscribed = true;
 
+	TArray<TPair<FName, ECookLoadType>> StartupPackageLoadTypes;
+	FCookLoadScope::SetCookerStartupComplete(StartupPackageLoadTypes);
+
+	for (TPair<FName, ECookLoadType>& Pair : StartupPackageLoadTypes)
+	{
+		switch (Pair.Value)
+		{
+		case ECookLoadType::EditorOnly:
+		{
+			OutStartupPackages.Remove(Pair.Key);
+			FInstigator* Instigator = NewPackages.Find(Pair.Key);
+			if (Instigator)
+			{
+				*Instigator = FInstigator(EInstigator::EditorOnlyLoad);
+			}
+			break;
+		}
+		case ECookLoadType::UsedInGame:
+		{
+			FInstigator* Instigator = NewPackages.Find(Pair.Key);
+			if (Instigator)
+			{
+				*Instigator = FInstigator(EInstigator::StartupPackageCookLoadScope);
+			}
+			break;
+		}
+		default:
+			break;
+		}
+	}
 	bTrackingInitialized = true;
 }
 
 TMap<FName, FInstigator> FPackageTracker::GetNewPackages()
 {
-	if (!bTrackingInitialized)
-	{
-		InitializeTracking();
-		check(bTrackingInitialized);
-	}
+	check(bTrackingInitialized);
 
 	FWriteScopeLock ScopeLock(Lock);
 	TMap<FName, FInstigator> Result = MoveTemp(NewPackages);
@@ -105,23 +140,6 @@ TMap<FName, FInstigator> FPackageTracker::GetNewPackages()
 	return Result;
 }
 
-void FPackageTracker::MarkLoadedPackagesAsNew()
-{
-	if (!bTrackingInitialized)
-	{
-		return;
-	}
-
-	LLM_SCOPE_BYTAG(Cooker);
-
-	FWriteScopeLock ScopeLock(Lock);
-	NewPackages.Reserve(LoadedPackages.Num());
-	for (UPackage* Package : LoadedPackages)
-	{
-		NewPackages.FindOrAdd(Package->GetFName(), FInstigator(EInstigator::StartupPackage));
-	}
-}
-
 void FPackageTracker::NotifyUObjectCreated(const class UObjectBase* Object, int32 Index)
 {
 	if (Object->GetClass() == UPackage::StaticClass())
@@ -191,8 +209,7 @@ void FPackageTracker::NotifyUObjectDeleted(const class UObjectBase* Object, int3
 
 void FPackageTracker::OnUObjectArrayShutdown()
 {
-	GUObjectArray.RemoveUObjectDeleteListener(this);
-	GUObjectArray.RemoveUObjectCreateListener(this);
+	Unsubscribe();
 }
 
 void FPackageTracker::RemapTargetPlatforms(const TMap<ITargetPlatform*, ITargetPlatform*>& Remap)
diff --git a/Engine/Source/Editor/UnrealEd/Private/Cooker/PackageTracker.h b/Engine/Source/Editor/UnrealEd/Private/Cooker/PackageTracker.h
--- a/Engine/Source/Editor/UnrealEd/Private/Cooker/PackageTracker.h
+++ b/Engine/Source/Editor/UnrealEd/Private/Cooker/PackageTracker.h
@@ -176,15 +176,11 @@ public:
 	FPackageTracker(UCookOnTheFlyServer& InCOTFS);
 	~FPackageTracker();
 
+	void InitializeTracking(TSet<FName>& OutStartupPackages);
+
 	/** Returns all packages that have been loaded since the last time GetNewPackages was called */
 	TMap<FName, FInstigator> GetNewPackages();
 
-	/**
-	 * Copy all LoadedPackages into NewPackages. Called when reachability of all packages needs to be recalculated
-	 * because a new session platform was added.
-	 */
-	void MarkLoadedPackagesAsNew();
-
 	virtual void NotifyUObjectCreated(const class UObjectBase* Object, int32 Index) override;
 	virtual void NotifyUObjectDeleted(const class UObjectBase* Object, int32 Index) override;
 	virtual void OnUObjectArrayShutdown() override;
@@ -234,7 +230,7 @@ public:
 	}
 
 private:
-	void InitializeTracking();
+	void Unsubscribe();
 
 	// Protects data for thread-safety
 	FRWLock Lock;
@@ -246,6 +242,7 @@ private:
 	// This list contains the UPackages loaded since last call to GetNewPackages
 	TMap<FName, FInstigator> NewPackages;
 	bool bTrackingInitialized = false;
+	bool bSubscribed = false;
 };
 
 } // namespace UE::Cook
diff --git a/Engine/Source/Editor/UnrealEd/Public/Commandlets/AssetRegistryGenerator.h b/Engine/Source/Editor/UnrealEd/Public/Commandlets/AssetRegistryGenerator.h
--- a/Engine/Source/Editor/UnrealEd/Public/Commandlets/AssetRegistryGenerator.h
+++ b/Engine/Source/Editor/UnrealEd/Public/Commandlets/AssetRegistryGenerator.h
@@ -51,7 +51,7 @@ public:
 	/**
 	 * Initializes manifest generator - creates manifest lists, hooks up delegates.
 	 */
-	void Initialize(const TArray<FName> &StartupPackages, bool bInitializeFromExisting);
+	void Initialize(bool bInitializeFromExisting);
 	/**
 	 * Used by DLC which does not Clone the entire global AR. Copies data from the global assetregistry
 	 * for all packages present in the previous AssetRegistry
@@ -156,9 +156,12 @@ public:
 	 *        If false, no manifest is written and all packages are implicitly assigned to chunk 0 by the automation tool.
 	 *        If true, packages are assigned to chunks based on settings (possibly all in chunk 0 if settings are empty)
 	 *        and a manifest of packagenames is written for each chunk
+	 * @param StartupPackages The packages the cooker recorded as being loaded at startup.
+	 *        These packages have special chunk rules.
 	 */
 	void FinalizeChunkIDs(const TSet<FName>& CookedPackages, const TSet<FName>& DevelopmentOnlyPackages,
-		UE::Cook::FCookSandbox& InSandboxFile, bool bGenerateStreamingInstallManifest);
+		UE::Cook::FCookSandbox& InSandboxFile, bool bGenerateStreamingInstallManifest,
+		const TSet<FName>& StartupPackages);
 
 	/**
 	 * Register a chunk data generator with this generator.
@@ -238,11 +241,6 @@ public:
 	 */
 	int32 GetPakchunkIndex(int32 ChunkId) const;
 
-	/**
-	 * Returns the chunks
-	 */
-	void GetChunkAssignments(TArray<TSet<FName>>& OutAssignments) const;
-
 	/**
 	 * Attempts to update the metadata for a package in an asset registry generator.
 	 * This is only called for CookByTheBook.
@@ -316,8 +314,6 @@ private:
 		TArray<FAssetDependency> PackageReferencers;
 	};
 	TMap<FName, FIterativelySkippedPackageUpdateData> PreviousPackagesToUpdate;
-	/** List of packages that were loaded at startup */
-	TSet<FName> StartupPackages;
 	/** List of packages that were successfully cooked */
 	TSet<FName> CookedPackages;
 	/** List of packages that were filtered out from cooking */
@@ -407,7 +403,7 @@ private:
 	int64 GetMaxChunkSizePerPlatform( const ITargetPlatform* Platform ) const;
 
 	/** Returns an array of chunks ID for a package name that have been assigned during the cook process. */
-	TArray<int32> GetExistingPackageChunkAssignments(FName PackageFName);
+	TArray<int32> GetExistingPackageChunkAssignments(FName PackageFName, const TSet<FName>& StartupPackages);
 
 	/**
 	 * Get chunks IDs for a package that were assigned to the package in the editor from AssetFileContextMenu.
@@ -417,7 +413,8 @@ private:
 
 	/** Calculate the final ChunkIds used by the package and store the package in the manifest for each of those chunks. */
 	void CalculateChunkIdsAndAssignToManifest(const FName& PackageFName, const FString& PackagePathName,
-		const FString& SandboxFilename, const FString& LastLoadedMapName, UE::Cook::FCookSandbox& InSandboxFile);
+		const FString& SandboxFilename, const FString& LastLoadedMapName, UE::Cook::FCookSandbox& InSandboxFile,
+		const TSet<FName>& StartupPackages);
 
 	/** Returns true if the specific platform desires multiple chunks suitable for streaming install */
 	bool ShouldPlatformGenerateStreamingInstallManifest(const ITargetPlatform* Platform) const;
diff --git a/Engine/Source/Runtime/CoreUObject/Private/Serialization/AsyncLoading.cpp b/Engine/Source/Runtime/CoreUObject/Private/Serialization/AsyncLoading.cpp
--- a/Engine/Source/Runtime/CoreUObject/Private/Serialization/AsyncLoading.cpp
+++ b/Engine/Source/Runtime/CoreUObject/Private/Serialization/AsyncLoading.cpp
@@ -6413,7 +6413,7 @@ void FAsyncPackage::AddImportDependency(const FName& PendingImport, const FName&
 		Info.ReferencerPackageOp = PackageAccessTrackingOps::NAME_Load;
 #endif
 #if WITH_EDITOR
-		Info.CookLoadType = ECookLoadType::Unexpected;
+		Info.CookLoadType = ECookLoadType::Unspecified;
 #endif
 		Info.SetInstancingContext(MoveTemp(InstancingContext));
 		PackageToStream = new FAsyncPackage(AsyncLoadingThread, Info, EDLBootNotificationManager);
diff --git a/Engine/Source/Runtime/CoreUObject/Private/Serialization/AsyncLoading2.cpp b/Engine/Source/Runtime/CoreUObject/Private/Serialization/AsyncLoading2.cpp
--- a/Engine/Source/Runtime/CoreUObject/Private/Serialization/AsyncLoading2.cpp
+++ b/Engine/Source/Runtime/CoreUObject/Private/Serialization/AsyncLoading2.cpp
@@ -528,7 +528,7 @@ struct FPackageReferencer
 #endif
 
 #if WITH_EDITOR
-	ECookLoadType CookLoadType = ECookLoadType::Unexpected;
+	ECookLoadType CookLoadType = ECookLoadType::Unspecified;
 #endif
 
 	static FPackageReferencer FromImport(FName ReferencerName)
@@ -539,7 +539,7 @@ struct FPackageReferencer
 		Result.ReferencerPackageOp = PackageAccessTrackingOps::NAME_Load;
 #endif
 #if WITH_EDITOR
-		Result.CookLoadType = ECookLoadType::Unexpected;
+		Result.CookLoadType = ECookLoadType::Unspecified;
 #endif
 		return Result;
 	}
diff --git a/Engine/Source/Runtime/CoreUObject/Private/Serialization/AsyncPackage.h b/Engine/Source/Runtime/CoreUObject/Private/Serialization/AsyncPackage.h
--- a/Engine/Source/Runtime/CoreUObject/Private/Serialization/AsyncPackage.h
+++ b/Engine/Source/Runtime/CoreUObject/Private/Serialization/AsyncPackage.h
@@ -52,7 +52,7 @@ struct FAsyncPackageDesc
 		, Priority(InPriority)
 		, PIEInstanceID(InPIEInstanceID)
 #if WITH_EDITOR
-		, CookLoadType(ECookLoadType::Unexpected)
+		, CookLoadType(ECookLoadType::Unspecified)
 #endif
 	{
 		check(!PackagePath.IsEmpty());
diff --git a/Engine/Source/Runtime/CoreUObject/Private/UObject/ICookInfo.cpp b/Engine/Source/Runtime/CoreUObject/Private/UObject/ICookInfo.cpp
--- a/Engine/Source/Runtime/CoreUObject/Private/UObject/ICookInfo.cpp
+++ b/Engine/Source/Runtime/CoreUObject/Private/UObject/ICookInfo.cpp
@@ -3,11 +3,46 @@
 #include "UObject/ICookInfo.h"
 
 #if WITH_EDITOR
+#include "HAL/CriticalSection.h"
+#include "Misc/ScopeLock.h"
+#include "Templates/UniquePtr.h"
+#include "UObject/UObjectArray.h"
+#include "UObject/Package.h"
 
+#include <atomic>
+#endif
+
+#if WITH_EDITOR
 
 namespace UE::Cook
 {
 
+/**
+ * Tracks package loads under an FCookLoadScope that occur before the cooker has initialized;
+ * after cooker initialization these loads are tracked by the cooker's FPackageTracker.
+ */
+class FCookLoadScopeStartupListener : public FUObjectArray::FUObjectCreateListener
+{
+public:
+	FCookLoadScopeStartupListener();
+	~FCookLoadScopeStartupListener();
+
+	virtual void NotifyUObjectCreated(const UObjectBase* Object, int32 Index) override;
+	virtual void OnUObjectArrayShutdown() override;
+
+	void Unsubscribe();
+
+public:
+	TArray<TPair<FName, ECookLoadType>> StartupPackageLoadTypes;
+
+private:
+	bool bSubscribed = false;
+};
+
+static std::atomic<FCookLoadScopeStartupListener*> GCookLoadScopeStartupListener { nullptr };
+static FCriticalSection GCookLoadScopeStartupListenerLock;
+static thread_local ECookLoadType GCookLoadType = ECookLoadType::Unspecified;
+
 const TCHAR* LexToString(EInstigator Value)
 {
 	switch (Value)
@@ -63,24 +98,105 @@ const TCHAR* GetReferencedSetOpName()
 	return TEXT("ReferencedSet");
 }
 
+FCookLoadScopeStartupListener::FCookLoadScopeStartupListener()
+{
+	GUObjectArray.AddUObjectCreateListener(this);
+	bSubscribed = true;
 }
 
-static thread_local ECookLoadType GCookLoadType = ECookLoadType::Unexpected;
+FCookLoadScopeStartupListener::~FCookLoadScopeStartupListener()
+{
+	Unsubscribe();
+}
+
+void FCookLoadScopeStartupListener::NotifyUObjectCreated(const class UObjectBase* Object, int32 Index)
+{
+	if (Object->GetClass() != UPackage::StaticClass())
+	{
+		return;
+	}
+
+	ECookLoadType CookLoadType = FCookLoadScope::GetCurrentValue();
+	if (CookLoadType == ECookLoadType::Unspecified)
+	{
+		return;
+	}
+
+	FScopeLock ListenerScopeLock(&GCookLoadScopeStartupListenerLock);
+	if (!GCookLoadScopeStartupListener)
+	{
+		// We have been called from the callback after we already unsubscribed
+		// and deleted on another thread. Do not access anything on *this.
+		return;
+	}
+	this->StartupPackageLoadTypes.Add(TPair<FName, ECookLoadType> { Object->GetFName(), CookLoadType});
+}
+
+void FCookLoadScopeStartupListener::OnUObjectArrayShutdown()
+{
+	Unsubscribe();
+}
+
+void FCookLoadScopeStartupListener::Unsubscribe()
+{
+	if (bSubscribed)
+	{
+		GUObjectArray.RemoveUObjectCreateListener(this);
+		bSubscribed = false;
+	}
+}
+
+void InitializeCookGlobals()
+{
+	check(!GCookLoadScopeStartupListener);
+	// We need to construct outside of the GCookLoadScopeStartupListenerLock, because the constructor subscribes to
+	// GUObjectArray which has its own lock; see the comment in SetCookerStartupComplete
+	GCookLoadScopeStartupListener = new FCookLoadScopeStartupListener();
+}
+
+} // namespace UE::Cook
 
 FCookLoadScope::FCookLoadScope(ECookLoadType ScopeType)
-	: PreviousScope(GCookLoadType)
+	: PreviousScope(UE::Cook::GCookLoadType)
 {
-	GCookLoadType = ScopeType;
+	UE::Cook::GCookLoadType = ScopeType;
 }
 
 FCookLoadScope::~FCookLoadScope()
 {
-	GCookLoadType = PreviousScope;
+	UE::Cook::GCookLoadType = PreviousScope;
 }
 
 ECookLoadType FCookLoadScope::GetCurrentValue()
 {
-	return GCookLoadType;
+	return UE::Cook::GCookLoadType;
 }
 
-#endif
+void FCookLoadScope::SetCookerStartupComplete(TArray<TPair<FName, ECookLoadType>>& OutStartupPackageLoadTypes)
+{
+	using namespace UE::Cook;
+
+	if (!IsRunningCookCommandlet())
+	{
+		// InitializeCookGlobals is only called when IsRunningCookCommandlet;
+		// We need that hook for this function to work, so we do not implement
+		// StartupPackageLoadTypes when cooking through means other than running as CookCommandlet.
+		return;
+	}
+
+	FCookLoadScopeStartupListener* Listener = GCookLoadScopeStartupListener;
+	check(Listener); // Set by InitializeCookGlobals
+	// We need to unsubscribe from GUObjectArray outside of GCookLoadScopeStartupListenerScopeLock, because
+	// GUObjectArray might have its own lock that it holds when calling NotifyUObjectCreated, and we enter
+	// GCookLoadScopeStartupListenerScopeLock inside of NotifyUObjectCreated. To avoid locking in a different
+	// order and therefore deadlocking if another thread (e.g. AsyncLoadingThread) is adding a UPackage during
+	// SetCookerStartupComplete, we therefore need to not hold our lock while unsubscribing.
+	Listener->Unsubscribe();
+
+	FScopeLock ListenerScopeLock(&GCookLoadScopeStartupListenerLock);
+	GCookLoadScopeStartupListener = nullptr;
+	OutStartupPackageLoadTypes = MoveTemp(Listener->StartupPackageLoadTypes);
+	delete Listener;
+}
+
+#endif // WITH_EDITOR
diff --git a/Engine/Source/Runtime/CoreUObject/Private/UObject/LinkerLoad.cpp b/Engine/Source/Runtime/CoreUObject/Private/UObject/LinkerLoad.cpp
--- a/Engine/Source/Runtime/CoreUObject/Private/UObject/LinkerLoad.cpp
+++ b/Engine/Source/Runtime/CoreUObject/Private/UObject/LinkerLoad.cpp
@@ -3658,7 +3658,7 @@ bool FLinkerLoad::VerifyImportInner(const int32 ImportIndex, FString& WarningSuf
 			}
 			FPackagePath PackagePathToLoad = FPackagePath::FromPackageNameChecked(PackageToLoad);
 #if WITH_EDITOR
-			FCookLoadScope ResetCookLoadScopeToUnexpected(ECookLoadType::Unexpected);
+			FCookLoadScope ResetCookLoadScopeToUnspecified(ECookLoadType::Unspecified);
 #endif
 
 			Package = LoadPackageInternal(Package, PackagePathToLoad, InternalLoadFlags | LOAD_IsVerifying, this, nullptr /* InReaderOverride */,
diff --git a/Engine/Source/Runtime/CoreUObject/Public/UObject/ICookInfo.h b/Engine/Source/Runtime/CoreUObject/Public/UObject/ICookInfo.h
--- a/Engine/Source/Runtime/CoreUObject/Public/UObject/ICookInfo.h
+++ b/Engine/Source/Runtime/CoreUObject/Public/UObject/ICookInfo.h
@@ -2,13 +2,16 @@
 
 #pragma once
 
+#include "HAL/Platform.h"
+#include "Misc/CoreMiscDefines.h"
+
 #if WITH_EDITOR
 #include "Containers/Array.h"
 #include "Containers/UnrealString.h"
 #include "Delegates/Delegate.h"
-#include "HAL/Platform.h"
 #include "UObject/CookEnums.h"
 #include "UObject/NameTypes.h"
+#include "Templates/Tuple.h"
 #endif
 
 #if WITH_EDITOR
@@ -19,6 +22,41 @@ class ITargetPlatform;
 class UPackage;
 
 namespace UE::Cook { class IMPCollector; }
+#endif
+
+/**
+ * A scope around loads when cooking that indicates whether the loaded package is needed in game or not.
+ * The default is Unexpected. Declare an FCookLoadScope to set the value.
+ * 
+ * If the package is marked as EditorOnly, that only suppresses the reference from the LoadPackage call.
+ * The package can still be pulled into the cook by UsedInGame references from the AssetRegistry or by being
+ * stored as an import in a cooked package.
+ *
+ * Packages that are declared in the AssetRegistry as an editoronly reference by the package that loads them
+ * are implicitly marked as ECookLoadType::EditorOnly for that load, and do not need an explicitly declared
+ * loadtype to be marked as EditorOnly.
+ *
+ * LoadTypes also apply to Startup packages, or packages loaded by systems without the load being owned by
+ * a referencer package. EditorOnly still suppresses these packages (for that single load) from being added
+ * to the cook, UsedInGame still forces them to be added. For Startup packages, marking the package as
+ * EditorOnly both removes its auto-inclusion in the cook, and also removes the special chunk rule for startup
+ * packages from it, if it ends up being pulled into the cook by another referencer. Startup packages are
+ * automatically added to Chunk0 when the project is using multiple pak files (aka streaming chunks);
+ * ECookLoadType::EditorOnly removes that chunk rule for the given package.
+ *
+ * This enum is declared in non-WITH_EDITOR builds to remove the need for #if WITH_EDITOR boilerplate, but
+ * it is ignored; all of the functions/classes that use it are noops if !WITH_EDITOR.
+ */
+enum class ECookLoadType : uint8
+{
+	Unspecified,
+	EditorOnly,
+	UsedInGame,
+
+	Unexpected UE_DEPRECATED(5.6, "Use ECookLoadType::Unspecified") = Unspecified,
+};
+
+#if WITH_EDITOR
 
 namespace UE::Cook
 {
@@ -35,6 +73,7 @@ namespace UE::Cook
 	CallbackMacro(NotYetRequested, true) \
 	CallbackMacro(Unspecified, false) \
 	CallbackMacro(StartupPackage, true) \
+	CallbackMacro(StartupPackageCookLoadScope, true) \
 	CallbackMacro(AlwaysCookMap, true) \
 	CallbackMacro(IniMapSection, false) \
 	CallbackMacro(IniAllMaps, true) \
@@ -198,21 +237,13 @@ COREUOBJECT_API const TCHAR* GetReferencedSetFilename();
 /** Return the name of the op used to store the ReferencedSet in the zenserver oplog. */
 COREUOBJECT_API const TCHAR* GetReferencedSetOpName();
 
-} // namespace UE::Cook
-
 /**
- * A scope around loads when cooking that indicates whether the loaded package is needed in game or not.
- * The default is Unexpected. Declare an FCookLoadScope to set the value.
- * Packages loads that are declared via AssetRegistry dependencies (imports or softobjectpaths in
- * the editor package) change the default to UsedInGame or EditorOnly dependending on whether they were saved
- * in an EditorOnly property or scope.
+ * Initialize some globals that are used to track engine activity relevant to the cook before the cooker
+ * is created. This should only be called when setting IsRunningCookCommandlet=true.
  */
-enum class ECookLoadType : uint8
-{
-	Unexpected,
-	EditorOnly,
-	UsedInGame,
-};
+COREUOBJECT_API void InitializeCookGlobals();
+
+} // namespace UE::Cook
 
 /** Set the ECookLoadType value in the current scope.*/
 struct FCookLoadScope
@@ -221,9 +252,24 @@ struct FCookLoadScope
 	COREUOBJECT_API ~FCookLoadScope();
 
 	COREUOBJECT_API static ECookLoadType GetCurrentValue();
+	COREUOBJECT_API static void SetCookerStartupComplete(TArray<TPair<FName, ECookLoadType>>& OutStartupPackageLoadTypes);
 
 private:
 	ECookLoadType PreviousScope;
 };
 
+#else // WITH_EDITOR
+
+/** Defined in !WITH_EDITOR to remove the need for #ifdef boilerplate, but it is a noop. */
+struct FCookLoadScope
+{
+	FORCEINLINE explicit FCookLoadScope(ECookLoadType ScopeType)
+	{
+	}
+
+	FORCEINLINE ~FCookLoadScope()
+	{
+	}
+};
+
 #endif // WITH_EDITOR
\ No newline at end of file
diff --git a/Engine/Source/Runtime/Engine/Private/StreamableManager.cpp b/Engine/Source/Runtime/Engine/Private/StreamableManager.cpp
--- a/Engine/Source/Runtime/Engine/Private/StreamableManager.cpp
+++ b/Engine/Source/Runtime/Engine/Private/StreamableManager.cpp
@@ -357,7 +357,7 @@ FStreamableHandle::FStreamableHandle()
 	, StreamablesLoading(0)
 	, OwningManager(nullptr)
 #if WITH_EDITOR
-	, CookLoadType(ECookLoadType::Unexpected)
+	, CookLoadType(ECookLoadType::Unspecified)
 #endif
 {
 }
diff --git a/Engine/Source/Runtime/Launch/Private/LaunchEngineLoop.cpp b/Engine/Source/Runtime/Launch/Private/LaunchEngineLoop.cpp
--- a/Engine/Source/Runtime/Launch/Private/LaunchEngineLoop.cpp
+++ b/Engine/Source/Runtime/Launch/Private/LaunchEngineLoop.cpp
@@ -132,6 +132,7 @@
 	#include "ShaderCompiler.h"
 	#include "Virtualization/VirtualizationSystem.h"
 	#include "AnimationUtils.h"
+	#include "UObject/ICookInfo.h"
 
 	#if PLATFORM_WINDOWS
 		#include "Windows/AllowWindowsPlatformTypes.h"
@@ -2399,6 +2400,7 @@ int32 FEngineLoop::PreInitPreStartupScreen(const TCHAR* CmdLine)
 		if (CommandletName == TEXTVIEW("cookcommandlet"))
 		{
 			PRIVATE_GIsRunningCookCommandlet = true;
+			UE::Cook::InitializeCookGlobals();
 			PRIVATE_GIsRunningDLCCookCommandlet = (FCString::Strfind(FCommandLine::Get(), TEXT("-dlcname=")) != nullptr);
 		}
 #endif
