Fix Font Thumbnail Renderer, which currently fails CanVisualizeAsset since it unnecessarily derives from UTextureThumbnailRenderer despite not rendering a texture asset

diff --git a/Engine/Source/Editor/UnrealEd/Classes/ThumbnailRendering/FontThumbnailRenderer.h b/Engine/Source/Editor/UnrealEd/Classes/ThumbnailRendering/FontThumbnailRenderer.h
--- a/Engine/Source/Editor/UnrealEd/Classes/ThumbnailRendering/FontThumbnailRenderer.h
+++ b/Engine/Source/Editor/UnrealEd/Classes/ThumbnailRendering/FontThumbnailRenderer.h
@@ -9,14 +9,14 @@
 
 #include "CoreMinimal.h"
 #include "UObject/ObjectMacros.h"
-#include "ThumbnailRendering/TextureThumbnailRenderer.h"
+#include "ThumbnailRendering/ThumbnailRenderer.h"
 #include "FontThumbnailRenderer.generated.h"
 
 class FCanvas;
 class FRenderTarget;
 
 UCLASS()
-class UFontThumbnailRenderer : public UTextureThumbnailRenderer
+class UFontThumbnailRenderer : public UThumbnailRenderer
 {
 	GENERATED_UCLASS_BODY()
 
