Update viewport axis visualizer.
- Use manipulator gizmo colors instead of raw RGB.
- Use both light and dark lines to improve visibility against varied backgrounds.
- Place axis letters at the end of the axis lines.
- Add dot at the origin to prevent visual gaps.
- Increase offset to avoid cutting off letters at edge of viewport.

diff --git a/Engine/Source/Editor/UnrealEd/Private/EditorViewportClient.cpp b/Engine/Source/Editor/UnrealEd/Private/EditorViewportClient.cpp
--- a/Engine/Source/Editor/UnrealEd/Private/EditorViewportClient.cpp
+++ b/Engine/Source/Editor/UnrealEd/Private/EditorViewportClient.cpp
@@ -3609,18 +3609,40 @@ void FEditorViewportClient::DrawAxes(FViewport* InViewport, FCanvas* Canvas, con
 
 	const int32 SizeY = InViewport->GetSizeXY().Y / Canvas->GetDPIScale();
 
-	const FIntPoint AxisOrigin( 30, SizeY - 30 );
-	const float AxisSize = 25.f;
+	const FIntPoint AxisOrigin(45.0f, SizeY - 45.0f);
+	const float AxisSize = 25.0f;
 
 	UFont* Font = GEngine->GetSmallFont();
 	int32 XL, YL;
 	StringSize(Font, XL, YL, TEXT("Z"));
 
 	FVector AxisVec;
+	FVector TextVec;
 	FIntPoint AxisEnd;
+	FIntPoint TextCenter;
+	float TextDistance = 12.0f;
 	FCanvasLineItem LineItem;
+	LineItem.LineThickness = 1.0f;
 	FCanvasTextItem TextItem( FVector2D::ZeroVector, FText::GetEmpty(), Font, FLinearColor::White );
+	FVector2D InsideCenterSize = FVector2D(2, 2);
+	FCanvasTileItem InsideTileItem(FVector2D::ZeroVector, InsideCenterSize, FLinearColor(0.2f, 0.2f, 0.2f));
+	FVector2D OutsideCenterSize = FVector2D(4, 4);
+	FCanvasTileItem OutsideTileItem(FVector2D::ZeroVector, OutsideCenterSize, FLinearColor(0.05f, 0.05f, 0.05f));
 	
+	FLinearColor XBaseColor = FLinearColor(0.594f, 0.0197f, 0.0f);
+	FLinearColor XShadowColor = FLinearColor(0.086f, 0.012f, 0.004f);
+	FLinearColor YBaseColor = FLinearColor(0.1349f, 0.3959f, 0.0f);
+	FLinearColor YShadowColor = FLinearColor(0.039f, 0.075f, 0.008f);
+	FLinearColor ZBaseColor = FLinearColor(0.0251f, 0.207f, 0.85f);
+	FLinearColor ZShadowColor = FLinearColor(0.012f, 0.035f, 0.094f);
+
+	auto GetShadowOffset = [this](FIntPoint EndPoint, FIntPoint OriginPoint) -> FIntPoint
+		{
+			int32 X = (FMath::Abs(EndPoint.X - OriginPoint.X) > FMath::Abs(EndPoint.Y - OriginPoint.Y)) ? 0 : 1;
+			int32 Y = (FMath::Abs(EndPoint.X - OriginPoint.X) > FMath::Abs(EndPoint.Y - OriginPoint.Y)) ? 1 : 0;
+			return FIntPoint(X, Y);
+		};
+		
 	const TStaticArray<FText, 3> AxisNames = []()
 	{
 		static bool bUseForwardRightUpDisplayNames = []()
@@ -3652,27 +3674,37 @@ void FEditorViewportClient::DrawAxes(FViewport* InViewport, FCanvas* Canvas, con
 			return FText();
 		}
 	};
-	
+
 	if( ( InAxis & EAxisList::X ) ==  EAxisList::X)
 	{
 		AxisVec = AxisSize * ViewTM.InverseTransformVector( FVector(1,0,0) );
 		AxisEnd = AxisOrigin + FIntPoint( AxisVec.Y, -AxisVec.Z );
-		LineItem.SetColor( FLinearColor::Red );
-		TextItem.SetColor( FLinearColor::Red );
+		TextVec = (AxisSize + TextDistance) * ViewTM.InverseTransformVector(FVector(1, 0, 0));
+		TextCenter = AxisOrigin + FIntPoint(TextVec.Y, -TextVec.Z);
+		LineItem.SetColor(XShadowColor);
+		LineItem.Draw(Canvas, AxisOrigin + GetShadowOffset(AxisEnd, AxisOrigin), AxisEnd + GetShadowOffset(AxisEnd, AxisOrigin));
+		LineItem.SetColor(XBaseColor);
 		LineItem.Draw( Canvas, AxisOrigin, AxisEnd );
 		TextItem.Text = GetAxisName(EAxisList::X);
-		TextItem.Draw( Canvas, FVector2D(AxisEnd.X + 2, AxisEnd.Y - 0.5*YL) );
+		TextItem.SetColor(XBaseColor);
+		TextItem.EnableShadow(XShadowColor);
+		TextItem.Draw( Canvas, FVector2D(TextCenter.X - 0.5*XL, TextCenter.Y - 0.5*YL) );
 	}
 
 	if( ( InAxis & EAxisList::Y ) == EAxisList::Y)
 	{
 		AxisVec = AxisSize * ViewTM.InverseTransformVector( FVector(0,1,0) );
 		AxisEnd = AxisOrigin + FIntPoint( AxisVec.Y, -AxisVec.Z );
-		LineItem.SetColor( FLinearColor::Green );
-		TextItem.SetColor( FLinearColor::Green );
+		TextVec = (AxisSize + TextDistance) * ViewTM.InverseTransformVector(FVector(0, 1, 0));
+		TextCenter = AxisOrigin + FIntPoint(TextVec.Y, -TextVec.Z);
+		LineItem.SetColor(YShadowColor);
+		LineItem.Draw(Canvas, AxisOrigin + GetShadowOffset(AxisEnd, AxisOrigin), AxisEnd + GetShadowOffset(AxisEnd, AxisOrigin));
+		LineItem.SetColor(YBaseColor);
 		LineItem.Draw( Canvas, AxisOrigin, AxisEnd );
 		TextItem.Text = GetAxisName(EAxisList::Y);
-		TextItem.Draw( Canvas, FVector2D(AxisEnd.X + 2, AxisEnd.Y - 0.5*YL) );
+		TextItem.SetColor(YBaseColor);
+		TextItem.EnableShadow(YShadowColor);
+		TextItem.Draw( Canvas, FVector2D(TextCenter.X - 0.5*XL, TextCenter.Y - 0.5*YL) );
 
 	}
 
@@ -3680,14 +3712,23 @@ void FEditorViewportClient::DrawAxes(FViewport* InViewport, FCanvas* Canvas, con
 	{
 		AxisVec = AxisSize * ViewTM.InverseTransformVector( FVector(0,0,1) );
 		AxisEnd = AxisOrigin + FIntPoint( AxisVec.Y, -AxisVec.Z );
-		LineItem.SetColor( FLinearColor::Blue );
-		TextItem.SetColor( FLinearColor::Blue );
+		TextVec = (AxisSize + TextDistance) * ViewTM.InverseTransformVector(FVector(0, 0, 1));
+		TextCenter = AxisOrigin + FIntPoint(TextVec.Y, -TextVec.Z);
+		LineItem.SetColor(ZShadowColor);
+		LineItem.Draw(Canvas, AxisOrigin + GetShadowOffset(AxisEnd, AxisOrigin), AxisEnd + GetShadowOffset(AxisEnd, AxisOrigin));
+		LineItem.SetColor(ZBaseColor);
 		LineItem.Draw( Canvas, AxisOrigin, AxisEnd );
 		TextItem.Text = GetAxisName(EAxisList::Z);
-		TextItem.Draw( Canvas, FVector2D(AxisEnd.X + 2, AxisEnd.Y - 0.5*YL) );
+		TextItem.SetColor(ZBaseColor);
+		TextItem.EnableShadow(ZShadowColor);
+		TextItem.Draw( Canvas, FVector2D(TextCenter.X - 0.5*XL, TextCenter.Y - 0.5*YL) );
 	}
+
+	OutsideTileItem.Draw(Canvas, FIntPoint(AxisOrigin.X - FMath::Floor(OutsideCenterSize.X * 0.5), AxisOrigin.Y - FMath::Floor(OutsideCenterSize.Y * 0.5)));
+	InsideTileItem.Draw(Canvas, FIntPoint(AxisOrigin.X - FMath::Floor(InsideCenterSize.X * 0.5), AxisOrigin.Y - FMath::Floor(InsideCenterSize.Y * 0.5) ));
 }
 
+
 /** Convert the specified number (in cm or unreal units) into a readable string with relevant si units */
 FString FEditorViewportClient::UnrealUnitsToSiUnits(float UnrealUnits)
 {
@@ -3755,7 +3796,7 @@ void FEditorViewportClient::DrawScaleUnits(FViewport* InViewport, FCanvas* Canva
 	StringSize(Font, TextWidth, TextHeight, *DisplayText);
 
 	// Origin is the bottom left of the scale
-	const FIntPoint StartPoint(80, InViewport->GetSizeXY().Y/Canvas->GetDPIScale() - 30);
+	const FIntPoint StartPoint(130, InViewport->GetSizeXY().Y/Canvas->GetDPIScale() - 45);
 	const FIntPoint EndPoint = StartPoint + (UnitsPerPixel != 0 ? FIntPoint(SegmentWidthUnits / UnitsPerPixel, 0) : FIntPoint(0,0));
 
 	// Sort out the color for the text and widget
