For delivery on release

#rnx

diff --git a/Engine/Source/Editor/UnrealEd/Classes/Commandlets/GatherTextFromAssetsCommandlet.h b/Engine/Source/Editor/UnrealEd/Classes/Commandlets/GatherTextFromAssetsCommandlet.h
--- a/Engine/Source/Editor/UnrealEd/Classes/Commandlets/GatherTextFromAssetsCommandlet.h
+++ b/Engine/Source/Editor/UnrealEd/Classes/Commandlets/GatherTextFromAssetsCommandlet.h
@@ -117,12 +117,16 @@ private:
 
 	TArray<FPackagePendingGather> PackagesPendingGather;
 
+	TSet<FName> PackagesWithDuplicateLocalizationIds;
+
 	/** Run a GC if the free system memory is less than this value (or zero to disable) */
 	uint64 MinFreeMemoryBytes;
 
 	/** Run a GC if the used process memory is greater than this value (or zero to disable) */
 	uint64 MaxUsedMemoryBytes;
 
+	uint64 NumPackagesDupLocId;
+
 	/** Array of objects that should be kept alive during the next call to CollectGarbage (used by PurgeGarbage and AddReferencedObjects) */
 	TSet<TObjectPtr<UObject>> ObjectsToKeepAlive;
 	/** Path to the directory where output reports etc will be saved.*/
@@ -133,4 +137,5 @@ private:
 	bool bFixMissingGatherCache;
 	bool bShouldGatherFromEditorOnlyData;
 	bool bShouldExcludeDerivedClasses;
+	bool bFixPackageLocalizationIdConflict;
 };
diff --git a/Engine/Source/Editor/UnrealEd/Private/Commandlets/GatherTextFromAssetsCommandlet.cpp b/Engine/Source/Editor/UnrealEd/Private/Commandlets/GatherTextFromAssetsCommandlet.cpp
--- a/Engine/Source/Editor/UnrealEd/Private/Commandlets/GatherTextFromAssetsCommandlet.cpp
+++ b/Engine/Source/Editor/UnrealEd/Private/Commandlets/GatherTextFromAssetsCommandlet.cpp
@@ -3,6 +3,7 @@
 #include "Commandlets/GatherTextFromAssetsCommandlet.h"
 #include "UObject/Class.h"
 #include "HAL/FileManager.h"
+#include "Internationalization/TextPackageNamespaceUtil.h"
 #include "Logging/StructuredLog.h"
 #include "Misc/CommandLine.h"
 #include "Misc/ConfigCacheIni.h"
@@ -280,9 +281,11 @@ UGatherTextFromAssetsCommandlet::UGatherTextFromAssetsCommandlet(const FObjectIn
 	: Super(ObjectInitializer)
 	, MinFreeMemoryBytes(0)
 	, MaxUsedMemoryBytes(0)
+	, NumPackagesDupLocId(0)
 	, bSkipGatherCache(false)
 	, bShouldGatherFromEditorOnlyData(false)
 	, bShouldExcludeDerivedClasses(false)
+	, bFixPackageLocalizationIdConflict(false)
 {
 }
 
@@ -570,6 +573,7 @@ bool UGatherTextFromAssetsCommandlet::BuildCollectionFilter(FARFilter& InOutFilt
 		}
 	}
 
+
 	return !bHasFailedToGetACollection;
 }
 
@@ -1154,21 +1158,31 @@ void UGatherTextFromAssetsCommandlet::ProcessAndRemoveCachedPackages(TMap<FName,
 			const bool bIsExternalActorPackage = PackageNameStr.ToView().Contains(FPackagePath::GetExternalActorsFolderName());
 
 			// Track the package localization ID of this package (if known) and detect duplicates
+			bool bThisPackageHasLocIdConflictToFix = false;
 			if (!PackagePendingGather.PackageLocalizationId.IsEmpty())
 			{
+				// if this package's localization ID is a duplicate
 				if (const FName* ExistingLongPackageName = AssignedPackageLocalizationIds.Find(PackagePendingGather.PackageLocalizationId))
 				{
 					UE_LOG(LogGatherTextFromAssetsCommandlet, Warning, TEXT("Package '%s' and '%s' have the same localization ID (%s). Please reset one of these (Asset Localization -> Reset Localization ID) to avoid conflicts."), *PackageNameStr, *ExistingLongPackageName->ToString(), *PackagePendingGather.PackageLocalizationId);
+					if (bFixPackageLocalizationIdConflict) 
+					{
+						bThisPackageHasLocIdConflictToFix = true;
+						// add this package to the list of packages with a duplicate localization ID
+						PackagesWithDuplicateLocalizationIds.Add(PackagePendingGather.PackageName);
+						NumPackagesDupLocId++;
+					}
 				}
 				else
 				{
+					// This package is not a duplicate so it is added to Assigned Package Localization IDs dictionary
 					AssignedPackageLocalizationIds.Add(PackagePendingGather.PackageLocalizationId, PackagePendingGather.PackageName);
 				}
 			}
 
 			if (PackagePendingGather.PackageLocCacheState != EPackageLocCacheState::Cached)
 			{
-				// External actors actors must be gathered via their owner world rather than via a raw LoadPackage call
+				// External actors must be gathered via their owner world rather than via a raw LoadPackage call
 				// Remove them from PackagesToGather as the owner world is merged back in below
 				if (bIsExternalActorPackage)
 				{
@@ -1201,7 +1215,7 @@ void UGatherTextFromAssetsCommandlet::ProcessAndRemoveCachedPackages(TMap<FName,
 
 			// If we're reporting or fixing assets with a stale gather cache then we still need to load this 
 			// package in order to do that, but the PackageLocCacheState prevents it being gathered again
-			if (bReportStaleGatherCache || bFixStaleGatherCache)
+			if (bReportStaleGatherCache || bFixStaleGatherCache || bThisPackageHasLocIdConflictToFix)
 			{
 				check(PackagePendingGather.PackageLocCacheState == EPackageLocCacheState::Cached);
 				return false;
@@ -1212,6 +1226,8 @@ void UGatherTextFromAssetsCommandlet::ProcessAndRemoveCachedPackages(TMap<FName,
 
 		UE_LOG(LogGatherTextFromAssetsCommandlet, Display, TEXT("Ingesting the cached localization data for %d package(s) took %.2f seconds."), InitialNumPackagesToGather, FPlatformTime::Seconds() - IngestStartTime);
 	}
+	UE_LOG(LogGatherTextFromAssetsCommandlet, Display, TEXT("Number of packages with duplicate loc ids: %d"), NumPackagesDupLocId);
+	UE_LOG(LogGatherTextFromAssetsCommandlet, Display, TEXT("Number of packages to load: %d"), PackagesPendingGather.Num());
 }
 
 void UGatherTextFromAssetsCommandlet::MergeInExternalActorsWithStaleOrMissingCaches(TMap<FName, TSet<FGuid>>& StaleExternalActors)
@@ -1250,6 +1266,19 @@ bool UGatherTextFromAssetsCommandlet::LoadAndProcessUncachedPackages(TArray<FNam
 	int32 NumPackagesFailedLoading = 0;
 	
 	FScopedSlowTask SlowTask(PackageCount, NSLOCTEXT("GatherTextCommandlet", "LoadAndProcessUncachedPackages", "Loading and Gathering Packages..."));
+	UE_LOG(LogGatherTextFromAssetsCommandlet, Display, TEXT("Number of packages with duplicate loc ids: %d"), NumPackagesDupLocId);
+	UE_LOG(LogGatherTextFromAssetsCommandlet, Display, TEXT("Number of packages to load: %d"), PackagesPendingGather.Num());
+
+	// Auxiliary lambda function for resaving packages
+	auto ResavePackage = [this, &LogOutputRedirector](
+		const FStringView PackageNameStr,
+		UPackage* Package,
+		const FString& PackageFileName
+		) -> bool
+	{
+		FLoadPackageLogOutputRedirector::FScopedCapture ScopedCapture(&LogOutputRedirector, PackageNameStr);
+		return FLocalizedAssetSCCUtil::SavePackageWithSCC(SourceControlInfo, Package, PackageFileName);
+	};
 
 	while (PackagesPendingGather.Num() > 0)
 	{
@@ -1288,6 +1317,31 @@ bool UGatherTextFromAssetsCommandlet::LoadAndProcessUncachedPackages(TArray<FNam
 			continue;
 		}
 
+		// if fixing duplicate localization package IDs mode is enabled, 
+		// and the package is in the list of packages with duplicate localization IDs,
+		// reset the package's localization ID
+		bool bPackageLocIdWasReset = false;
+		if (bFixPackageLocalizationIdConflict)
+		{
+			if (PackagesWithDuplicateLocalizationIds.Contains(PackagePendingGather.PackageName))
+			{
+				// This package's localization ID and name have been found in the duplicates dictionary, and the flag for fixing duplicates is on
+				// So resetting the localization ID for that package
+				UE_LOG(LogGatherTextFromAssetsCommandlet, Display, TEXT("About to reset localization id for package: '%s'."), *PackageNameStr);
+				TextNamespaceUtil::ClearPackageNamespace(Package);
+				TextNamespaceUtil::EnsurePackageNamespace(Package);
+				bPackageLocIdWasReset = true;
+				UE_LOG(LogGatherTextFromAssetsCommandlet, Display, TEXT("Localization ID has been reset for package: '%s'."), *PackageNameStr);
+				
+				// Now resaving to ensure the in-memory IDs are updated in the package before the data gatherer runs over it
+				UE_LOG(LogGatherTextFromAssetsCommandlet, Display, TEXT("Resaving package: '%s'..."), *PackageNameStr);
+				if (!ResavePackage(PackageNameStr, Package, PackagePendingGather.PackageFilename))
+				{
+					UE_LOG(LogGatherTextFromAssetsCommandlet, Warning, TEXT("Failed to resave package: '%s'."), *PackageNameStr);
+				}
+			}
+		}
+
 		// Tick background tasks
 		UE::Private::GatherTextFromAssetsCommandlet::TickBackgroundTasks();
 		if (EmbeddedContext)
@@ -1330,18 +1384,17 @@ bool UGatherTextFromAssetsCommandlet::LoadAndProcessUncachedPackages(TArray<FNam
 				bSavePackage = true;
 			}
 
+			// if the package localization ID was reset in duplicate localization ID fixing mode, it needs to be resaved
+			if(bPackageLocIdWasReset)
+			{ 
+				bSavePackage = true;
+			}
+
 			// Re-save the package to attempt to fix it?
 			if (bSavePackage)
 			{
 				UE_LOG(LogGatherTextFromAssetsCommandlet, Display, TEXT("Resaving package: '%s'..."), *PackageNameStr);
-
-				bool bSavedPackage = false;
-				{
-					FLoadPackageLogOutputRedirector::FScopedCapture ScopedCapture(&LogOutputRedirector, PackageNameStr);
-					bSavedPackage = FLocalizedAssetSCCUtil::SavePackageWithSCC(SourceControlInfo, Package, PackagePendingGather.PackageFilename);
-				}
-
-				if (!bSavedPackage)
+				if (!ResavePackage(PackageNameStr, Package, PackagePendingGather.PackageFilename))
 				{
 					UE_LOG(LogGatherTextFromAssetsCommandlet, Warning, TEXT("Failed to resave package: '%s'."), *PackageNameStr);
 				}
@@ -1402,6 +1455,16 @@ bool UGatherTextFromAssetsCommandlet::LoadAndProcessUncachedPackages(TArray<FNam
 
 			GatherableTextDataArray.Reset();
 		}
+		else if (bPackageLocIdWasReset)
+		{
+			// Resaving a second time in case of package localization ID reset
+			// to fix the on-disk cache to match the new in-memory IDs
+			UE_LOG(LogGatherTextFromAssetsCommandlet, Display, TEXT("Resaving package: '%s'..."), *PackageNameStr);
+			if (!ResavePackage(PackageNameStr, Package, PackagePendingGather.PackageFilename))
+			{
+				UE_LOG(LogGatherTextFromAssetsCommandlet, Warning, TEXT("Failed to resave package: '%s'."), *PackageNameStr);
+			}
+		}
 
 		if (HasExceededMemoryLimit(/*bLog*/true))
 		{
@@ -1418,6 +1481,10 @@ bool UGatherTextFromAssetsCommandlet::LoadAndProcessUncachedPackages(TArray<FNam
 			}
 		}
 	}
+
+	// Clear list of packages with duplicate localization IDs
+	PackagesWithDuplicateLocalizationIds.Empty();
+
 	UE_LOG(LogGatherTextFromAssetsCommandlet, Display, TEXT("Loaded %d packages. %d failed."), NumPackagesProcessed, NumPackagesFailedLoading);
 	return true;
 }
@@ -1770,6 +1837,7 @@ bool UGatherTextFromAssetsCommandlet::ConfigureFromScript(const FString& GatherT
 	ReadBoolFlagWithFallback(TEXT("ReportStaleGatherCache"), bReportStaleGatherCache);
 	ReadBoolFlagWithFallback(TEXT("FixStaleGatherCache"), bFixStaleGatherCache);
 	ReadBoolFlagWithFallback(TEXT("FixMissingGatherCache"), bFixMissingGatherCache);
+	ReadBoolFlagWithFallback(TEXT("FixPackageLocalizationIdConflict"), bFixPackageLocalizationIdConflict);
 
 	// Read some settings from the editor config
 	{
