Changed default value of experimental editor settings flag 'Enable Level Instance Landscape Edit' to enabled.

diff --git a/Engine/Source/Editor/UnrealEd/Private/Settings/SettingsClasses.cpp b/Engine/Source/Editor/UnrealEd/Private/Settings/SettingsClasses.cpp
--- a/Engine/Source/Editor/UnrealEd/Private/Settings/SettingsClasses.cpp
+++ b/Engine/Source/Editor/UnrealEd/Private/Settings/SettingsClasses.cpp
@@ -194,6 +194,7 @@ UEditorExperimentalSettings::UEditorExperimentalSettings( const FObjectInitializ
 	, HDREditorNITLevel(160.0f)
 	, bUseOpenCLForConvexHullDecomp(false)
 	, bAllowPotentiallyUnsafePropertyEditing(false)
+	, bEnableLevelInstanceLandscapeEdit(true)
 {
 }
 
