IncrementalCook: Silently discard build dependencies on /Engine/Transient/* packages rather than disallowing iterativecook when those dependencies are present. Some systems such as niagara store their transient package data in /Engine/Transient/SubPath rather than the main /Engine/Transient package.
#rnx
#incrementalcook

diff --git a/Engine/Source/Editor/UnrealEd/Private/TargetDomain/TargetDomainUtils.cpp b/Engine/Source/Editor/UnrealEd/Private/TargetDomain/TargetDomainUtils.cpp
--- a/Engine/Source/Editor/UnrealEd/Private/TargetDomain/TargetDomainUtils.cpp
+++ b/Engine/Source/Editor/UnrealEd/Private/TargetDomain/TargetDomainUtils.cpp
@@ -300,12 +300,8 @@ enum class EPackageMountPoint
 	Content,
 	GeneratedContent,
 };
-EPackageMountPoint GetPackageMountPoint(FName PackageName, FName TransientPackageName)
+EPackageMountPoint GetPackageMountPoint(FName PackageName)
 {
-	if (PackageName == TransientPackageName)
-	{
-		return EPackageMountPoint::Transient;
-	}
 	TStringBuilder<256> StringBuffer;
 	PackageName.ToString(StringBuffer);
 	if (
@@ -315,7 +311,8 @@ EPackageMountPoint GetPackageMountPoint(FName PackageName, FName TransientPackag
 		!StringBuffer.ToView().StartsWith('/') ||
 		// Ignore /Memory and /Temp packages
 		FPackageName::IsMemoryPackage(StringBuffer) ||
-		FPackageName::IsTempPackage(StringBuffer))
+		FPackageName::IsTempPackage(StringBuffer) ||
+		FPackageName::IsInEngineTransientPackages(StringBuffer))
 	{
 		return EPackageMountPoint::Transient;
 	}
@@ -589,7 +586,6 @@ FCookDependencies FCookDependencies::Collect(UPackage* Package, const ITargetPla
 		Result.RedirectorsAffectingPackageDependencies.SetNum(Algo::Unique(Result.RedirectorsAffectingPackageDependencies));
 	}
 
-	FName TransientPackageName = GetTransientPackage()->GetFName();
 	Algo::Sort(Result.CookDependencies);
 	Result.CookDependencies.SetNum(Algo::Unique(Result.CookDependencies));
 	for (TPair<FName, EBuildDependency>& Pair : PackageBuildDependencies)
@@ -599,7 +595,7 @@ FCookDependencies FCookDependencies::Collect(UPackage* Package, const ITargetPla
 			// Skip dependencies to self, for performance
 			continue;
 		}
-		if (GetPackageMountPoint(Pair.Key, TransientPackageName) != EPackageMountPoint::Content)
+		if (GetPackageMountPoint(Pair.Key) != EPackageMountPoint::Content)
 		{
 			// Skip script and temp packages
 			continue;
@@ -624,7 +620,7 @@ FCookDependencies FCookDependencies::Collect(UPackage* Package, const ITargetPla
 	for (TArray<FName>::TIterator Iter(RuntimeDependencies); Iter; ++Iter)
 	{
 		FName PackageName = *Iter;
-		EPackageMountPoint MountPoint = GetPackageMountPoint(PackageName, TransientPackageName);
+		EPackageMountPoint MountPoint = GetPackageMountPoint(PackageName);
 		switch (MountPoint)
 		{
 		case EPackageMountPoint::GeneratedContent:
@@ -1007,7 +1003,7 @@ FIncrementalCookAttachments FIncrementalCookAttachments::Collect(UPackage* Packa
 		LogText << TEXT("Could not collect CookDependencies for package '") << Package->GetFName() << TEXT("'");
 		JoinMessagesIntoErrorReason(LogText, Messages);
 
-		// IterativeTODO: This error occurs due to dependencies on _Verse and on some transient packages.
+		// IterativeTODO: This error occurs due to dependencies on _Verse.
 		// Raise Verbosity to Error once that is fixed.
 		UE_LOG(LogCook, Verbose, TEXT("%s"), *LogText);
 	}
diff --git a/Engine/Source/Runtime/CoreUObject/Private/Misc/PackageName.cpp b/Engine/Source/Runtime/CoreUObject/Private/Misc/PackageName.cpp
--- a/Engine/Source/Runtime/CoreUObject/Private/Misc/PackageName.cpp
+++ b/Engine/Source/Runtime/CoreUObject/Private/Misc/PackageName.cpp
@@ -3178,6 +3178,13 @@ bool FPackageName::IsTempPackage(FStringView InPackageName)
 	return FPathViews::IsParentPathOf(FLongPackagePathsSingleton::Get().TempRootPath, InPackageName);
 }
 
+bool FPackageName::IsInEngineTransientPackages(FStringView InPackageName)
+{
+	constexpr FStringView EngineTransientPrefix(TEXTVIEW("/Engine/Transient"));
+	return InPackageName.StartsWith(EngineTransientPrefix, ESearchCase::IgnoreCase) &&
+		(InPackageName.Len() == EngineTransientPrefix.Len() || InPackageName[EngineTransientPrefix.Len()] == '/');
+}
+
 bool FPackageName::IsLocalizedPackage(FStringView InPackageName)
 {
 	// Minimum valid package name length is "/A/L10N"
diff --git a/Engine/Source/Runtime/CoreUObject/Public/Misc/PackageName.h b/Engine/Source/Runtime/CoreUObject/Public/Misc/PackageName.h
--- a/Engine/Source/Runtime/CoreUObject/Public/Misc/PackageName.h
+++ b/Engine/Source/Runtime/CoreUObject/Public/Misc/PackageName.h
@@ -854,6 +854,12 @@ public:
 	 */
 	static COREUOBJECT_API bool IsTempPackage(FStringView InPackageName);
 
+	/**
+	 * Checks the root of the package's path to see if it is the EngineTransient package in specific, OR
+	 * one of the packages in the directory of EngineTransient packages.
+	 */
+	static COREUOBJECT_API bool IsInEngineTransientPackages(FStringView InPackageName);
+
 	/**
 	 * Checks the root of the package's path to see if it is a localized package
 	 * @return true if the root of the path matches any localized root path
