Reenable support for the UBlueprint::bRecompileOnLoad  config variable

[RN] The [/Script/Engine.Blueprint] bRecompileOnLoad config value is once again honored by the blueprint compiler. It has been ignored for several releases and now warns when it is used to suppress compilation on load

diff --git a/Engine/Source/Editor/UnrealEd/Private/Kismet2/BlueprintEditorUtils.cpp b/Engine/Source/Editor/UnrealEd/Private/Kismet2/BlueprintEditorUtils.cpp
--- a/Engine/Source/Editor/UnrealEd/Private/Kismet2/BlueprintEditorUtils.cpp
+++ b/Engine/Source/Editor/UnrealEd/Private/Kismet2/BlueprintEditorUtils.cpp
@@ -1909,7 +1909,20 @@ bool FBlueprintEditorUtils::ShouldRegenerateBlueprint(UBlueprint* Blueprint)
 
 bool FBlueprintEditorUtils::IsCompileOnLoadDisabled(UBlueprint* Blueprint)
 {
-	bool bCompilationDisabled = false;
+	bool bCompilationDisabled = !Blueprint->bRecompileOnLoad;
+	if(bCompilationDisabled)
+	{
+		// bRecompileOnLoad has been unused for awhile, warn once about its use,
+		// in case it is unintentional or the user has forgotten they've disabled
+		// it:
+		static bool bHasWarned = false;
+		if(!bHasWarned)
+		{
+			bHasWarned = true;
+			UE_LOG(LogBlueprint, Warning, 
+				TEXT("Blueprints are loading but RecompileOnLoad is disabled - no bytecode will be available unless SkipByteCodeSerialization is disabled"));
+		}
+	}
 	if(Blueprint->GetLinker())
 	{
 		bCompilationDisabled = (Blueprint->GetLinker()->LoadFlags & LOAD_DisableCompileOnLoad) != LOAD_None;
@@ -1918,7 +1931,7 @@ bool FBlueprintEditorUtils::IsCompileOnLoadDisabled(UBlueprint* Blueprint)
 	// it's already cooked, if will already be up to date (likely they shouldn't exist, but in case they do we need to make them work)
 	if (Blueprint->GetOutermost()->bIsCookedForEditor)
 	{
-		return true;
+		bCompilationDisabled = true;
 	}
 	return bCompilationDisabled;
 }
diff --git a/Engine/Source/Editor/UnrealEd/Public/Kismet2/BlueprintEditorUtils.h b/Engine/Source/Editor/UnrealEd/Public/Kismet2/BlueprintEditorUtils.h
--- a/Engine/Source/Editor/UnrealEd/Public/Kismet2/BlueprintEditorUtils.h
+++ b/Engine/Source/Editor/UnrealEd/Public/Kismet2/BlueprintEditorUtils.h
@@ -289,6 +289,7 @@ public:
 	/**
 	 * Whether or not the blueprint should regenerate its class on load or not.  This prevents macros and other BP types not marked for reconstruction from being recompiled all the time
 	 */
+	UE_DEPRECATED(5.6, "Replace with IsCompileOnLoadDisabled, ShouldRegenerateBlueprint handles reentrancy which is no longer possible.")
 	static UNREALED_API bool ShouldRegenerateBlueprint(UBlueprint* Blueprint);
 
 	/** Returns true if compilation for the given blueprint has been disabled */
