IncrementalCook: Fix error in IterativeValidate:
"Package %s is a generated package, but ParentGenerator %s is not a generator package. The generated package cannot be loaded." Generated packages need to not proceed past request state until after their generator has queued them, to handle RequiresCachedCookedPlatformDataBeforeSplit generators.
#rnx
#incrementalcook

diff --git a/Engine/Source/Editor/UnrealEd/Private/CookOnTheFlyServer.cpp b/Engine/Source/Editor/UnrealEd/Private/CookOnTheFlyServer.cpp
--- a/Engine/Source/Editor/UnrealEd/Private/CookOnTheFlyServer.cpp
+++ b/Engine/Source/Editor/UnrealEd/Private/CookOnTheFlyServer.cpp
@@ -2753,6 +2753,32 @@ void UCookOnTheFlyServer::PumpRequests(UE::Cook::FTickStackData& StackData, int3
 		{
 			continue;
 		}
+		if (PackageData->IsGenerated() && !CookWorkerClient)
+		{
+			// We only need to check for this on the CookDirector (or in SPCook), because it is at the start of making
+			// every package active. CookWorkers do not have the information they need to check for this.
+			// 
+			// Generated packages cannot be requested before their generator queues them for some generators,
+			// because some generators require that they save on the same CooKWorker as the generator
+			// (EGeneratedRequiresGenerator::Save) and we don't know that assignment yet, and some generators require
+			// that the generator run BeginCacheCookedPlatformData before calling ShouldSplit
+			// (RequiresCachedCookedPlatformDataBeforeSplit).
+			// Generated packages being queued before their generator queues them can occur during IterativeValidate
+			// incremental cooks: if the generator is up to date, the incremental cook will explore the generator and
+			// all of its generated packages, but IterativeValidate will decide they need to be cooked and will queue
+			// them for cooking all at the same time.
+			TRefCountPtr<FGenerationHelper> ParentGenerationHelper = PackageData->GetParentGenerationHelper();
+			if (!CookWorkerClient &&
+				(!ParentGenerationHelper || !ParentGenerationHelper->GetDirectorAPI().HasStartedQueueGeneratedPackages()))
+			{
+				// If the generated package's generator has not yet called QueueGeneratedPackages, temporarily demote the
+				// generated package back to Idle; it will be requeued during QueueGeneratedPackages. Also reset its
+				// reachability, so that we don't early exit from that upcoming QueueDiscoveredPackage call.
+				PackageData->ResetReachable(EReachability::Runtime);
+				DemoteToIdle(*PackageData, ESendFlags::QueueAdd, ESuppressCookReason::NotYetReadyForRequest);
+				continue;
+			}
+		}
 		if (PackageData->GetPlatformsNeedingCommitNum(GetCookPhase()) == 0)
 		{
 			ESuppressCookReason SuppressCookReason = PackageData->HasAnyCommittedPlatforms()
@@ -2851,6 +2877,8 @@ void UCookOnTheFlyServer::DemoteToIdle(UE::Cook::FPackageData& PackageData, UE::
 		bool bPrintDiagnostic = !bCookListMode &
 			!!((GCookProgressDisplay & ((int32)ECookProgressDisplayMode::Instigators | (int32)ECookProgressDisplayMode::PackageNames)));
 
+		// Suppress the message if it's a temporary demotion
+		bPrintDiagnostic &= (Reason != ESuppressCookReason::NotYetReadyForRequest);
 		// Suppress the message in cases that cause large spam like NotInCurrentPlugin for DLC cooks.
 		bPrintDiagnostic &= (Reason != ESuppressCookReason::NotInCurrentPlugin);
 		// Iterative cooks: suppress the diagnostic for packages that were iteratively skipped
diff --git a/Engine/Source/Editor/UnrealEd/Private/Cooker/CookGenerationHelper.cpp b/Engine/Source/Editor/UnrealEd/Private/Cooker/CookGenerationHelper.cpp
--- a/Engine/Source/Editor/UnrealEd/Private/Cooker/CookGenerationHelper.cpp
+++ b/Engine/Source/Editor/UnrealEd/Private/Cooker/CookGenerationHelper.cpp
@@ -1111,6 +1111,11 @@ void FGenerationHelper::NotifyStartQueueGeneratedPackages(UCookOnTheFlyServer& C
 	SetKeepForQueueResults();
 }
 
+bool FGenerationHelper::FDirectorAPI::HasStartedQueueGeneratedPackages() const
+{
+	return GenerationHelper.GetOwnerInfo().SavedOnWorker.IsValid();
+}
+
 void FGenerationHelper::EndQueueGeneratedPackages(UCookOnTheFlyServer& COTFS)
 {
 	bHasFinishedQueueGeneratedPackages = true;
diff --git a/Engine/Source/Editor/UnrealEd/Private/Cooker/CookGenerationHelper.h b/Engine/Source/Editor/UnrealEd/Private/Cooker/CookGenerationHelper.h
--- a/Engine/Source/Editor/UnrealEd/Private/Cooker/CookGenerationHelper.h
+++ b/Engine/Source/Editor/UnrealEd/Private/Cooker/CookGenerationHelper.h
@@ -201,6 +201,7 @@ public:
 	TArrayView<FCookGenerationInfo> GetPackagesToGenerate();
 	/** Return the GenerationInfo used to save the generator package's UPackage. Does not call Initialize. */
 	FCookGenerationInfo& GetOwnerInfo();
+	const FCookGenerationInfo& GetOwnerInfo() const;
 	/** Return owner FPackageData. Does not call Initialize. */
 	FPackageData& GetOwner();
 	/** Return the GenerationInfo for the given PackageData, or null if not found. */
@@ -435,6 +436,27 @@ public:
 	/** Debug setting that forces an otherwise unspecified order. Generated packages save before their generator does. */
 	static bool IsGeneratedSavedFirst();
 
+	/**
+	 * An API to access data that's only available on the CookDirector (or SingleProcess cook). This is present for
+	 * readability at callsites, it just accesses data on the GenerationHelper.
+	 */
+	struct FDirectorAPI
+	{
+	public:
+		/** Report whether NotifyStartQueueGeneratedPackages has been called. Returns false on CookWorkers. */
+		bool HasStartedQueueGeneratedPackages() const;
+
+	private:
+		explicit FDirectorAPI(FGenerationHelper& InGenerationHelper);
+		FDirectorAPI(const FDirectorAPI&) = delete;
+		FDirectorAPI(FDirectorAPI&&) = delete;
+
+		FGenerationHelper& GenerationHelper;
+
+		friend struct FGenerationHelper;
+	};
+	FDirectorAPI GetDirectorAPI() const;
+
 private:
 	enum class EInitializeStatus : uint8
 	{
@@ -629,6 +651,11 @@ inline UE::Cook::FCookGenerationInfo& FGenerationHelper::GetOwnerInfo()
 	return OwnerInfo;
 }
 
+inline const UE::Cook::FCookGenerationInfo& FGenerationHelper::GetOwnerInfo() const
+{
+	return OwnerInfo;
+}
+
 inline UE::Cook::FPackageData& FGenerationHelper::GetOwner()
 {
 	return *OwnerInfo.PackageData;
@@ -766,4 +793,14 @@ inline int32& FGenerationHelper::GetMPCookNextAssignmentIndex()
 	return MPCookNextAssignmentIndex;
 }
 
+inline FGenerationHelper::FDirectorAPI FGenerationHelper::GetDirectorAPI() const
+{
+	return FDirectorAPI(const_cast<FGenerationHelper&>(*this));
+}
+
+inline FGenerationHelper::FDirectorAPI::FDirectorAPI(FGenerationHelper& InGenerationHelper)
+	: GenerationHelper(InGenerationHelper)
+{
+}
+
 } // namespace UE::Cook
\ No newline at end of file
diff --git a/Engine/Source/Editor/UnrealEd/Private/Cooker/CookTypes.cpp b/Engine/Source/Editor/UnrealEd/Private/Cooker/CookTypes.cpp
--- a/Engine/Source/Editor/UnrealEd/Private/Cooker/CookTypes.cpp
+++ b/Engine/Source/Editor/UnrealEd/Private/Cooker/CookTypes.cpp
@@ -84,6 +84,7 @@ const TCHAR* LexToString(ESuppressCookReason Reason)
 	case ESuppressCookReason::MultiprocessAssignmentError: return TEXT("MultiprocessAssignmentError");
 	case ESuppressCookReason::RetractedByCookDirector: return TEXT("RetractedByCookDirector");
 	case ESuppressCookReason::CookFilter: return TEXT("CookFilter");
+	case ESuppressCookReason::NotYetReadyForRequest: return TEXT("NotYetReadyForRequest");
 	default: return TEXT("Invalid");
 	}
 }
@@ -178,6 +179,7 @@ EStateChangeReason ConvertToStateChangeReason(ESuppressCookReason Reason)
 	case ESuppressCookReason::CookCanceled: return EStateChangeReason::CookerShutdown;
 	case ESuppressCookReason::MultiprocessAssignmentError: return EStateChangeReason::ReassignAbortedPackages;
 	case ESuppressCookReason::RetractedByCookDirector: return EStateChangeReason::Retraction;
+	case ESuppressCookReason::NotYetReadyForRequest: return EStateChangeReason::Retraction;
 	default: return EStateChangeReason::CookSuppressed;
 	}
 }
diff --git a/Engine/Source/Editor/UnrealEd/Private/Cooker/CookTypes.h b/Engine/Source/Editor/UnrealEd/Private/Cooker/CookTypes.h
--- a/Engine/Source/Editor/UnrealEd/Private/Cooker/CookTypes.h
+++ b/Engine/Source/Editor/UnrealEd/Private/Cooker/CookTypes.h
@@ -169,6 +169,7 @@ namespace UE::Cook
 		MultiprocessAssignmentError,
 		RetractedByCookDirector,
 		CookFilter,
+		NotYetReadyForRequest,
 		Count,
 		BitCount = FPlatformMath::ConstExprCeilLogTwo(Count),
 	};
