IncrementalCook: Fix discarding of generated packages if they do not yet have their ParentGenerationHelper yet. We cannot create it if it does not exist, but we can fetch it from the generator if the generator already has it.
This fixes the bad cook with warning message "GenerationHelper for package %s is still allocated at end of cooksession. This is unexpected and could indicate some generated packages are missing."
#rnx

diff --git a/Engine/Source/Editor/UnrealEd/Private/CookOnTheFlyServer.cpp b/Engine/Source/Editor/UnrealEd/Private/CookOnTheFlyServer.cpp
--- a/Engine/Source/Editor/UnrealEd/Private/CookOnTheFlyServer.cpp
+++ b/Engine/Source/Editor/UnrealEd/Private/CookOnTheFlyServer.cpp
@@ -2767,7 +2767,7 @@ void UCookOnTheFlyServer::PumpRequests(UE::Cook::FTickStackData& StackData, int3
 			// incremental cooks: if the generator is up to date, the incremental cook will explore the generator and
 			// all of its generated packages, but IterativeValidate will decide they need to be cooked and will queue
 			// them for cooking all at the same time.
-			TRefCountPtr<FGenerationHelper> ParentGenerationHelper = PackageData->GetParentGenerationHelper();
+			TRefCountPtr<FGenerationHelper> ParentGenerationHelper = PackageData->GetOrFindParentGenerationHelper();
 			if (!CookWorkerClient &&
 				(!ParentGenerationHelper || !ParentGenerationHelper->GetDirectorAPI().HasStartedQueueGeneratedPackages()))
 			{
