Allow the cooker to bypass path length checks if the package writer (e.g. zen store) allows it.
ICookedPackageWriter::FCookCapabilities::bIgnorePathLengthLimits specifies that a writer allow arbitrarily long paths.
Add a command line argument -nolongpaths for windows platform to allow testing path length restrictions even if long paths are enabled at an OS level.

#tests Tested cooking an asset with a 310 character long path on all combinations of zen store/no zen store, OS allowing/disallowing long paths

diff --git a/Engine/Source/Developer/IoStoreUtilities/Public/ZenStoreWriter.h b/Engine/Source/Developer/IoStoreUtilities/Public/ZenStoreWriter.h
--- a/Engine/Source/Developer/IoStoreUtilities/Public/ZenStoreWriter.h
+++ b/Engine/Source/Developer/IoStoreUtilities/Public/ZenStoreWriter.h
@@ -101,6 +101,7 @@ public:
 		FCookCapabilities Result;
 		Result.bDiffModeSupported = true;
 		Result.bOplogAttachments = true;
+		Result.bIgnorePathLengthLimits = true;
 		Result.HeaderFormat = EPackageHeaderFormat::ZenPackageSummary;
 		return Result;
 	}
diff --git a/Engine/Source/Editor/UnrealEd/Private/Cooker/CookSavePackage.cpp b/Engine/Source/Editor/UnrealEd/Private/Cooker/CookSavePackage.cpp
--- a/Engine/Source/Editor/UnrealEd/Private/Cooker/CookSavePackage.cpp
+++ b/Engine/Source/Editor/UnrealEd/Private/Cooker/CookSavePackage.cpp
@@ -245,6 +245,11 @@ void FSaveCookedPackageContext::SetupPlatform(const ITargetPlatform* InTargetPla
 	PlatFilename = Filename.Replace(TEXT("[Platform]"), *TargetPlatform->PlatformName());
 	bPlatformSetupSuccessful = false;
 
+	// It's safe to get this cook context even if we may fail to cook this package as the context is per-platform and not per-package.
+	CookContext = &COTFS.FindOrCreateSaveContext(TargetPlatform);
+	SavePackageContext = &CookContext->SaveContext;
+	PackageWriter = CookContext->PackageWriter;
+
 	if (CommitType == EReachability::Runtime)
 	{
 		ArchiveCookContext.Emplace(Package,
@@ -313,13 +318,16 @@ void FSaveCookedPackageContext::SetupPlatform(const ITargetPlatform* InTargetPla
 			}
 		}
 
-		const FString FullFilename = FPaths::ConvertRelativePathToFull(PlatFilename);
-		if (FullFilename.Len() >= FPlatformMisc::GetMaxPathLength())
+		if (!PackageWriter->GetCookCapabilities().bIgnorePathLengthLimits)
 		{
-			LogCookerMessage(FString::Printf(TEXT("Couldn't save package, filename is too long (%d >= %d): %s"),
-				FullFilename.Len(), FPlatformMisc::GetMaxPathLength(), *FullFilename), EMessageSeverity::Error);
-			SavePackageResult = ESavePackageResult::Error;
-			return;
+			const FString FullFilename = FPaths::ConvertRelativePathToFull(PlatFilename);
+			if (FullFilename.Len() >= FPlatformMisc::GetMaxPathLength())
+			{
+				LogCookerMessage(FString::Printf(TEXT("Couldn't save package, filename is too long (%d >= %d): %s"),
+					FullFilename.Len(), FPlatformMisc::GetMaxPathLength(), *FullFilename), EMessageSeverity::Error);
+				SavePackageResult = ESavePackageResult::Error;
+				return;
+			}
 		}
 
 		bEndianSwap = (!TargetPlatform->IsLittleEndian()) ^ (!PLATFORM_LITTLE_ENDIAN);
@@ -373,9 +381,6 @@ void FSaveCookedPackageContext::SetupPlatform(const ITargetPlatform* InTargetPla
 			TEXT("Committing BuildDependencies for %s"), *PackageName);
 	}
 
-	CookContext = &COTFS.FindOrCreateSaveContext(TargetPlatform);
-	SavePackageContext = &CookContext->SaveContext;
-	PackageWriter = CookContext->PackageWriter;
 	ICookedPackageWriter::FBeginPackageInfo Info;
 	Info.PackageName = Package->GetFName();
 	Info.LooseFilePath = PlatFilename;
diff --git a/Engine/Source/Runtime/Core/Private/Windows/WindowsPlatformMisc.cpp b/Engine/Source/Runtime/Core/Private/Windows/WindowsPlatformMisc.cpp
--- a/Engine/Source/Runtime/Core/Private/Windows/WindowsPlatformMisc.cpp
+++ b/Engine/Source/Runtime/Core/Private/Windows/WindowsPlatformMisc.cpp
@@ -1096,6 +1096,11 @@ int32 FWindowsPlatformMisc::GetMaxPathLength()
 
 		FLongPathsEnabled()
 		{
+			if (FParse::Param( FCommandLine::Get(), TEXT("NoLongPaths")))
+			{
+				bValue = false;
+				return;
+			}
 			HMODULE Handle = GetModuleHandle(TEXT("ntdll.dll"));
 			if (Handle == NULL)
 			{
diff --git a/Engine/Source/Runtime/Core/Public/Serialization/PackageWriter.h b/Engine/Source/Runtime/Core/Public/Serialization/PackageWriter.h
--- a/Engine/Source/Runtime/Core/Public/Serialization/PackageWriter.h
+++ b/Engine/Source/Runtime/Core/Public/Serialization/PackageWriter.h
@@ -280,6 +280,9 @@ public:
 		/** If true, this writer can save and load extra data for each package across cooks. */
 		bool bOplogAttachments = false;
 
+		/** If true, this writer can ignore the path length limits of the host platform and handle paths of arbitrary length in BeginPackage */
+		bool bIgnorePathLengthLimits = false;
+
 		/** What header format is produced as output by this writer. */
 		EPackageHeaderFormat HeaderFormat = EPackageHeaderFormat::PackageFileSummary;
 	};
