Provide TMap key or prop in property softpath to conform to assumptions from 3231246, this fixes an ensure when diffing building props which have a tmap value that is always reporting itself as changed in their sparse class data

[RN] Fixed an ensure that fired when diffing blueprint assets that contained TMaps

diff --git a/Engine/Source/Editor/UnrealEd/Private/DiffUtils.cpp b/Engine/Source/Editor/UnrealEd/Private/DiffUtils.cpp
--- a/Engine/Source/Editor/UnrealEd/Private/DiffUtils.cpp
+++ b/Engine/Source/Editor/UnrealEd/Private/DiffUtils.cpp
@@ -550,7 +550,7 @@ FPropertyPath FPropertySoftPath::ResolvePath(const UObject* Object) const
 						Ret.AddProperty(MakKeyInfo);
 					}
 					else if(ensure( PropertyChain[I+1].PropertyName == MapProperty->ValueProp->GetFName() ))
-					{	
+					{
 						++I;
 
 						UpdateContainerAddress( MapProperty->ValueProp, MapHelper.GetValuePtr(InternalIndex), ContainerAddress, ContainerStruct );
@@ -971,13 +971,15 @@ static void IdenticalHelper(const FProperty* AProperty, const FProperty* BProper
 			FScriptMapHelper::FIterator IteratorB(MapHelperB);
 			for (; IteratorA && IteratorB; ++IteratorA, ++IteratorB)
 			{
-				DiffParameters.RootPath = FPropertySoftPath(RootPath, IteratorA.GetLogicalIndex());
+				const FPropertySoftPath EntryIndexPath = FPropertySoftPath(RootPath, IteratorA.GetLogicalIndex());
 
+				DiffParameters.RootPath = FPropertySoftPath(EntryIndexPath, BPropAsMap->KeyProp);
 				IdenticalHelper(APropAsMap->KeyProp, BPropAsMap->KeyProp, MapHelperA.GetKeyPtr(IteratorA), MapHelperB.GetKeyPtr(IteratorB),
 					OwningOuterA, OwningOuterB, DiffParameters, DifferingSubProperties);
+					
+				DiffParameters.RootPath = FPropertySoftPath(EntryIndexPath, BPropAsMap->ValueProp);
 				IdenticalHelper(APropAsMap->ValueProp, BPropAsMap->ValueProp, MapHelperA.GetValuePtr(IteratorA), MapHelperB.GetValuePtr(IteratorB),
 					OwningOuterA, OwningOuterB, DiffParameters, DifferingSubProperties);
-
 			}
 		}
 		else
