Cooker: Fix an unfiltered AssetRegistry being written (and sometimes being so large that it asserts because it's more than 2GB) when no packages were cooked in the -cooksinglepackagenorefs case.
#rnx

diff --git a/Engine/Source/Editor/UnrealEd/Private/Commandlets/AssetRegistryGenerator.cpp b/Engine/Source/Editor/UnrealEd/Private/Commandlets/AssetRegistryGenerator.cpp
--- a/Engine/Source/Editor/UnrealEd/Private/Commandlets/AssetRegistryGenerator.cpp
+++ b/Engine/Source/Editor/UnrealEd/Private/Commandlets/AssetRegistryGenerator.cpp
@@ -1645,7 +1645,13 @@ void FAssetRegistryGenerator::FinalizeChunkIDs(const TSet<FName>& InCookedPackag
 		// Create a new FAssetRegistryState and call InitializeFromExistingAndPrune, then move the result to State.
 		// This is faster than calling State.PruneAssetData. PruneAssetData removes iteratively all the packages one by one and it is slow.
 		FAssetRegistryState PrunedAssetRegistry;
-		PrunedAssetRegistry.InitializeFromExistingAndPrune(State, AllPackages, TSet<FName>(), TSet<int32>(), DevelopmentSaveOptions);
+		if (!AllPackages.IsEmpty())
+		{
+			// Only call InitializeFromExistingAndPrune if we recorded at least one package.
+			// InitializeFromExistingAndPrune does not handle the empty case because it assumes
+			// that an empty RequiredPackages means no required packages.
+			PrunedAssetRegistry.InitializeFromExistingAndPrune(State, AllPackages, TSet<FName>(), TSet<int32>(), DevelopmentSaveOptions);
+		}
 		State = MoveTemp(PrunedAssetRegistry);
 	}
 
