Detect stalled cook :
- New class FStalledIterationCounter keeping track of for how many iteration a value didn't change.
- New class FStallDetctor used to detect when a cook is not making progress.
- CookOnTheFlyServer : add new member FStallDetector. New function IsStalled detecting when the cook is stalled.
- CookCommandlet : Check if the cook is stalled in the director and print a fatal log.

#rnx

diff --git a/Engine/Source/Editor/UnrealEd/Classes/CookOnTheSide/CookOnTheFlyServer.h b/Engine/Source/Editor/UnrealEd/Classes/CookOnTheSide/CookOnTheFlyServer.h
--- a/Engine/Source/Editor/UnrealEd/Classes/CookOnTheSide/CookOnTheFlyServer.h
+++ b/Engine/Source/Editor/UnrealEd/Classes/CookOnTheSide/CookOnTheFlyServer.h
@@ -188,6 +188,7 @@ namespace UE::Cook
 	class FRequestQueue;
 	class FSaveCookedPackageContext;
 	class FSoftGCHistory;
+	class FStallDetector;
 	class FWorkerRequestsLocal;
 	class FWorkerRequestsRemote;
 	class ICookOnTheFlyRequestManager;
@@ -866,6 +867,8 @@ public:
 		float GCDurationSeconds);
 	UNREALED_API bool NeedsDiagnosticSecondGC() const;
 
+	bool IsStalled();
+
 	/**
 	 * RequestPackage to be cooked
 	 *
@@ -1662,6 +1665,8 @@ private:
 	EIdleStatus IdleStatus = EIdleStatus::Done;
 	TUniquePtr<UE::Cook::FODSCClientData> ODSCClientData;
 
+	TUniquePtr<UE::Cook::FStallDetector> StallDetector;
+
 	friend FAssetRegistryGenerator;
 	friend FIterativeValidatePackageWriter;
 	friend UE::Cook::FAssetRegistryMPCollector;
diff --git a/Engine/Source/Editor/UnrealEd/Private/Commandlets/CookCommandlet.cpp b/Engine/Source/Editor/UnrealEd/Private/Commandlets/CookCommandlet.cpp
--- a/Engine/Source/Editor/UnrealEd/Private/Commandlets/CookCommandlet.cpp
+++ b/Engine/Source/Editor/UnrealEd/Private/Commandlets/CookCommandlet.cpp
@@ -774,4 +774,16 @@ void UCookCommandlet::ConditionalCollectGarbage(uint32 TickResults, UCookOnTheFl
 	COTFS.EvaluateGarbageCollectionResults(bWasDueToOOM, bPartialGC, TickResults,
 		NumObjectsBeforeGC, MemStatsBeforeGC, AllocatorStatsBeforeGC,
 		NumObjectsAfterGC, MemStatsAfterGC, AllocatorStatsAfterGC, GCDurationSeconds);
+
+	if (!COTFS.IsCookWorkerMode() && !bSoft)
+	{
+		// Only check stalling in the director and if it's not a SoftGC.
+		// SoftGC cannot cause OOM stalls because it keeps packages in memory that will be used dater.
+
+		bool bCookStalled = COTFS.IsStalled();
+		if (bCookStalled)
+		{
+			UE_LOG(LogCook, Fatal, TEXT("Cook stalled probably because of memory settings being too restrictive and triggering garbage collection."));
+		}
+	}
 }
diff --git a/Engine/Source/Editor/UnrealEd/Private/CookOnTheFlyServer.cpp b/Engine/Source/Editor/UnrealEd/Private/CookOnTheFlyServer.cpp
--- a/Engine/Source/Editor/UnrealEd/Private/CookOnTheFlyServer.cpp
+++ b/Engine/Source/Editor/UnrealEd/Private/CookOnTheFlyServer.cpp
@@ -45,6 +45,7 @@
 #include "Cooker/NetworkFileCookOnTheFlyRequestManager.h"
 #include "Cooker/OnDemandShaderCompilation.h"
 #include "Cooker/PackageTracker.h"
+#include "Cooker/StallDetector.h"
 #include "Cooker/WorkerRequestsLocal.h"
 #include "Cooker/WorkerRequestsRemote.h"
 #include "CookerSettings.h"
@@ -5657,6 +5658,7 @@ void UCookOnTheFlyServer::Initialize( ECookMode::Type DesiredCookMode, ECookInit
 	AssetRegistry = IAssetRegistry::Get();
 	CachedDependencies = MakeUnique<UE::Cook::FCachedDependencies>();
 	GCDiagnosticContext = MakeUnique<UE::Cook::FCookGCDiagnosticContext>();
+	StallDetector = MakeUnique<UE::Cook::FStallDetector>();
 
 	if (!IsCookWorkerMode())
 	{
@@ -12827,4 +12829,9 @@ void UCookOnTheFlyServer::BroadcastCookFinished()
 	UE::Cook::FDelegates::CookFinished.Broadcast(*this);
 }
 
+bool UCookOnTheFlyServer::IsStalled()
+{
+	return StallDetector->IsStalled(PackageDatas->GetNumCooked(), PackageDatas->GetMonitor().GetNumInProgress());
+}
+
 #undef LOCTEXT_NAMESPACE
diff --git a/Engine/Source/Editor/UnrealEd/Private/Cooker/StallDetector.cpp b/Engine/Source/Editor/UnrealEd/Private/Cooker/StallDetector.cpp
new file mode 100644
--- /dev/null
+++ b/Engine/Source/Editor/UnrealEd/Private/Cooker/StallDetector.cpp
@@ -0,0 +1,29 @@
+// Copyright Epic Games, Inc. All Rights Reserved.
+
+#include "StallDetector.h"
+
+namespace UE::Cook
+{
+	// Number of iteration without progress after which the cook is considered stalled for the number of cooked packages.
+	const static int32 StalledIterationLimitPackageCooked = 10;
+
+	// Number of iteration without progress after which the cook is considered stalled for the number of in progress packages.
+	const static int32 StalledIterationLimitPackageInProgress = 10;
+
+
+	bool FStallDetector::IsStalled(int32 NewPackageCooked, int32 NewPackageInProgress)
+	{
+		// First check if the number of cooked package is stalled
+		StalledPackageCooked.Update(NewPackageCooked);
+		if (StalledPackageCooked.StalledIterationCount < UE::Cook::StalledIterationLimitPackageCooked)
+			return false;
+
+		// Now check if the number of package in progress is stalled
+		StalledPackageInProgress.Update(NewPackageInProgress);
+		if (StalledPackageInProgress.StalledIterationCount < UE::Cook::StalledIterationLimitPackageInProgress)
+			return false;
+
+		// No progress so the cook is stalled
+		return true;
+	}
+}
diff --git a/Engine/Source/Editor/UnrealEd/Private/Cooker/StallDetector.h b/Engine/Source/Editor/UnrealEd/Private/Cooker/StallDetector.h
new file mode 100644
--- /dev/null
+++ b/Engine/Source/Editor/UnrealEd/Private/Cooker/StallDetector.h
@@ -0,0 +1,23 @@
+// Copyright Epic Games, Inc. All Rights Reserved.
+
+#pragma once
+
+#include "StalledIterationCounter.h"
+
+namespace UE::Cook
+{
+
+	/**
+	 * Class used to detect when the cook is not making progress. It does it by keeping track of the number
+	 * of cooked packages and packages in progress.
+	 */
+	class FStallDetector
+	{
+	public:
+		bool IsStalled(int32 NewPackageCooked, int32 NewPackageInProgress);
+
+	private:
+		FStalledIterationCounter StalledPackageCooked;
+		FStalledIterationCounter StalledPackageInProgress;
+	};
+}
diff --git a/Engine/Source/Editor/UnrealEd/Private/Cooker/StalledIterationCounter.cpp b/Engine/Source/Editor/UnrealEd/Private/Cooker/StalledIterationCounter.cpp
new file mode 100644
--- /dev/null
+++ b/Engine/Source/Editor/UnrealEd/Private/Cooker/StalledIterationCounter.cpp
@@ -0,0 +1,23 @@
+// Copyright Epic Games, Inc. All Rights Reserved.
+
+#include "StalledIterationCounter.h"
+
+namespace UE::Cook
+{
+
+	FStalledIterationCounter::FStalledIterationCounter()
+		: Value(0)
+		, StalledIterationCount(0)
+	{}
+
+	void FStalledIterationCounter::Update(int32 NewValue)
+	{
+		if (Value != NewValue)
+		{
+			Value = NewValue;
+			StalledIterationCount = 0;
+		}
+
+		++StalledIterationCount;
+	}
+}
diff --git a/Engine/Source/Editor/UnrealEd/Private/Cooker/StalledIterationCounter.h b/Engine/Source/Editor/UnrealEd/Private/Cooker/StalledIterationCounter.h
new file mode 100644
--- /dev/null
+++ b/Engine/Source/Editor/UnrealEd/Private/Cooker/StalledIterationCounter.h
@@ -0,0 +1,22 @@
+// Copyright Epic Games, Inc. All Rights Reserved.
+
+#pragma once
+
+#include "HAL/Platform.h"
+
+namespace UE::Cook
+{
+
+	/**
+	 * Small structure used to keep track of a value and for how many iteration the value didn't change. It's used
+	 * more specifically in the cooker to detect when the cooker is stalling and no progress is being done.
+	 */
+	struct FStalledIterationCounter
+	{
+		FStalledIterationCounter();
+		void Update(int32 NewValue);
+
+		int32 Value;
+		int32 StalledIterationCount;
+	};
+}
