Allow cook dependencies to be fetched in batch from zenserver using the GetChunks RPC.

diff --git a/Engine/Source/Developer/IoStoreUtilities/Internal/ZenStoreHttpClient.h b/Engine/Source/Developer/IoStoreUtilities/Internal/ZenStoreHttpClient.h
--- a/Engine/Source/Developer/IoStoreUtilities/Internal/ZenStoreHttpClient.h
+++ b/Engine/Source/Developer/IoStoreUtilities/Internal/ZenStoreHttpClient.h
@@ -4,6 +4,7 @@
 
 #include "Async/Future.h"
 #include "Containers/Array.h"
+#include "Containers/ArrayView.h"
 #include "Containers/StringView.h"
 #include "Experimental/ZenServerInterface.h"
 #include "HAL/CriticalSection.h"
@@ -11,6 +12,7 @@
 #include "IO/IoDispatcher.h"
 #include "Misc/StringBuilder.h"
 #include "Serialization/CompactBinaryPackage.h"
+#include "Templates/Function.h"
 #include "Templates/UniquePtr.h"
 
 class FCbPackage;
@@ -51,6 +53,12 @@ public:
 
 	TIoStatusOr<uint64> GetChunkSize(const FIoChunkId& Id);
 	TIoStatusOr<FIoBuffer> ReadChunk(const FIoChunkId& Id, uint64 Offset = 0, uint64 Size = ~0ull);
+	TIoStatusOr<FIoBuffer> ReadChunk(const FIoHash& RawHash, uint64 Offset = 0, uint64 Size = ~0ull);
+	// TODO: Expand API to add partial chunk requests and ModTag field
+	UE_INTERNAL void ReadChunks(TConstArrayView<FIoChunkId> Ids, TFunctionRef<void(const FIoChunkId& Id, TIoStatusOr<FIoBuffer> Result)> ChunkCallback);
+	UE_INTERNAL void ReadChunks(TConstArrayView<FIoHash> RawHashes, TFunctionRef<void(const FIoHash& RawHash, TIoStatusOr<FIoBuffer> Result)> ChunkCallback);
+
+	UE_DEPRECATED(5.6, "Use ReadChunk(FIoHash) instead")
 	TIoStatusOr<FIoBuffer> ReadOpLogAttachment(FStringView Id);
 
 #if UE_WITH_ZEN
@@ -76,9 +84,13 @@ private:
 	bool Download(Zen::FZenHttpRequest& Request, FStringView Uri, TArray64<uint8>* Buffer, Zen::EContentType AcceptType);
 	bool Post(Zen::FZenHttpRequest& Request, FStringView Uri, FCbObjectView Obj);
 	bool Post(Zen::FZenHttpRequest& Request, FStringView Uri, FMemoryView Payload);
+	bool Rpc(Zen::FZenHttpRequest& Request, FStringView Uri, FCbObjectView RpcRequestObject, FCbPackage &OutResponse);
 	bool Delete(Zen::FZenHttpRequest& Request, FStringView Uri);
 	bool ShouldRecoverAndRetry(Zen::FZenHttpRequest& Request);
 
+	template<typename T>
+	void ReadChunks_Internal(TConstArrayView<T> Requests, TFunctionRef<void(const T& Id, TIoStatusOr<FIoBuffer> Result)> ChunkCallback);
+
 	static const uint32 PoolEntryCount;
 	struct SaltGenerator
 	{
@@ -99,6 +111,7 @@ private:
 	SaltGenerator SaltGen;
 	FString ProjectPath;
 	FString OplogPath;
+	FString OplogRpcPath;
 	FString OplogNewEntryPath;
 	FString OplogPrepNewEntryPath;
 	FString TempDirPath;
diff --git a/Engine/Source/Developer/IoStoreUtilities/Private/ZenStoreHttpClient.cpp b/Engine/Source/Developer/IoStoreUtilities/Private/ZenStoreHttpClient.cpp
--- a/Engine/Source/Developer/IoStoreUtilities/Private/ZenStoreHttpClient.cpp
+++ b/Engine/Source/Developer/IoStoreUtilities/Private/ZenStoreHttpClient.cpp
@@ -130,6 +130,7 @@ FZenStoreHttpClient::TryCreateProject(FStringView InProjectId,
 
 	ProjectPath = WriteToString<128>("/prj/", InProjectId);
 	OplogPath = WriteToString<128>("/prj/", InProjectId, "/oplog/", InOplogId);
+	OplogRpcPath = WriteToString<128>("/prj/", InProjectId, "/oplog/", InOplogId, "/rpc");
 	OplogNewEntryPath = WriteToString<128>("/prj/", InProjectId, "/oplog/", InOplogId, "/new");
 	OplogPrepNewEntryPath = WriteToString<128>("/prj/", InProjectId, "/oplog/", InOplogId, "/prep");
 
@@ -289,6 +290,17 @@ bool FZenStoreHttpClient::Post(Zen::FZenHttpRequest& Request, FStringView Uri, F
 	return Res == Zen::FZenHttpRequest::Result::Success;
 }
 
+bool FZenStoreHttpClient::Rpc(Zen::FZenHttpRequest& Request, FStringView Uri, FCbObjectView RpcRequestObject, FCbPackage &OutResponse)
+{
+	Zen::FZenHttpRequest::Result Res = Request.PerformRpc(Uri, RpcRequestObject, OutResponse);
+	if (ShouldRecoverAndRetry(Request))
+	{
+		Request.Reset();
+		Res = Request.PerformRpc(Uri, RpcRequestObject, OutResponse);
+	}
+	return Res == Zen::FZenHttpRequest::Result::Success;
+}
+
 bool FZenStoreHttpClient::Delete(Zen::FZenHttpRequest& Request, FStringView Uri)
 {
 	Zen::FZenHttpRequest::Result Res = Request.PerformBlockingDelete(Uri);
@@ -526,6 +538,149 @@ TIoStatusOr<FIoBuffer> FZenStoreHttpClient::ReadChunk(const FIoChunkId& Id, uint
 	return ReadOpLogUri(ChunkUri, Offset, Size);
 }
 
+TIoStatusOr<FIoBuffer> FZenStoreHttpClient::ReadChunk(const FIoHash& RawHash, uint64 Offset, uint64 Size)
+{
+	TRACE_CPUPROFILER_EVENT_SCOPE(ZenStoreHttp_ReadChunk);
+	TStringBuilder<128> ChunkUri;
+	ChunkUri << OplogPath << '/' << RawHash;
+	return ReadOpLogUri(ChunkUri, Offset, Size);
+}
+
+static void SerializeRequest(FCbWriter& Writer, const FIoChunkId& ChunkId)
+{
+	Writer << "Oid" << ChunkId;
+}
+
+static void SerializeRequest(FCbWriter& Writer, const FIoHash& RawHash)
+{
+	Writer << "RawHash" << RawHash;
+}
+
+static void DeserializeResponse(const FCbObjectView ChunkObject, FIoChunkId& ChunkId)
+{
+	FCbObjectId ObjectId = ChunkObject[ANSITEXTVIEW("Id")].AsObjectId();
+	ChunkId.Set(ObjectId.GetView());
+}
+
+static void DeserializeResponse(const FCbObjectView ChunkObject, FIoHash& RawHash)
+{
+	RawHash = ChunkObject[ANSITEXTVIEW("Id")].AsHash();
+}
+
+static const TCHAR* DescribeRequest(const FIoChunkId& ChunkId)
+{
+	return TEXT("Oid");
+}
+
+static const TCHAR* DescribeRequest(const FIoHash& RawHash)
+{
+	return TEXT("RawHash");
+}
+
+template<typename T>
+void FZenStoreHttpClient::ReadChunks_Internal(TConstArrayView<T> Requests, TFunctionRef<void(const T& Id, TIoStatusOr<FIoBuffer> Result)> ChunkCallback)
+{
+	TRACE_CPUPROFILER_EVENT_SCOPE(ZenStoreHttp_ReadChunks_Internal);
+	FCbWriter Writer;
+	Writer.BeginObject();
+	Writer << "method" << "getchunks";
+
+	Writer.BeginObject("Request");
+	Writer.BeginArray("Chunks");
+
+	TSet<T> InProgressRequests;
+	for (const T& Request : Requests)
+	{
+		InProgressRequests.Add(Request);
+	}
+	for (const T& Request : InProgressRequests)
+	{
+		Writer.BeginObject();
+		SerializeRequest(Writer, Request);
+		Writer.EndObject();
+	}
+
+	Writer.EndArray();
+	Writer.EndObject();
+	Writer.EndObject();
+
+	UE::Zen::FZenScopedRequestPtr ZenHttpRequest(RequestPool.Get());
+	FCbPackage ResultPackage;
+	if (Rpc(*ZenHttpRequest, OplogRpcPath, Writer.Save().AsObjectView(), ResultPackage))
+	{
+		const FCbObjectView ResultObjectView(ResultPackage.GetObject());
+		for (FCbFieldView ChunkField : ResultObjectView[ANSITEXTVIEW("Chunks")])
+		{
+			const FCbObjectView ChunkObject(ChunkField.AsObjectView());
+			T Response;
+			DeserializeResponse(ChunkObject, Response);
+
+			ON_SCOPE_EXIT
+			{
+				InProgressRequests.Remove(Response);
+			};
+
+			const FIoHash RawHash = ChunkObject[ANSITEXTVIEW("RawHash")].AsHash();
+			if (const FCbAttachment* Attachment = ResultPackage.FindAttachment(RawHash))
+			{
+				if (const FCompressedBuffer& Compressed = Attachment->AsCompressedBinary(); !Compressed.IsNull())
+				{
+					FIoBuffer Decompressed(Compressed.GetRawSize());
+					if (!FCompressedBufferReader(Compressed).TryDecompressTo(Decompressed.GetMutableView()))
+					{
+						ChunkCallback(Response, (FIoStatus)(FIoStatusBuilder(EIoErrorCode::CompressionError) << TEXT("Failed to decompress, ")
+							<< DescribeRequest(Response) << TEXT("=") << *WriteToString<32>(Response)));
+					}
+					else
+					{
+						// Completion status - success
+						ChunkCallback(Response, Decompressed);
+					}
+				}
+				else if (const FSharedBuffer& Uncompressed = Attachment->AsBinary(); !Uncompressed.IsNull())
+				{
+					// Completion status - success
+					ChunkCallback(Response, FIoBuffer(FIoBuffer::Clone, Uncompressed.GetData(), Uncompressed.GetSize()));
+				}
+				else
+				{
+					ChunkCallback(Response, (FIoStatus)(FIoStatusBuilder(EIoErrorCode::ReadError) << TEXT("Unexpected attachment type, ")
+						<< DescribeRequest(Response) << TEXT("=") << *WriteToString<32>(Response)));
+				}
+			}
+			else
+			{
+				ChunkCallback(Response, (FIoStatus)(FIoStatusBuilder(EIoErrorCode::NotFound) << TEXT("Attachment not found, ")
+					<< DescribeRequest(Response) << TEXT("=") << *WriteToString<32>(Response)));
+			}
+		}
+
+		for (const T& InProgressRequest : InProgressRequests)
+		{
+			ChunkCallback(InProgressRequest, (FIoStatus)(FIoStatusBuilder(EIoErrorCode::NotFound) << TEXT("Missing response, ")
+				<< DescribeRequest(InProgressRequest) << TEXT("=") << *WriteToString<32>(InProgressRequest)));
+		}
+	}
+	else
+	{
+		for (const T& InProgressRequest : InProgressRequests)
+		{
+			ChunkCallback(InProgressRequest, (FIoStatus)(FIoStatusBuilder(EIoErrorCode::NotFound) << TEXT("Transport/request failure, ")
+				<< DescribeRequest(InProgressRequest) << TEXT("=") << *WriteToString<32>(InProgressRequest)));
+		}
+	}
+}
+
+void FZenStoreHttpClient::ReadChunks(TConstArrayView<FIoChunkId> Ids, TFunctionRef<void(const FIoChunkId& Id, TIoStatusOr<FIoBuffer> Result)> ChunkCallback)
+{
+	return ReadChunks_Internal(Ids, ChunkCallback);
+}
+
+void FZenStoreHttpClient::ReadChunks(TConstArrayView<FIoHash> RawHashes, TFunctionRef<void(const FIoHash& RawHash, TIoStatusOr<FIoBuffer> Result)> ChunkCallback)
+{
+	return ReadChunks_Internal(RawHashes, ChunkCallback);
+}
+
 TIoStatusOr<FIoBuffer> FZenStoreHttpClient::ReadOpLogAttachment(FStringView Id)
 {
 	TStringBuilder<128> ChunkUri;
@@ -801,6 +956,27 @@ TIoStatusOr<FIoBuffer> FZenStoreHttpClient::ReadChunk(const FIoChunkId& Id, uint
 	return FIoBuffer();
 }
 
+TIoStatusOr<FIoBuffer> FZenStoreHttpClient::ReadChunk(const FIoHash& RawHash, uint64 Offset, uint64 Size)
+{
+	return FIoBuffer();
+}
+
+void FZenStoreHttpClient::ReadChunks(TConstArrayView<FIoChunkId> Ids, TFunctionRef<void(const FIoChunkId& Id, TIoStatusOr<FIoBuffer> Result)> ChunkCallback)
+{
+	for (const FIoChunkId& Id : Ids)
+	{
+		ChunkCallback(Id, (FIoStatus)(FIoStatusBuilder(EIoErrorCode::Unknown) << TEXT("Unsupported, Oid=") << *WriteToString<32>(Id)));
+	}
+}
+
+void FZenStoreHttpClient::ReadChunks(TConstArrayView<FIoHash> RawHashes, TFunctionRef<void(const FIoHash& RawHash, TIoStatusOr<FIoBuffer> Result)> ChunkCallback)
+{
+	for (const FIoHash& RawHash : RawHashes)
+	{
+		ChunkCallback(RawHash, (FIoStatus)(FIoStatusBuilder(EIoErrorCode::Unknown) << TEXT("Unsupported, RawHash=") << *WriteToString<32>(RawHash)));
+	}
+}
+
 TIoStatusOr<FIoBuffer> FZenStoreHttpClient::ReadOpLogAttachment(FStringView Id)
 {
 	return FIoBuffer();
diff --git a/Engine/Source/Developer/IoStoreUtilities/Private/ZenStoreWriter.cpp b/Engine/Source/Developer/IoStoreUtilities/Private/ZenStoreWriter.cpp
--- a/Engine/Source/Developer/IoStoreUtilities/Private/ZenStoreWriter.cpp
+++ b/Engine/Source/Developer/IoStoreUtilities/Private/ZenStoreWriter.cpp
@@ -1150,7 +1150,7 @@ FCbObject FZenStoreWriter::GetOplogAttachment(FName PackageName, FUtf8StringView
 		}
 		AttachmentHash = Existing.Hash;
 	}
-	TIoStatusOr<FIoBuffer> BufferResult = HttpClient->ReadOpLogAttachment(WriteToString<48>(AttachmentHash));
+	TIoStatusOr<FIoBuffer> BufferResult = HttpClient->ReadChunk(AttachmentHash);
 	if (!BufferResult.IsOk())
 	{
 		return FCbObject();
@@ -1165,6 +1165,117 @@ FCbObject FZenStoreWriter::GetOplogAttachment(FName PackageName, FUtf8StringView
 	return FCbObject(SharedBuffer);
 }
 
+void FZenStoreWriter::GetOplogAttachments(TArrayView<FName> PackageNames,
+	TArrayView<FUtf8StringView> AttachmentKeys,
+	TFunctionRef<void(FName PackageName, FUtf8StringView AttachmentKey, FCbObject&& Attachment)> Callback)
+{
+	const int MaximumHashCount = PackageNames.Num() * AttachmentKeys.Num();
+	TArray<FIoHash> AttachmentHashes;
+	AttachmentHashes.Reserve(MaximumHashCount);
+	struct FAttachmentHashParam
+	{
+		FName PackageName;
+		FUtf8StringView AttachmentKey;
+	};
+	TMultiMap<FIoHash, FAttachmentHashParam> AttachmentHashParams;
+	AttachmentHashParams.Reserve(MaximumHashCount);
+
+	TArray<FAttachmentHashParam> InvalidAttachmentHashParams;
+	InvalidAttachmentHashParams.Reserve(MaximumHashCount);
+
+	TArray<const UTF8CHAR*, TInlineAllocator<2>> AttachmentIds;
+	for (FUtf8StringView AttachmentKey : AttachmentKeys)
+	{
+		AttachmentIds.Add(UE::FZenStoreHttpClient::FindAttachmentId(AttachmentKey));
+	}
+
+	{
+		FReadScopeLock _(EntriesLock);
+
+		for (FName PackageName : PackageNames)
+		{
+			const int32* Idx = PackageNameToIndex.Find(PackageName);
+
+			for (int32 InputAttachmentIndex = 0; InputAttachmentIndex < AttachmentKeys.Num(); ++InputAttachmentIndex)
+			{
+				FUtf8StringView AttachmentKey = AttachmentKeys[InputAttachmentIndex];
+				const UTF8CHAR* AttachmentId = AttachmentIds[InputAttachmentIndex];
+
+				FIoHash AttachmentHash;
+				ON_SCOPE_EXIT
+				{
+					if (AttachmentHash.IsZero())
+					{
+						InvalidAttachmentHashParams.Emplace(PackageName, AttachmentKey);
+					}
+					else
+					{
+						AttachmentHashes.Add(AttachmentHash);
+						AttachmentHashParams.Emplace(AttachmentHash, FAttachmentHashParam{PackageName, AttachmentKey});
+					}
+				};
+
+				if (!Idx || !AttachmentId)
+				{
+					continue;
+				}
+
+				FUtf8StringView AttachmentIdView(AttachmentId);
+
+				const FOplogCookInfo& CookInfo = CookedPackagesInfo[*Idx];
+				int32 AttachmentIndex = Algo::LowerBound(CookInfo.Attachments, AttachmentIdView,
+					[](const FOplogCookInfo::FAttachment& Existing, FUtf8StringView AttachmentIdView)
+					{
+						return FUtf8StringView(Existing.Key).Compare(AttachmentIdView, ESearchCase::IgnoreCase) < 0;
+					});
+				if (AttachmentIndex == CookInfo.Attachments.Num())
+				{
+					continue;
+				}
+				const FOplogCookInfo::FAttachment& Existing = CookInfo.Attachments[AttachmentIndex];
+				if (!FUtf8StringView(Existing.Key).Equals(AttachmentIdView, ESearchCase::IgnoreCase))
+				{
+					continue;
+				}
+				AttachmentHash = Existing.Hash;
+			}
+		}
+	}
+
+	// Invoke the callback for all invalid attachment hashes
+	for (FAttachmentHashParam& InvalidAttachmentHashParam : InvalidAttachmentHashParams)
+	{
+		Callback(InvalidAttachmentHashParam.PackageName, InvalidAttachmentHashParam.AttachmentKey, FCbObject());
+	}
+	
+	if (AttachmentHashes.IsEmpty())
+	{
+		return;
+	}
+
+	HttpClient->ReadChunks(AttachmentHashes, [Callback,&AttachmentHashParams](const FIoHash& RawHash, TIoStatusOr<FIoBuffer> Result)
+	{
+		for (auto It(AttachmentHashParams.CreateConstKeyIterator(RawHash)); It; ++It)
+		{
+			const FAttachmentHashParam& Param = It.Value();
+			if (!Result.IsOk())
+			{
+				Callback(Param.PackageName, Param.AttachmentKey, FCbObject());
+				continue;
+			}
+
+			FIoBuffer Buffer = Result.ConsumeValueOrDie();
+			if (Buffer.DataSize() == 0)
+			{
+				Callback(Param.PackageName, Param.AttachmentKey, FCbObject());
+				continue;
+			}
+			FSharedBuffer SharedBuffer = IoBufferToSharedBuffer(Buffer);
+			Callback(Param.PackageName, Param.AttachmentKey, FCbObject(SharedBuffer));
+		}
+	});
+}
+
 IPackageWriter::ECommitStatus FZenStoreWriter::GetCommitStatus(FName PackageName)
 {
 	FReadScopeLock _(EntriesLock);
diff --git a/Engine/Source/Developer/IoStoreUtilities/Public/ZenStoreWriter.h b/Engine/Source/Developer/IoStoreUtilities/Public/ZenStoreWriter.h
--- a/Engine/Source/Developer/IoStoreUtilities/Public/ZenStoreWriter.h
+++ b/Engine/Source/Developer/IoStoreUtilities/Public/ZenStoreWriter.h
@@ -147,6 +147,9 @@ public:
 
 	IOSTOREUTILITIES_API TUniquePtr<FAssetRegistryState> LoadPreviousAssetRegistry() override;
 	IOSTOREUTILITIES_API virtual FCbObject GetOplogAttachment(FName PackageName, FUtf8StringView AttachmentKey) override;
+	IOSTOREUTILITIES_API virtual void GetOplogAttachments(TArrayView<FName> PackageNames,
+		TArrayView<FUtf8StringView> AttachmentKeys,
+		TFunctionRef<void(FName PackageName, FUtf8StringView AttachmentKey, FCbObject&& Attachment)> Callback) override;
 	IOSTOREUTILITIES_API virtual ECommitStatus GetCommitStatus(FName PackageName) override;
 
 	IOSTOREUTILITIES_API virtual void RemoveCookedPackages(TArrayView<const FName> PackageNamesToRemove) override;
diff --git a/Engine/Source/Editor/UnrealEd/Private/Cooker/DiffPackageWriter.h b/Engine/Source/Editor/UnrealEd/Private/Cooker/DiffPackageWriter.h
--- a/Engine/Source/Editor/UnrealEd/Private/Cooker/DiffPackageWriter.h
+++ b/Engine/Source/Editor/UnrealEd/Private/Cooker/DiffPackageWriter.h
@@ -90,6 +90,12 @@ public:
 	{
 		return Inner->GetOplogAttachment(PackageName, AttachmentKey);
 	}
+	virtual void GetOplogAttachments(TArrayView<FName> PackageNames,
+		TArrayView<FUtf8StringView> AttachmentKeys,
+		TFunctionRef<void(FName PackageName, FUtf8StringView AttachmentKey, FCbObject&& Attachment)> Callback) override
+	{
+		return Inner->GetOplogAttachments(PackageNames, AttachmentKeys, Callback);
+	}
 	virtual ECommitStatus GetCommitStatus(FName PackageName) override
 	{
 		return Inner->GetCommitStatus(PackageName);
@@ -247,6 +253,12 @@ public:
 	{
 		return Inner->GetOplogAttachment(PackageName, AttachmentKey);
 	}
+	virtual void GetOplogAttachments(TArrayView<FName> PackageNames,
+		TArrayView<FUtf8StringView> AttachmentKeys,
+		TFunctionRef<void(FName PackageName, FUtf8StringView AttachmentKey, FCbObject&& Attachment)> Callback) override
+	{
+		return Inner->GetOplogAttachments(PackageNames, AttachmentKeys, Callback);
+	}
 	virtual ECommitStatus GetCommitStatus(FName PackageName) override
 	{
 		return Inner->GetCommitStatus(PackageName);
diff --git a/Engine/Source/Editor/UnrealEd/Private/Cooker/LooseCookedPackageWriter.cpp b/Engine/Source/Editor/UnrealEd/Private/Cooker/LooseCookedPackageWriter.cpp
--- a/Engine/Source/Editor/UnrealEd/Private/Cooker/LooseCookedPackageWriter.cpp
+++ b/Engine/Source/Editor/UnrealEd/Private/Cooker/LooseCookedPackageWriter.cpp
@@ -444,6 +444,20 @@ FCbObject FLooseCookedPackageWriter::GetOplogAttachment(FName PackageName, FUtf8
 	return FCbObject();
 }
 
+void FLooseCookedPackageWriter::GetOplogAttachments(TArrayView<FName> PackageNames,
+	TArrayView<FUtf8StringView> AttachmentKeys,
+	TFunctionRef<void(FName PackageName, FUtf8StringView AttachmentKey, FCbObject&& Attachment)> Callback)
+{
+	/** Oplog attachments are not implemented by FLooseCookedPackageWriter */
+	for (FName PackageName: PackageNames)
+	{
+		for (FUtf8StringView AttachmentKey: AttachmentKeys)
+		{
+			Callback(PackageName, AttachmentKey, FCbObject());
+		}
+	}
+}
+
 IPackageWriter::ECommitStatus FLooseCookedPackageWriter::GetCommitStatus(FName PackageName)
 {
 	// GetCommitStatus is not implemented by FLooseCookedPackageWriter
diff --git a/Engine/Source/Editor/UnrealEd/Private/Cooker/LooseCookedPackageWriter.h b/Engine/Source/Editor/UnrealEd/Private/Cooker/LooseCookedPackageWriter.h
--- a/Engine/Source/Editor/UnrealEd/Private/Cooker/LooseCookedPackageWriter.h
+++ b/Engine/Source/Editor/UnrealEd/Private/Cooker/LooseCookedPackageWriter.h
@@ -54,6 +54,9 @@ public:
 	virtual void EndCook(const FCookInfo& Info) override;
 	virtual TUniquePtr<FAssetRegistryState> LoadPreviousAssetRegistry() override;
 	virtual FCbObject GetOplogAttachment(FName PackageName, FUtf8StringView AttachmentKey) override;
+	virtual void GetOplogAttachments(TArrayView<FName> PackageNames,
+		TArrayView<FUtf8StringView> AttachmentKeys,
+		TFunctionRef<void(FName PackageName, FUtf8StringView AttachmentKey, FCbObject&& Attachment)> Callback) override;
 	virtual ECommitStatus GetCommitStatus(FName PackageName) override;
 	virtual void RemoveCookedPackages(TArrayView<const FName> PackageNamesToRemove) override;
 	virtual void RemoveCookedPackages() override;
diff --git a/Engine/Source/Editor/UnrealEd/Private/TargetDomain/TargetDomainUtils.cpp b/Engine/Source/Editor/UnrealEd/Private/TargetDomain/TargetDomainUtils.cpp
--- a/Engine/Source/Editor/UnrealEd/Private/TargetDomain/TargetDomainUtils.cpp
+++ b/Engine/Source/Editor/UnrealEd/Private/TargetDomain/TargetDomainUtils.cpp
@@ -58,7 +58,9 @@ public:
 
 	bool IsValid() const;
 	void CommitPackage(FName PackageName, TArrayView<IPackageWriter::FCommitAttachmentInfo> Attachments);
-	FCbObject GetOplogAttachment(FName PackageName, FUtf8StringView AttachmentKey);
+	void GetOplogAttachments(TArrayView<FName> PackageName,
+		TArrayView<FUtf8StringView> AttachmentKeys,
+		TFunctionRef<void(FName PackageName, FUtf8StringView AttachmentKey, FCbObject&& Attachment)> Callback);
 
 private:
 	struct FOplogEntry
@@ -1017,41 +1019,70 @@ void FIncrementalCookAttachments::Fetch(TArrayView<FName> PackageNames, const IT
 	ICookedPackageWriter* PackageWriter,
 	TUniqueFunction<void(FName PackageName, FIncrementalCookAttachments&& Result)>&& Callback)
 {
-	for (FName PackageName : PackageNames)
+	struct FInProgressResult
 	{
-		FCbObject DependenciesObj;
-		FCbObject BuildDefinitionsObj;
-		IPackageWriter::ECommitStatus CommitStatus = IPackageWriter::ECommitStatus::NotCommitted;
-		if (TargetPlatform)
+		FIncrementalCookAttachments Result;
+		int ReceivedAttachmentCount = 0;
+	};
+	TMap<FName, FInProgressResult> InProgressResults;
+	TArray<FUtf8StringView> RequestedAttachments = {CookDependenciesAttachmentKey, BuildDefinitionsAttachmentKey};
+
+	if (TargetPlatform || GEditorDomainOplog)
+	{
+		auto OnOplogAttachment = 
+		[PackageWriter, Callback = MoveTemp(Callback), RequestedAttachmentNum = RequestedAttachments.Num(), &InProgressResults]
+		(FName PackageName, FUtf8StringView AttachmentKey, FCbObject&& Attachment)
 		{
-			check(PackageWriter);
-			CommitStatus = PackageWriter->GetCommitStatus(PackageName);
-			DependenciesObj = PackageWriter->GetOplogAttachment(PackageName, CookDependenciesAttachmentKey);
-			BuildDefinitionsObj = PackageWriter->GetOplogAttachment(PackageName, CookDependenciesAttachmentKey);
+			FInProgressResult& InProgressResult = InProgressResults.FindOrAdd(PackageName);
+			InProgressResult.ReceivedAttachmentCount++;
+			if (AttachmentKey == CookDependenciesAttachmentKey)
+			{
+				if (PackageWriter)
+				{
+					InProgressResult.Result.CommitStatus = PackageWriter->GetCommitStatus(PackageName);
+				}
+				else
+				{
+					InProgressResult.Result.CommitStatus = Attachment ? IPackageWriter::ECommitStatus::Success
+					: IPackageWriter::ECommitStatus::NotCommitted;
+				}
+
+				if (LoadFromCompactBinary(MoveTemp(Attachment), InProgressResult.Result.Dependencies))
+				{
+					InProgressResult.Result.Dependencies.PackageName = PackageName;
+				}
+			}
+			else if (AttachmentKey == BuildDefinitionsAttachmentKey)
+			{
+				LoadFromCompactBinary(MoveTemp(Attachment), InProgressResult.Result.BuildDefinitions);
+			}
+
+			if (InProgressResult.ReceivedAttachmentCount == RequestedAttachmentNum)
+			{
+				Callback(PackageName, MoveTemp(InProgressResult.Result));
+				InProgressResults.Remove(PackageName);
+			}
+		};
+
+		if (PackageWriter)
+		{
+			PackageWriter->GetOplogAttachments(PackageNames,
+				RequestedAttachments,OnOplogAttachment);
 		}
 		else
 		{
-			if (!GEditorDomainOplog)
-			{
-				Callback(PackageName, FIncrementalCookAttachments());
-				continue;
-			}
-			DependenciesObj = GEditorDomainOplog->GetOplogAttachment(PackageName, CookDependenciesAttachmentKey);
-			BuildDefinitionsObj = GEditorDomainOplog->GetOplogAttachment(PackageName, BuildDefinitionsAttachmentKey);
-			CommitStatus = DependenciesObj ? IPackageWriter::ECommitStatus::Success
-				: IPackageWriter::ECommitStatus::NotCommitted;
+			GEditorDomainOplog->GetOplogAttachments(PackageNames,
+				RequestedAttachments,OnOplogAttachment);
 		}
-
-		FIncrementalCookAttachments Result;
-		if (LoadFromCompactBinary(DependenciesObj, Result.Dependencies))
+	}
+	else
+	{
+		for (FName PackageName : PackageNames)
 		{
-			Result.Dependencies.PackageName = PackageName;
+			Callback(PackageName, FIncrementalCookAttachments());
 		}
-		LoadFromCompactBinary(MoveTemp(BuildDefinitionsObj), Result.BuildDefinitions);
-		Result.CommitStatus = CommitStatus;
-
-		Callback(PackageName, MoveTemp(Result));
 	}
+	check(InProgressResults.Num() == 0);
 }
 
 bool IsIterativeEnabled(FName PackageName, bool bAllowAllClasses, const FAssetPackageData* OverrideAssetPackageData)
@@ -1318,51 +1349,115 @@ FSharedBuffer IoBufferToSharedBuffer(FIoBuffer& InBuffer)
 	return FSharedBuffer{ FSharedBuffer::TakeOwnership(DataPtr, DataSize, FMemory::Free) };
 };
 
-FCbObject FEditorDomainOplog::GetOplogAttachment(FName PackageName, FUtf8StringView AttachmentKey)
+void FEditorDomainOplog::GetOplogAttachments(TArrayView<FName> PackageNames,
+	TArrayView<FUtf8StringView> AttachmentKeys,
+	TFunctionRef<void(FName PackageName, FUtf8StringView AttachmentKey, FCbObject&& Attachment)> Callback)
 {
-	FScopeLock ScopeLock(&Lock);
-	InitializeRead();
+	const int MaximumHashCount = PackageNames.Num() * AttachmentKeys.Num();
+	TArray<FIoHash> AttachmentHashes;
+	AttachmentHashes.Reserve(MaximumHashCount);
+	struct FAttachmentHashParam
+	{
+		FName PackageName;
+		FUtf8StringView AttachmentKey;
+	};
+	TMultiMap<FIoHash, FAttachmentHashParam> AttachmentHashParams;
+	AttachmentHashParams.Reserve(MaximumHashCount);
+
+	TArray<FAttachmentHashParam> InvalidAttachmentHashParams;
+	InvalidAttachmentHashParams.Reserve(MaximumHashCount);
 
-	FOplogEntry* Entry = Entries.Find(PackageName);
-	if (!Entry)
+	TArray<const UTF8CHAR*, TInlineAllocator<2>> AttachmentIds;
+	for (FUtf8StringView AttachmentKey : AttachmentKeys)
 	{
-		return FCbObject();
+		AttachmentIds.Add(UE::FZenStoreHttpClient::FindAttachmentId(AttachmentKey));
 	}
 
-	const UTF8CHAR* AttachmentId = UE::FZenStoreHttpClient::FindAttachmentId(AttachmentKey);
-	if (!AttachmentId)
 	{
-		return FCbObject();
-	}
-	FUtf8StringView AttachmentIdView(AttachmentId);
+		FScopeLock _(&Lock);
+		InitializeRead();
 
-	int32 AttachmentIndex = Algo::LowerBound(Entry->Attachments, AttachmentIdView,
-		[](const FOplogEntry::FAttachment& Existing, FUtf8StringView AttachmentIdView)
+		for (FName PackageName : PackageNames)
 		{
-			return FUtf8StringView(Existing.Key).Compare(AttachmentIdView, ESearchCase::IgnoreCase) < 0;
-		});
-	if (AttachmentIndex == Entry->Attachments.Num())
-	{
-		return FCbObject();
+			FOplogEntry* Entry = Entries.Find(PackageName);
+
+			for (int32 InputAttachmentIndex = 0; InputAttachmentIndex < AttachmentKeys.Num(); ++InputAttachmentIndex)
+			{
+				FUtf8StringView AttachmentKey = AttachmentKeys[InputAttachmentIndex];
+				const UTF8CHAR* AttachmentId = AttachmentIds[InputAttachmentIndex];
+
+				FIoHash AttachmentHash;
+				ON_SCOPE_EXIT
+				{
+					if (AttachmentHash.IsZero())
+					{
+						InvalidAttachmentHashParams.Emplace(PackageName, AttachmentKey);
+					}
+					else
+					{
+						AttachmentHashes.Add(AttachmentHash);
+						AttachmentHashParams.Emplace(AttachmentHash, FAttachmentHashParam{PackageName, AttachmentKey});
+					}
+				};
+
+				if (!Entry || !AttachmentId)
+				{
+					continue;
+				}
+
+				FUtf8StringView AttachmentIdView(AttachmentId);
+
+				int32 AttachmentIndex = Algo::LowerBound(Entry->Attachments, AttachmentIdView,
+					[](const FOplogEntry::FAttachment& Existing, FUtf8StringView AttachmentIdView)
+					{
+						return FUtf8StringView(Existing.Key).Compare(AttachmentIdView, ESearchCase::IgnoreCase) < 0;
+					});
+				if (AttachmentIndex == Entry->Attachments.Num())
+				{
+					continue;
+				}
+				const FOplogEntry::FAttachment& Existing = Entry->Attachments[AttachmentIndex];
+				if (!FUtf8StringView(Existing.Key).Equals(AttachmentIdView, ESearchCase::IgnoreCase))
+				{
+					continue;
+				}
+				AttachmentHash = Existing.Hash;
+			}
+		}
 	}
-	const FOplogEntry::FAttachment& Existing = Entry->Attachments[AttachmentIndex];
-	if (!FUtf8StringView(Existing.Key).Equals(AttachmentIdView, ESearchCase::IgnoreCase))
+
+	// Invoke the callback for all invalid attachment hashes
+	for (FAttachmentHashParam& InvalidAttachmentHashParam : InvalidAttachmentHashParams)
 	{
-		return FCbObject();
+		Callback(InvalidAttachmentHashParam.PackageName, InvalidAttachmentHashParam.AttachmentKey, FCbObject());
 	}
-	TIoStatusOr<FIoBuffer> BufferResult = HttpClient.ReadOpLogAttachment(WriteToString<48>(Existing.Hash));
-	if (!BufferResult.IsOk())
+	
+	if (AttachmentHashes.IsEmpty())
 	{
-		return FCbObject();
+		return;
 	}
-	FIoBuffer Buffer = BufferResult.ValueOrDie();
-	if (Buffer.DataSize() == 0)
+
+	HttpClient.ReadChunks(AttachmentHashes, [Callback,&AttachmentHashParams](const FIoHash& RawHash, TIoStatusOr<FIoBuffer> Result)
 	{
-		return FCbObject();
-	}
+		for (auto It(AttachmentHashParams.CreateConstKeyIterator(RawHash)); It; ++It)
+		{
+			const FAttachmentHashParam& Param = It.Value();
+			if (!Result.IsOk())
+			{
+				Callback(Param.PackageName, Param.AttachmentKey, FCbObject());
+				continue;
+			}
 
-	FSharedBuffer SharedBuffer = IoBufferToSharedBuffer(Buffer);
-	return FCbObject(SharedBuffer);
+			FIoBuffer Buffer = Result.ConsumeValueOrDie();
+			if (Buffer.DataSize() == 0)
+			{
+				Callback(Param.PackageName, Param.AttachmentKey, FCbObject());
+				continue;
+			}
+			FSharedBuffer SharedBuffer = IoBufferToSharedBuffer(Buffer);
+			Callback(Param.PackageName, Param.AttachmentKey, FCbObject(SharedBuffer));
+		}
+	});
 }
 
 void CommitEditorDomainCookAttachments(FName PackageName, TArrayView<IPackageWriter::FCommitAttachmentInfo> Attachments)
diff --git a/Engine/Source/Runtime/Core/Public/Serialization/PackageWriter.h b/Engine/Source/Runtime/Core/Public/Serialization/PackageWriter.h
--- a/Engine/Source/Runtime/Core/Public/Serialization/PackageWriter.h
+++ b/Engine/Source/Runtime/Core/Public/Serialization/PackageWriter.h
@@ -347,6 +347,14 @@ public:
 	 */
 	virtual FCbObject GetOplogAttachment(FName PackageName, FUtf8StringView AttachmentKey) = 0;
 
+	/**
+	 * Callback is called for every cross product pair of PackageNames X AttachmentKeys. The Attachment object
+	 * passed to the callback is empty if not found.
+	 */
+	virtual void GetOplogAttachments(TArrayView<FName> PackageNames,
+		TArrayView<FUtf8StringView> AttachmentKeys,
+		TFunctionRef<void(FName PackageName, FUtf8StringView AttachmentKey, FCbObject&& Attachment)> Callback) = 0;
+
 	/**
 	 * Report commit status for the given package. If cooking incrementally, this reports the status
 	 * from a previous cook if it has not yet been committed in the current cook.
