Fix for compile error on Mac

diff --git a/Engine/Source/Developer/IoStoreUtilities/Private/ZenStoreWriter.cpp b/Engine/Source/Developer/IoStoreUtilities/Private/ZenStoreWriter.cpp
--- a/Engine/Source/Developer/IoStoreUtilities/Private/ZenStoreWriter.cpp
+++ b/Engine/Source/Developer/IoStoreUtilities/Private/ZenStoreWriter.cpp
@@ -1176,6 +1176,11 @@ void FZenStoreWriter::GetOplogAttachments(TArrayView<FName> PackageNames,
 	{
 		FName PackageName;
 		FUtf8StringView AttachmentKey;
+
+		FAttachmentHashParam(const FName& InPackageName, FUtf8StringView InAttachmentKey)
+		: PackageName(InPackageName), AttachmentKey(InAttachmentKey)
+		{
+		}
 	};
 	TMultiMap<FIoHash, FAttachmentHashParam> AttachmentHashParams;
 	AttachmentHashParams.Reserve(MaximumHashCount);
diff --git a/Engine/Source/Editor/UnrealEd/Private/TargetDomain/TargetDomainUtils.cpp b/Engine/Source/Editor/UnrealEd/Private/TargetDomain/TargetDomainUtils.cpp
--- a/Engine/Source/Editor/UnrealEd/Private/TargetDomain/TargetDomainUtils.cpp
+++ b/Engine/Source/Editor/UnrealEd/Private/TargetDomain/TargetDomainUtils.cpp
@@ -1360,6 +1360,11 @@ void FEditorDomainOplog::GetOplogAttachments(TArrayView<FName> PackageNames,
 	{
 		FName PackageName;
 		FUtf8StringView AttachmentKey;
+
+		FAttachmentHashParam(const FName& InPackageName, FUtf8StringView InAttachmentKey)
+		: PackageName(InPackageName), AttachmentKey(InAttachmentKey)
+		{
+		}
 	};
 	TMultiMap<FIoHash, FAttachmentHashParam> AttachmentHashParams;
 	AttachmentHashParams.Reserve(MaximumHashCount);
