Enable AutoRTFM editor builds to construct UObjects within transactions

- Switch FPropertyBagRepository::CriticalSection to FTransactionallySafeCriticalSection
- Enable AutoRTFM instrumentation of UnrealEd, which subscribes to the WITH_EDITOR FCoreUObjectDelegates::OnObjectConstructed with FConstructionDuringUpgradeTracker::OnObjectConstructed

diff --git a/Engine/Source/Editor/UnrealEd/Private/Cooker/CookSavePackage.cpp b/Engine/Source/Editor/UnrealEd/Private/Cooker/CookSavePackage.cpp
--- a/Engine/Source/Editor/UnrealEd/Private/Cooker/CookSavePackage.cpp
+++ b/Engine/Source/Editor/UnrealEd/Private/Cooker/CookSavePackage.cpp
@@ -106,7 +106,9 @@ void UCookOnTheFlyServer::SaveCookedPackage(UE::Cook::FSaveCookedPackageContext&
 			bool IsFirstPass = true;
 			for (;;)
 			{
+#if !UE_AUTORTFM
 				try
+#endif
 				{
 					LLM_SCOPE_BYTAG(Cooker_SavePackage);
 					FScopedActivePackage ScopedActivePackage(*this, Context.PackageData.GetPackageName(), NAME_None);
@@ -119,12 +121,14 @@ void UCookOnTheFlyServer::SaveCookedPackage(UE::Cook::FSaveCookedPackageContext&
 						Context.SavePackageResult = GEditor->Save(Package, Context.World, *Context.PlatFilename, SaveArgs);
 					}
 				}
+#if !UE_AUTORTFM
 				catch (std::exception&)
 				{
 					UE_LOG(LogCook, Warning, TEXT("Tried to save package %s for target platform %s but threw an exception"),
 						*Package->GetName(), *TargetPlatform->PlatformName());
 					Context.SavePackageResult = ESavePackageResult::Error;
 				}
+#endif
 
 				bool IsAnotherSaveNeeded = Context.PackageWriter->IsAnotherSaveNeeded(Context.SavePackageResult, SaveArgs);
 				if (IsFirstPass)
diff --git a/Engine/Source/Editor/UnrealEd/UnrealEd.Build.cs b/Engine/Source/Editor/UnrealEd/UnrealEd.Build.cs
--- a/Engine/Source/Editor/UnrealEd/UnrealEd.Build.cs
+++ b/Engine/Source/Editor/UnrealEd/UnrealEd.Build.cs
@@ -12,8 +12,7 @@ public class UnrealEd : ModuleRules
 			throw new BuildException("Unable to instantiate UnrealEd module for non-editor targets.");
 		}
 
-		bEnableExceptions = true;
-		bDisableAutoRTFMInstrumentation = true;
+		bEnableExceptions = !Target.bUseAutoRTFMCompiler;
 
 		PrivatePCHHeaderFile = "Private/UnrealEdPrivatePCH.h";
 		SharedPCHHeaderFile = "Public/UnrealEdSharedPCH.h";
diff --git a/Engine/Source/Runtime/CoreUObject/Internal/UObject/PropertyBagRepository.h b/Engine/Source/Runtime/CoreUObject/Internal/UObject/PropertyBagRepository.h
--- a/Engine/Source/Runtime/CoreUObject/Internal/UObject/PropertyBagRepository.h
+++ b/Engine/Source/Runtime/CoreUObject/Internal/UObject/PropertyBagRepository.h
@@ -3,8 +3,8 @@
 #pragma once
 
 #include "CoreTypes.h"
+#include "Misc/TransactionallySafeCriticalSection.h"
 #include "Containers/Map.h"
-#include "HAL/CriticalSection.h"
 #include "UObject/GCObject.h"
 #include "Templates/FunctionFwd.h"
 
@@ -33,7 +33,7 @@ class FPropertyBagRepository : public FGCObject
 
 private:
 	friend class FPropertyBagRepositoryLock;
-	mutable FCriticalSection CriticalSection;
+	mutable FTransactionallySafeCriticalSection CriticalSection;
 
 	// Lifetimes/ownership:
 	// Managed within UObjectBase and synced with object lifetime. The repo tracks pointers to bags, not the bags themselves.
