Do not create rendering scene when initializing inactive worlds to avoid using GPU resources unnecessarily

This is important for any automated processes running in the Editor that may load maps which in turn create worlds. Without this change, such automated process may accumulate a significant number of rendering scenes.

diff --git a/Engine/Source/Editor/UnrealEd/Private/EditorEngine.cpp b/Engine/Source/Editor/UnrealEd/Private/EditorEngine.cpp
--- a/Engine/Source/Editor/UnrealEd/Private/EditorEngine.cpp
+++ b/Engine/Source/Editor/UnrealEd/Private/EditorEngine.cpp
@@ -6022,8 +6022,9 @@ void UEditorEngine::InitializeNewlyCreatedInactiveWorld(UWorld* World)
 		// Create the world without a physics scene because creating too many physics scenes causes deadlock issues in PhysX. The scene will be created when it is opened in the level editor.
 		// Also, don't create an FXSystem because it consumes too much video memory. This is also created when the level editor opens this world.
 		// Do not create AISystem/Navigation for inactive world. These ones will also be created when the level editor opens this world. if required.
+		// Do not create rendering scene to avoid using GPU resources unnecessarily.
 		World->InitWorld(GetEditorWorldInitializationValues()
-			.CreatePhysicsScene(false)
+			.InitializeScenes(false) // skip creating physics and rendering scenes
 			.CreateFXSystem(false)
 			.CreateAISystem(false)
 			.CreateNavigation(bCooking)
