Allow more information in the iterative validate class summary.

diff --git a/Engine/Source/Editor/UnrealEd/Private/Cooker/IterativeValidatePackageWriter.cpp b/Engine/Source/Editor/UnrealEd/Private/Cooker/IterativeValidatePackageWriter.cpp
--- a/Engine/Source/Editor/UnrealEd/Private/Cooker/IterativeValidatePackageWriter.cpp
+++ b/Engine/Source/Editor/UnrealEd/Private/Cooker/IterativeValidatePackageWriter.cpp
@@ -791,19 +791,35 @@ void FIterativeValidatePackageWriter::EndCook(const FCookInfo& Info)
 		{
 			TStringBuilder<1024> MessageWithDiagnostics;
 			MessageWithDiagnostics << Message;
-			TMap<FTopLevelAssetPath, int32> ClassFalsePositiveCounts = GetSummaryStatusCounts(
+			TMap<FTopLevelAssetPath, TArray<FName>> ClassFalsePositiveCounts = GetClassStatusSummary(
 				EPackageStatus::DeclaredUnmodified_FoundModified_FalsePositive);
-			int32 NumPrinted = 0;
-			constexpr int32 MaxNumPrinted = 5;
-			for (const TPair<FTopLevelAssetPath, int32>& Pair : ClassFalsePositiveCounts)
+			int32 NumClassesPrinted = 0;
+			constexpr int32 MaxNumClassesPrinted = 25;
+			for (const TPair<FTopLevelAssetPath, TArray<FName>>& Pair : ClassFalsePositiveCounts)
 			{
-				MessageWithDiagnostics << TEXT("\n\t") << WriteToString<256>(Pair.Key) << TEXT(": ") << Pair.Value;
-				if (++NumPrinted >= MaxNumPrinted)
+				MessageWithDiagnostics << TEXT("\n\t") << WriteToString<256>(Pair.Key) << TEXT(": ") << Pair.Value.Num();
+
+				int32 NumPackagesPrinted = 0;
+				constexpr int32 MaxNumPackagesPrinted = 10;
+				for (FName PackageName : Pair.Value)
+				{
+					MessageWithDiagnostics << TEXT("\n\t\t") << PackageName;
+					if (++NumPackagesPrinted >= MaxNumPackagesPrinted)
+					{
+						break;
+					}
+				}
+				if (Pair.Value.Num() > NumPackagesPrinted)
+				{
+					MessageWithDiagnostics << TEXT("\n\t\t...");
+				}
+
+				if (++NumClassesPrinted >= MaxNumClassesPrinted)
 				{
 					break;
 				}
 			}
-			if (ClassFalsePositiveCounts.Num() > NumPrinted)
+			if (ClassFalsePositiveCounts.Num() > NumClassesPrinted)
 			{
 				MessageWithDiagnostics << TEXT("\n\t...");
 			}
@@ -1074,10 +1090,10 @@ void FIterativeValidatePackageWriter::MarkPackageCompletedOnDirector(FName Packa
 	}
 
 	int32& TotalCount = TotalStatusCounts.FindOrAdd(Status->Status, 0);
-	int32& ClassCount = ClassStatusCounts.FindOrAdd(Status->AssetClass).FindOrAdd(Status->Status, 0);
+	TArray<FName>& ClassStatusArray = ClassStatusSummary.FindOrAdd(Status->AssetClass).FindOrAdd(Status->Status);
 	++TotalCount;
-	++ClassCount;
-	if (LoggingSoftMaximum >= 0 && TotalCount > LoggingSoftMaximum && ClassCount > 1)
+	ClassStatusArray.Add(PackageName);
+	if (LoggingSoftMaximum >= 0 && TotalCount > LoggingSoftMaximum && ClassStatusArray.Num() > 1)
 	{
 		// Suppress reported logs for this package
 		return;
@@ -1232,21 +1248,21 @@ FIterativeValidatePackageWriter::FStatusCounts FIterativeValidatePackageWriter::
 	return StatusCounts;
 }
 
-TMap<FTopLevelAssetPath, int32> FIterativeValidatePackageWriter::GetSummaryStatusCounts(EPackageStatus PackageStatus)
+TMap<FTopLevelAssetPath, TArray<FName>> FIterativeValidatePackageWriter::GetClassStatusSummary(EPackageStatus PackageStatus)
 {
-	TMap<FTopLevelAssetPath, int32> Result;
-	for (const TPair<FTopLevelAssetPath, TMap<EPackageStatus, int32>>& Pair : ClassStatusCounts)
+	TMap<FTopLevelAssetPath, TArray<FName>> Result;
+	for (const TPair<FTopLevelAssetPath, TMap<EPackageStatus, TArray<FName>>>& Pair : ClassStatusSummary)
 	{
-		const int32* Count = Pair.Value.Find(PackageStatus);
-		if (Count)
+		const TArray<FName>* Packages = Pair.Value.Find(PackageStatus);
+		if (Packages)
 		{
-			Result.Add(Pair.Key, *Count);
+			Result.Add(Pair.Key, *Packages);
 		}
 	}
 
-	Result.ValueSort([](const int32& A, const int32& B)
+	Result.ValueSort([](const TArray<FName>& A, const TArray<FName>& B)
 		{
-			return A > B;
+			return A.Num() > B.Num();
 		});
 	return Result;
 }
diff --git a/Engine/Source/Editor/UnrealEd/Private/Cooker/IterativeValidatePackageWriter.h b/Engine/Source/Editor/UnrealEd/Private/Cooker/IterativeValidatePackageWriter.h
--- a/Engine/Source/Editor/UnrealEd/Private/Cooker/IterativeValidatePackageWriter.h
+++ b/Engine/Source/Editor/UnrealEd/Private/Cooker/IterativeValidatePackageWriter.h
@@ -114,7 +114,7 @@ protected:
 	EPackageStatus GetPackageStatus(FName PackageName) const;
 	void SetPackageStatus(FName PackageName, EPackageStatus NewStatus);
 	FStatusCounts CountPackagesByStatus();
-	TMap<FTopLevelAssetPath, int32> GetSummaryStatusCounts(EPackageStatus PackageStatus);
+	TMap<FTopLevelAssetPath, TArray<FName>> GetClassStatusSummary(EPackageStatus PackageStatus);
 	bool IsAnotherSaveNeededInternal(FSavePackageResultStruct& PreviousResult, FSavePackageArgs& SaveArgs);
 	void MarkPackageCompletedOnDirector(FName PackageName, UE::Cook::FWorkerId WorkerId);
 
@@ -122,7 +122,7 @@ protected:
 	TMap<FName, FPackageStatusInfo> PackageStatusMap;
 	TMap<FName, TArray<FMessage>> PackageMessageMap;
 	TMap<EPackageStatus, int32> TotalStatusCounts; // Only populated on director
-	TMap<FTopLevelAssetPath, TMap<EPackageStatus, int32>> ClassStatusCounts; // Only populated on director
+	TMap<FTopLevelAssetPath, TMap<EPackageStatus, TArray<FName>>> ClassStatusSummary; // Only populated on director
 	TSet<FName> PackageIgnoreList;
 	FString MetadataPath;
 	UCookOnTheFlyServer& COTFS;
