[New Gizmos] Migrate frustum select drag tool to ITF

- Add SelectionMode.EnableITFTools CVAR to enable ITF based drag tools

- Add UDragToolsBehaviorSource class where DragTools and their required behaviors are setup and registered.

- FEditorViewportClient hosts a UDragToolsBehaviorSource instance used for drag tools behavior registering/unregistering/tracking.

- FEditorViewportClient takes care of drawing drag tools UI (e.g. selection box for Frustum Select tool)

- FDragTool is now deriving from IClickDragBehaviorTarget, in order to share code between old and ITF implementation of Drag Tools.

- FDragTool_ActorFrustumSelect modified accordingly.

- UModeManagerInteractiveToolsContext::InputKey now has an internal flag indicating if it should bypass Alt input or not. Used in this CL when registering Drag Tools, so that they can be used with the usual hotkeys, which e.g. for Frustum Select include Alt.

SelectionMode.EnableITFTools 1 to use ITF based Frustum Select tool.
SelectionMode.EnableITFTools 0 to use legacy Frustum Select tool.

Note: bUseHoverFeedback was created and initialized in a couple of places of FDragTool_ActorFrustumSelect, but never used. It has been removed (of course we can keep it if needed!)

diff --git a/Engine/Source/Editor/Experimental/EditorInteractiveToolsFramework/Private/EdModeInteractiveToolsContext.cpp b/Engine/Source/Editor/Experimental/EditorInteractiveToolsFramework/Private/EdModeInteractiveToolsContext.cpp
--- a/Engine/Source/Editor/Experimental/EditorInteractiveToolsFramework/Private/EdModeInteractiveToolsContext.cpp
+++ b/Engine/Source/Editor/Experimental/EditorInteractiveToolsFramework/Private/EdModeInteractiveToolsContext.cpp
@@ -961,9 +961,12 @@ bool UModeManagerInteractiveToolsContext::InputKey(FEditorViewportClient* Viewpo
 				// to wherever they get handled.
 				// Someday these kinds of prioritizations will be handled by having camera manipulation be
 				// in a common input router so that behavior priorities can determine the ordering.
-				if (ViewportClient && ViewportClient->IsAltPressed() && InputRouter->HasActiveMouseCapture() == false)
+				if (BypassAltModifier())
 				{
-					return false;
+					if (ViewportClient && ViewportClient->IsAltPressed() && InputRouter->HasActiveMouseCapture() == false)
+					{
+						return false;
+					}
 				}
 
 				FInputDeviceState InputState = CurrentMouseState;
@@ -1283,6 +1286,10 @@ bool UModeManagerInteractiveToolsContext::OnChildEdModeDeactivated(UEdModeIntera
 	return false;
 }
 
+void UModeManagerInteractiveToolsContext::SetBypassAltModifier(bool bInBypassAltModifier)
+{
+	bBypassAltModifier = bInBypassAltModifier;
+}
 
 void UEdModeInteractiveToolsContext::InitializeContextFromModeManagerContext(UModeManagerInteractiveToolsContext* ModeManagerToolsContext)
 {
@@ -1307,4 +1314,4 @@ void UEdModeInteractiveToolsContext::InitializeContextFromModeManagerContext(UMo
 FRay UEdModeInteractiveToolsContext::GetLastWorldRay() const
 {
 	return ParentModeManagerToolsContext->GetLastWorldRay();
-}
\ No newline at end of file
+}
diff --git a/Engine/Source/Editor/Experimental/EditorInteractiveToolsFramework/Public/EdModeInteractiveToolsContext.h b/Engine/Source/Editor/Experimental/EditorInteractiveToolsFramework/Public/EdModeInteractiveToolsContext.h
--- a/Engine/Source/Editor/Experimental/EditorInteractiveToolsFramework/Public/EdModeInteractiveToolsContext.h
+++ b/Engine/Source/Editor/Experimental/EditorInteractiveToolsFramework/Public/EdModeInteractiveToolsContext.h
@@ -338,6 +338,27 @@ public:
 	 */
 	bool OnChildEdModeDeactivated(UEdModeInteractiveToolsContext* ChildToolsContext);
 
+protected:
+	//~Begin ITF Alt processing Bypass
+
+	/**
+	 * These functions are declared protected to avoid having to deal with deprecation later on.
+	 * The Alt modifier bypass will be removed once a complete switch to ITF happens
+	 */
+
+	// Needs to use Alt modifier for Drag Tools to work
+	friend class FEditorViewportClient;
+
+	/** Enable/Disable Alt modifier bypass */
+	void SetBypassAltModifier(bool bInBypassAltModifier);
+
+	/** Is this Interactive Tools Context skipping Alt processing? */
+	bool BypassAltModifier() const { return bBypassAltModifier; }
+
+	/** Should we skip processing Alt modifier inputs? See UModeManagerInteractiveToolsContext::InputKey */
+	bool bBypassAltModifier = true;
+
+	//~End ITF Alt processing Bypass
 };
 
 
@@ -364,4 +385,4 @@ public:
 protected:
 	UPROPERTY()
 	TObjectPtr<UModeManagerInteractiveToolsContext> ParentModeManagerToolsContext = nullptr;
-};
\ No newline at end of file
+};
diff --git a/Engine/Source/Editor/UnrealEd/Private/DragTool_FrustumSelect.cpp b/Engine/Source/Editor/UnrealEd/Private/DragTool_FrustumSelect.cpp
--- a/Engine/Source/Editor/UnrealEd/Private/DragTool_FrustumSelect.cpp
+++ b/Engine/Source/Editor/UnrealEd/Private/DragTool_FrustumSelect.cpp
@@ -32,7 +32,7 @@
 //
 ///////////////////////////////////////////////////////////////////////////////
 
-namespace UE::LevelEditor::Private
+namespace UE::DragTools::Private
 {
 	TArray<FTypedElementHandle> GetElementsIntersectingFrustum(const AActor* Actor,
 		const FConvexVolume& InFrustum,
@@ -53,48 +53,143 @@ namespace UE::LevelEditor::Private
 
 		return {};
 	}
+
+	bool IsCurrentModeSupported(const FEditorViewportClient* InViewportClient)
+	{
+		static TArray<FEditorModeID> UnsupportedModes = { FBuiltinEditorModes::EM_Landscape };
+
+		if (InViewportClient)
+		{
+			if (const FEditorModeTools* ModeTools = InViewportClient->GetModeTools())
+			{
+				for (const FEditorModeID UnsupportedMode : UnsupportedModes)
+				{
+					if (!ModeTools->EnsureNotInMode(UnsupportedMode))
+					{
+						return false;
+					}
+				}
+			}
+		}
+
+		return true;
+	}
 }
 
 void FDragTool_ActorFrustumSelect::AddDelta( const FVector& InDelta )
 {
-	FIntPoint MousePos;
-	EditorViewportClient->Viewport->GetMousePos(MousePos);
+	if (EditorViewportClient)
+	{
+		if (FViewport* const Viewport = EditorViewportClient->Viewport)
+		{
+			FIntPoint MousePos;
+			Viewport->GetMousePos(MousePos);
 
-	EndWk = FVector(MousePos);
-	End = EndWk;
+			FInputDeviceRay DragPos;
+			DragPos.ScreenPosition = MousePos;
 
-	const bool bUseHoverFeedback = GEditor != NULL && GetDefault<ULevelEditorViewportSettings>()->bEnableViewportHoverFeedback;
+			// Former AddDelta logic is now in OnClickDrag
+			OnClickDrag(DragPos);
+		}
+	}
 }
 
 void FDragTool_ActorFrustumSelect::StartDrag(FEditorViewportClient* InViewportClient, const FVector& InStart, const FVector2D& InStartScreen)
 {
-	FDragTool::StartDrag(InViewportClient, InStart, InStartScreen);
+	EditorViewportClient = InViewportClient;
 
-	const bool bUseHoverFeedback = GEditor != NULL && GetDefault<ULevelEditorViewportSettings>()->bEnableViewportHoverFeedback;
+	FInputDeviceRay PressPos;
+	PressPos.WorldRay.Origin = InStart;
+	PressPos.ScreenPosition = InStartScreen;
 
-	// Remove any active hover objects
-	FLevelEditorViewportClient::ClearHoverFromObjects();
+	// Former StartDrag logic is now in OnClickPress
+	OnClickPress(PressPos);
+}
 
-	FIntPoint MousePos;
-	InViewportClient->Viewport->GetMousePos(MousePos);
+void FDragTool_ActorFrustumSelect::EndDrag()
+{
+	FInputDeviceRay ReleasePos;
+	ReleasePos.ScreenPosition = FVector2D(End);
 
-	Start = FVector(InStartScreen.X, InStartScreen.Y, 0);
-	End = EndWk = Start;
+	// Former EndDrag logic is now in OnClickRelease
+	OnClickRelease(ReleasePos);
 }
 
-void FDragTool_ActorFrustumSelect::EndDrag()
+void FDragTool_ActorFrustumSelect::Render(const FSceneView* View, FCanvas* Canvas)
+{
+	if (!Canvas || !View)
+	{
+		return;
+	}
+
+	FCanvasBoxItem BoxItem(
+		FVector2D(Start.X, Start.Y) / Canvas->GetDPIScale(),
+		FVector2D(End.X - Start.X, End.Y - Start.Y) / Canvas->GetDPIScale()
+	);
+
+	BoxItem.SetColor(FLinearColor::White);
+	Canvas->DrawItem(BoxItem);
+}
+
+FInputRayHit FDragTool_ActorFrustumSelect::CanBeginClickDragSequence(const FInputDeviceRay& InPressPos)
+{
+	// Todo: this could be retrieved from a command for customization
+	FInputChord ActivationChord(EModifierKey::Control | EModifierKey::Alt, EKeys::LeftMouseButton);
+
+	return IsActivationChordPressed(ActivationChord) && UE::DragTools::Private::IsCurrentModeSupported(LevelViewportClient)
+			 ? FInputRayHit(TNumericLimits<float>::Max()) // bHit is true. Depth is max to lose the standard tiebreaker.
+			 : FInputRayHit();
+}
+
+void FDragTool_ActorFrustumSelect::OnClickPress(const FInputDeviceRay& InPressPos)
 {
+	if (EditorViewportClient)
+	{
+		// Signal that this tool is now active
+		OnActivateTool().Broadcast();
+
+		FDragTool::StartDrag(EditorViewportClient, InPressPos.WorldRay.Origin, InPressPos.ScreenPosition);
+
+		// Remove any active hover objects
+		FLevelEditorViewportClient::ClearHoverFromObjects();
+
+		Start = FVector(InPressPos.ScreenPosition.X, InPressPos.ScreenPosition.Y, 0);
+		End = EndWk = Start;
+	}
+}
+
+void FDragTool_ActorFrustumSelect::OnClickDrag(const FInputDeviceRay& InDragPos)
+{
+	EndWk = FVector(InDragPos.ScreenPosition.X, InDragPos.ScreenPosition.Y, 0);
+	End = EndWk;
+}
+
+void FDragTool_ActorFrustumSelect::OnClickRelease(const FInputDeviceRay& InReleasePos)
+{
+	End = FVector(InReleasePos.ScreenPosition.X, InReleasePos.ScreenPosition.Y, 0);
+
+	if (!EditorViewportClient)
+	{
+		return;
+	}
+
+	FViewport* Viewport = EditorViewportClient->Viewport;
+	if (!Viewport)
+	{
+		return;
+	}
+
 	UBrushEditingSubsystem* BrushSubsystem = GEditor->GetEditorSubsystem<UBrushEditingSubsystem>();
 	const bool bGeometryMode = BrushSubsystem ? BrushSubsystem->IsGeometryEditorModeActive() : false;
 
-	FSceneViewFamilyContext ViewFamily(FSceneViewFamily::ConstructionValues(EditorViewportClient->Viewport, EditorViewportClient->GetScene(), EditorViewportClient->EngineShowFlags ));
+	FSceneViewFamilyContext ViewFamily(FSceneViewFamily::ConstructionValues(Viewport, EditorViewportClient->GetScene(), EditorViewportClient->EngineShowFlags ));
 	FSceneView* SceneView = EditorViewportClient->CalcSceneView(&ViewFamily);
 
 	// Generate a frustum out of the dragged box
 	FConvexVolume Frustum;
 	CalculateFrustum( SceneView, Frustum, true );
 
-	FScopedTransaction Transaction( NSLOCTEXT("ActorFrustumSelect", "MarqueeSelectTransation", "Marquee Select" ) );
+	FScopedTransaction Transaction( NSLOCTEXT("ActorFrustumSelect", "MarqueeSelectTransaction", "Marquee Select" ) );
 
 	bool bShouldSelect = true;
 
@@ -111,7 +206,7 @@ void FDragTool_ActorFrustumSelect::EndDrag()
 	const bool bEditorModeHandledSelection = ModeTools->FrustumSelect(Frustum, EditorViewportClient, bLeftMouseButtonDown);
 
 	// Let the component visualizers try to handle the selection.
-	const bool bComponentVisHandledSelection = !bEditorModeHandledSelection && GUnrealEd->ComponentVisManager.HandleFrustumSelect(Frustum, EditorViewportClient, EditorViewportClient->Viewport);
+	const bool bComponentVisHandledSelection = !bEditorModeHandledSelection && GUnrealEd->ComponentVisManager.HandleFrustumSelect(Frustum, EditorViewportClient, Viewport);
 
 	if( !bEditorModeHandledSelection && !bComponentVisHandledSelection)
 	{
@@ -128,8 +223,8 @@ void FDragTool_ActorFrustumSelect::EndDrag()
 			bGeometryMode
 		};
 
-		const int32 ViewportSizeX = EditorViewportClient->Viewport->GetSizeXY().X;
-		const int32 ViewportSizeY = EditorViewportClient->Viewport->GetSizeXY().Y;
+		const int32 ViewportSizeX = Viewport->GetSizeXY().X;
+		const int32 ViewportSizeY = Viewport->GetSizeXY().Y;
 
 		if( Start.X > End.X )
 		{
@@ -149,7 +244,7 @@ void FDragTool_ActorFrustumSelect::EndDrag()
 			for(FActorIterator It(EditorViewportClient->GetWorld()); It; ++It)
 			{
 				AActor* Actor = *It;
-				ElementsToSelect.Append(UE::LevelEditor::Private::GetElementsIntersectingFrustum(Actor, Frustum, EditorViewportClient, LevelViewportClient, SeletionArgs));
+				ElementsToSelect.Append(UE::DragTools::Private::GetElementsIntersectingFrustum(Actor, Frustum, EditorViewportClient, LevelViewportClient, SeletionArgs));
 			}
 		}
 		else
@@ -168,7 +263,7 @@ void FDragTool_ActorFrustumSelect::EndDrag()
 			// Typed Element selection
 			{
 				FTypedElementListRef ElementList = UTypedElementRegistry::GetInstance()->CreateElementList();
-				EditorViewportClient->Viewport->GetElementHandlesInRect(BoxRect, ElementList);
+				Viewport->GetElementHandlesInRect(BoxRect, ElementList);
 
 				if (bStrictDragSelection)
 				{
@@ -197,11 +292,9 @@ void FDragTool_ActorFrustumSelect::EndDrag()
 			// We need this old code to support the BSP
 			TSet<AActor*> BSPActors;
 			TSet<UModel*> HitModels;
-			EditorViewportClient->Viewport->GetActorsAndModelsInHitProxy( BoxRect, BSPActors, HitModels );
+			Viewport->GetActorsAndModelsInHitProxy( BoxRect, BSPActors, HitModels );
 			BSPActors.Empty(HitModels.Num());
 
-			
-
 			if (HitModels.Num() > 0)
 			{
 				// Check every model to see if its BSP surfaces should be selected
@@ -228,7 +321,7 @@ void FDragTool_ActorFrustumSelect::EndDrag()
 					AActor* Actor = *It;
 					if (bStrictDragSelection)
 					{
-						ElementsToSelect.Append(UE::LevelEditor::Private::GetElementsIntersectingFrustum(Actor, Frustum, EditorViewportClient, LevelViewportClient, SeletionArgs));
+						ElementsToSelect.Append(UE::DragTools::Private::GetElementsIntersectingFrustum(Actor, Frustum, EditorViewportClient, LevelViewportClient, SeletionArgs));
 					}
 					else
 					{
@@ -255,13 +348,14 @@ void FDragTool_ActorFrustumSelect::EndDrag()
 	FLevelEditorViewportClient::ClearHoverFromObjects();
 
 	FDragTool::EndDrag();
+
+	// Signal that this tool is no longer active
+	OnDeactivateTool().Broadcast();
 }
 
-void FDragTool_ActorFrustumSelect::Render(const FSceneView* View, FCanvas* Canvas )
+void FDragTool_ActorFrustumSelect::OnTerminateDragSequence()
 {
-	FCanvasBoxItem BoxItem( FVector2D(Start.X, Start.Y) / Canvas->GetDPIScale(), FVector2D(End.X-Start.X, End.Y-Start.Y) / Canvas->GetDPIScale());
-	BoxItem.SetColor( FLinearColor::White );
-	Canvas->DrawItem( BoxItem );
+	Start = End = FVector::ZeroVector;
 }
 
 bool FDragTool_ActorFrustumSelect::IntersectsFrustum( const UModel& InModel, int32 NodeIndex, const FConvexVolume& InFrustum, bool bUseStrictSelection ) const
@@ -280,6 +374,11 @@ bool FDragTool_ActorFrustumSelect::IntersectsFrustum( const UModel& InModel, int
 
 void FDragTool_ActorFrustumSelect::CalculateFrustum( const FSceneView* InView, FConvexVolume& OutFrustum, bool bUseBoxFrustum ) const
 {
+	if (!InView || ! EditorViewportClient)
+	{
+		return;
+	}
+
 	if( bUseBoxFrustum )
 	{
 		FVector CamPoint = EditorViewportClient->GetViewLocation();
diff --git a/Engine/Source/Editor/UnrealEd/Private/EditorDragTools.cpp b/Engine/Source/Editor/UnrealEd/Private/EditorDragTools.cpp
--- a/Engine/Source/Editor/UnrealEd/Private/EditorDragTools.cpp
+++ b/Engine/Source/Editor/UnrealEd/Private/EditorDragTools.cpp
@@ -4,7 +4,57 @@
 #include "InputCoreTypes.h"
 #include "EditorViewportClient.h"
 #include "Editor.h"
+#include "EditorModeManager.h"
+#include "EditorModes.h"
 #include "SnappingUtils.h"
+#include "Tools/UEdMode.h"
+
+namespace UE::DragTools
+{
+void UnregisterDragTools(FEditorViewportClient* InViewportClient)
+{
+	if (InViewportClient)
+	{
+		InViewportClient->UnregisterDragTools();
+	}
+}
+
+void RegisterDragTools(FEditorViewportClient* InViewportClient)
+{
+	if (InViewportClient)
+	{
+		InViewportClient->RegisterDragTools();
+	}
+}
+
+// CVar initializer
+static TAutoConsoleVariable<int32> CVarEnableITFTools(
+	TEXT("SelectionMode.EnableITFTools"),
+	0,
+	TEXT("Is the ITF version of Selection Mode Tools enabled?\nCurrently supports: Frustum Select."),
+	FConsoleVariableDelegate::CreateLambda(
+		[](const IConsoleVariable* InVariable)
+		{
+			if (FEditorViewportClient* ViewportClient = GLevelEditorModeTools().GetFocusedViewportClient())
+			{
+				if (InVariable->GetBool())
+				{
+					RegisterDragTools(ViewportClient);
+				}
+				else
+				{
+					UnregisterDragTools(ViewportClient);
+				}
+			}
+		}
+	)
+);
+
+bool UseITFTools()
+{
+	return CVarEnableITFTools->GetBool();
+}
+}
 
 ///////////////////////////////////////////////////////////////////////////////////////////////////////////////////////
 //
@@ -66,3 +116,46 @@ void FDragTool::EndDrag()
 	Start = End = EndWk = FVector::ZeroVector;
 	bIsDragging = false;
 }
+
+void FDragTool::OnUpdateModifierState(int InModifierID, bool bInIsOn)
+{
+	constexpr int ShiftKeyMod = 1;
+	constexpr int AltKeyMod = 2;
+	constexpr int CtrlKeyMod = 3;
+
+	switch (InModifierID)
+	{
+	case ShiftKeyMod:
+		bShiftDown = bInIsOn;
+		break;
+	case CtrlKeyMod:
+		bControlDown = bInIsOn;
+		break;
+	case AltKeyMod:
+		bAltDown = bInIsOn;
+		break;
+	default:;
+	}
+}
+
+bool FDragTool::IsActivationChordPressed(const FInputChord& InChord) const
+{
+	bool bIsChordPressed = true;
+
+	if (InChord.NeedsControl())
+	{
+		bIsChordPressed = bIsChordPressed && bControlDown;
+	}
+
+	if (InChord.NeedsAlt())
+	{
+		bIsChordPressed = bIsChordPressed && bAltDown;
+	}
+
+	if (InChord.NeedsShift())
+	{
+		bIsChordPressed = bIsChordPressed && bShiftDown;
+	}
+
+	return bIsChordPressed;
+}
diff --git a/Engine/Source/Editor/UnrealEd/Private/EditorViewportClient.cpp b/Engine/Source/Editor/UnrealEd/Private/EditorViewportClient.cpp
--- a/Engine/Source/Editor/UnrealEd/Private/EditorViewportClient.cpp
+++ b/Engine/Source/Editor/UnrealEd/Private/EditorViewportClient.cpp
@@ -79,6 +79,7 @@
 #include "Settings/EditorStyleSettings.h"
 #include "GameFramework/ActorPrimitiveColorHandler.h"
 #include "LandscapeRender.h"
+#include "Tools/DragToolsBehaviourSource.h"
 
 #define LOCTEXT_NAMESPACE "EditorViewportClient"
 
@@ -546,6 +547,8 @@ FEditorViewportClient::FEditorViewportClient(FEditorModeTools* InModeTools, FPre
 	RequestUpdateDPIScale();
 
 	FSlateApplication::Get().OnWindowDPIScaleChanged().AddRaw(this, &FEditorViewportClient::HandleWindowDPIScaleChanged);
+
+	RegisterDragTools();
 }
 
 FEditorViewportClient::~FEditorViewportClient()
@@ -584,6 +587,8 @@ FEditorViewportClient::~FEditorViewportClient()
 		FSlateApplication::Get().OnWindowDPIScaleChanged().RemoveAll(this);
 	}
 
+	UnregisterDragTools();
+
 	ModeTools.Reset();
 }
 
@@ -4017,6 +4022,11 @@ void FEditorViewportClient::AddReferencedObjects( FReferenceCollector& Collector
 			StereoViewState.GetReference()->AddReferencedObjects(Collector);
 		}
 	}
+
+	if (DragToolsBehaviorSource) 
+	{
+		Collector.AddReferencedObject(DragToolsBehaviorSource);
+	}
 }
 
 FString FEditorViewportClient::GetReferencerName() const
@@ -4662,6 +4672,14 @@ void FEditorViewportClient::DrawPreviewLightVisualization(const FSceneView* View
 
 void FEditorViewportClient::RenderDragTool(const FSceneView* View, FCanvas* Canvas)
 {
+	if (DragToolsBehaviorSource)
+	{
+		if (FDragTool* ActiveDragTool = DragToolsBehaviorSource->GetActiveTool())
+		{
+			ActiveDragTool->Render(View, Canvas);
+		}
+	}
+
 	MouseDeltaTracker->RenderDragTool(View, Canvas);
 }
 
@@ -6731,6 +6749,47 @@ float FEditorViewportClient::GetMinimumOrthoZoom() const
 	return FMath::Max(GetDefault<ULevelEditorViewportSettings>()->MinimumOrthographicZoom, MIN_ORTHOZOOM);
 }
 
+void FEditorViewportClient::RegisterDragTools()
+{
+	if (UE::DragTools::UseITFTools())
+	{
+		if (!DragToolsBehaviorSource)
+		{
+			if (const FEditorModeTools* EditorModeTools = GetModeTools())
+			{
+				if (UModeManagerInteractiveToolsContext* InteractiveToolContext = EditorModeTools->GetInteractiveToolsContext())
+				{
+					// Alt modifier input required by Drag Tools, but currently skipped by ITF in Tool Context.
+					// Let's make sure Alt keys are processed
+					InteractiveToolContext->SetBypassAltModifier(false);
+
+					DragToolsBehaviorSource = NewObject<UDragToolsBehaviorSource>();
+					DragToolsBehaviorSource->Initialize(InteractiveToolContext);
+
+					InteractiveToolContext->InputRouter->RegisterSource(DragToolsBehaviorSource);
+				}
+			}
+		}
+	}
+}
+
+void FEditorViewportClient::UnregisterDragTools()
+{
+	if (DragToolsBehaviorSource)
+	{
+		if (const FEditorModeTools* EditorModeTools = GetModeTools())
+		{
+			if (UModeManagerInteractiveToolsContext* InteractiveToolContext = EditorModeTools->GetInteractiveToolsContext())
+			{
+				// Revert to previous state, with Alt Modifier not being processed
+				InteractiveToolContext->SetBypassAltModifier(true);
+				InteractiveToolContext->InputRouter->DeregisterSource(DragToolsBehaviorSource);
+				DragToolsBehaviorSource = nullptr;
+			}
+		}
+	}
+}
+
 ////////////////
 
 bool FEditorViewportStats::bInitialized(false);
diff --git a/Engine/Source/Editor/UnrealEd/Private/MouseDeltaTracker.cpp b/Engine/Source/Editor/UnrealEd/Private/MouseDeltaTracker.cpp
--- a/Engine/Source/Editor/UnrealEd/Private/MouseDeltaTracker.cpp
+++ b/Engine/Source/Editor/UnrealEd/Private/MouseDeltaTracker.cpp
@@ -370,9 +370,12 @@ void FMouseDeltaTracker::ConditionalBeginUsingDragTool( FEditorViewportClient* I
 				}
 				else
 				{
-					if (LeftMouseButtonDown && bControlDown && bAltDown)
+					if (!UE::DragTools::UseITFTools())
 					{
-						DragTool = InViewportClient->MakeDragTool(EDragTool::FrustumSelect);
+						if (LeftMouseButtonDown && bControlDown && bAltDown)
+						{
+							DragTool = InViewportClient->MakeDragTool(EDragTool::FrustumSelect);
+						}
 					}
 				}
 			}
diff --git a/Engine/Source/Editor/UnrealEd/Private/Tools/DragToolsBehaviourSource.cpp b/Engine/Source/Editor/UnrealEd/Private/Tools/DragToolsBehaviourSource.cpp
new file mode 100644
--- /dev/null
+++ b/Engine/Source/Editor/UnrealEd/Private/Tools/DragToolsBehaviourSource.cpp
@@ -0,0 +1,138 @@
+// Copyright Epic Games, Inc. All Rights Reserved.
+
+#include "Tools/DragToolsBehaviourSource.h"
+#include "BaseBehaviors/ClickDragBehavior.h"
+#include "BaseBehaviors/KeyAsModifierInputBehavior.h"
+#include "DragTool_FrustumSelect.h"
+#include "EdModeInteractiveToolsContext.h"
+
+void UDragToolsBehaviorSource::Initialize(const UEditorInteractiveToolsContext* InInteractiveToolsContext)
+{
+	if (!InInteractiveToolsContext)
+	{
+		return;
+	}
+
+	TObjectPtr<UInteractiveToolManager> ToolManager = InInteractiveToolsContext->ToolManager;
+	if (!ToolManager)
+	{
+		return;
+	}
+
+	IToolsContextQueriesAPI* ContextQueriesAPI = ToolManager->GetContextQueriesAPI();
+	if (!ContextQueriesAPI)
+	{
+		return;
+	}
+
+	FViewport* Viewport = ContextQueriesAPI->GetFocusedViewport();
+	if (!Viewport)
+	{
+		return;
+	}
+
+	FEditorViewportClient* EditorViewportClient = static_cast<FEditorViewportClient*>(Viewport->GetClient());
+	if (!EditorViewportClient)
+	{
+		return;
+	}
+
+	// Some Drag Tools require or can make use specifically of Level Viewport Client.
+	// This can be nullptr
+	FLevelEditorViewportClient* LevelViewportClient = static_cast<FLevelEditorViewportClient*>(EditorViewportClient);
+
+	BehaviorSet = NewObject<UInputBehaviorSet>();
+
+	constexpr int ShiftKeyMod = 1;
+	constexpr int AltKeyMod = 2;
+	constexpr int CtrlKeyMod = 3;
+
+	// In case New TRS Gizmos are enabled, we need Drag Tools to be processed before the Gizmo ones, otherwise Drag Behaviors will never be triggered.
+	// See:
+	// * UTransformGizmo::SetupBehaviors()			[DEFAULT_GIZMO_PRIORITY]
+	// * UTransformGizmo::SetupIndirectBehaviors()	[DEFAULT_GIZMO_PRIORITY-1]
+	static constexpr int DEFAULT_DRAG_TOOLS_PRIORITY = FInputCapturePriority::DEFAULT_GIZMO_PRIORITY-2;
+
+	// Frustum Select
+	{
+		TSharedPtr<FDragTool> FrustumSelectTool;
+
+		if (LevelViewportClient)
+		{
+			// Initialize using Level Viewport Client if available
+			FrustumSelectTool = MakeShared<FDragTool_ActorFrustumSelect>(LevelViewportClient);
+		}
+		else
+		{
+			FrustumSelectTool = MakeShared<FDragTool_ActorFrustumSelect>(EditorViewportClient);
+		}
+
+		DragTools.Add(FrustumSelectTool);
+
+		// Add drag input behavior
+		UClickDragInputBehavior* MouseBehavior = NewObject<UClickDragInputBehavior>();
+		MouseBehavior->Initialize(FrustumSelectTool.Get());
+		MouseBehavior->SetDefaultPriority(DEFAULT_DRAG_TOOLS_PRIORITY);
+		BehaviorSet->Add(MouseBehavior);
+
+		// Register modifier listeners
+		UKeyAsModifierInputBehavior* KeyModifierInputBehavior = NewObject<UKeyAsModifierInputBehavior>();
+		KeyModifierInputBehavior->Initialize(FrustumSelectTool.Get(), ShiftKeyMod, FInputDeviceState::IsShiftKeyDown);
+		KeyModifierInputBehavior->Initialize(FrustumSelectTool.Get(), CtrlKeyMod, FInputDeviceState::IsCtrlKeyDown);
+		KeyModifierInputBehavior->Initialize(FrustumSelectTool.Get(), AltKeyMod, FInputDeviceState::IsAltKeyDown);
+		KeyModifierInputBehavior->SetDefaultPriority(DEFAULT_DRAG_TOOLS_PRIORITY);
+		BehaviorSet->Add(KeyModifierInputBehavior);
+
+		FrustumSelectTool->OnActivateTool().AddUObject(
+			this, &UDragToolsBehaviorSource::OnActivateTool, EDragTool::FrustumSelect
+		);
+		FrustumSelectTool->OnDeactivateTool().AddUObject(
+			this, &UDragToolsBehaviorSource::OnDeactivateTool, EDragTool::FrustumSelect
+		);
+	}
+
+	// Add other drag tools
+}
+
+FDragTool* UDragToolsBehaviorSource::GetActiveTool() const
+{
+	if (ActiveToolType > 0)
+	{
+		int32 ToolIndex = GetToolIndex(static_cast<EDragTool::Type>(ActiveToolType));
+		if (DragTools.IsValidIndex(ToolIndex))
+		{
+			if (const TSharedPtr<FDragTool>& ActiveTool = DragTools[ToolIndex])
+			{
+				return ActiveTool.Get();
+			}
+		}
+	}
+
+	return nullptr;
+}
+
+int32 UDragToolsBehaviorSource::GetToolIndex(EDragTool::Type InDragToolType)
+{
+	static TArray<int32> ToolIndexes = { EDragTool::FrustumSelect };
+
+	int32 ToolIndex = ToolIndexes.Find(InDragToolType);
+	if (ToolIndex != INDEX_NONE)
+	{
+		return ToolIndex;
+	}
+
+	return -1;
+}
+
+void UDragToolsBehaviorSource::OnActivateTool(EDragTool::Type InDragToolType)
+{
+	ActiveToolType = InDragToolType;
+}
+
+void UDragToolsBehaviorSource::OnDeactivateTool(EDragTool::Type InDragToolType)
+{
+	if (ActiveToolType == InDragToolType)
+	{
+		ActiveToolType = -1;
+	}
+}
diff --git a/Engine/Source/Editor/UnrealEd/Private/Tools/DragToolsBehaviourSource.h b/Engine/Source/Editor/UnrealEd/Private/Tools/DragToolsBehaviourSource.h
new file mode 100644
--- /dev/null
+++ b/Engine/Source/Editor/UnrealEd/Private/Tools/DragToolsBehaviourSource.h
@@ -0,0 +1,61 @@
+// Copyright Epic Games, Inc. All Rights Reserved.
+
+#pragma once
+
+#include "InputBehaviorSet.h"
+#include "UObject/Object.h"
+#include "DragTool_FrustumSelect.h"
+#include "DragToolsBehaviourSource.generated.h"
+
+class UEditorInteractiveToolsContext;
+
+/**
+ * Hosts Drag Tools and needed behaviors.
+ * Handles an Input Behavior Set, and keeps track of currently active Drag tool.
+ */
+UCLASS(Transient, MinimalAPI)
+class UDragToolsBehaviorSource final : public UObject, public IInputBehaviorSource
+{
+	GENERATED_BODY()
+
+public:
+	//~ Begin IInputBehaviorSource
+	virtual const UInputBehaviorSet* GetInputBehaviors() const override { return BehaviorSet; }
+	//~ End IClickDragBehaviorTarget
+
+	/**
+	 * Creates and stores the drag tools.
+	 * Instantiates the BehaviorSet hosting all behaviors required by drag tools.
+	 *
+	 * @param InInteractiveToolsContext: used to retrieve the viewport client used by drag tools. Can be gathered e.g.
+	 * from UEdMode::GetInteractiveToolsContext.
+	 */
+	void Initialize(const UEditorInteractiveToolsContext* InInteractiveToolsContext);
+
+	/**
+	 * Returns the drag tool currently being used, if any
+	 */
+	FDragTool* GetActiveTool() const;
+
+private:
+	static int32 GetToolIndex(EDragTool::Type InDragToolType);
+
+	void OnActivateTool(EDragTool::Type InDragToolType);
+	void OnDeactivateTool(EDragTool::Type InDragToolType);
+
+	/**
+	 * Hosting drag tools behaviors
+	 */
+	UPROPERTY(Transient)
+	TObjectPtr<UInputBehaviorSet> BehaviorSet;
+
+	/**
+	 * Available Drag Tools, access by using the index provided by GetToolIndex(EDragTool::Type)
+	 */
+	TArray<TSharedPtr<FDragTool>> DragTools;
+
+	/**
+	 * Which type of Drag tool is currently active
+	 */
+	int ActiveToolType = -1;
+};
diff --git a/Engine/Source/Editor/UnrealEd/Public/DragTool_FrustumSelect.h b/Engine/Source/Editor/UnrealEd/Public/DragTool_FrustumSelect.h
--- a/Engine/Source/Editor/UnrealEd/Public/DragTool_FrustumSelect.h
+++ b/Engine/Source/Editor/UnrealEd/Public/DragTool_FrustumSelect.h
@@ -1,6 +1,5 @@
 // Copyright Epic Games, Inc. All Rights Reserved.
 
-
 #pragma once
 
 #include "CoreMinimal.h"
@@ -43,13 +42,21 @@ public:
 	 * @param	InStart				Where the mouse was when the drag started.
 	 */
 	UNREALED_API virtual void StartDrag(FEditorViewportClient* InViewportClient, const FVector& InStart, const FVector2D& InStartScreen) override;
-	
+
 	/**
 	 * Ends a mouse drag behavior (the user has let go of the mouse button).
 	 */
 	UNREALED_API virtual void EndDrag() override;
 	UNREALED_API virtual void Render(const FSceneView* View,FCanvas* Canvas) override;
 
+	//~ Begin IClickDragBehaviorTarget
+	UNREALED_API virtual FInputRayHit CanBeginClickDragSequence(const FInputDeviceRay& InPressPos) override;
+	UNREALED_API virtual void OnClickPress(const FInputDeviceRay& InPressPos) override;
+	UNREALED_API virtual void OnClickDrag(const FInputDeviceRay& InDragPos) override;
+	UNREALED_API virtual void OnClickRelease(const FInputDeviceRay& InReleasePos) override;
+	UNREALED_API virtual void OnTerminateDragSequence() override;
+	//~ End IClickDragBehaviorTarget
+
 private:
 	/** 
 	 * Calculates a frustum to check actors against 
@@ -69,10 +76,10 @@ private:
 	 * @param bUseStrictSelection	true if the node must be entirely within the frustum
 	 */
 	UNREALED_API bool IntersectsFrustum( const UModel& InModel, int32 NodeIndex, const FConvexVolume& InFrustum, bool bUseStrictSelection ) const;
-	
+
 	/** Adds a hover effect to the passed in actor */
 	UNREALED_API void AddHoverEffect( AActor& InActor );
-	
+
 	/** Adds a hover effect to the passed in bsp surface */
 	UNREALED_API void AddHoverEffect( UModel& InModel, int32 SurfIndex );
 
diff --git a/Engine/Source/Editor/UnrealEd/Public/EditorDragTools.h b/Engine/Source/Editor/UnrealEd/Public/EditorDragTools.h
--- a/Engine/Source/Editor/UnrealEd/Public/EditorDragTools.h
+++ b/Engine/Source/Editor/UnrealEd/Public/EditorDragTools.h
@@ -3,19 +3,43 @@
 
 #pragma once
 
+#include "BaseBehaviors/BehaviorTargetInterfaces.h"
 #include "CoreMinimal.h"
+#include "Editor.h"
 
 class FCanvas;
 class FEditorModeTools;
 class FEditorViewportClient;
 class FPrimitiveDrawInterface;
 class FSceneView;
+class FUICommandInfo;
+
+/**
+ * ITF related helper functions
+ */
+namespace UE::DragTools
+{
+/**
+ * Returns true if ITF-based drag tools should be used.
+ * If an ITF version of a drag tool is not available yet, legacy will be used.
+ */
+UNREALED_API bool UseITFTools();
+
+/** Register/Activate available ITF-based drag tools for default ed mode ("Selection Mode") */
+void RegisterDragTools(FEditorViewportClient* InViewportClient);
+
+/** Deregisters/Deactivates ITF-based drag tools for default ed mode ("Selection Mode") */
+void UnregisterDragTools(FEditorViewportClient* InViewportClient);
+
+} // namespace UE::DragTools
 
 /**
  * The base class that all drag tools inherit from.
  * The drag tools implement special behaviors for the user clicking and dragging in a viewport.
+ *
+ * Inheriting from IClickDragBehaviorTarget so that Drag Tools logic can be shared between legacy and ITF usage of the same tools
  */
-class FDragTool
+class FDragTool : public IClickDragBehaviorTarget
 {
 public:
 	UNREALED_API FDragTool(FEditorModeTools* InModeTools);
@@ -55,6 +79,39 @@ public:
 	/** Does this drag tool need to have the mouse movement converted to the viewport orientation? */
 	bool bConvertDelta;
 
+	//~ Begin IModifierToggleBehaviorTarget
+	virtual void OnUpdateModifierState(int InModifierID, bool bInIsOn) override;
+	//~ End IModifierToggleBehaviorTarget
+
+	//~ Begin IClickDragBehaviorTarget
+	virtual FInputRayHit CanBeginClickDragSequence(const FInputDeviceRay& InPressPos) override
+	{
+		return FInputRayHit();
+	}
+
+	virtual void OnClickPress(const FInputDeviceRay& InPressPos) override
+	{
+	}
+
+	virtual void OnClickDrag(const FInputDeviceRay& InDragPos) override
+	{
+	}
+
+	virtual void OnClickRelease(const FInputDeviceRay& InReleasePos) override
+	{
+	}
+
+	virtual void OnTerminateDragSequence() override
+	{
+	}
+	//~ End IClickDragBehaviorTarget
+
+	bool IsActivationChordPressed(const FInputChord& InChord) const;
+
+	DECLARE_MULTICAST_DELEGATE(FOnToolStateChange);
+	FOnToolStateChange& OnActivateTool() { return OnToolActivatedDelegate; }
+	FOnToolStateChange& OnDeactivateTool() { return OnToolDeactivatedDelegate; }
+
 protected:
 	FEditorModeTools* ModeTools;
 
@@ -73,4 +130,7 @@ protected:
 	bool bMiddleMouseButtonDown;
 	/** True if we are dragging */
 	bool bIsDragging;
+
+	FOnToolStateChange OnToolActivatedDelegate;
+	FOnToolStateChange OnToolDeactivatedDelegate;
 };
diff --git a/Engine/Source/Editor/UnrealEd/Public/EditorViewportClient.h b/Engine/Source/Editor/UnrealEd/Public/EditorViewportClient.h
--- a/Engine/Source/Editor/UnrealEd/Public/EditorViewportClient.h
+++ b/Engine/Source/Editor/UnrealEd/Public/EditorViewportClient.h
@@ -36,6 +36,7 @@ struct FTypedElementHandle;
 class IAssetFactoryInterface;
 class SEditorViewport;
 class UActorFactory;
+class UDragToolsBehaviorSource;
 class UTypedElementViewportInteraction;
 enum class EViewStatusForScreenPercentage;
 struct FGizmoState;
@@ -2134,6 +2135,23 @@ private:
 	 * true during FEditorViewportClient::StopTracking. Used to guard against reentry.
 	 */
 	bool bIsTrackingBeingStopped;
+
+public:
+	/**
+	 * Initializes and registers ITF-based drag tools
+	 */
+	UNREALED_API void RegisterDragTools();
+
+	/**
+	 * Unregistering of ITF drag tools
+	 */
+	UNREALED_API void UnregisterDragTools();
+
+protected:
+	/**
+	 * Hosting and accessing Drag Tools
+	 */
+	TObjectPtr<UDragToolsBehaviorSource> DragToolsBehaviorSource;
 };
 
 
