[ContentBrowser] Second batch of update for the thumbnails based on feedback
- Removed extra border from DetailsPanel Thumbnail
- LeftAligned AssetView (TileView)
- Removed extra ClassIcon overlay from AnimationBlueprint assets
- Removed all old Overlays for the new CB 2.0
- Keep Highlight style when editing the thumbnail in EditMode
- Primitives change option removed from BP
- CB tooltip "more info" changed to match SDocumentationTooltip style
- Inverted color highlight text when searching
- Changed default icon for DragAndDrop widgets with a newer one

diff --git a/Engine/Plugins/Editor/ContentBrowser/ContentBrowserAssetDataSource/Source/ContentBrowserAssetDataSource/Private/ContentBrowserAssetDataSource.cpp b/Engine/Plugins/Editor/ContentBrowser/ContentBrowserAssetDataSource/Source/ContentBrowserAssetDataSource/Private/ContentBrowserAssetDataSource.cpp
--- a/Engine/Plugins/Editor/ContentBrowser/ContentBrowserAssetDataSource/Source/ContentBrowserAssetDataSource/Private/ContentBrowserAssetDataSource.cpp
+++ b/Engine/Plugins/Editor/ContentBrowser/ContentBrowserAssetDataSource/Source/ContentBrowserAssetDataSource/Private/ContentBrowserAssetDataSource.cpp
@@ -2836,6 +2836,13 @@ bool UContentBrowserAssetDataSource::CanCopyItem(const FContentBrowserItemData&
 		return false;
 	}
 
+	// Destination must not be self (folder)
+	const FName VirtualPath = InItem.GetVirtualPath();
+	if (InDestPath == VirtualPath)
+	{
+		ContentBrowserAssetData::SetOptionalErrorMessage(OutErrorMsg, LOCTEXT("CopyError_ItemInsideItself", "Can't copy folder inside itself"));
+		return false;
+	}
 	return true;
 }
 
@@ -2906,6 +2913,29 @@ bool UContentBrowserAssetDataSource::CanMoveItem(const FContentBrowserItemData&
 		return false;
 	}
 
+	// Destination must not be self (folder)
+	FName VirtualPath = InItem.GetVirtualPath();
+	if (InDestPath == VirtualPath)
+	{
+		ContentBrowserAssetData::SetOptionalErrorMessage(OutErrorMsg, LOCTEXT("MoveError_ItemInsideItself", "Can't move Item inside itself"));
+		return false;
+	}
+
+	// Cannot be moved to the same folder
+	FString VirtualPathAsString = VirtualPath.ToString();
+	int32 LastSlashIndex = INDEX_NONE;
+	VirtualPathAsString.FindLastChar('/', LastSlashIndex);
+	if (LastSlashIndex != INDEX_NONE)
+	{
+		VirtualPathAsString = VirtualPathAsString.Left(VirtualPathAsString.Len() - (VirtualPathAsString.Len() - LastSlashIndex));
+	}
+
+	if (InDestPath == FName(VirtualPathAsString))
+	{
+		ContentBrowserAssetData::SetOptionalErrorMessage(OutErrorMsg, LOCTEXT("MoveError_ItemInsideSamePath", "Can't move Item inside the same location"));
+		return false;
+	}
+
 	return true;
 }
 
diff --git a/Engine/Plugins/Editor/ContentBrowser/ContentBrowserFileDataSource/Source/ContentBrowserFileDataSource/Private/ContentBrowserFileDataSource.cpp b/Engine/Plugins/Editor/ContentBrowser/ContentBrowserFileDataSource/Source/ContentBrowserFileDataSource/Private/ContentBrowserFileDataSource.cpp
--- a/Engine/Plugins/Editor/ContentBrowser/ContentBrowserFileDataSource/Source/ContentBrowserFileDataSource/Private/ContentBrowserFileDataSource.cpp
+++ b/Engine/Plugins/Editor/ContentBrowser/ContentBrowserFileDataSource/Source/ContentBrowserFileDataSource/Private/ContentBrowserFileDataSource.cpp
@@ -1437,6 +1437,13 @@ bool UContentBrowserFileDataSource::CanCopyItem(const FContentBrowserItemData& I
 		return false;
 	}
 
+	// Destination must not be self (folder)
+	const FName VirtualPath = InItem.GetVirtualPath();
+	if (InDestPath == VirtualPath)
+	{
+		ContentBrowserFileData::SetOptionalErrorMessage(OutErrorMsg, LOCTEXT("CopyError_ItemInsideItself", "Can't copy folder inside itself"));
+		return false;
+	}
 	return true;
 }
 
@@ -1504,6 +1511,30 @@ bool UContentBrowserFileDataSource::CanMoveItem(const FContentBrowserItemData& I
 		return false;
 	}
 
+	
+	// Destination must not be self (folder)
+	FName VirtualPath = InItem.GetVirtualPath();
+	if (InDestPath == VirtualPath)
+	{
+		ContentBrowserFileData::SetOptionalErrorMessage(OutErrorMsg, LOCTEXT("MoveError_ItemInsideItself", "Can't move Item inside itself"));
+		return false;
+	}
+
+	// Cannot be moved to the same folder
+	FString VirtualPathAsString = VirtualPath.ToString();
+	int32 LastSlashIndex = INDEX_NONE;
+	VirtualPathAsString.FindLastChar('/', LastSlashIndex);
+	if (LastSlashIndex != INDEX_NONE)
+	{
+		VirtualPathAsString = VirtualPathAsString.Left(VirtualPathAsString.Len() - (VirtualPathAsString.Len() - LastSlashIndex));
+	}
+
+	if (InDestPath == FName(VirtualPathAsString))
+	{
+		ContentBrowserFileData::SetOptionalErrorMessage(OutErrorMsg, LOCTEXT("MoveError_ItemInsideSamePath", "Can't move Item inside the same location"));
+		return false;
+	}
+
 	return true;
 }
 
diff --git a/Engine/Plugins/Editor/EngineAssetDefinitions/Source/Private/Animation/AssetDefinition_AnimBlueprint.cpp b/Engine/Plugins/Editor/EngineAssetDefinitions/Source/Private/Animation/AssetDefinition_AnimBlueprint.cpp
--- a/Engine/Plugins/Editor/EngineAssetDefinitions/Source/Private/Animation/AssetDefinition_AnimBlueprint.cpp
+++ b/Engine/Plugins/Editor/EngineAssetDefinitions/Source/Private/Animation/AssetDefinition_AnimBlueprint.cpp
@@ -107,22 +107,6 @@ UFactory* UAssetDefinition_AnimBlueprint::GetFactoryForBlueprintType(UBlueprint*
 	}
 }
 
-TSharedPtr<SWidget> UAssetDefinition_AnimBlueprint::GetThumbnailOverlay(const FAssetData& AssetData) const
-{
-	const FSlateBrush* Icon = FSlateIconFinder::FindIconBrushForClass(UAnimBlueprint::StaticClass());
-
-	return SNew(SBorder)
-		.BorderImage(FAppStyle::GetNoBrush())
-		.Visibility(EVisibility::HitTestInvisible)
-		.Padding(FMargin(0.0f, 0.0f, 0.0f, 3.0f))
-		.HAlign(HAlign_Right)
-		.VAlign(VAlign_Bottom)
-		[
-			SNew(SImage)
-			.Image(Icon)
-		];
-}
-
 EAssetCommandResult UAssetDefinition_AnimBlueprint::PerformAssetDiff(const FAssetDiffArgs& DiffArgs) const
 {
 	const UBlueprint* OldBlueprint = Cast<UBlueprint>(DiffArgs.OldAsset);
diff --git a/Engine/Plugins/Editor/EngineAssetDefinitions/Source/Private/Animation/AssetDefinition_AnimBlueprint.h b/Engine/Plugins/Editor/EngineAssetDefinitions/Source/Private/Animation/AssetDefinition_AnimBlueprint.h
--- a/Engine/Plugins/Editor/EngineAssetDefinitions/Source/Private/Animation/AssetDefinition_AnimBlueprint.h
+++ b/Engine/Plugins/Editor/EngineAssetDefinitions/Source/Private/Animation/AssetDefinition_AnimBlueprint.h
@@ -29,7 +29,6 @@ public:
 	}
 	
 	virtual EAssetCommandResult PerformAssetDiff(const FAssetDiffArgs& DiffArgs) const override;
-	virtual TSharedPtr<SWidget> GetThumbnailOverlay(const FAssetData& AssetData) const override;
 
 	virtual void BuildFilters(TArray<FAssetFilterData>& OutFilters) const override;
 	
diff --git a/Engine/Source/Editor/ContentBrowser/Private/AssetViewWidgets.cpp b/Engine/Source/Editor/ContentBrowser/Private/AssetViewWidgets.cpp
--- a/Engine/Source/Editor/ContentBrowser/Private/AssetViewWidgets.cpp
+++ b/Engine/Source/Editor/ContentBrowser/Private/AssetViewWidgets.cpp
@@ -2084,7 +2084,6 @@ void SAssetTileItem::Construct( const FArguments& InArgs )
 
 	InitializeAssetNameHeights();
 
-	TSharedPtr<SWidget> Thumbnail;
 	if ( AssetItem.IsValid() && AssetThumbnail.IsValid() )
 	{
 		FAssetThumbnailConfig ThumbnailConfig;
@@ -2096,7 +2095,6 @@ void SAssetTileItem::Construct( const FArguments& InArgs )
 		ThumbnailConfig.ThumbnailLabel = InArgs._ThumbnailLabel;
 		ThumbnailConfig.HighlightedText = InArgs._HighlightText;
 		ThumbnailConfig.HintColorAndOpacity = InArgs._ThumbnailHintColorAndOpacity;
-		ThumbnailConfig.Padding= FMargin(2.0f);
 		ThumbnailConfig.GenericThumbnailSize = MakeAttributeSP(this, &SAssetTileItem::GetGenericThumbnailSize);
 
 		if (UE::Editor::ContentBrowser::IsNewStyleEnabled())
@@ -2221,7 +2219,6 @@ void SAssetTileItem::Construct( const FArguments& InArgs )
 						.FillHeight(1.f)
 						[
 							SNew(SBox)
-							.Visibility(this, &SAssetTileItem::GetNameAreaVisibility)
 							.Padding(FMargin(NameAreaBoxLeftRightBotPadding, NameAreaBoxTopPadding, NameAreaBoxLeftRightBotPadding, NameAreaBoxLeftRightBotPadding))
 							.Visibility(this, &SAssetTileItem::GetNameAreaVisibility)
 							[
@@ -2417,7 +2414,7 @@ void SAssetTileItem::OnMouseLeave(const FPointerEvent& MouseEvent)
 const FSlateBrush* SAssetTileItem::GetBorderImage() const
 {
 	const bool bIsSelected = IsSelected.IsBound() ? IsSelected.Execute() : false;
-	const bool bIsHoveredOrDraggedOver = IsHovered() || bDraggedOver;
+	const bool bIsHoveredOrDraggedOver = IsHovered() || bDraggedOver || (Thumbnail.IsValid() && Thumbnail->IsHovered());
 	if (bIsSelected && bIsHoveredOrDraggedOver)
 	{
 		static const FName SelectedHover("ContentBrowser.AssetTileItem.SelectedHoverBorder");
@@ -2498,7 +2495,7 @@ EVisibility SAssetTileItem::GetAssetClassLabelVisibility() const
 FSlateColor SAssetTileItem::GetAssetClassLabelTextColor() const
 {
 	const bool bIsSelected = IsSelected.IsBound() ? IsSelected.Execute() : false;
-	const bool bIsHoveredOrDraggedOver = IsHovered() || bDraggedOver;
+	const bool bIsHoveredOrDraggedOver = IsHovered() || bDraggedOver || (Thumbnail.IsValid() && Thumbnail->IsHovered());
 	if (bIsSelected || bIsHoveredOrDraggedOver)
 	{
 		return FStyleColors::White;
@@ -2588,7 +2585,7 @@ const FSlateBrush* SAssetTileItem::GetNameAreaBackgroundImage() const
 	const FName Normal = UE::Editor::ContentBrowser::IsNewStyleEnabled() ? "ContentBrowser.AssetTileItem.AssetContent" : "ContentBrowser.AssetTileItem.NameAreaBackground";
 
 	const bool bIsSelected = IsSelected.IsBound() ? IsSelected.Execute() : false;
-	const bool bIsHoveredOrDraggedOver = IsHovered() || bDraggedOver;
+	const bool bIsHoveredOrDraggedOver = IsHovered() || bDraggedOver || (Thumbnail.IsValid() && Thumbnail->IsHovered());
 	if (bIsSelected && bIsHoveredOrDraggedOver)
 	{
 		return FAppStyle::Get().GetBrush(SelectedHover);
@@ -2611,7 +2608,7 @@ const FSlateBrush* SAssetTileItem::GetNameAreaBackgroundImage() const
 const FSlateBrush* SAssetTileItem::GetAssetAreaOverlayBackgroundImage() const
 {
 	const bool bIsSelected = IsSelected.IsBound() ? IsSelected.Execute() : false;
-	const bool bIsHoveredOrDraggedOver = IsHovered() || bDraggedOver;
+	const bool bIsHoveredOrDraggedOver = IsHovered() || bDraggedOver || (Thumbnail.IsValid() && Thumbnail->IsHovered());
 	if (bIsSelected && bIsHoveredOrDraggedOver)
 	{
 		static const FLazyName SelectedHover("ContentBrowser.AssetTileItem.AssetBorderSelectedHoverBackground");
@@ -2639,7 +2636,7 @@ const FSlateBrush* SAssetTileItem::GetAssetAreaOverlayBackgroundImage() const
 FSlateColor SAssetTileItem::GetChipBackgroundColor() const
 {
 	const bool bIsSelected = IsSelected.IsBound() ? IsSelected.Execute() : false;
-	const bool bIsHoveredOrDraggedOver = IsHovered() || bDraggedOver;
+	const bool bIsHoveredOrDraggedOver = IsHovered() || bDraggedOver || (Thumbnail.IsValid() && Thumbnail->IsHovered());
 	if (bIsSelected && bIsHoveredOrDraggedOver)
 	{
 		return FStyleColors::PrimaryHover;
@@ -2662,13 +2659,13 @@ FSlateColor SAssetTileItem::GetChipBackgroundColor() const
 FSlateColor SAssetTileItem::GetNameAreaTextColor() const
 {
 	const bool bIsSelected = IsSelected.IsBound() ? IsSelected.Execute() : false;
-	const bool bIsHoveredOrDraggedOver = IsHovered() || bDraggedOver;
+	const bool bIsHoveredOrDraggedOver = IsHovered() || bDraggedOver || (Thumbnail.IsValid() && Thumbnail->IsHovered());
 	if (bIsSelected || bIsHoveredOrDraggedOver)
 	{
 		return FStyleColors::White;
 	}
 
-	return UE::Editor::ContentBrowser::IsNewStyleEnabled() ? FStyleColors::Foreground : FSlateColor::UseForeground();
+	return FSlateColor::UseForeground();
 }
 
 FOptionalSize SAssetTileItem::GetNameAreaMaxDesiredHeight() const
diff --git a/Engine/Source/Editor/ContentBrowser/Private/SAssetView.cpp b/Engine/Source/Editor/ContentBrowser/Private/SAssetView.cpp
--- a/Engine/Source/Editor/ContentBrowser/Private/SAssetView.cpp
+++ b/Engine/Source/Editor/ContentBrowser/Private/SAssetView.cpp
@@ -2708,6 +2708,7 @@ TSharedRef<SAssetTileView> SAssetView::CreateTileView()
 	return SNew(SAssetTileView)
 		.SelectionMode( SelectionMode )
 		.ListItemsSource(&FilteredAssetItems)
+		.ItemAlignment(EListItemAlignment::LeftAligned)
 		.OnGenerateTile(this, &SAssetView::MakeTileViewWidget)
 		.OnItemToString_Debug_Static(&FAssetViewItem::ItemToString_Debug)
 		.OnItemScrolledIntoView(this, &SAssetView::ItemScrolledIntoView)
diff --git a/Engine/Source/Editor/ContentBrowser/Public/AssetViewWidgets.h b/Engine/Source/Editor/ContentBrowser/Public/AssetViewWidgets.h
--- a/Engine/Source/Editor/ContentBrowser/Public/AssetViewWidgets.h
+++ b/Engine/Source/Editor/ContentBrowser/Public/AssetViewWidgets.h
@@ -758,6 +758,9 @@ protected:
 	FOptionalSize GetSourceControlIconSize() const;
 
 private:
+	/** Thumbnail Widget or Folder image holder */
+	TSharedPtr<SWidget> Thumbnail;
+
 	/** If false, the tooltip will not be displayed */
 	bool bShowType;
 
diff --git a/Engine/Source/Editor/EditorStyle/Private/StarshipStyle.cpp b/Engine/Source/Editor/EditorStyle/Private/StarshipStyle.cpp
--- a/Engine/Source/Editor/EditorStyle/Private/StarshipStyle.cpp
+++ b/Engine/Source/Editor/EditorStyle/Private/StarshipStyle.cpp
@@ -1248,18 +1248,27 @@ void FStarshipEditorStyle::FStyle::SetupGeneralStyles()
 		Set( "AssetThumbnail.AssetTileViewWhiteChipBorder", new IMAGE_BRUSH_SVG("Starship/ContentBrowser/CornerChipOuter2_20", FVector2D(20, 20)));
 		Set( "AssetThumbnail.AssetTileViewWhiteChip", new IMAGE_BRUSH_SVG("Starship/ContentBrowser/CornerChipInner2_20", FVector2D(20, 20)));
 		Set( "AssetThumbnail.AssetTileViewWhiteChipSharp", new IMAGE_BRUSH_SVG("Starship/ContentBrowser/CornerChipSharp_20", FVector2D(20, 20)));
+		// TODO: Update with correct ratio icons for the small one once ready
+		Set( "AssetThumbnail.AssetTileViewWhiteChipBorderSmall", new IMAGE_BRUSH_SVG("Starship/ContentBrowser/CornerChipOuter2_20", FVector2D(20, 20)));
+		Set( "AssetThumbnail.AssetTileViewWhiteChipSmall", new IMAGE_BRUSH_SVG("Starship/ContentBrowser/CornerChipInner2_20", FVector2D(20, 20)));
 		Set( "AssetThumbnail.AssetBorder", new FSlateRoundedBoxBrush(FStyleColors::Transparent, FVector4(4.0f, 4.0f, 4.0f, 4.0f), FStyleColors::Secondary, 2.f));
+		Set( "AssetThumbnail.AssetBorderSmall", new FSlateRoundedBoxBrush(FStyleColors::Transparent, FVector4(4.0f, 4.0f, 4.0f, 4.0f), FStyleColors::Secondary, 1.f));
 		FLinearColor OverlayColorAssetStatusOverlay = FStyleColors::Panel.GetSpecifiedColor();
 		OverlayColorAssetStatusOverlay.A = 0.75f;
 		Set("AssetThumbnail.AssetThumbnailStatusBar", new FSlateRoundedBoxBrush(OverlayColorAssetStatusOverlay, 2.f));
 		FLinearColor OverlayColorAssetThumbnailOverlay = FStyleColors::Panel.GetSpecifiedColor();
 		OverlayColorAssetThumbnailOverlay.A = 0.75f;
 		Set("AssetThumbnail.AssetThumbnailBar", new FSlateRoundedBoxBrush(OverlayColorAssetThumbnailOverlay, 4.f));
-		Set("AssetThumbnail.ToolTip.CommandBorder", new FSlateRoundedBoxBrush(FStyleColors::Transparent, 4.f, FStyleColors::White, 1.f));
+		Set("AssetThumbnail.ToolTip.CommandBorder", new FSlateRoundedBoxBrush(FStyleColors::Transparent, 4.f, FStyleColors::Foreground, 1.f));
 		Set("AssetThumbnail.ToolTip.DarkCommandBorder", new FSlateRoundedBoxBrush(FStyleColors::Transparent, 4.f, FStyleColors::Black, 1.f));
 		Set("AssetThumbnail.ToolTip.ForegroundCommandBorder", new FSlateRoundedBoxBrush(FStyleColors::Transparent, 4.f, FStyleColors::Foreground, 1.f));
 		Set("AssetThumbnail.Tooltip.Border", new FSlateRoundedBoxBrush(FStyleColors::Secondary, 0.f, COLOR("#484848FF"), 2.f));
 
+		const FTextBlockStyle TooltipTextSubdued = FTextBlockStyle(NormalText)
+			.SetFont(DEFAULT_FONT("Regular", 8))
+			.SetColorAndOpacity(FStyleColors::Foreground);
+		Set("AssetThumbnail.Tooltip.MoreInfoText", TooltipTextSubdued);
+
 		const FButtonStyle EditModePrimitives = FButtonStyle()
 			.SetNormal(FSlateRoundedBoxBrush(FStyleColors::Secondary, 4.f, COLOR("#121212FF"), 1.f))
 			.SetHovered(FSlateRoundedBoxBrush(FStyleColors::Hover, 4.f, COLOR("#121212FF"), 1.f))
diff --git a/Engine/Source/Editor/PropertyEditor/Private/UserInterface/PropertyEditor/SPropertyEditorAsset.cpp b/Engine/Source/Editor/PropertyEditor/Private/UserInterface/PropertyEditor/SPropertyEditorAsset.cpp
--- a/Engine/Source/Editor/PropertyEditor/Private/UserInterface/PropertyEditor/SPropertyEditorAsset.cpp
+++ b/Engine/Source/Editor/PropertyEditor/Private/UserInterface/PropertyEditor/SPropertyEditorAsset.cpp
@@ -46,6 +46,23 @@
 
 #define LOCTEXT_NAMESPACE "PropertyEditor"
 
+namespace UE::Editor::ContentBrowser
+{
+	static bool IsNewStyleEnabled()
+	{
+		static bool bIsNewStyleEnabled = false;
+		UE_CALL_ONCE([&]()
+		{
+			if (const IConsoleVariable* CVar = IConsoleManager::Get().FindConsoleVariable(TEXT("ContentBrowser.EnableNewStyle")))
+			{
+				ensureAlwaysMsgf(!EnumHasAnyFlags(CVar->GetFlags(), ECVF_Default), TEXT("The CVar should have already been set from commandline, @see: UnrealEdGlobals.cpp, UE::Editor::ContentBrowser::EnableContentBrowserNewStyleCVarRegistration."));
+				bIsNewStyleEnabled = CVar->GetBool();
+			}
+		});
+
+		return bIsNewStyleEnabled;
+	}
+}
 
 // Helper to retrieve the correct property that has the applicable metadata.
 static const FProperty* GetActualMetadataProperty(const FProperty* Property)
@@ -589,24 +606,34 @@ void SPropertyEditorAsset::Construct(const FArguments& InArgs, const TSharedPtr<
 
 		AssetThumbnail = MakeShareable( new FAssetThumbnail( Value.AssetData, InArgs._ThumbnailSize.X, InArgs._ThumbnailSize.Y, InArgs._ThumbnailPool ) );
 
+		float ThumbnailPadding = 1.f;
 		FAssetThumbnailConfig AssetThumbnailConfig;
-		TSharedPtr<IAssetTypeActions> AssetTypeActions;
-		if (ObjectClass != nullptr)
+		if (UE::Editor::ContentBrowser::IsNewStyleEnabled())
 		{
-			const UClass* EffectiveClass = ObjectClass;
-			if (EffectiveClass->GetPathName() != Value.AssetData.AssetClassPath.ToString())
+			ThumbnailPadding = 0.f;
+			AssetThumbnailConfig.bSmallAssetChipBorder = true;
+			AssetThumbnailConfig.AssetChipBorderImageOverride = TAttribute<const FSlateBrush*>::CreateSP(this, &SPropertyEditorAsset::GetThumbnailBorder);
+		}
+		else
+		{
+			TSharedPtr<IAssetTypeActions> AssetTypeActions;
+			if (ObjectClass != nullptr)
 			{
-				if (UClass* AssetDataClass = FindObject<UClass>(Value.AssetData.AssetClassPath); AssetDataClass && AssetDataClass->IsChildOf(EffectiveClass))
+				const UClass* EffectiveClass = ObjectClass;
+				if (EffectiveClass->GetPathName() != Value.AssetData.AssetClassPath.ToString())
 				{
-					EffectiveClass = AssetDataClass;
+					if (UClass* AssetDataClass = FindObject<UClass>(Value.AssetData.AssetClassPath); AssetDataClass && AssetDataClass->IsChildOf(EffectiveClass))
+					{
+						EffectiveClass = AssetDataClass;
+					}
 				}
-			}
-			FAssetToolsModule& AssetToolsModule = FModuleManager::LoadModuleChecked<FAssetToolsModule>(TEXT("AssetTools"));
-			AssetTypeActions = AssetToolsModule.Get().GetAssetTypeActionsForClass(EffectiveClass).Pin();
+				FAssetToolsModule& AssetToolsModule = FModuleManager::LoadModuleChecked<FAssetToolsModule>(TEXT("AssetTools"));
+				AssetTypeActions = AssetToolsModule.Get().GetAssetTypeActionsForClass(EffectiveClass).Pin();
 
-			if (AssetTypeActions.IsValid())
-			{
-				AssetThumbnailConfig.AssetTypeColorOverride = AssetTypeActions->GetTypeColor();
+				if (AssetTypeActions.IsValid())
+				{
+					AssetThumbnailConfig.AssetTypeColorOverride = AssetTypeActions->GetTypeColor();
+				}
 			}
 		}
 
@@ -622,7 +649,7 @@ void SPropertyEditorAsset::Construct(const FArguments& InArgs, const TSharedPtr<
 			[
 				SNew(SOverlay)
 				+SOverlay::Slot()
-				.Padding(1.0f)
+				.Padding(ThumbnailPadding)
 				[
 					SAssignNew(ThumbnailBorder, SBorder)
 					.Padding(0)
@@ -642,7 +669,7 @@ void SPropertyEditorAsset::Construct(const FArguments& InArgs, const TSharedPtr<
 				[
 					SNew(SImage)
 					.Image(this, &SPropertyEditorAsset::GetThumbnailBorder)
-					.Visibility(EVisibility::SelfHitTestInvisible)
+					.Visibility(UE::Editor::ContentBrowser::IsNewStyleEnabled() ? EVisibility::Collapsed : EVisibility::SelfHitTestInvisible)
 				]
 			]
 		];
diff --git a/Engine/Source/Editor/UnrealEd/Private/AssetThumbnail.cpp b/Engine/Source/Editor/UnrealEd/Private/AssetThumbnail.cpp
--- a/Engine/Source/Editor/UnrealEd/Private/AssetThumbnail.cpp
+++ b/Engine/Source/Editor/UnrealEd/Private/AssetThumbnail.cpp
@@ -118,6 +118,7 @@ public:
 		SLATE_ARGUMENT(bool, AllowAssetStatusThumbnailOverlay)
 		SLATE_ATTRIBUTE(TSharedPtr<IToolTip>, TooltipOverride)
 		SLATE_ARGUMENT(bool, ShowAssetChip)
+		SLATE_ARGUMENT(bool, SmallAssetChipBorder)
 		SLATE_ATTRIBUTE(const FSlateBrush*, AssetChipBorderImageOverride)
 		SLATE_ARGUMENT(bool, CanDisplayEditModePrimitiveTools)
 		SLATE_ATTRIBUTE(EVisibility, IsEditModeVisible)
@@ -263,7 +264,7 @@ public:
 			];
 		}
 
-		if( ThumbnailClass.Get() && bIsClassType )
+		if (ThumbnailClass.Get() && bIsClassType && !UE::Editor::ContentBrowser::IsNewStyleEnabled())
 		{
 			OverlayWidget->AddSlot()
 			.VAlign(VAlign_Bottom)
@@ -309,39 +310,49 @@ public:
 			{
 				// The asset color chip
 				LayeredImage = SNew(SLayeredImage)
-					.Image(FAppStyle::GetBrush(Style, ".AssetTileViewWhiteChipBorder"))
+					.Image(FAppStyle::GetBrush(Style, InArgs._SmallAssetChipBorder ? ".AssetTileViewWhiteChipBorderSmall" : ".AssetTileViewWhiteChipBorder"))
 					.ColorAndOpacity(FStyleColors::Secondary);
-				LayeredImage->AddLayer(FAppStyle::GetBrush(Style, ".AssetTileViewWhiteChip"), AssetColor);
+				LayeredImage->AddLayer(FAppStyle::GetBrush(Style, InArgs._SmallAssetChipBorder ? ".AssetTileViewWhiteChipSmall" : ".AssetTileViewWhiteChip"), AssetColor);
+
+				// Need this to be negative for the Overlay so that it overlaps with the actual border which will keep the ratio correct and won't let the thumbnail over the limits.
+				const FMargin SmallAssetChipMarginOverlay(-1.f);
+				// Asset chip margin
+				const FMargin AssetChipMargin(0.f);
+				const FMargin SmallAssetChipMargin(1.f);
 
 				OverlayWidget->AddSlot()
 					.HAlign(HAlign_Right)
 					.VAlign(VAlign_Top)
+					.Padding(InArgs._SmallAssetChipBorder ? SmallAssetChipMarginOverlay : -2.f)
 					[
 						LayeredImage.ToSharedRef()
 					];
 
-				const TSharedRef<SImage> OverlayBorder = SNew(SImage)
-					.Image(FAppStyle::GetBrush(Style, ".AssetBorder"))
+				// Image overlay to match the Hovered/Selected status, by design only part of the chip should be highlighted
+				const TSharedRef<SImage> OverlayImage = SNew(SImage)
+					.Image(FAppStyle::GetBrush(Style, InArgs._SmallAssetChipBorder ? ".AssetBorderSmall" : ".AssetBorder"))
 					.Visibility(EVisibility::HitTestInvisible);
 
 				OverlayWidget->AddSlot()
+				.Padding(InArgs._SmallAssetChipBorder ? SmallAssetChipMarginOverlay : -2.f)
 				[
-					OverlayBorder
+					OverlayImage
 				];
 
+				// Border used to contain the Thumbnail in the ratio
+				const TSharedRef<SBorder> OverlayBorder = SNew(SBorder)
+					.Padding(InArgs._SmallAssetChipBorder ? SmallAssetChipMargin : 2.f)
+					.BorderImage(FAppStyle::GetBrush(Style, InArgs._SmallAssetChipBorder ? ".AssetBorderSmall" : ".AssetBorder"));
+
+
 				if (AssetChipBorderImageOverride.IsSet())
 				{
-					const TSharedRef<SBorder> ThumbnailWidgetBorder = SNew(SBorder)
-						.Padding(0)
-						.BorderImage(FAppStyle::GetBrush(Style, ".AssetBorder"))
-						[
-							OverlayWidget
-						];
-
-					OverlayBorder->SetImage(AssetChipBorderImageOverride);
-					ThumbnailWidgetBorder->SetBorderImage(AssetChipBorderImageOverride);
-					ContentWidget = ThumbnailWidgetBorder;
+					OverlayImage->SetImage(AssetChipBorderImageOverride);
+					OverlayBorder->SetBorderImage(AssetChipBorderImageOverride);
 				}
+
+				OverlayBorder->SetContent(OverlayWidget);
+				ContentWidget = OverlayBorder;
 			}
 		}
 		else
@@ -365,7 +376,7 @@ public:
 			{
 				OverlayWidget->AddSlot()
 				[
-					SNew(SAssetThumbnailEditModeTools, AssetThumbnail)
+					SAssignNew(AssetThumbnailEditMode, SAssetThumbnailEditModeTools, AssetThumbnail)
 					.SmallView(InArgs._CanDisplayEditModePrimitiveTools)
 					.Visibility(EditModeVisibility)
 				];
@@ -436,92 +447,95 @@ public:
 		if( InArgs._AllowAssetSpecificThumbnailOverlay && AssetTypeActions.IsValid() )
 		{
 			FAssetActionThumbnailOverlayInfo OutThumbnailInfo;
-			if (UE::Editor::ContentBrowser::IsNewStyleEnabled()
-				&& AssetDefinition && AssetDefinition->GetThumbnailActionOverlay(AssetData, OutThumbnailInfo))
+			// Skip ALL older overlay if the new style is enabled
+			if (UE::Editor::ContentBrowser::IsNewStyleEnabled())
 			{
-				constexpr int32 OverlayZOrder = 1;
-				if (OutThumbnailInfo.ActionImageWidget.IsValid())
+				if (AssetDefinition && AssetDefinition->GetThumbnailActionOverlay(AssetData, OutThumbnailInfo))
 				{
-					constexpr float PaddingFromTopLeftBorder = 4.f;
-					constexpr float TopLeftImageSize = 20.f;
-					constexpr float BorderPaddingToActionImageWidget = 4.f;
-
-					OverlayWidget->AddSlot()
-					.ZOrder(OverlayZOrder)
-					.VAlign(VAlign_Top)
-					.HAlign(HAlign_Left)
-					.Padding(PaddingFromTopLeftBorder, PaddingFromTopLeftBorder, 0.f, 0.f)
-					[
-						SNew(SBox)
-						.WidthOverride(TopLeftImageSize)
-						.HeightOverride(TopLeftImageSize)
+					constexpr int32 OverlayZOrder = 1;
+					if (OutThumbnailInfo.ActionImageWidget.IsValid())
+					{
+						constexpr float PaddingFromTopLeftBorder = 4.f;
+						constexpr float TopLeftImageSize = 20.f;
+						constexpr float BorderPaddingToActionImageWidget = 4.f;
+
+						OverlayWidget->AddSlot()
+						.ZOrder(OverlayZOrder)
+						.VAlign(VAlign_Top)
+						.HAlign(HAlign_Left)
+						.Padding(PaddingFromTopLeftBorder, PaddingFromTopLeftBorder, 0.f, 0.f)
 						[
-							SNew(SBorder)
-							.BorderImage(FAppStyle::GetBrush(Style, ".AssetThumbnailBar"))
-							.Padding(BorderPaddingToActionImageWidget)
-							.Visibility_Lambda([this] () { return IsHovered() || (EditModeVisibility.IsSet() && EditModeVisibility.Get().IsVisible()) ? EVisibility::Collapsed : EVisibility::Visible; })
+							SNew(SBox)
+							.WidthOverride(TopLeftImageSize)
+							.HeightOverride(TopLeftImageSize)
 							[
-								OutThumbnailInfo.ActionImageWidget.ToSharedRef()
+								SNew(SBorder)
+								.BorderImage(FAppStyle::GetBrush(Style, ".AssetThumbnailBar"))
+								.Padding(BorderPaddingToActionImageWidget)
+								.Visibility_Lambda([this] () { return IsHovered() || (EditModeVisibility.IsSet() && EditModeVisibility.Get().IsVisible()) ? EVisibility::Collapsed : EVisibility::Visible; })
+								[
+									OutThumbnailInfo.ActionImageWidget.ToSharedRef()
+								]
 							]
-						]
-					];
-				}
+						];
+					}
 
-				PRAGMA_DISABLE_DEPRECATION_WARNINGS
-				if (OutThumbnailInfo.ActionButtonWidget.IsValid())
-				{
-					constexpr float CenterImageSize = 32.f;
+					PRAGMA_DISABLE_DEPRECATION_WARNINGS
+					if (OutThumbnailInfo.ActionButtonWidget.IsValid())
+					{
+						constexpr float CenterImageSize = 32.f;
 
-					OverlayWidget->AddSlot()
-					.ZOrder(OverlayZOrder)
-					.VAlign(VAlign_Center)
-					.HAlign(HAlign_Center)
-					[
-						SNew(SBorder)
-						.Padding(0.f)
-						.BorderImage(FStyleDefaults::GetNoBrush())
-						.Visibility(this, &SAssetThumbnail::GetActionButtonVisibility)
+						OverlayWidget->AddSlot()
+						.ZOrder(OverlayZOrder)
+						.VAlign(VAlign_Center)
+						.HAlign(HAlign_Center)
+						[
+							SNew(SBorder)
+							.Padding(0.f)
+							.BorderImage(FStyleDefaults::GetNoBrush())
+							.Visibility(this, &SAssetThumbnail::GetActionButtonVisibility)
+							[
+								SNew(SBox)
+								.WidthOverride(CenterImageSize)
+								.HeightOverride(CenterImageSize)
+								[
+									OutThumbnailInfo.ActionButtonWidget.ToSharedRef()
+								]
+							]
+						];
+					}
+					PRAGMA_ENABLE_DEPRECATION_WARNINGS
+					else
+					{
+						// Set the default style and padding for the Button
+						OutThumbnailInfo.ActionButtonArgs
+							.ButtonStyle(&FAppStyle::GetWidgetStyle<FButtonStyle>(Style, ".Action.Button"))
+							.ContentPadding(8.f);
+
+						TAttribute<EVisibility> ActionVisibility = TAttribute<EVisibility>::CreateSP(this, &SAssetThumbnail::GetActionButtonVisibility);
+						TSharedRef<SButton> ActionButton = SNew(SButton);
+						ActionButton->Construct(OutThumbnailInfo.ActionButtonArgs);
+						ActionButton->SetVisibility(ActionVisibility);
+						ActionButton->SetToolTipText(OutThumbnailInfo.ActionButtonArgs._ToolTipText);
+						ActionButton->SetToolTip(OutThumbnailInfo.ActionButtonArgs._ToolTip);
+						ActionButton->SetCursor(EMouseCursor::Default);
+						ActionButton->SetContent(OutThumbnailInfo.ActionImageWidget.ToSharedRef());
+
+						constexpr float CenterImageSize = 32.f;
+
+						OverlayWidget->AddSlot()
+						.ZOrder(OverlayZOrder)
+						.VAlign(VAlign_Center)
+						.HAlign(HAlign_Center)
 						[
 							SNew(SBox)
 							.WidthOverride(CenterImageSize)
 							.HeightOverride(CenterImageSize)
 							[
-								OutThumbnailInfo.ActionButtonWidget.ToSharedRef()
+								ActionButton
 							]
-						]
-					];
-				}
-				PRAGMA_ENABLE_DEPRECATION_WARNINGS
-				else
-				{
-					// Set the default style and padding for the Button
-					OutThumbnailInfo.ActionButtonArgs
-						.ButtonStyle(&FAppStyle::GetWidgetStyle<FButtonStyle>(Style, ".Action.Button"))
-						.ContentPadding(8.f);
-
-					TAttribute<EVisibility> ActionVisibility = TAttribute<EVisibility>::CreateSP(this, &SAssetThumbnail::GetActionButtonVisibility);
-					TSharedRef<SButton> ActionButton = SNew(SButton);
-					ActionButton->Construct(OutThumbnailInfo.ActionButtonArgs);
-					ActionButton->SetVisibility(ActionVisibility);
-					ActionButton->SetToolTipText(OutThumbnailInfo.ActionButtonArgs._ToolTipText);
-					ActionButton->SetToolTip(OutThumbnailInfo.ActionButtonArgs._ToolTip);
-					ActionButton->SetCursor(EMouseCursor::Default);
-					ActionButton->SetContent(OutThumbnailInfo.ActionImageWidget.ToSharedRef());
-
-					constexpr float CenterImageSize = 32.f;
-
-					OverlayWidget->AddSlot()
-                    .ZOrder(OverlayZOrder)
-                    .VAlign(VAlign_Center)
-                    .HAlign(HAlign_Center)
-                    [
-                    	SNew(SBox)
-						.WidthOverride(CenterImageSize)
-						.HeightOverride(CenterImageSize)
-						[
-							ActionButton
-						]
-                    ];
+						];
+					}
 				}
 			}
 			else
@@ -613,6 +627,12 @@ public:
 	{
 		SCompoundWidget::OnMouseLeave( MouseEvent );
 
+		if (AssetThumbnailEditMode.IsValid())
+		{
+			// Keep the Hovered state the same as the ThumbnailEdit mode during OnMouseLeave
+			SetHover(AssetThumbnailEditMode->IsHovered());
+		}
+
 		if ( bAllowRealTimeOnHovered )
 		{
 			AssetThumbnail->SetRealTime( false );
@@ -791,8 +811,9 @@ public:
 				PromptBox->SetPadding(FMargin(0.f, 2.f, 0.f, 0.f));
 				PromptBox->SetContent(
 					SNew(SRichTextBlock)
+					.TextStyle(&FAppStyle::GetWidgetStyle<FTextBlockStyle>(Style, ".Tooltip.MoreInfoText"))
 					.Justification(ETextJustify::Center)
-					.Text(NSLOCTEXT("AssetThumbnail", "MoreInfoTooltip", "Hold <WrappedCommand/> for more"))
+					.Text(NSLOCTEXT("AssetThumbnail", "MoreInfoTooltip", "Hold<WrappedCommand/>for more"))
 					+ SRichTextBlock::WidgetDecorator(TEXT("WrappedCommand"), this, &SAssetThumbnail::OnCreateWidgetDecoratorWidget)
 				);
 			}
@@ -829,17 +850,21 @@ public:
 private:
 	FSlateWidgetRun::FWidgetRunInfo OnCreateWidgetDecoratorWidget(const FTextRunInfo& InRunInfo, const ISlateStyle* InStyle) const
 	{
-		TSharedRef<SWidget> CtrlAltWidget = SNew(SBorder)
-			.VAlign(VAlign_Center)
-			.HAlign(HAlign_Center)
-			.BorderImage(FAppStyle::GetBrush(Style, ".ToolTip.CommandBorder"))
+		TSharedRef<SWidget> CtrlAltWidget = SNew(SBox)
+			.Padding(5.f, 0.f)
 			[
-				SNew(STextBlock)
+				SNew(SBorder)
+				.VAlign(VAlign_Center)
+				.BorderImage(FAppStyle::GetBrush(Style, ".ToolTip.CommandBorder"))
+				[
+					SNew(STextBlock)
+					.TextStyle(&FAppStyle::GetWidgetStyle<FTextBlockStyle>(Style, ".Tooltip.MoreInfoText"))
 #if PLATFORM_MAC
-				.Text(NSLOCTEXT("AssetThumbnail", "CommandOptionLabel", " Command + Option "))
+					.Text(NSLOCTEXT("AssetThumbnail", "CommandOptionLabel", " Command + Option "))
 #else
-				.Text(NSLOCTEXT("AssetThumbnail", "CtrlAltLabel", " Ctrl + Alt "))
+					.Text(NSLOCTEXT("AssetThumbnail", "CtrlAltLabel", " Ctrl + Alt "))
 #endif
+				]
 			];
 		const TSharedRef<FSlateFontMeasure> FontMeasure = FSlateApplication::Get().GetRenderer()->GetFontMeasureService();
 		const int16 Baseline = FontMeasure->GetBaseline(FStyleDefaults::GetFontInfo());
@@ -1406,6 +1431,7 @@ private:
 	}
 
 private:
+	TSharedPtr<SAssetThumbnailEditModeTools> AssetThumbnailEditMode;
 	TSharedPtr<STextBlock> GenericLabelTextBlock;
 	TSharedPtr<STextBlock> HintTextBlock;
 	TSharedPtr<SImage> GenericThumbnailImage;
@@ -1615,6 +1641,8 @@ TSharedRef<SWidget> FAssetThumbnail::MakeThumbnailWidget( const FAssetThumbnailC
 			.TooltipOverride(InConfig.TooltipOverride)
 			.AllowAssetStatusThumbnailOverlay(InConfig.bAllowAssetStatusThumbnailOverlay)
 			.ShowAssetChip(InConfig.bShowAssetChip)
+			.IsEditModeVisible(InConfig.IsEditModeVisible)
+			.SmallAssetChipBorder(InConfig.bSmallAssetChipBorder)
 			.AssetChipBorderImageOverride(InConfig.AssetChipBorderImageOverride)
 			.ColorStripOrientation(InConfig.ColorStripOrientation);
 	}
diff --git a/Engine/Source/Editor/UnrealEd/Private/DragAndDrop/AssetDragDropOp.cpp b/Engine/Source/Editor/UnrealEd/Private/DragAndDrop/AssetDragDropOp.cpp
--- a/Engine/Source/Editor/UnrealEd/Private/DragAndDrop/AssetDragDropOp.cpp
+++ b/Engine/Source/Editor/UnrealEd/Private/DragAndDrop/AssetDragDropOp.cpp
@@ -90,6 +90,7 @@ TSharedPtr<SWidget> FAssetDragDropOp::GetDefaultDecorator() const
 		{
 			FAssetThumbnailConfig AssetThumbnailConfig;
 			AssetThumbnailConfig.bShowAssetChip = true;
+			AssetThumbnailConfig.bSmallAssetChipBorder = true;
 			AssetThumbnailConfig.AssetChipBorderImageOverride = FAppStyle::GetBrush("ContentBrowser.ThumbnailDragDropBackground");
 			ThumbnailWidget = AssetThumbnail->MakeThumbnailWidget(AssetThumbnailConfig);
 		}
diff --git a/Engine/Source/Editor/UnrealEd/Private/SAssetThumbnailEditModeTools.cpp b/Engine/Source/Editor/UnrealEd/Private/SAssetThumbnailEditModeTools.cpp
--- a/Engine/Source/Editor/UnrealEd/Private/SAssetThumbnailEditModeTools.cpp
+++ b/Engine/Source/Editor/UnrealEd/Private/SAssetThumbnailEditModeTools.cpp
@@ -106,7 +106,7 @@ void SAssetThumbnailEditModeTools::Construct(const FArguments& InArgs, const TSh
 
 EVisibility SAssetThumbnailEditModeTools::GetPrimitiveToolsVisibility() const
 {
-	const bool bIsVisible = bInSmallView && (GetSceneThumbnailInfo() != nullptr);
+	const bool bIsVisible = bInSmallView && (GetSceneThumbnailInfo() != nullptr && GetSceneThumbnailInfoWithPrimitive() != nullptr);
 	return bIsVisible ? EVisibility::Visible : EVisibility::Collapsed;
 }
 
@@ -215,6 +215,8 @@ FReply SAssetThumbnailEditModeTools::OnMouseButtonDown(const FGeometry& MyGeomet
 				bModifiedThumbnailWhileDragging = false;
 				DragStartLocation = FIntPoint(FMath::TruncToInt32(MouseEvent.GetScreenSpacePosition().X), FMath::TruncToInt32(MouseEvent.GetScreenSpacePosition().Y));
 
+				// Keep the widget hovered so that the view can maintain a hovered state on the whole widget while dragging the mouse even outside its bounds
+				SetHover(true);
 				return FReply::Handled().CaptureMouse(AsShared()).UseHighPrecisionMouseMovement(AsShared()).PreventThrottling();
 			}
 		}
@@ -241,6 +243,9 @@ FReply SAssetThumbnailEditModeTools::OnMouseButtonUp(const FGeometry& MyGeometry
 			bModifiedThumbnailWhileDragging = false;
 		}
 
+		// Reset the Hover to being changed automatically
+		// since the mouse position is reset as well inside the thumbnail there will be no need to force it changing here
+		SetHover(TAttribute<bool>());
 		return FReply::Handled().ReleaseMouseCapture().SetMousePos(DragStartLocation);
 	}
 
diff --git a/Engine/Source/Editor/UnrealEd/Public/AssetThumbnail.h b/Engine/Source/Editor/UnrealEd/Public/AssetThumbnail.h
--- a/Engine/Source/Editor/UnrealEd/Public/AssetThumbnail.h
+++ b/Engine/Source/Editor/UnrealEd/Public/AssetThumbnail.h
@@ -69,6 +69,7 @@ struct FAssetThumbnailConfig
 		, AssetSystemInfoProvider(nullptr)
 		, bAllowAssetStatusThumbnailOverlay( false )
 		, bShowAssetChip(true)
+		, bSmallAssetChipBorder(false)
 		, AssetChipBorderImageOverride()
 		, bCanDisplayEditModePrimitiveTools(false)
 		, IsEditModeVisible()
@@ -92,6 +93,8 @@ struct FAssetThumbnailConfig
 	TAttribute<TSharedPtr<IToolTip>> TooltipOverride;
 	bool bAllowAssetStatusThumbnailOverlay;
 	bool bShowAssetChip;
+	bool bSmallAssetChipBorder;
+	/** Border override brush needs to have an OutlineWidth of 2 if normal size or 1 if small size  */
 	TAttribute<const FSlateBrush*> AssetChipBorderImageOverride;
 	bool bCanDisplayEditModePrimitiveTools;
 	TAttribute<EVisibility> IsEditModeVisible;
