Deprecate MemoryMaxUsedVirtual and MemoryMaxUsedPhysical and add a warning if there are settings for any of those 2 variables.

#rnx

diff --git a/Engine/Source/Editor/UnrealEd/Classes/CookOnTheSide/CookOnTheFlyServer.h b/Engine/Source/Editor/UnrealEd/Classes/CookOnTheSide/CookOnTheFlyServer.h
--- a/Engine/Source/Editor/UnrealEd/Classes/CookOnTheSide/CookOnTheFlyServer.h
+++ b/Engine/Source/Editor/UnrealEd/Classes/CookOnTheSide/CookOnTheFlyServer.h
@@ -352,8 +352,12 @@ private:
 	float CookProgressRetryBusyPeriodSeconds = 0.f;
 
 	// Memory Limits for when to Collect Garbage
+	UE_DEPRECATED(5.6, "MemoryMaxUsedVirtual is deprecated and will be removed in future version.")
 	uint64 MemoryMaxUsedVirtual;
+
+	UE_DEPRECATED(5.6, "MemoryMaxUsedPhysical is deprecated and will be removed in future version.")
 	uint64 MemoryMaxUsedPhysical;
+
 	uint64 MemoryMinFreeVirtual;
 	uint64 MemoryMinFreePhysical;
 	FGenericPlatformMemoryStats::EMemoryPressureStatus MemoryTriggerGCAtPressureLevel;
diff --git a/Engine/Source/Editor/UnrealEd/Private/CookOnTheFlyServer.cpp b/Engine/Source/Editor/UnrealEd/Private/CookOnTheFlyServer.cpp
--- a/Engine/Source/Editor/UnrealEd/Private/CookOnTheFlyServer.cpp
+++ b/Engine/Source/Editor/UnrealEd/Private/CookOnTheFlyServer.cpp
@@ -1928,6 +1928,7 @@ void UCookOnTheFlyServer::UpdateDisplay(UE::Cook::FTickStackData& StackData, boo
 
 FString UCookOnTheFlyServer::GetCookSettingsForMemoryLogText() const
 {
+	PRAGMA_DISABLE_DEPRECATION_WARNINGS
 	return FString::Printf(
 		TEXT("\n\tMemoryMaxUsedVirtual %dMiB")
 		TEXT("\n\tMemoryMaxUsedPhysical %dMiB")
@@ -1945,6 +1946,7 @@ FString UCookOnTheFlyServer::GetCookSettingsForMemoryLogText() const
 		(bUseSoftGC && SoftGCTimeFractionBudget > 0) ? *FString::Printf(TEXT(" (%.3g budget, %.0fs min period)"),
 			SoftGCTimeFractionBudget, SoftGCMinimumPeriodSeconds) : TEXT("")
 	);
+	PRAGMA_ENABLE_DEPRECATION_WARNINGS
 }
 
 namespace UE::Cook::Pollable
@@ -5837,16 +5839,39 @@ void FInitializeConfigSettings::LoadLocal(const FString& InOutputDirectoryOverri
 		}
 		return false;
 	};
+
+	PRAGMA_DISABLE_DEPRECATION_WARNINGS
 	MemoryMaxUsedVirtual = 0;
 	MemoryMaxUsedPhysical = 0;
+	PRAGMA_ENABLE_DEPRECATION_WARNINGS
+
 	MemoryMinFreeVirtual = 0;
 	MemoryMinFreePhysical = 0;
 	bUseSoftGC = false;
 	SoftGCStartNumerator = 5;
 	SoftGCDenominator = 10;
 
+	PRAGMA_DISABLE_DEPRECATION_WARNINGS
 	ReadMemorySetting(TEXT("MemoryMaxUsedVirtual"), MemoryMaxUsedVirtual);
 	ReadMemorySetting(TEXT("MemoryMaxUsedPhysical"), MemoryMaxUsedPhysical);
+
+	if (MemoryMaxUsedVirtual != 0)
+	{
+		UE_LOG(LogCook, Warning, TEXT("Setting MemoryMaxUsedVirtual will be deprecated in future version. Please remove it from the settings files."
+			"On systems with too little memory to load the minimum set required to cook some packages, frequent garbage collection will stall progress, "
+			"we now detect this stall and terminate the cook with an assertion rather than the previous behavior of MemoryMaxUsedVirtual, "
+			"which would terminate it with an OutOfMemory message."))
+	}
+
+	if (MemoryMaxUsedPhysical != 0)
+	{
+		UE_LOG(LogCook, Warning, TEXT("Setting MemoryMaxUsedPhysical will be deprecated in future version. Please remove it from the settings files."
+			"On systems with too little memory to load the minimum set required to cook some packages, frequent garbage collection will stall progress, "
+			"we now detect this stall and terminate the cook with an assertion rather than the previous behavior of MemoryMaxUsedPhysical, "
+			"which would terminate it with an OutOfMemory message."))
+	}
+	PRAGMA_ENABLE_DEPRECATION_WARNINGS
+
 	ReadMemorySetting(TEXT("MemoryMinFreeVirtual"), MemoryMinFreeVirtual);
 	ReadMemorySetting(TEXT("MemoryMinFreePhysical"), MemoryMinFreePhysical);
 	FString ConfigText(TEXT("None"));
diff --git a/Engine/Source/Editor/UnrealEd/Private/Cooker/CookDirector.cpp b/Engine/Source/Editor/UnrealEd/Private/Cooker/CookDirector.cpp
--- a/Engine/Source/Editor/UnrealEd/Private/Cooker/CookDirector.cpp
+++ b/Engine/Source/Editor/UnrealEd/Private/Cooker/CookDirector.cpp
@@ -1329,8 +1329,11 @@ void FCookDirector::ActivateMachineResourceReduction()
 
 	// When running a multiprocess cook, we remove the Memory triggers and trigger GC based solely on PressureLevel.
 	// But keep the Soft GC settings
+	PRAGMA_DISABLE_DEPRECATION_WARNINGS
 	COTFS.MemoryMaxUsedPhysical = 0;
 	COTFS.MemoryMaxUsedVirtual = 0;
+	PRAGMA_ENABLE_DEPRECATION_WARNINGS
+
 	COTFS.MemoryMinFreeVirtual = 0;
 	COTFS.MemoryMinFreePhysical = 0;
 	COTFS.MemoryTriggerGCAtPressureLevel = FGenericPlatformMemoryStats::EMemoryPressureStatus::Critical;
diff --git a/Engine/Source/Editor/UnrealEd/Private/Cooker/CookGarbageCollect.cpp b/Engine/Source/Editor/UnrealEd/Private/Cooker/CookGarbageCollect.cpp
--- a/Engine/Source/Editor/UnrealEd/Private/Cooker/CookGarbageCollect.cpp
+++ b/Engine/Source/Editor/UnrealEd/Private/Cooker/CookGarbageCollect.cpp
@@ -449,6 +449,7 @@ bool UCookOnTheFlyServer::PumpHasExceededMaxMemory(uint32& OutResultFlags)
 	// if MemoryMaxUsed is set, we won't GC until at least that much mem is used
 	// this can be useful if you demand that amount of memory as your min spec
 	bool bMaxUsedTriggered = false;
+	PRAGMA_DISABLE_DEPRECATION_WARNINGS
 	if (MemoryMaxUsedVirtual > 0 || MemoryMaxUsedPhysical > 0)
 	{
 		// check validity of trigger :
@@ -475,6 +476,7 @@ bool UCookOnTheFlyServer::PumpHasExceededMaxMemory(uint32& OutResultFlags)
 			bMaxUsedTriggered = true;
 		}
 	}
+	PRAGMA_ENABLE_DEPRECATION_WARNINGS
 
 	bool bPeriodicTriggered = false;
 	bool bPressureTriggered = false;
@@ -506,12 +508,15 @@ bool UCookOnTheFlyServer::PumpHasExceededMaxMemory(uint32& OutResultFlags)
 	if (bMinFreeTriggered || bMaxUsedTriggered)
 	{
 		const bool bOnlyTriggerIfBothMinFreeAndMaxUsedTrigger = true;
+
+		PRAGMA_DISABLE_DEPRECATION_WARNINGS
 		if (!bOnlyTriggerIfBothMinFreeAndMaxUsedTrigger ||
 			((bMinFreeTriggered || (MemoryMinFreeVirtual <= 0 && MemoryMinFreePhysical <= 0)) &&
 				(bMaxUsedTriggered || (MemoryMaxUsedVirtual <= 0 && MemoryMaxUsedPhysical <= 0))))
 		{
 			bTriggerGC = true;
 		}
+		PRAGMA_ENABLE_DEPRECATION_WARNINGS
 	}
 	if (bPressureTriggered)
 	{
diff --git a/Engine/Source/Editor/UnrealEd/Private/Cooker/CookTypes.cpp b/Engine/Source/Editor/UnrealEd/Private/Cooker/CookTypes.cpp
--- a/Engine/Source/Editor/UnrealEd/Private/Cooker/CookTypes.cpp
+++ b/Engine/Source/Editor/UnrealEd/Private/Cooker/CookTypes.cpp
@@ -814,8 +814,12 @@ void FInitializeConfigSettings::MoveOrCopy(SourceType&& Source, TargetType&& Tar
 	Target.PackagesPerGC = Source.PackagesPerGC;
 	Target.MemoryExpectedFreedToSpreadRatio = Source.MemoryExpectedFreedToSpreadRatio;
 	Target.IdleTimeToGC = Source.IdleTimeToGC;
+
+	PRAGMA_DISABLE_DEPRECATION_WARNINGS
 	Target.MemoryMaxUsedVirtual = Source.MemoryMaxUsedVirtual;
 	Target.MemoryMaxUsedPhysical = Source.MemoryMaxUsedPhysical;
+	PRAGMA_ENABLE_DEPRECATION_WARNINGS
+
 	Target.MemoryMinFreeVirtual = Source.MemoryMinFreeVirtual;
 	Target.MemoryMinFreePhysical = Source.MemoryMinFreePhysical;
 	Target.MemoryTriggerGCAtPressureLevel = Source.MemoryTriggerGCAtPressureLevel;
