Localization fixes:  The Blueprint Graph Editor will be localized for all Functions/Events/Actions again.

Localization is reintroduced in all Function look-up cases.  Sometime during the development of UE5.3, the Event Function DisplayName look-up were executing on the SKEL Objects rather than the real objects, causing issues in the look-up name resulting in inconsistency between Events and Functions.  That has now been fixed-up in addition to reintroducing localization so names should be consistent across all Blueprint Graph Editor areas.

1. Events Area uses Localization or DisplayNames
2. Functions Area uses Localization or DisplayNames
3. Overrides Menu uses Localization or DisplayNames
4. Graph Editor uses Localization or FriendlyNames
5. Actions Menu uses Localization or FriendlyNames

diff --git a/Engine/Source/Editor/BlueprintGraph/Private/EdGraphSchema_K2.cpp b/Engine/Source/Editor/BlueprintGraph/Private/EdGraphSchema_K2.cpp
--- a/Engine/Source/Editor/BlueprintGraph/Private/EdGraphSchema_K2.cpp
+++ b/Engine/Source/Editor/BlueprintGraph/Private/EdGraphSchema_K2.cpp
@@ -5679,21 +5679,25 @@ void UEdGraphSchema_K2::GetGraphDisplayInformation(const UEdGraph& Graph, /*out*
 		DisplayInfo.Notes.Add(TEXT("intermediate build product"));
 	}
 
-	if( GEditor && GetDefault<UEditorStyleSettings>()->bShowFriendlyNames )
+	// We disallow friendly names so that the Function name appears as entered (unless localized)
+	constexpr bool bAllowFriendlyNames = false;
+	if (GraphType == GT_Function && Function)
 	{
-		if (GraphType == GT_Function && Function)
+		DisplayInfo.DisplayName = ObjectTools::GetUserFacingFunctionName(Function, bAllowFriendlyNames);
+	}
+	else
+	{
+		// This is essentially GetUserFacingFunctionName but without having an actual UFunction to pull metadata from (no localization)
+		const bool bShowFriendlyNames = bAllowFriendlyNames && (GEditor && GetDefault<UEditorStyleSettings>()->bShowFriendlyNames);
+		if (bShowFriendlyNames)
 		{
-			DisplayInfo.DisplayName = GetFriendlySignatureName(Function);
+			DisplayInfo.DisplayName = FText::FromString(FName::NameToDisplayString(DisplayInfo.PlainName.ToString(), false));
 		}
 		else
 		{
-			DisplayInfo.DisplayName = FText::FromString(FName::NameToDisplayString(DisplayInfo.PlainName.ToString(), false));
+			DisplayInfo.DisplayName = DisplayInfo.PlainName;
 		}
 	}
-	else
-	{
-		DisplayInfo.DisplayName = DisplayInfo.PlainName;
-	}
 }
 
 bool UEdGraphSchema_K2::IsSelfPin(const UEdGraphPin& Pin) const 
diff --git a/Engine/Source/Editor/UnrealEd/Private/ObjectTools.cpp b/Engine/Source/Editor/UnrealEd/Private/ObjectTools.cpp
--- a/Engine/Source/Editor/UnrealEd/Private/ObjectTools.cpp
+++ b/Engine/Source/Editor/UnrealEd/Private/ObjectTools.cpp
@@ -5268,33 +5268,56 @@ namespace ObjectTools
 
 	FText GetUserFacingFunctionName(const UFunction* Function, bool bAllowFriendlyNames)
 	{
-		FText ReturnDisplayName;
-
 		if (Function != nullptr)
 		{
-			static const FName NAME_DisplayName { TEXT("DisplayName") };
+			// Function->GetDisplayNameText() and Function->GetMetaDataText(NAME_DisplayName,...) are equivalent.
+			// However, we cannot simply use those functions since they enforce friendly names and do not work on SKEL classes.
+			// We implement the same functionality a third time.
+			static const FName NAME_DisplayName{ TEXT("DisplayName") };
 
-			// Functions do not use friendly names because they can be manually input by a user in the editor (and it would otherwise not adhere to their name)
-			// There is a long-term goal of removing friendly names from the Engine.  However, we keep them in the case of FullTitle as it helps the nodes
-			// be decipherable in a zoomed-out view.
-			if (GEditor && bAllowFriendlyNames && GetDefault<UEditorStyleSettings>()->bShowFriendlyNames)
+			FString FunctionDisplayName = Function->GetName();
+			FString FunctionFriendlyName;
+			if (const FString* FoundMetaData = Function->FindMetaData(NAME_DisplayName))
 			{
-				ReturnDisplayName = Function->GetDisplayNameText();
+				FunctionDisplayName = *FoundMetaData;
 			}
-			else if (const FString* OverrideDisplayName = Function->FindMetaData(NAME_DisplayName))
+			else
 			{
-				ReturnDisplayName = FText::FromString(*OverrideDisplayName);
+				FunctionFriendlyName = FName::NameToDisplayString(FunctionDisplayName, false);
 			}
 
-			if (ReturnDisplayName.IsEmpty())
+			// There is a long-term Localization vision that no code is localized, but we are choosing to always localize in UE5.
+			if (constexpr bool bUseLocalization = true)
+			{
+				// We are sometimes calling this function on a SKEL_ClassName_C Actor rather than the instance
+				// So we should keep traversing the hierarchy finding the right function name
+				const UFunction* BaseMostFunction = Function;
+				for (; BaseMostFunction->GetSuperFunction(); BaseMostFunction = BaseMostFunction->GetSuperFunction());
+
+				static const FTextKey LocalizationNamespace = TEXT("UObjectDisplayNames");
+				FString LocalizationKey = BaseMostFunction->GetFullGroupName(false);
+				FText OutLocalizedText;
+				if (FText::FindTextInLiveTable_Advanced(LocalizationNamespace, LocalizationKey, OutLocalizedText, FunctionFriendlyName.IsEmpty() ? &FunctionDisplayName : &FunctionFriendlyName))
+				{
+					return OutLocalizedText;
+				}
+			}
+
+			// Functions can opt to not use friendly names because they can be manually input by a user in the editor (and it would otherwise not adhere to their name)
+			// There is a long-term goal of removing friendly names from the Engine.  However, we keep them in some cases as it helps the nodes
+			// be decipherable in a zoomed-out view.
+			const bool bUseFriendlyNames = bAllowFriendlyNames && GEditor && GetDefault<UEditorStyleSettings>()->bShowFriendlyNames;
+			if (bUseFriendlyNames && !FunctionFriendlyName.IsEmpty())
+			{
+				return FText::FromString(FunctionFriendlyName);
+			}
+			else
 			{
-				// Previous (and similar) code paths would go through FField::GetMetaDataText(DisplayName) which attempts localization
-				// However, we do not localize function names (and we've explicitly requested non-friendly names), so just show us the real name
-				ReturnDisplayName = FText::FromString(Function->GetName());
+				return FText::FromString(FunctionDisplayName);
 			}
 		}
 
-		return ReturnDisplayName;
+		return FText::GetEmpty();
 	}
 
 	FString GetDefaultTooltipForFunction(const UFunction* Function)
