Anim In Editor: Support for evaluating subsequences in in isoloation when baking animation sequences.  If Subsequenes are present we by default use the subsequence range, but if we are evaluating in isolation, which is an editor only/sequencer setting, we need to use the moviescene range.

diff --git a/Engine/Plugins/Animation/ControlRig/Source/ControlRigEditor/Private/ControlRigSequencerEditorLibrary.cpp b/Engine/Plugins/Animation/ControlRig/Source/ControlRigEditor/Private/ControlRigSequencerEditorLibrary.cpp
--- a/Engine/Plugins/Animation/ControlRig/Source/ControlRigEditor/Private/ControlRigSequencerEditorLibrary.cpp
+++ b/Engine/Plugins/Animation/ControlRig/Source/ControlRigEditor/Private/ControlRigSequencerEditorLibrary.cpp
@@ -1187,6 +1187,7 @@ bool UControlRigSequencerEditorLibrary::ExportAnimSequenceFromSequencer(UAnimSeq
 			AESQ.RootToLocalTransform = RootToLocalTransform;
 			AESQ.MovieSceneSequence = Sequence;
 			AESQ.RootMovieSceneSequence = RootMovieSceneSequence;
+			AESQ.bForceUseOfMovieScenePlaybackRange = ParentSequencer.Get()->GetSequencerSettings()->ShouldEvaluateSubSequencesInIsolation();
 			bResult = MovieSceneToolHelpers::ExportToAnimSequence(AnimSequence, ExportOptions, AESQ, SkeletalMeshComp);
 		}
 		else
@@ -1270,6 +1271,7 @@ bool UControlRigSequencerEditorLibrary::BakeToControlRig(UWorld* World, ULevelSe
 				AESP.RootToLocalTransform = RootToLocalTransform;
 				AESP.MovieSceneSequence = LevelSequence;
 				AESP.RootMovieSceneSequence = LevelSequence;
+				AESP.bForceUseOfMovieScenePlaybackRange = false;
 				bResult = MovieSceneToolHelpers::ExportToAnimSequence(TempAnimSequence, ExportOptions, AESP, SkeletalMeshComp);
 				if (bResult == false)
 				{
diff --git a/Engine/Plugins/Animation/ControlRig/Source/ControlRigEditor/Private/Sequencer/ControlRigParameterTrackEditor.cpp b/Engine/Plugins/Animation/ControlRig/Source/ControlRigEditor/Private/Sequencer/ControlRigParameterTrackEditor.cpp
--- a/Engine/Plugins/Animation/ControlRig/Source/ControlRigEditor/Private/Sequencer/ControlRigParameterTrackEditor.cpp
+++ b/Engine/Plugins/Animation/ControlRig/Source/ControlRigEditor/Private/Sequencer/ControlRigParameterTrackEditor.cpp
@@ -1184,6 +1184,7 @@ void FControlRigParameterTrackEditor::BakeToControlRig(UClass* InClass, FGuid Ob
 				AESP.RootToLocalTransform = RootToLocalTransform;
 				AESP.MovieSceneSequence = OwnerSequence;
 				AESP.RootMovieSceneSequence = RootSequence;
+				AESP.bForceUseOfMovieScenePlaybackRange = ParentSequencer.Get()->GetSequencerSettings()->ShouldEvaluateSubSequencesInIsolation();
 				bool bResult = MovieSceneToolHelpers::ExportToAnimSequence(TempAnimSequence, AnimSeqExportOption, AESP, SkelMeshComp);
 				if (bResult == false)
 				{
@@ -1397,7 +1398,14 @@ void FControlRigParameterTrackEditor::BakeInvertedPose(UControlRig* InControlRig
 	
 	const FScopedTransaction Transaction(LOCTEXT("BakeInvertedPose_Transaction", "Bake Inverted Pose"));
 
-	UnFbx::FLevelSequenceAnimTrackAdapter AnimTrackAdapter(ParentSequencer.Get(), MovieSceneSequence, RootMovieSceneSequence, RootToLocalTransform);
+	UnFbx::FLevelSequenceAnimTrackAdapter::FAnimTrackSettings Settings;
+	Settings.MovieScenePlayer = ParentSequencer.Get();
+	Settings.MovieSceneSequence = MovieSceneSequence;
+	Settings.RootMovieSceneSequence = RootMovieSceneSequence;
+	Settings.RootToLocalTransform = RootToLocalTransform;
+	Settings.bForceUseOfMovieScenePlaybackRange = ParentSequencer.Get()->GetSequencerSettings()->ShouldEvaluateSubSequencesInIsolation();
+
+	UnFbx::FLevelSequenceAnimTrackAdapter AnimTrackAdapter(Settings);
 	int32 AnimationLength = AnimTrackAdapter.GetLength();
 	FScopedSlowTask Progress(AnimationLength, LOCTEXT("BakingToControlRig_SlowTask", "Baking To Control Rig..."));
 	Progress.MakeDialog(true);
@@ -1433,7 +1441,7 @@ void FControlRigParameterTrackEditor::BakeInvertedPose(UControlRig* InControlRig
 	AESP.RootToLocalTransform = RootToLocalTransform;
 	AESP.MovieSceneSequence = MovieSceneSequence;
 	AESP.RootMovieSceneSequence = RootMovieSceneSequence;
-
+	AESP.bForceUseOfMovieScenePlaybackRange = ParentSequencer.Get()->GetSequencerSettings()->ShouldEvaluateSubSequencesInIsolation();
 	MovieSceneToolHelpers::BakeToSkelMeshToCallbacks(AESP,SkelMeshComp, ExportOptions,
 		InitCallback, StartCallback, TickCallback, EndCallback);
 
diff --git a/Engine/Source/Editor/MovieSceneTools/Private/MovieSceneToolHelpers.cpp b/Engine/Source/Editor/MovieSceneTools/Private/MovieSceneToolHelpers.cpp
--- a/Engine/Source/Editor/MovieSceneTools/Private/MovieSceneToolHelpers.cpp
+++ b/Engine/Source/Editor/MovieSceneTools/Private/MovieSceneToolHelpers.cpp
@@ -4148,7 +4148,14 @@ bool MovieSceneToolHelpers::BakeToSkelMeshToCallbacks(const FAnimExportSequenceP
 		}
 	}
 
-	UnFbx::FLevelSequenceAnimTrackAdapter AnimTrackAdapter(AESP.Player, AESP.MovieSceneSequence, AESP.RootMovieSceneSequence, AESP.RootToLocalTransform);
+	UnFbx::FLevelSequenceAnimTrackAdapter::FAnimTrackSettings AnimTrackSettings;
+	AnimTrackSettings.MovieScenePlayer = AESP.Player;
+	AnimTrackSettings.MovieSceneSequence = AESP.MovieSceneSequence;
+	AnimTrackSettings.RootMovieSceneSequence = AESP.RootMovieSceneSequence;
+	AnimTrackSettings.RootToLocalTransform = AESP.RootToLocalTransform;
+	AnimTrackSettings.bForceUseOfMovieScenePlaybackRange = AESP.bForceUseOfMovieScenePlaybackRange;
+
+	UnFbx::FLevelSequenceAnimTrackAdapter AnimTrackAdapter(AnimTrackSettings);
 	if (ExportOptions->bUseCustomTimeRange)
 	{
 		const FFrameRate TickResolution = MovieScene->GetTickResolution();
diff --git a/Engine/Source/Editor/MovieSceneTools/Private/TrackEditors/SkeletalAnimationTrackEditor.cpp b/Engine/Source/Editor/MovieSceneTools/Private/TrackEditors/SkeletalAnimationTrackEditor.cpp
--- a/Engine/Source/Editor/MovieSceneTools/Private/TrackEditors/SkeletalAnimationTrackEditor.cpp
+++ b/Engine/Source/Editor/MovieSceneTools/Private/TrackEditors/SkeletalAnimationTrackEditor.cpp
@@ -1519,6 +1519,7 @@ void FSkeletalAnimationTrackEditor::OnSequencerSaved(ISequencer& )
 						AESP.RootToLocalTransform = RootToLocalTransform;
 						AESP.MovieSceneSequence = MovieSceneSequence;
 						AESP.RootMovieSceneSequence = RootMovieSceneSequence;
+						AESP.bForceUseOfMovieScenePlaybackRange = SequencerPtr.Get()->GetSequencerSettings()->ShouldEvaluateSubSequencesInIsolation();
 						bool bResult = MovieSceneToolHelpers::ExportToAnimSequence(AnimSequence, AnimSeqExportOption, AESP, SkelMeshComp);
 
 						AnimSeqExportOption->bExportMorphTargets = bSavedExportMorphTargets;
@@ -1685,6 +1686,7 @@ bool FSkeletalAnimationTrackEditor::CreateAnimationSequence(const TArray<UObject
 				AESP.RootToLocalTransform = RootToLocalTransform;
 				AESP.MovieSceneSequence = MovieSceneSequence;
 				AESP.RootMovieSceneSequence = RootMovieSceneSequence;
+				AESP.bForceUseOfMovieScenePlaybackRange = ParentSequencer->GetSequencerSettings()->ShouldEvaluateSubSequencesInIsolation();
 				AnimSeqExportOption->CustomDisplayRate = ParentSequencer->GetFocusedDisplayRate();
 				bResult  = MovieSceneToolHelpers::ExportToAnimSequence(AnimSequence, AnimSeqExportOption, AESP, SkelMeshComp);
 			}
diff --git a/Engine/Source/Editor/MovieSceneTools/Public/MovieSceneToolHelpers.h b/Engine/Source/Editor/MovieSceneTools/Public/MovieSceneToolHelpers.h
--- a/Engine/Source/Editor/MovieSceneTools/Public/MovieSceneToolHelpers.h
+++ b/Engine/Source/Editor/MovieSceneTools/Public/MovieSceneToolHelpers.h
@@ -73,6 +73,7 @@ struct FAnimExportSequenceParameters
 	UMovieSceneSequence* RootMovieSceneSequence;
 	IMovieScenePlayer* Player;
 	FMovieSceneSequenceTransform RootToLocalTransform;
+	bool bForceUseOfMovieScenePlaybackRange = false;
 };
 
 //callback's used by skel mesh recorders
diff --git a/Engine/Source/Editor/UnrealEd/Private/Fbx/FbxMainExport.cpp b/Engine/Source/Editor/UnrealEd/Private/Fbx/FbxMainExport.cpp
--- a/Engine/Source/Editor/UnrealEd/Private/Fbx/FbxMainExport.cpp
+++ b/Engine/Source/Editor/UnrealEd/Private/Fbx/FbxMainExport.cpp
@@ -1795,6 +1795,19 @@ FbxNode* FFbxExporter::FLevelSequenceNodeNameAdapter::GetFbxNode(UObject* InObje
 	return nullptr;
 }
 
+FLevelSequenceAnimTrackAdapter::FLevelSequenceAnimTrackAdapter(const FAnimTrackSettings& Settings)
+{
+	MovieScenePlayer = Settings.MovieScenePlayer;
+	MovieSceneSequence = Settings.MovieSceneSequence;
+	RootMovieSceneSequence = Settings.RootMovieSceneSequence;
+	MovieScene = Settings.MovieSceneSequence->GetMovieScene();
+	RootToLocalTransform = Settings.RootToLocalTransform;
+	AnimTrack = Settings.AnimTrack;
+	bForceUseOfMovieScenePlaybackRange = Settings.bForceUseOfMovieScenePlaybackRange;
+}
+
+
+//5.6 deprecrated
 FLevelSequenceAnimTrackAdapter::FLevelSequenceAnimTrackAdapter(IMovieScenePlayer* InMovieScenePlayer, UMovieSceneSequence* InMovieSceneSequence, UMovieSceneSequence* InRootMovieSceneSequence, const FMovieSceneSequenceTransform& InRootToLocalTransform, UMovieSceneSkeletalAnimationTrack* InAnimTrack)
 {
 	MovieScenePlayer = InMovieScenePlayer;
@@ -1803,6 +1816,7 @@ FLevelSequenceAnimTrackAdapter::FLevelSequenceAnimTrackAdapter(IMovieScenePlayer
 	MovieScene = MovieSceneSequence->GetMovieScene();
 	RootToLocalTransform = InRootToLocalTransform;
 	AnimTrack = InAnimTrack;
+	bForceUseOfMovieScenePlaybackRange = false;
 }
 
 //5.5 deprecrated
@@ -1822,6 +1836,7 @@ FLevelSequenceAnimTrackAdapter::FLevelSequenceAnimTrackAdapter(IMovieScenePlayer
 			RootMovieSceneSequence = OwnerSceneSequence;
 		}
 	}
+	bForceUseOfMovieScenePlaybackRange = false;
 }
 
 void FLevelSequenceAnimTrackAdapter::SetRange(const FFrameNumber& StartFrame, const FFrameNumber& EndFrame)
@@ -1849,7 +1864,7 @@ TRange<FFrameNumber> FLevelSequenceAnimTrackAdapter::GetSequenceRange() const
 	TRange<FFrameNumber> Range = MovieScene->GetPlaybackRange();
 	FFrameTime StartTime = FFrameRate::TransformTime(UE::MovieScene::DiscreteInclusiveLower(MovieScene->GetPlaybackRange()).Value, MovieScene->GetTickResolution(), MovieScene->GetDisplayRate());
 
-	if (SubSequenceData)
+	if (SubSequenceData && bForceUseOfMovieScenePlaybackRange == false)
 	{
 		Range = SubSequenceData->PlayRange.Value;
 	}
@@ -2084,7 +2099,13 @@ bool FFbxExporter::ExportLevelSequenceTracks(UMovieSceneSequence* MovieSceneSequ
 	{
 		if (!bSkip3DTransformTrack)
 		{
-			FLevelSequenceAnimTrackAdapter AnimTrackAdapter(MovieScenePlayer, MovieSceneSequence, RootMovieSceneSequence, RootToLocalTransform, nullptr);
+			UnFbx::FLevelSequenceAnimTrackAdapter::FAnimTrackSettings Settings;
+			Settings.MovieScenePlayer = MovieScenePlayer;
+			Settings.MovieSceneSequence = MovieSceneSequence;
+			Settings.RootMovieSceneSequence = RootMovieSceneSequence;
+			Settings.RootToLocalTransform = RootToLocalTransform;
+			Settings.bForceUseOfMovieScenePlaybackRange = false;
+			FLevelSequenceAnimTrackAdapter AnimTrackAdapter(Settings);
 			ExportLevelSequenceBaked3DTransformTrack(AnimTrackAdapter, FbxActor, MovieScenePlayer, InSequenceID, TransformTracks, BoundObject, MovieScene->GetPlaybackRange(), RootToLocalTransform);
 		}
 
@@ -2140,9 +2161,16 @@ bool FFbxExporter::ExportLevelSequenceTracks(UMovieSceneSequence* MovieSceneSequ
 	{
 		TArray<USkeletalMeshComponent*> SkeletalMeshComponents;
 		SkeletalMeshComp->GetOwner()->GetComponents(SkeletalMeshComponents);
+		UnFbx::FLevelSequenceAnimTrackAdapter::FAnimTrackSettings Settings;
+		Settings.MovieScenePlayer = MovieScenePlayer;
+		Settings.MovieSceneSequence = MovieSceneSequence;
+		Settings.RootMovieSceneSequence = RootMovieSceneSequence;
+		Settings.RootToLocalTransform = RootToLocalTransform;
+		Settings.bForceUseOfMovieScenePlaybackRange = false;
 		for (USkeletalMeshComponent* SkeletalMeshComponent : SkeletalMeshComponents)
 		{
-			FLevelSequenceAnimTrackAdapter AnimTrackAdapter(MovieScenePlayer, MovieSceneSequence,RootMovieSceneSequence, RootToLocalTransform, SkeletalAnimationTrack);
+			Settings.AnimTrack = SkeletalAnimationTrack;
+			FLevelSequenceAnimTrackAdapter AnimTrackAdapter(Settings);
 			ExportAnimTrack(AnimTrackAdapter, Actor, SkeletalMeshComponent, 1.0 / DisplayRate.AsDecimal());
 		}
 	}
diff --git a/Engine/Source/Editor/UnrealEd/Private/FbxExporter.h b/Engine/Source/Editor/UnrealEd/Private/FbxExporter.h
--- a/Engine/Source/Editor/UnrealEd/Private/FbxExporter.h
+++ b/Engine/Source/Editor/UnrealEd/Private/FbxExporter.h
@@ -78,11 +78,22 @@ namespace UnFbx
 	class FLevelSequenceAnimTrackAdapter : public IAnimTrackAdapter
 	{
 	public:
-
+		struct FAnimTrackSettings
+		{
+			IMovieScenePlayer* MovieScenePlayer = nullptr;
+			UMovieSceneSequence* MovieSceneSequence = nullptr;
+			UMovieSceneSequence* RootMovieSceneSequence = nullptr;
+			FMovieSceneSequenceTransform RootToLocalTransform;
+			UMovieSceneSkeletalAnimationTrack* AnimTrack = nullptr;
+			bool bForceUseOfMovieScenePlaybackRange = false; //needed for when evaluate in isolation is used by Sequencer, the subsequence range will not be correct
+		};
+		UNREALED_API FLevelSequenceAnimTrackAdapter(const FAnimTrackSettings& Settings);
+		
+		UE_DEPRECATED(5.6, "FLevelSequenceAnimTrackAdapter constructor taking multiple arguments is deprecated, use version that takes FAnimTrackSettings instead")
 		UNREALED_API FLevelSequenceAnimTrackAdapter(IMovieScenePlayer* InMovieScenePlayer, UMovieSceneSequence* InMovieSceneSequence, UMovieSceneSequence* InRootMovieSceneSequence, const FMovieSceneSequenceTransform& InRootToLocalTransform, UMovieSceneSkeletalAnimationTrack* InAnimTrack = nullptr);
 		UE_DEPRECATED(5.5, "FLevelSequenceAnimTrackAdapter constructor taking a movie scene is deprecated, use version that takes movie scene sequences")
 		UNREALED_API FLevelSequenceAnimTrackAdapter(IMovieScenePlayer* InMovieScenePlayer, UMovieScene* InMovieScene, const FMovieSceneSequenceTransform& InRootToLocalTransform, UMovieSceneSkeletalAnimationTrack* InAnimTrack = nullptr);
-
+		
 		UNREALED_API virtual int32 GetLocalStartFrame() const override;
 		UNREALED_API virtual int32 GetStartFrame() const override;
 		UNREALED_API virtual int32 GetLength() const override;
@@ -99,7 +110,7 @@ namespace UnFbx
 		IMovieScenePlayer* MovieScenePlayer;
 		UMovieSceneSequence* MovieSceneSequence;
 		UMovieSceneSequence* RootMovieSceneSequence;
-
+		bool bForceUseOfMovieScenePlaybackRange;
 		UMovieScene* MovieScene;
 		FMovieSceneSequenceTransform RootToLocalTransform;
 		UMovieSceneSkeletalAnimationTrack* AnimTrack;
