Improve display of Far View Plane slider so that the maximum value more intuitively represents "Infinity" rather than zero.

UE-238977

diff --git a/Engine/Source/Editor/UnrealEd/Private/ViewportToolbar/UnrealEdViewportToolbar.cpp b/Engine/Source/Editor/UnrealEd/Private/ViewportToolbar/UnrealEdViewportToolbar.cpp
--- a/Engine/Source/Editor/UnrealEd/Private/ViewportToolbar/UnrealEdViewportToolbar.cpp
+++ b/Engine/Source/Editor/UnrealEd/Private/ViewportToolbar/UnrealEdViewportToolbar.cpp
@@ -2296,6 +2296,14 @@ TSharedRef<SWidget> CreateFarViewPlaneMenuWidget(const TSharedRef<SEditorViewpor
 {
 	TWeakPtr<FEditorViewportClient> ViewportClientWeak = InViewport->GetViewportClient();
 
+	/*
+	 * FEditorViewportClient treats a far clip plane value of 0.0 as "infinity".
+	 * This Spin Box transforms the maximum value to that 0.0 and back again,
+	 * allowing the maximum value to be treated as infinity and creating a more
+	 * natural interface.
+	 */
+	constexpr float MaxValue = 100000.0f;
+
 	return
 		// clang-format off
 		SNew(SBox)
@@ -2311,24 +2319,37 @@ TSharedRef<SWidget> CreateFarViewPlaneMenuWidget(const TSharedRef<SEditorViewpor
 				[
 					SNew(SSpinBox<float>)
 					.Style(&FAppStyle::Get(), "Menu.SpinBox")
-					.ToolTipText(LOCTEXT("FarViewPlaneTooltip", "Distance to use as the far view plane, or zero to enable an infinite far view plane"))
-					.MinValue(0.0f)
-					.MaxValue(100000.0f)
+					.ToolTipText(LOCTEXT("FarViewPlaneTooltip", "Distance to use as the far view plane"))
+					.MinValue(0.01f)
+					.MaxValue(MaxValue)
+					.SliderExponent(3.0f) // Gives better precision for smaller ranges
+					.OnGetDisplayValue_Lambda([](float InValue)
+					{
+						if (InValue >= MaxValue)
+						{
+							return TOptional<FText>(LOCTEXT("Infinity", "Infinity"));
+						}
+						return TOptional<FText>();
+					})
 					.Font(FAppStyle::GetFontStyle(TEXT("MenuItem.Font")))
 					.Value_Lambda([ViewportClientWeak]()
 					{
 						if (TSharedPtr<FEditorViewportClient> ViewportClient = ViewportClientWeak.Pin())
 						{
-							return ViewportClient->GetFarClipPlaneOverride();
+							const float Override = ViewportClient->GetFarClipPlaneOverride(); 
+							if (Override > 0.0f)
+							{
+								return Override;
+							}
 						}
 
-						return 100000.0f;
+						return MaxValue;
 					})
 					.OnValueChanged_Lambda([ViewportClientWeak](float InNewValue)
 					{
 						if (TSharedPtr<FEditorViewportClient> ViewportClient = ViewportClientWeak.Pin())
 						{
-							ViewportClient->OverrideFarClipPlane(InNewValue);
+							ViewportClient->OverrideFarClipPlane(InNewValue >= MaxValue ? 0.0f : InNewValue);
 							ViewportClient->Invalidate();
 						}
 					})
diff --git a/Engine/Source/Runtime/Slate/Private/Widgets/Input/SSpinBox.cpp b/Engine/Source/Runtime/Slate/Private/Widgets/Input/SSpinBox.cpp
--- a/Engine/Source/Runtime/Slate/Private/Widgets/Input/SSpinBox.cpp
+++ b/Engine/Source/Runtime/Slate/Private/Widgets/Input/SSpinBox.cpp
@@ -90,6 +90,8 @@ void SSpinBox<NumericType>::Construct(const FArguments& InArgs)
 	SupportDynamicSliderMinValue = InArgs._SupportDynamicSliderMinValue;
 	OnDynamicSliderMaxValueChanged = InArgs._OnDynamicSliderMaxValueChanged;
 	OnDynamicSliderMinValueChanged = InArgs._OnDynamicSliderMinValueChanged;
+	
+	OnGetDisplayValue = InArgs._OnGetDisplayValue;
 
 	bEnableWheel = InArgs._EnableWheel;
 	bBroadcastValueChangesPerKey = InArgs._BroadcastValueChangesPerKey;
@@ -154,7 +156,7 @@ void SSpinBox<NumericType>::Construct(const FArguments& InArgs)
 		[
 			SAssignNew(TextBlock, STextBlock)
 				.Font(InArgs._Font)
-				.Text(this, &SSpinBox<NumericType>::GetValueAsText)
+				.Text(this, &SSpinBox<NumericType>::GetDisplayValue)
 				.MinDesiredWidth(this, &SSpinBox<NumericType>::GetTextMinDesiredWidth)
 				.Justification(InArgs._Justification)
 		]
@@ -890,6 +892,20 @@ FText SSpinBox<NumericType>::GetValueAsText() const
 	return FText::FromString(GetValueAsString());
 }
 
+template <typename NumericType>
+FText SSpinBox<NumericType>::GetDisplayValue() const
+{
+	if (OnGetDisplayValue.IsBound())
+	{
+		const TOptional<FText> OverrideValue = OnGetDisplayValue.Execute(ValueAttribute.Get());
+		if (OverrideValue.IsSet())
+		{
+			return OverrideValue.GetValue();
+		}
+	}
+	return FText::FromString(GetValueAsString());
+}
+
 template <typename NumericType>
 void SSpinBox<NumericType>::TextField_OnTextChanged(const FText& NewText)
 {
diff --git a/Engine/Source/Runtime/Slate/Public/Widgets/Input/SSpinBox.h b/Engine/Source/Runtime/Slate/Public/Widgets/Input/SSpinBox.h
--- a/Engine/Source/Runtime/Slate/Public/Widgets/Input/SSpinBox.h
+++ b/Engine/Source/Runtime/Slate/Public/Widgets/Input/SSpinBox.h
@@ -72,6 +72,9 @@ public:
 
 	/** Notification when the max/min spinner values are changed (only apply if SupportDynamicSliderMaxValue or SupportDynamicSliderMinValue are true) */
 	DECLARE_DELEGATE_FourParams(FOnDynamicSliderMinMaxValueChanged, NumericType, TWeakPtr<SWidget>, bool, bool);
+	
+	/** Optional customization of the display value based on the current value. */
+	DECLARE_DELEGATE_RetVal_OneParam(TOptional<FText>, FOnGetDisplayValue, NumericType);
 
 	SLATE_BEGIN_ARGS(SSpinBox<NumericType>)
 		: _Style(&FCoreStyle::Get().GetWidgetStyle<FSpinBoxStyle>("SpinBox"))
@@ -161,6 +164,8 @@ public:
 		SLATE_EVENT(FSimpleDelegate, OnBeginSliderMovement)
 		/** Called right after the slider handle is released by the user */
 		SLATE_EVENT(FOnValueChanged, OnEndSliderMovement)
+		/** Called to allow customization of what text is displayed when not typing. An empty value falls back to the default behavior. */
+		SLATE_EVENT(FOnGetDisplayValue, OnGetDisplayValue)
 		/** Whether to clear keyboard focus when pressing enter to commit changes */
 		SLATE_ATTRIBUTE(bool, ClearKeyboardFocusOnCommit)
 		/** Whether to select all text when pressing enter to commit changes */
@@ -333,6 +338,9 @@ protected:
 
 	/** @return the value being observed by the spinbox as FText - todo: spinbox FText support (reimplement me) */
 	SLATE_API FText GetValueAsText() const;
+	
+	/** @return the value to be displayed when not manually editing text */
+	SLATE_API FText GetDisplayValue() const;
 
 	/**
 	 * Invoked when the text in the text field changes
@@ -438,6 +446,7 @@ private:
 	TAttribute< TOptional<NumericType> > WheelStep;
 	FOnDynamicSliderMinMaxValueChanged OnDynamicSliderMaxValueChanged;
 	FOnDynamicSliderMinMaxValueChanged OnDynamicSliderMinValueChanged;
+	FOnGetDisplayValue OnGetDisplayValue;
 
 	/** Prevents the spinbox from being smaller than desired in certain cases (e.g. when it is empty) */
 	TAttribute<float> MinDesiredWidth;
