Re-enable garbage reference warnings in the cooker

diff --git a/Engine/Source/Editor/UnrealEd/Private/Commandlets/CookCommandlet.cpp b/Engine/Source/Editor/UnrealEd/Private/Commandlets/CookCommandlet.cpp
--- a/Engine/Source/Editor/UnrealEd/Private/Commandlets/CookCommandlet.cpp
+++ b/Engine/Source/Editor/UnrealEd/Private/Commandlets/CookCommandlet.cpp
@@ -205,9 +205,6 @@ int32 UCookCommandlet::Main(const FString& CmdLineParams)
 	COOK_STAT(FParse::Value(*Params, TEXT("CookCultures="), DetailedCookStats::CookCultures));
 	COOK_STAT(FParse::Value(*Params, TEXT("CookLabel="), DetailedCookStats::CookLabel));
 	
-	// Temporarily disable garbage ref reporting pending a fix to FWorldPartitionLevelHelper::MoveExternalActorsToLevel
-	TGuardConsoleVariable<int32> GarbageRefGuard(TEXT("gc.GarbageReferenceTrackingEnabled"), 0);
-
 	ITargetPlatformManagerModule& TPM = GetTargetPlatformManagerRef();
 	if ( bCookOnTheFly )
 	{
