IncrementalCook: Refactor FCookDependencies -> UE::Cook::FPackageArtifacts, in preparation for adding multiple sets of dependencies in each package, one for each build result. Move FPackageArtifacts and its helper types into its own header.
#rnx
#incrementalcook

diff --git a/Engine/Source/Editor/UnrealEd/Private/Cooker/CookFunctionLibrary.cpp b/Engine/Source/Editor/UnrealEd/Private/Cooker/CookFunctionLibrary.cpp
--- a/Engine/Source/Editor/UnrealEd/Private/Cooker/CookFunctionLibrary.cpp
+++ b/Engine/Source/Editor/UnrealEd/Private/Cooker/CookFunctionLibrary.cpp
@@ -5,6 +5,7 @@
 #include "Commandlets/Commandlet.h"
 #include "Cooker/AsyncIODelete.h"
 #include "Cooker/CookDependency.h"
+#include "Cooker/CookPackageArtifacts.h"
 #include "Cooker/CookSandbox.h"
 #include "Cooker/LooseCookedPackageWriter.h"
 #include "CookOnTheSide/CookLog.h"
@@ -16,7 +17,6 @@
 #include "Serialization/CompactBinaryPackage.h"
 #include "Serialization/CompactBinaryWriter.h"
 #include "Subsystems/UnrealEditorSubsystem.h"
-#include "TargetDomain/TargetDomainUtils.h"
 #include "UObject/ArchiveCookContext.h"
 #include "UObject/FieldPath.h"
 #include "UObject/SavePackage.h"
@@ -26,17 +26,17 @@ extern UNREALED_API UEditorEngine* GEditor;
 
 namespace UE::Private
 {
-	void SaveDepsToFile(const UE::TargetDomain::FCookDependencies& CookDependencies, const FString& Filename)
+	void SaveDepsToFile(const UE::Cook::FPackageArtifacts& Artifacts, const FString& Filename)
 	{
 		FCbWriter Writer;
 		Writer.BeginObject();
-		Writer << "CookTestSnapshot" << CookDependencies;
+		Writer << "CookTestSnapshot" << Artifacts;
 		Writer.EndObject();
 		TUniquePtr<FArchive> FileArchive(IFileManager::Get().CreateDebugFileWriter(*Filename));
 		Writer.Save(*FileArchive);
 	}
 
-	bool LoadDepsFromFile(UE::TargetDomain::FCookDependencies& CookDependencies, const FString& Filename)
+	bool LoadDepsFromFile(UE::Cook::FPackageArtifacts& Artifacts, const FString& Filename)
 	{
 		TArray<uint8> Data;
 		{
@@ -48,7 +48,7 @@ namespace UE::Private
 		FSharedBuffer SharedBuffer = FSharedBuffer::MakeView(MakeMemoryView(Data)); // strange boilerplate, i don't understand it
 		FCbObject CbObject(SharedBuffer);
 		FCbField TestSnapshot = CbObject["CookTestSnapshot"];
-		return LoadFromCompactBinary(TestSnapshot.AsObject(), CookDependencies);
+		return LoadFromCompactBinary(TestSnapshot.AsObject(), Artifacts);
 	}
 
 	FString GetDepsFilename(const UObject* Object, const FString& DestinationSubfolder)
@@ -61,7 +61,7 @@ namespace UE::Private
 
 void UCookFunctionLibrary::CookAsset(UObject* Object, const FString& ForPlatform, const FString& DestinationSubfolder, const FString& CookCommandlineArgs)
 {
-	using namespace UE::TargetDomain;
+	using namespace UE::Cook;
 	if (Object == nullptr)
 	{
 		UE_LOG(LogCook, Warning, TEXT("CookAsset expected an object to cook, but received nullptr"));
@@ -214,7 +214,7 @@ void UCookFunctionLibrary::CookAsset(UObject* Object, const FString& ForPlatform
 		TArray<FName> RuntimeDepsTodo; // I don't believe we need valid data in here at this time
 
 		TArray<TPair<ELogVerbosity::Type, FString>> Messages;
-		FCookDependencies Deps = FCookDependencies::Collect(Package, TargetPlatform, &Result, Todo, MoveTemp(RuntimeDepsTodo), &Messages);
+		FPackageArtifacts Deps = FPackageArtifacts::Collect(Package, TargetPlatform, &Result, Todo, MoveTemp(RuntimeDepsTodo), &Messages);
 		if (!Deps.IsValid())
 		{
 			TStringBuilder<256> ErrorMessage;
diff --git a/Engine/Source/Editor/UnrealEd/Private/Cooker/CookGenerationHelper.cpp b/Engine/Source/Editor/UnrealEd/Private/Cooker/CookGenerationHelper.cpp
--- a/Engine/Source/Editor/UnrealEd/Private/Cooker/CookGenerationHelper.cpp
+++ b/Engine/Source/Editor/UnrealEd/Private/Cooker/CookGenerationHelper.cpp
@@ -6,6 +6,7 @@
 #include "AssetRegistry/AssetData.h"
 #include "Cooker/CookDirector.h"
 #include "Cooker/CookGarbageCollect.h"
+#include "Cooker/CookPackageArtifacts.h"
 #include "Cooker/CookPlatformManager.h"
 #include "Cooker/CookWorkerServer.h"
 #include "Cooker/IWorkerRequests.h"
@@ -1280,7 +1281,7 @@ void FGenerationHelper::FinishGeneratorPlatformSave(FPackageData& PackageData, b
 }
 
 void FGenerationHelper::FinishGeneratedPlatformSave(FPackageData& PackageData,
-	UE::TargetDomain::FGeneratedPackageResultStruct& OutGeneratedResult)
+	FGeneratedPackageResultStruct& OutGeneratedResult)
 {
 	ConditionalInitialize();
 
diff --git a/Engine/Source/Editor/UnrealEd/Private/Cooker/CookGenerationHelper.h b/Engine/Source/Editor/UnrealEd/Private/Cooker/CookGenerationHelper.h
--- a/Engine/Source/Editor/UnrealEd/Private/Cooker/CookGenerationHelper.h
+++ b/Engine/Source/Editor/UnrealEd/Private/Cooker/CookGenerationHelper.h
@@ -26,7 +26,7 @@
 class ITargetPlatform;
 class UCookOnTheFlyServer;
 namespace UE::Cook { class FCookGCDiagnosticContext; }
-namespace UE::TargetDomain { struct FGeneratedPackageResultStruct; }
+namespace UE::Cook { struct FGeneratedPackageResultStruct; }
 
 namespace UE::Cook
 {
@@ -300,7 +300,7 @@ public:
 	 * Calculates AssetRegistryData.
 	 */
 	void FinishGeneratedPlatformSave(FPackageData& PackageData,
-		UE::TargetDomain::FGeneratedPackageResultStruct& GeneratedResult);
+		FGeneratedPackageResultStruct& GeneratedResult);
 	/**
 	 * Return the AssetPackageData that was loaded for the given PackageData from the previous cook's
 	 * AssetRegistry, or nullptr if not an incremental cook or it was not previously cooked.
diff --git a/Engine/Source/Editor/UnrealEd/Private/Cooker/CookPackageArtifacts.cpp b/Engine/Source/Editor/UnrealEd/Private/Cooker/CookPackageArtifacts.cpp
new file mode 100644
--- /dev/null
+++ b/Engine/Source/Editor/UnrealEd/Private/Cooker/CookPackageArtifacts.cpp
@@ -0,0 +1,1033 @@
+// Copyright Epic Games, Inc. All Rights Reserved.
+
+#include "CookPackageArtifacts.h"
+
+#include "Algo/Sort.h"
+#include "Algo/Unique.h"
+#include "AssetRegistry/IAssetRegistry.h"
+#include "Cooker/CookConfigAccessTracker.h"
+#include "Cooker/CookDependency.h"
+#include "Cooker/PackageBuildDependencyTracker.h"
+#include "CookOnTheSide/CookLog.h"
+#include "CookPackageSplitter.h"
+#include "DerivedDataBuildDefinition.h"
+#include "DerivedDataBuildKey.h"
+#include "EditorDomain/EditorDomain.h"
+#include "EditorDomain/EditorDomainUtils.h"
+#include "HAL/PlatformFile.h"
+#include "HAL/PlatformFileManager.h"
+#include "Hash/Blake3.h"
+#include "Misc/Guid.h"
+#include "Misc/StringBuilder.h"
+#include "Serialization/CompactBinary.h"
+#include "Serialization/CompactBinarySerialization.h"
+#include "Serialization/CompactBinaryWriter.h"
+#include "TargetDomain/TargetDomainUtils.h"
+#include "UObject/CoreRedirects.h"
+
+namespace UE::Cook
+{
+
+// Change TargetDomainVersion to a new guid when all TargetDomain keys in an incremental cook need to be invalidated
+FGuid TargetDomainVersion(0x10E1BAE8, 0x6CA04033, 0xA7E57A37, 0xD967F493);
+
+// Bump PackageArtifactsVersion version when the serialization of PackageArtifacts has changed and we want to add
+// backwards compatibility rather than invalidating everything.
+constexpr uint32 PackageArtifactsVersion = 0x00000002;
+
+static const FUtf8StringView PackageArtifactsAttachmentKey = UTF8TEXTVIEW("CookPackageArtifacts");
+static const FUtf8StringView BuildDefinitionsAttachmentKey = UTF8TEXTVIEW("BuildDefinitionsAttachmentKey");
+
+// Constructor/Destructor defined here in cpp rather than header so we can 
+// avoid needing the definition of FCookDependency in the header; it is needed
+// for construct/destruct of TArray<FCookDependency>.
+FPackageArtifacts::FPackageArtifacts() = default;
+FPackageArtifacts::~FPackageArtifacts() = default;
+FPackageArtifacts::FPackageArtifacts(const FPackageArtifacts&) = default;
+FPackageArtifacts::FPackageArtifacts(FPackageArtifacts&&) = default;
+FPackageArtifacts& FPackageArtifacts::operator=(const FPackageArtifacts&) = default;
+FPackageArtifacts& FPackageArtifacts::operator=(FPackageArtifacts&&) = default;
+
+
+bool FPackageArtifacts::IsValid() const
+{
+	return bValid;
+}
+
+bool FPackageArtifacts::HasSaveResults() const
+{
+	return bHasSaveResults;
+}
+
+bool FPackageArtifacts::HasKeyMatch(const FAssetPackageData* OverrideAssetPackageData)
+{
+	if (!bValid)
+	{
+		return false;
+	}
+	if (StoredKey.IsZero())
+	{
+		return false;
+	}
+	if (CurrentKey.IsZero())
+	{
+		ECurrentKeyResult Result = TryCalculateCurrentKey(OverrideAssetPackageData);
+		if (Result != ECurrentKeyResult::Success)
+		{
+			CurrentKey.Reset();
+			return false;
+		}
+	}
+	return CurrentKey == StoredKey;
+}
+
+FPackageArtifacts::ECurrentKeyResult FPackageArtifacts::TryCalculateCurrentKey(
+	const FAssetPackageData* OverrideAssetPackageData,
+	TArray<TPair<ELogVerbosity::Type, FString>>* OutMessages)
+{
+	IAssetRegistry* AssetRegistry = IAssetRegistry::Get();
+	if (PackageName.IsNone())
+	{
+		if (OutMessages) OutMessages->Emplace(ELogVerbosity::Error, TEXT("PackageName is not set."));
+		return ECurrentKeyResult::Error;
+	}
+	if (!AssetRegistry)
+	{
+		if (OutMessages) OutMessages->Emplace(ELogVerbosity::Error, TEXT("AssetRegistry is unavailable."));
+		return ECurrentKeyResult::Error;
+	}
+	FEditorDomain* EditorDomain = FEditorDomain::Get();
+	if (!EditorDomain)
+	{
+		if (OutMessages) OutMessages->Emplace(ELogVerbosity::Error, TEXT("EditorDomain is unavailable."));
+		return ECurrentKeyResult::Error;
+	}
+	FBlake3 KeyBuilder;
+	KeyBuilder.Update(&TargetDomainVersion, sizeof(TargetDomainVersion));
+
+	UE::EditorDomain::FPackageDigest PackageDigest = OverrideAssetPackageData ?
+		UE::EditorDomain::CalculatePackageDigest(*OverrideAssetPackageData, PackageName) :
+		EditorDomain->GetPackageDigest(PackageName);
+	if (!PackageDigest.IsSuccessful())
+	{
+		if (OutMessages) OutMessages->Emplace(ELogVerbosity::Error, PackageDigest.GetStatusString());
+		return ECurrentKeyResult::Error;
+	}
+	KeyBuilder.Update(&PackageDigest.Hash, sizeof(PackageDigest.Hash));
+
+	if (!ClassDependencies.IsEmpty())
+	{
+		TArray<FTopLevelAssetPath, TInlineAllocator<1>> ClassPaths;
+		ClassPaths.Reserve(ClassDependencies.Num());
+		for (const FString& ClassPathStr : ClassDependencies)
+		{
+			FTopLevelAssetPath& ClassPath = ClassPaths.Emplace_GetRef(ClassPathStr);
+			if (!ClassPath.IsValid())
+			{
+				if (OutMessages)
+				{
+					OutMessages->Emplace(ELogVerbosity::Error, FString::Printf(
+						TEXT("ClassDependency failed: %s is not a valid TopLevelAssetPath."), *ClassPathStr));
+				}
+				return ECurrentKeyResult::Error;
+			}
+		}
+		FString AppendDigestError;
+		if (!UE::EditorDomain::TryAppendClassDigests(KeyBuilder, ClassPaths, &AppendDigestError))
+		{
+			if (OutMessages)
+			{
+				OutMessages->Emplace(ELogVerbosity::Error, FString::Printf(
+					TEXT("ClassDependency failed: %s"), *AppendDigestError));
+			}
+			return ECurrentKeyResult::Error;
+		}
+	}
+
+	for (FName PackageDependency : BuildPackageDependencies)
+	{
+		PackageDigest = EditorDomain->GetPackageDigest(PackageDependency);
+		if (!PackageDigest.IsSuccessful())
+		{
+			if (OutMessages)
+			{
+				OutMessages->Emplace(ELogVerbosity::Error, FString::Printf(TEXT("Could not create PackageDigest for %s: %s"),
+					*PackageDependency.ToString(), *PackageDigest.GetStatusString()));
+			}
+			return ECurrentKeyResult::Error;
+		}
+		KeyBuilder.Update(&PackageDigest.Hash, sizeof(PackageDigest.Hash));
+	}
+
+	FCoreRedirects::AppendHashOfRedirectsAffectingPackages(KeyBuilder, RuntimePackageDependencies);
+	FCoreRedirects::AppendHashOfRedirectsAffectingPackages(KeyBuilder, ScriptPackageDependencies);
+	FCoreRedirects::AppendHashOfRedirectsAffectingPackages(KeyBuilder, RedirectorsAffectingPackageDependencies);
+
+	if (!ConfigDependencies.IsEmpty())
+	{
+#if UE_WITH_CONFIG_TRACKING
+		using namespace UE::ConfigAccessTracking;
+		FCookConfigAccessTracker& ConfigTracker = FCookConfigAccessTracker::Get();
+#endif
+		for (const FString& ConfigDependency : ConfigDependencies)
+		{
+			FString Value;
+#if UE_WITH_CONFIG_TRACKING
+			Value = ConfigTracker.GetValue(ConfigDependency);
+#endif
+			uint8 Marker = 0;
+			KeyBuilder.Update(&Marker, sizeof(Marker));
+			if (!Value.IsEmpty())
+			{
+				KeyBuilder.Update(*Value, Value.Len() * sizeof(Value[0]));
+			}
+		}
+	}
+
+	ECurrentKeyResult Result = ECurrentKeyResult::Success;
+	if (!CookDependencies.IsEmpty())
+	{
+		UE::Cook::FCookDependencyContext Context(&KeyBuilder,
+			[&Result, OutMessages](ELogVerbosity::Type Verbosity, FString&& Message)
+			{
+				if (OutMessages)
+				{
+					OutMessages->Emplace(Verbosity, MoveTemp(Message));
+				}
+			},
+			[&Result](ELogVerbosity::Type Verbosity)
+			{
+				ECurrentKeyResult NewResult = Verbosity <= ELogVerbosity::Error
+					? ECurrentKeyResult::Error : ECurrentKeyResult::Invalidated;
+				Result = FMath::Max(Result, NewResult);
+			},
+			PackageName);
+
+		for (UE::Cook::FCookDependency& CookDependency : CookDependencies)
+		{
+			CookDependency.UpdateHash(Context);
+		}
+	}
+
+	if (Result != ECurrentKeyResult::Error) // -V547
+	{
+		CurrentKey = KeyBuilder.Finalize();
+	}
+	return Result;
+}
+
+void FPackageArtifacts::Reset()
+{
+	BuildPackageDependencies.Reset();
+	ConfigDependencies.Reset();
+	RuntimePackageDependencies.Reset();
+	ScriptPackageDependencies.Reset();
+	RedirectorsAffectingPackageDependencies.Reset();
+	PackageName = FName();
+	StoredKey = FIoHash::Zero;
+	CurrentKey = FIoHash::Zero;
+	bValid = false;
+}
+
+void FPackageArtifacts::Empty()
+{
+	Reset();
+	BuildPackageDependencies.Empty();
+	ConfigDependencies.Empty();
+	RuntimePackageDependencies.Empty();
+	RedirectorsAffectingPackageDependencies.Empty();
+}
+
+enum class EPackageMountPoint
+{
+	Transient,
+	Script,
+	Content,
+	GeneratedContent,
+};
+EPackageMountPoint GetPackageMountPoint(FName PackageName)
+{
+	TStringBuilder<256> StringBuffer;
+	PackageName.ToString(StringBuffer);
+	if (
+		// Some packages get renamed to "TrashedPackage" during blueprint compilation and are no longer valid for
+		// saving but might have been dereferenced by TObjectPtr during PostLoad/PreSave. We need to discard these
+		// packages, which we can do by requiring a valid package name; all valid packages start with /MountPoint/.
+		!StringBuffer.ToView().StartsWith('/') ||
+		// Ignore /Memory and /Temp packages
+		FPackageName::IsMemoryPackage(StringBuffer) ||
+		FPackageName::IsTempPackage(StringBuffer) ||
+		FPackageName::IsInEngineTransientPackages(StringBuffer))
+	{
+		return EPackageMountPoint::Transient;
+	}
+	if (FPackageName::IsScriptPackage(StringBuffer))
+	{
+		return EPackageMountPoint::Script;
+	}
+	if (ICookPackageSplitter::IsUnderGeneratedPackageSubPath(StringBuffer))
+	{
+		return EPackageMountPoint::GeneratedContent;
+	}
+	return EPackageMountPoint::Content;
+}
+
+static void JoinMessagesIntoErrorReason(FStringBuilderBase& OutText,
+	TArray<TPair<ELogVerbosity::Type, FString>>& Messages)
+{
+	if (Messages.IsEmpty())
+	{
+		OutText << TEXT(".");
+	}
+	else
+	{
+		OutText << TEXT(":");
+		for (TPair<ELogVerbosity::Type, FString>& MessagePair : Messages)
+		{
+			OutText << TEXT("\n\t") << MessagePair.Value;
+		}
+	}
+}
+
+FPackageArtifacts FPackageArtifacts::Collect(UPackage* Package, const ITargetPlatform* TargetPlatform,
+	FSavePackageResultStruct* SaveResult, const FGeneratedPackageResultStruct* GeneratedResult,
+	TArray<FName>&& RuntimeDependencies, TArray<TPair<ELogVerbosity::Type, FString>>* OutMessages)
+{
+	if (!Package)
+	{
+		if (OutMessages) OutMessages->Emplace(ELogVerbosity::Error, TEXT("Invalid null package."));
+		return FPackageArtifacts();
+	}
+	IAssetRegistry* AssetRegistry = IAssetRegistry::Get();
+	if (!AssetRegistry)
+	{
+		if (OutMessages) OutMessages->Emplace(ELogVerbosity::Error, TEXT("AssetRegistry is unavailable."));
+		return FPackageArtifacts();
+	}
+	FEditorDomain* EditorDomain = FEditorDomain::Get();
+	if (!EditorDomain)
+	{
+		if (OutMessages) OutMessages->Emplace(ELogVerbosity::Error, TEXT("EditorDomain is unavailable."));
+		return FPackageArtifacts();
+	}
+
+	FPackageArtifacts Result;
+	Result.PackageName = Package->GetFName();
+	Result.bHasSaveResults = SaveResult != nullptr;
+	enum class EBuildDependency
+	{
+		PackageOnly,
+		Transitive,
+	};
+
+	TMap<FName, EBuildDependency> PackageBuildDependencies;
+
+	TArray<FName> AssetDependencies;
+	if (!GeneratedResult)
+	{
+		AssetRegistry->GetDependencies(Result.PackageName, AssetDependencies,
+			UE::AssetRegistry::EDependencyCategory::Package, UE::AssetRegistry::EDependencyQuery::Game);
+		RuntimeDependencies.Append(MoveTemp(AssetDependencies));
+	}
+	else
+	{
+		// GeneratedResult->PackageDependencies are already incorporated into the PackageHash of the generated
+		// package by FCookGenerationInfo::CreatePackageHash, so we do not need to add them into the BuildDependencies
+	}
+
+#if UE_WITH_PACKAGE_ACCESS_TRACKING
+	FPackageBuildDependencyTracker& Tracker = FPackageBuildDependencyTracker::Get();
+
+	if (Tracker.IsEnabled())
+	{
+		TArray<TPair<FBuildDependencyAccessData, FResultProjectionList>> AccessDatas
+			= Tracker.GetAccessDatas(Result.PackageName);
+
+		PackageBuildDependencies.Reserve(AccessDatas.Num());
+		for (TPair<FBuildDependencyAccessData, FResultProjectionList>& Pair : AccessDatas)
+		{
+			FBuildDependencyAccessData& AccessData = Pair.Key;
+			if (AccessData.TargetPlatform == TargetPlatform || AccessData.TargetPlatform == nullptr)
+			{
+				FResultProjectionList& ProjectionList = Pair.Value;
+				constexpr bool bAutoTransitiveDependenciesEnabled = false;
+				if (!bAutoTransitiveDependenciesEnabled)
+				{
+					// We have not yet enabled marking the auto-added dependencies from TObjectPtr resolve
+					// as transitive, because it causes a performance regression and we are still working on fixing the
+					// regression.
+					PackageBuildDependencies.FindOrAdd(AccessData.ReferencedPackage, EBuildDependency::PackageOnly);
+				}
+				else if (ProjectionList.bHasAll)
+				{
+					// TObjectPtr UE::BuildProjection::All dependencies are added as transitive build dependencies.
+					// We have to do this to be conservative, since we do not know which bytes from the target are
+					// dependended upon and which of the target's build dependencies influence those bytes.
+					PackageBuildDependencies.FindOrAdd(AccessData.ReferencedPackage) = EBuildDependency::Transitive;
+				}
+				else
+				{
+					for (FTopLevelAssetPath ClassPath : ProjectionList.Classes)
+					{
+						Result.ClassDependencies.Add(ClassPath.ToString());
+					}
+
+					for (FName ResultProjection : ProjectionList.ResultProjections)
+					{
+						if (ResultProjection == UE::Cook::ResultProjection::PackageAndClass)
+						{
+							PackageBuildDependencies.FindOrAdd(AccessData.ReferencedPackage, EBuildDependency::PackageOnly);
+						}
+						else
+						{
+							UE_LOG(LogCook, Error,
+								TEXT("When saving %s, found ResultProjection %s, which is a system-specific ResultProjection, and this is not yet implemented. ")
+								TEXT("Find the call to UE_COOK_RESULTPROJECTION_SCOPED passing in this name and remove it."),
+								*Result.PackageName.ToString(), *ResultProjection.ToString());
+						}
+					}
+				}
+			}
+		}
+	}
+	else
+#endif
+	{
+		// When PackageAccessTracking is disabled, defensively treat all asset dependencies as transitive build
+		// dependencies.
+		for (FName AssetDependency : AssetDependencies)
+		{
+			PackageBuildDependencies.FindOrAdd(AssetDependency) = EBuildDependency::Transitive;
+		}
+	}
+
+#if UE_WITH_CONFIG_TRACKING
+	{
+		using namespace UE::ConfigAccessTracking;
+		FCookConfigAccessTracker& ConfigTracker = FCookConfigAccessTracker::Get();
+		if (ConfigTracker.IsEnabled())
+		{
+			TArray<FConfigAccessData> ConfigKeys = ConfigTracker.GetPackageRecords(Result.PackageName, TargetPlatform);
+			Result.ConfigDependencies.Reserve(ConfigKeys.Num());
+			for (const FConfigAccessData& ConfigKey : ConfigKeys)
+			{
+				Result.ConfigDependencies.Add(ConfigKey.FullPathToString());
+			}
+		}
+	}
+#endif
+	if (SaveResult)
+	{
+		// Convert some FCookDependency types into our specially-handled dependencies, and put the more
+		// generic ones into this->CookDependencies
+		Result.CookDependencies.Reserve(SaveResult->CookDependencies.Num());
+
+		// Additional dependencies can be added from SettingsObject dependencies. Add these included dependencies
+		// to a followup list that we process in the second round, if they exist. Handle detecting cycles in the
+		// included dependencies; cycle detection is currently simple: nested includes are not yet required
+		// so give an error if we ever make it to a third round.
+		TArray<UE::Cook::FCookDependency> NextIncludedDependenciesBuffer[2];
+		TArray<UE::Cook::FCookDependency>* RemainingDependencies = &SaveResult->CookDependencies;
+		int32 IncludedDependenciesRound = 0;
+		constexpr int32 MaxNumberOfIncludedDependenciesRounds = 2;
+		while (!RemainingDependencies->IsEmpty())
+		{
+			TArray<UE::Cook::FCookDependency>* NextIncludedDependencies =
+				&NextIncludedDependenciesBuffer[IncludedDependenciesRound % UE_ARRAY_COUNT(NextIncludedDependenciesBuffer)];
+			if (IncludedDependenciesRound++ > 0)
+			{
+				NextIncludedDependencies->Reset();
+				if (IncludedDependenciesRound > MaxNumberOfIncludedDependenciesRounds)
+				{
+					if (OutMessages)
+					{
+						OutMessages->Emplace(ELogVerbosity::Error,
+							TEXT("More rounds than expected when handling included dependencies."));
+					}
+					return FPackageArtifacts();
+				}
+			}
+
+			for (UE::Cook::FCookDependency& CookDependency : *RemainingDependencies)
+			{
+				switch (CookDependency.GetType())
+				{
+				case UE::Cook::ECookDependency::TransitiveBuild:
+				{
+					// Build dependencies from a package to itself have a performance cost and serve no
+					// purpose, so remove them. They can occur in some systems that naively add a build dependency
+					// from one object to another without checking whether the second object is in the same package.
+					FName DependencyPackageName = CookDependency.GetPackageName();
+					if (DependencyPackageName == Result.PackageName)
+					{
+						continue;
+					}
+					if (CookDependency.IsAlsoAddRuntimeDependency())
+					{
+						RuntimeDependencies.Add(DependencyPackageName);
+					}
+					PackageBuildDependencies.FindOrAdd(DependencyPackageName) = EBuildDependency::Transitive;
+					break;
+				}
+				case UE::Cook::ECookDependency::Package:
+					// Add it as a PackageOnly build dependency (lowest level) unless a higher level already exists.
+					PackageBuildDependencies.FindOrAdd(CookDependency.GetPackageName(), EBuildDependency::PackageOnly);
+					break;
+				case UE::Cook::ECookDependency::Config:
+					Result.ConfigDependencies.Add(FString(CookDependency.GetConfigPath()));
+					break;
+				case UE::Cook::ECookDependency::SettingsObject:
+				{
+					const UObject* SettingsObject = CookDependency.GetSettingsObject();
+					if (SettingsObject)
+					{
+						// We rely on the object to be rooted because we use its pointer as a key for the lifetime of
+						// the cook process, so it being garbage collected and something else allocated on the same
+						// pointer would break our key. IsRooted should have been validated by
+						// FCookDependency::SettingsObject.
+						check(SettingsObject->IsRooted());
+						FCookDependencyGroups::FRecordedDependencies& IncludeDependencies =
+							FCookDependencyGroups::Get().FindOrCreate(reinterpret_cast<UPTRINT>(SettingsObject));
+						if (!IncludeDependencies.bInitialized)
+						{
+							IncludeDependencies.Artifacts = CollectSettingsObject(SettingsObject,
+								&IncludeDependencies.Messages);
+							IncludeDependencies.bInitialized = true;
+						}
+						if (!IncludeDependencies.Artifacts.IsValid())
+						{
+							if (OutMessages)
+							{
+								TStringBuilder<256> ErrorText;
+								ErrorText << TEXT("Dependencies for SettingsObject ") << SettingsObject->GetPathName()
+									<< TEXT(" are unavailable");
+								JoinMessagesIntoErrorReason(ErrorText, IncludeDependencies.Messages);
+								OutMessages->Emplace(ELogVerbosity::Error, FString(*ErrorText));
+							}
+							return FPackageArtifacts();
+						}
+
+						for (const UE::Cook::FCookDependency& IncludeDependency :
+							IncludeDependencies.Artifacts.GetCookDependencies())
+						{
+							NextIncludedDependencies->Add(IncludeDependency);
+						}
+					}
+					break;
+				}
+				case UE::Cook::ECookDependency::NativeClass:
+					Result.ClassDependencies.Add(FString(CookDependency.GetClassPath()));
+					break;
+				default:
+					Result.CookDependencies.Add(MoveTemp(CookDependency));
+					break;
+				}
+			}
+			RemainingDependencies = NextIncludedDependencies;
+		}
+
+		Result.RedirectorsAffectingPackageDependencies = SaveResult->UntrackedSoftPackageReferences;
+		Result.RedirectorsAffectingPackageDependencies.Sort(FNameLexicalLess());
+		Result.RedirectorsAffectingPackageDependencies.SetNum(Algo::Unique(Result.RedirectorsAffectingPackageDependencies));
+	}
+
+	Algo::Sort(Result.CookDependencies);
+	Result.CookDependencies.SetNum(Algo::Unique(Result.CookDependencies));
+	for (TPair<FName, EBuildDependency>& Pair : PackageBuildDependencies)
+	{
+		if (Pair.Key == Result.PackageName)
+		{
+			// Skip dependencies to self, for performance
+			continue;
+		}
+		if (GetPackageMountPoint(Pair.Key) != EPackageMountPoint::Content)
+		{
+			// Skip script and temp packages
+			continue;
+		}
+		if (Pair.Value == EBuildDependency::Transitive)
+		{
+			Result.TransitiveBuildDependencies.Add(UE::Cook::FCookDependency::Package(Pair.Key));
+		}
+		else
+		{
+			Result.BuildPackageDependencies.Add(Pair.Key);
+		}
+	}
+	Algo::Sort(Result.TransitiveBuildDependencies);
+	Result.BuildPackageDependencies.Sort(FNameLexicalLess());
+	Algo::Sort(Result.ConfigDependencies);
+	Result.ConfigDependencies.SetNum(Algo::Unique(Result.ConfigDependencies));
+	Algo::Sort(Result.ClassDependencies);
+	Result.ClassDependencies.SetNum(Algo::Unique(Result.ClassDependencies));
+
+
+	for (TArray<FName>::TIterator Iter(RuntimeDependencies); Iter; ++Iter)
+	{
+		FName PackageName = *Iter;
+		EPackageMountPoint MountPoint = GetPackageMountPoint(PackageName);
+		switch (MountPoint)
+		{
+		case EPackageMountPoint::GeneratedContent:
+		case EPackageMountPoint::Content:
+			Result.RuntimePackageDependencies.Add(PackageName);
+			break;
+		case EPackageMountPoint::Script:
+			Result.ScriptPackageDependencies.Add(PackageName);
+			break;
+		default:
+			break;
+		}
+	}
+	Result.RuntimePackageDependencies.Sort(FNameLexicalLess());
+	Result.RuntimePackageDependencies.SetNum(Algo::Unique(Result.RuntimePackageDependencies), EAllowShrinking::Yes);
+	Result.ScriptPackageDependencies.Sort(FNameLexicalLess());
+	Result.ScriptPackageDependencies.SetNum(Algo::Unique(Result.ScriptPackageDependencies), EAllowShrinking::Yes);
+
+	const FAssetPackageData* AssetPackageData = GeneratedResult ? &GeneratedResult->AssetPackageData : nullptr;
+	ECurrentKeyResult CurrentKeyResult = Result.TryCalculateCurrentKey(AssetPackageData, OutMessages);
+	if (CurrentKeyResult == ECurrentKeyResult::Error)
+	{
+		return FPackageArtifacts();
+	}
+	Result.StoredKey = Result.CurrentKey;
+	Result.bValid = true;
+
+	return Result;
+}
+
+FPackageArtifacts FPackageArtifacts::CollectSettingsObject(const UObject* Object,
+	TArray<TPair<ELogVerbosity::Type, FString>>* OutMessages)
+{
+	if (!Object)
+	{
+		if (OutMessages)
+		{
+			OutMessages->Emplace(ELogVerbosity::Error, TEXT("Invalid null Object."));
+		}
+		return FPackageArtifacts();
+	}
+
+	UClass* Class = Object->GetClass();
+	if (!Class->HasAnyClassFlags(CLASS_Config | CLASS_PerObjectConfig))
+	{
+		if (OutMessages)
+		{
+			OutMessages->Emplace(ELogVerbosity::Error, FString::Printf(
+				TEXT("Class %s is not a config class."), *Class->GetPathName()));
+		}
+		return FPackageArtifacts();
+	}
+	if (!Class->HasAnyClassFlags(CLASS_PerObjectConfig) && Object != Class->GetDefaultObject())
+	{
+		if (OutMessages)
+		{
+			OutMessages->Emplace(ELogVerbosity::Error, FString::Printf(
+				TEXT("Class %s is not a per-object-config class."), *Class->GetPathName()));
+		}
+		return FPackageArtifacts();
+	}
+
+	FPackageArtifacts Result;
+	TArray<UE::ConfigAccessTracking::FConfigAccessData> ConfigDatas;
+	const_cast<UObject*>(Object)->LoadConfig(nullptr /* ConfigClass */, nullptr /* Filename */,
+		UE::LCPF_None /* PropagationFlags */, nullptr /* PropertyToLoad */, &ConfigDatas);
+	Result.CookDependencies.Reserve(ConfigDatas.Num() + 1);
+	for (const UE::ConfigAccessTracking::FConfigAccessData& ConfigData : ConfigDatas)
+	{
+		Result.CookDependencies.Add(UE::Cook::FCookDependency::Config(ConfigData));
+	}
+
+	// In addition to adding the config dependencies, add a dependency on the class schema. If the current class has
+	// config fields A,B,C, we add dependencies on those config values. But if the class header is modified to have
+	// additional config field D then we need to rebuild packages that depend on it to record the new dependency on D.
+	UClass* NativeClass = Class;
+	while (NativeClass && !NativeClass->IsNative())
+	{
+		NativeClass = NativeClass->GetSuperClass();
+	}
+	if (NativeClass)
+	{
+		Result.CookDependencies.Add(UE::Cook::FCookDependency::NativeClass(NativeClass));
+	}
+
+	Algo::Sort(Result.CookDependencies);
+	Result.CookDependencies.SetNum(Algo::Unique(Result.CookDependencies));
+	Result.bValid = true;
+	return Result;
+}
+
+}
+
+bool LoadFromCompactBinary(FCbObjectView ObjectView, UE::Cook::FPackageArtifacts& Artifacts)
+{
+	using namespace UE::Cook;
+
+	Artifacts.Reset();
+	int32 Version = -1;
+
+	for (FCbFieldViewIterator FieldView(ObjectView.CreateViewIterator()); FieldView; )
+	{
+		const FCbFieldViewIterator Last = FieldView;
+		if (FieldView.GetName().Equals(UTF8TEXTVIEW("Version")))
+		{
+			Version = FieldView.AsInt32();
+			if ((FieldView++).HasError() || Version != PackageArtifactsVersion)
+			{
+				return false;
+			}
+		}
+		if (FieldView.GetName().Equals(UTF8TEXTVIEW("StoredKey")))
+		{
+			if (!LoadFromCompactBinary(FieldView++, Artifacts.StoredKey))
+			{
+				return false;
+			}
+		}
+		if (FieldView.GetName().Equals(UTF8TEXTVIEW("HasSaveResults")))
+		{
+			if (!LoadFromCompactBinary(FieldView++, Artifacts.bHasSaveResults))
+			{
+				return false;
+			}
+		}
+		if (FieldView.GetName().Equals(UTF8TEXTVIEW("BuildPackageDependencies")))
+		{
+			if (!LoadFromCompactBinary(FieldView++, Artifacts.BuildPackageDependencies))
+			{
+				return false;
+			}
+		}
+		if (FieldView.GetName().Equals(UTF8TEXTVIEW("ConfigDependencies")))
+		{
+			if (!LoadFromCompactBinary(FieldView++, Artifacts.ConfigDependencies))
+			{
+				return false;
+			}
+		}
+		if (FieldView.GetName().Equals(UTF8TEXTVIEW("RuntimePackageDependencies")))
+		{
+			if (!LoadFromCompactBinary(FieldView++, Artifacts.RuntimePackageDependencies))
+			{
+				return false;
+			}
+		}
+		if (FieldView->GetName().Equals(UTF8TEXTVIEW("RedirectorsAffectingPackageDependencies")))
+		{
+			if (!LoadFromCompactBinary(FieldView++, Artifacts.RedirectorsAffectingPackageDependencies))
+			{
+				return false;
+			}
+		}
+		if (FieldView.GetName().Equals(UTF8TEXTVIEW("ScriptPackageDependencies")))
+		{
+			if (!LoadFromCompactBinary(FieldView++, Artifacts.ScriptPackageDependencies))
+			{
+				return false;
+			}
+		}
+		if (FieldView.GetName().Equals(UTF8TEXTVIEW("CookDependencies")))
+		{
+			if (!LoadFromCompactBinary(FieldView++, Artifacts.CookDependencies))
+			{
+				return false;
+			}
+		}
+		if (FieldView.GetName().Equals(UTF8TEXTVIEW("TransitiveBuildDependencies")))
+		{
+			if (!LoadFromCompactBinary(FieldView++, Artifacts.TransitiveBuildDependencies))
+			{
+				return false;
+			}
+		}
+		if (FieldView.GetName().Equals(UTF8TEXTVIEW("ClassDependencies")))
+		{
+			if (!LoadFromCompactBinary(FieldView++, Artifacts.ClassDependencies))
+			{
+				return false;
+			}
+		}
+		if (FieldView == Last)
+		{
+			++FieldView;
+		}
+	}
+	if (Version == -1)
+	{
+		return false;
+	}
+	Artifacts.bValid = true;
+	return true;
+}
+
+FCbWriter& operator<<(FCbWriter& Writer, const UE::Cook::FPackageArtifacts& Artifacts)
+{
+	using namespace UE::Cook;
+
+	Writer.BeginObject();
+	Writer << "Version" << PackageArtifactsVersion;
+	Writer << "StoredKey" << Artifacts.StoredKey;
+	Writer << "HasSaveResults" << Artifacts.bHasSaveResults;
+	if (!Artifacts.BuildPackageDependencies.IsEmpty())
+	{
+		Writer << "BuildPackageDependencies" << Artifacts.BuildPackageDependencies;
+	}
+	if (!Artifacts.ConfigDependencies.IsEmpty())
+	{
+		Writer << "ConfigDependencies" << Artifacts.ConfigDependencies;
+	}
+	if (!Artifacts.RuntimePackageDependencies.IsEmpty())
+	{
+		Writer << "RuntimePackageDependencies" << Artifacts.RuntimePackageDependencies;
+	}
+	if (!Artifacts.RedirectorsAffectingPackageDependencies.IsEmpty())
+	{
+		Writer << "RedirectorsAffectingPackageDependencies" << Artifacts.RedirectorsAffectingPackageDependencies;
+	}
+	if (!Artifacts.ScriptPackageDependencies.IsEmpty())
+	{
+		Writer << "ScriptPackageDependencies" << Artifacts.ScriptPackageDependencies;
+	}
+	if (!Artifacts.CookDependencies.IsEmpty())
+	{
+		Writer << "CookDependencies" << Artifacts.CookDependencies;
+	}
+	if (!Artifacts.TransitiveBuildDependencies.IsEmpty())
+	{
+		Writer << "TransitiveBuildDependencies" << Artifacts.TransitiveBuildDependencies;
+	}
+	if (!Artifacts.ClassDependencies.IsEmpty())
+	{
+		Writer << "ClassDependencies" << Artifacts.ClassDependencies;
+	}
+
+	Writer.EndObject();
+	return Writer;
+}
+
+namespace UE::Cook
+{
+
+FCookDependencyGroups& FCookDependencyGroups::Get()
+{
+	static FCookDependencyGroups Singleton;
+	return Singleton;
+}
+
+FCookDependencyGroups::FRecordedDependencies& FCookDependencyGroups::FindOrCreate(UPTRINT Key)
+{
+	return Groups.FindOrAdd(Key);
+}
+
+FBuildDefinitionList FBuildDefinitionList::Collect(UPackage* Package, const ITargetPlatform* TargetPlatform,
+	TArray<TPair<ELogVerbosity::Type, FString>>* OutMessages)
+{
+	using namespace UE::DerivedData;
+
+	FBuildDefinitionList Result;
+
+	// TODO_BuildDefinitionList: Calculate and store BuildDefinitionList on the PackageData, or collect it here from some other source.
+	if (Result.Definitions.IsEmpty())
+	{
+		if (OutMessages) OutMessages->Emplace(ELogVerbosity::Error, TEXT("Not yet implemented"));
+		return FBuildDefinitionList();
+	}
+
+	TArray<FBuildDefinition>& Defs = Result.Definitions;
+	Algo::Sort(Defs, [](const FBuildDefinition& A, const FBuildDefinition& B)
+		{
+			return A.GetKey().Hash < B.GetKey().Hash;
+		});
+
+	return Result;
+}
+
+void FBuildDefinitionList::Reset()
+{
+	Definitions.Reset();
+}
+
+void FBuildDefinitionList::Empty()
+{
+	Definitions.Empty();
+}
+
+}
+
+bool LoadFromCompactBinary(FCbObject&& Object, UE::Cook::FBuildDefinitionList& Definitions)
+{
+	using namespace UE::DerivedData;
+
+	FCbField DefinitionsField = Object["BuildDefinitions"];
+	FCbArray DefinitionsArrayField = DefinitionsField.AsArray();
+	if (DefinitionsField.HasError())
+	{
+		return false;
+	}
+	TArray<FBuildDefinition>& Defs = Definitions.Definitions;
+	Defs.Empty(DefinitionsArrayField.Num());
+	for (FCbField& BuildDefinitionObj : DefinitionsArrayField)
+	{
+		FOptionalBuildDefinition BuildDefinition = FBuildDefinition::Load(TEXTVIEW("TargetDomainBuildDefinitionList"),
+			BuildDefinitionObj.AsObject());
+		if (!BuildDefinition)
+		{
+			Defs.Empty();
+			return false;
+		}
+		Defs.Add(MoveTemp(BuildDefinition).Get());
+	}
+
+	return true;
+}
+
+FCbWriter& operator<<(FCbWriter& Writer, const UE::Cook::FBuildDefinitionList& Definitions)
+{
+	using namespace UE::DerivedData;
+
+	Writer.BeginObject();
+	Writer.BeginArray("BuildDefinitions");
+	for (const FBuildDefinition& BuildDefinition : Definitions.Definitions)
+	{
+		BuildDefinition.Save(Writer);
+	}
+	Writer.EndArray();
+	return Writer;
+}
+
+namespace UE::Cook
+{
+
+void FIncrementalCookAttachments::Reset()
+{
+	Artifacts.Reset();
+	BuildDefinitions.Reset();
+}
+
+void FIncrementalCookAttachments::Empty()
+{
+	Artifacts.Empty();
+	BuildDefinitions.Empty();
+}
+
+void FIncrementalCookAttachments::AppendCommitAttachments(TArray<IPackageWriter::FCommitAttachmentInfo>& OutAttachments)
+{
+	if (Artifacts.IsValid())
+	{
+		FCbWriter Writer;
+		Writer << Artifacts;
+		IPackageWriter::FCommitAttachmentInfo& Attachment = OutAttachments.Emplace_GetRef();
+		Attachment.Key = PackageArtifactsAttachmentKey;
+		Attachment.Value = Writer.Save().AsObject();
+	}
+
+	if (!BuildDefinitions.Definitions.IsEmpty())
+	{
+		FCbWriter Writer;
+		Writer << BuildDefinitions;
+		IPackageWriter::FCommitAttachmentInfo& Attachment = OutAttachments.Emplace_GetRef();
+		Attachment.Key = BuildDefinitionsAttachmentKey;
+		Attachment.Value = Writer.Save().AsObject();
+	}
+}
+
+FIncrementalCookAttachments FIncrementalCookAttachments::Collect(UPackage* Package,
+	const ITargetPlatform* TargetPlatform, FSavePackageResultStruct* SaveResult,
+	const FGeneratedPackageResultStruct* GeneratedResult, TArray<FName>&& RuntimeDependencies)
+{
+	FIncrementalCookAttachments Result;
+
+	Result.CommitStatus = IPackageWriter::ECommitStatus::NotCommitted;
+
+	TArray<TPair<ELogVerbosity::Type, FString>> Messages;
+	Result.Artifacts = FPackageArtifacts::Collect(Package, TargetPlatform, SaveResult,
+		GeneratedResult, MoveTemp(RuntimeDependencies), &Messages);
+	if (!Result.Artifacts.IsValid())
+	{
+		TStringBuilder<256> LogText;
+		LogText << TEXT("Could not collect PackageArtifacts for package '") << Package->GetFName() << TEXT("'");
+		JoinMessagesIntoErrorReason(LogText, Messages);
+
+		// IterativeTODO: This error occurs due to dependencies on _Verse.
+		// Raise Verbosity to Error once that is fixed.
+		UE_LOG(LogCook, Verbose, TEXT("%s"), *LogText);
+	}
+
+	Result.BuildDefinitions = FBuildDefinitionList::Collect(Package, TargetPlatform);
+
+	return Result;
+}
+
+void FIncrementalCookAttachments::Fetch(TArrayView<FName> PackageNames, const ITargetPlatform* TargetPlatform,
+	ICookedPackageWriter* PackageWriter,
+	TUniqueFunction<void(FName PackageName, FIncrementalCookAttachments&& Result)>&& Callback)
+{
+	using namespace UE::TargetDomain;
+
+	struct FInProgressResult
+	{
+		FIncrementalCookAttachments Result;
+		int ReceivedAttachmentCount = 0;
+	};
+	TMap<FName, FInProgressResult> InProgressResults;
+	TArray<FUtf8StringView> RequestedAttachments = { PackageArtifactsAttachmentKey, BuildDefinitionsAttachmentKey };
+
+	if (TargetPlatform || GEditorDomainOplog)
+	{
+		auto OnOplogAttachment =
+			[PackageWriter, Callback = MoveTemp(Callback), RequestedAttachmentNum = RequestedAttachments.Num(), &InProgressResults]
+			(FName PackageName, FUtf8StringView AttachmentKey, FCbObject&& Attachment)
+			{
+				FInProgressResult& InProgressResult = InProgressResults.FindOrAdd(PackageName);
+				InProgressResult.ReceivedAttachmentCount++;
+				if (AttachmentKey == PackageArtifactsAttachmentKey)
+				{
+					if (PackageWriter)
+					{
+						InProgressResult.Result.CommitStatus = PackageWriter->GetCommitStatus(PackageName);
+					}
+					else
+					{
+						InProgressResult.Result.CommitStatus = Attachment ? IPackageWriter::ECommitStatus::Success
+							: IPackageWriter::ECommitStatus::NotCommitted;
+					}
+
+					if (LoadFromCompactBinary(MoveTemp(Attachment), InProgressResult.Result.Artifacts))
+					{
+						InProgressResult.Result.Artifacts.PackageName = PackageName;
+					}
+				}
+				else if (AttachmentKey == BuildDefinitionsAttachmentKey)
+				{
+					LoadFromCompactBinary(MoveTemp(Attachment), InProgressResult.Result.BuildDefinitions);
+				}
+
+				if (InProgressResult.ReceivedAttachmentCount == RequestedAttachmentNum)
+				{
+					Callback(PackageName, MoveTemp(InProgressResult.Result));
+					InProgressResults.Remove(PackageName);
+				}
+			};
+
+		if (PackageWriter)
+		{
+			PackageWriter->GetOplogAttachments(PackageNames,
+				RequestedAttachments, OnOplogAttachment);
+		}
+		else
+		{
+			GEditorDomainOplog->GetOplogAttachments(PackageNames,
+				RequestedAttachments, OnOplogAttachment);
+		}
+	}
+	else
+	{
+		for (FName PackageName : PackageNames)
+		{
+			Callback(PackageName, FIncrementalCookAttachments());
+		}
+	}
+	check(InProgressResults.Num() == 0);
+}
+
+} // namespace UE::Cook
diff --git a/Engine/Source/Editor/UnrealEd/Private/Cooker/CookPackageArtifacts.h b/Engine/Source/Editor/UnrealEd/Private/Cooker/CookPackageArtifacts.h
new file mode 100644
--- /dev/null
+++ b/Engine/Source/Editor/UnrealEd/Private/Cooker/CookPackageArtifacts.h
@@ -0,0 +1,260 @@
+// Copyright Epic Games, Inc. All Rights Reserved.
+
+#pragma once
+
+#include "AssetRegistry/AssetData.h"
+#include "AssetRegistry/IAssetRegistry.h"
+#include "Containers/Array.h"
+#include "Containers/ArrayView.h"
+#include "Containers/Map.h"
+#include "Containers/UnrealString.h"
+#include "DerivedDataBuildDefinition.h"
+#include "IO/IoHash.h"
+#include "Logging/LogVerbosity.h"
+#include "Serialization/PackageWriter.h"
+#include "Templates/Tuple.h"
+#include "UObject/NameTypes.h"
+
+class FCbObject;
+class FCbObjectView;
+class FCbWriter;
+class ITargetPlatform;
+class UObject;
+class UPackage;
+struct FSavePackageResultStruct;
+template <typename FuncType> class TUniqueFunction;
+
+namespace UE::Cook { class FCookDependency; }
+namespace UE::Cook { class FPackageArtifacts; }
+namespace UE::Cook { struct FBuildDefinitionList; }
+namespace UE::Cook { struct FIncrementalCookAttachments; }
+
+bool LoadFromCompactBinary(FCbObjectView ObjectView, UE::Cook::FPackageArtifacts& CookAttachments);
+FCbWriter& operator<<(FCbWriter& Writer, const UE::Cook::FPackageArtifacts& CookAttachments);
+bool LoadFromCompactBinary(FCbObject&& Object, UE::Cook::FBuildDefinitionList& Definitions);
+FCbWriter& operator<<(FCbWriter& Writer, const UE::Cook::FBuildDefinitionList& Definitions);
+
+namespace UE::Cook
+{
+
+/**
+ * Information collected from a CookPackageSplitter Generated package after it is saved; this information is needed to collect
+ * the CookDependencies for the generated package.
+ */
+struct FGeneratedPackageResultStruct
+{
+	FAssetPackageData AssetPackageData;
+	TArray<FAssetDependency> PackageDependencies;
+};
+
+/**
+ * Non-runtime data recorded about each package and stored in the cook Oplog as attachments to the package.
+ * Includes BuildResults built from the package that can be used for future incremental cooks, and the dependencies
+ * discovered for those BuildResults while the package was loading and cook-saving.
+ *
+ * Notes about the dependencies:
+ * All dependencies except for those marked Runtime contribute to the BuildResult's TargetDomain Key. If HasKeyMatch
+ * returns false after fetching this structure for a package at the beginning of cook, then the package is not
+ * iteratively skippable and needs to be recooked, and this structure needs to be recalculated for the package.
+ *
+ * Runtime fields on the dependencies are used to inform the cook of discovered softreferences that need to be added to
+ * the cook when the package is cooked.
+ */
+class FPackageArtifacts
+{
+public:
+	FPackageArtifacts();
+	~FPackageArtifacts();
+	FPackageArtifacts(const FPackageArtifacts&);
+	FPackageArtifacts(FPackageArtifacts&&);
+	FPackageArtifacts& operator=(const FPackageArtifacts&);
+	FPackageArtifacts& operator=(FPackageArtifacts&&);
+	// Build Dependencies
+	const TArray<FName>& GetBuildPackageDependencies() const;
+	const TArray<FString>& GetConfigDependencies() const;
+
+	// Runtime Dependencies
+	const TArray<FName>& GetRuntimePackageDependencies() const;
+	// Script dependencies, needed for TryCalculateCurrentKey
+	const TArray<FName>& GetScriptPackageDependencies() const;
+
+	// Cook Dependencies
+	const TArray<UE::Cook::FCookDependency>& GetCookDependencies() const;
+	const TArray<UE::Cook::FCookDependency>& GetTransitiveBuildDependencies() const;
+
+	// Data about the Key and Package that are not dependencies
+	/**
+		* True if the structure has been calculated or fetched and accurately reports dependencies and
+		* key for the package. False if the stucture is default, has been reset, or was marked invalid.
+		*/
+	bool IsValid() const;
+	bool HasSaveResults() const;
+	FName GetPackageName() const;
+	const FIoHash& GetStoredKey() const;
+	const FIoHash& GetCurrentKey() const;
+	bool HasKeyMatch(const FAssetPackageData* OverrideAssetPackageData);
+
+	// Modifying the structure data
+	enum class ECurrentKeyResult
+	{
+		Success,
+		Invalidated,
+		Error,
+	};
+	ECurrentKeyResult TryCalculateCurrentKey(const FAssetPackageData* OverrideAssetPackageData,
+		TArray<TPair<ELogVerbosity::Type, FString>>* OutMessages = nullptr);
+	void Reset();
+	void Empty();
+
+	// Construction functions
+	/**
+		* Read dependencies for the given targetplatform of the given package out of global dependency trackers
+		* that have recorded its data during the package's load/other/save operations in the current cook session.
+		*/
+	static FPackageArtifacts Collect(UPackage* Package, const ITargetPlatform* TargetPlatform,
+		FSavePackageResultStruct* SaveResult, const FGeneratedPackageResultStruct* GeneratedResult,
+		TArray<FName>&& RuntimeDependencies, TArray<TPair<ELogVerbosity::Type, FString>>* OutMessages = nullptr);
+	static FPackageArtifacts CollectSettingsObject(const UObject* Object,
+		TArray<TPair<ELogVerbosity::Type, FString>>* OutMessages);
+
+	// Fetch function to load the dependencies from a PackageStore is not yet implemented independently for
+	// this structure. Use FIncrementalCookAttachments instead.
+
+private:
+	TArray<FName> BuildPackageDependencies;
+	TArray<FString> ConfigDependencies;
+	TArray<FName> RuntimePackageDependencies;
+	TArray<FName> ScriptPackageDependencies;
+	TArray<FName> RedirectorsAffectingPackageDependencies;
+	TArray<UE::Cook::FCookDependency> CookDependencies;
+	TArray<UE::Cook::FCookDependency> TransitiveBuildDependencies;
+	TArray<FString> ClassDependencies;
+	FName PackageName;
+	FIoHash StoredKey;
+	FIoHash CurrentKey;
+	bool bHasSaveResults = false;
+	bool bValid = false;
+
+	friend bool ::LoadFromCompactBinary(FCbObjectView ObjectView, FPackageArtifacts& CookAttachments);
+	friend FCbWriter& ::operator<<(FCbWriter& Writer, const FPackageArtifacts& CookAttachments);
+	friend FIncrementalCookAttachments;
+};
+
+/**
+ * Non-persistent cache of groups of cookdependencies. Dependencies to a CookDependencyGroup are not persistently
+ * recorded into the oplog, instead we make a copy of all of their dependencies and append those dependencies onto
+ * the CookDependencies that are written for a package.
+ *
+ * Example: The cookdependencies used by the CDO of a settings object that itself is configured by config values.
+ *          The settings object's class's schema and the list of config settings are included in the cookdependencies.
+ */
+class FCookDependencyGroups
+{
+public:
+	struct FRecordedDependencies
+	{
+		FPackageArtifacts Artifacts;
+		TArray<TPair<ELogVerbosity::Type, FString>> Messages;
+		bool bInitialized = false;
+	};
+
+	static FCookDependencyGroups& Get();
+	FRecordedDependencies& FindOrCreate(UPTRINT Key);
+
+private:
+	TMap<UPTRINT, FRecordedDependencies> Groups;
+};
+
+/** Wrapper around TArray<FBuildDefinition>, used to provide custom functions for compactbinary, collection, and fetch */
+struct FBuildDefinitionList
+{
+public:
+	TArray<UE::DerivedData::FBuildDefinition> Definitions;
+
+	void Reset();
+	void Empty();
+
+	/** Collect DDC BuildDefinitions that were issued from the load/save of the given package and platform. */
+	static FBuildDefinitionList Collect(UPackage* Package, const ITargetPlatform* TargetPlatform,
+		TArray<TPair<ELogVerbosity::Type, FString>>* OutMessages = nullptr);
+
+private:
+	friend bool ::LoadFromCompactBinary(FCbObject&& ObjectView, FBuildDefinitionList& Definitions);
+	friend FCbWriter& ::operator<<(FCbWriter& Writer, const FBuildDefinitionList& Definitions);
+};
+
+/** All of the metadata that is written/read to the oplog for the incremental cook of a package. */
+struct FIncrementalCookAttachments
+{
+	FPackageArtifacts Artifacts;
+	FBuildDefinitionList BuildDefinitions;
+	IPackageWriter::ECommitStatus CommitStatus = IPackageWriter::ECommitStatus::NotCommitted;
+
+	void Reset();
+	void Empty();
+	void AppendCommitAttachments(TArray<IPackageWriter::FCommitAttachmentInfo>& OutAttachments);
+
+	static void Fetch(TArrayView<FName> PackageNames, const ITargetPlatform* TargetPlatform,
+		ICookedPackageWriter* PackageWriter,
+		TUniqueFunction<void(FName PackageName, FIncrementalCookAttachments&& Result)>&& Callback);
+	static FIncrementalCookAttachments Collect(UPackage* Package,
+		const ITargetPlatform* TargetPlatform, FSavePackageResultStruct* SaveResult,
+		const FGeneratedPackageResultStruct* GeneratedResult, TArray<FName>&& RuntimeDependencies);
+};
+
+} // namespace UE::Cook
+
+
+////////////////////////////////
+// Inline Implementations
+////////////////////////////////
+
+namespace UE::Cook
+{
+
+inline const TArray<FName>& FPackageArtifacts::GetBuildPackageDependencies() const
+{
+	return BuildPackageDependencies;
+}
+
+inline const TArray<FString>& FPackageArtifacts::GetConfigDependencies() const
+{
+	return ConfigDependencies;
+}
+
+inline const TArray<FName>& FPackageArtifacts::GetRuntimePackageDependencies() const
+{
+	return RuntimePackageDependencies;
+}
+
+inline const TArray<FName>& FPackageArtifacts::GetScriptPackageDependencies() const
+{
+	return ScriptPackageDependencies;
+}
+
+inline const TArray<UE::Cook::FCookDependency>& FPackageArtifacts::GetCookDependencies() const
+{
+	return CookDependencies;
+}
+
+inline const TArray<UE::Cook::FCookDependency>& FPackageArtifacts::GetTransitiveBuildDependencies() const
+{
+	return TransitiveBuildDependencies;
+}
+
+inline FName FPackageArtifacts::GetPackageName() const
+{
+	return PackageName;
+}
+
+inline const FIoHash& FPackageArtifacts::GetStoredKey() const
+{
+	return StoredKey;
+}
+
+inline const FIoHash& FPackageArtifacts::GetCurrentKey() const
+{
+	return CurrentKey;
+}
+
+} // namespace UE::Cook
diff --git a/Engine/Source/Editor/UnrealEd/Private/Cooker/CookRequestCluster.cpp b/Engine/Source/Editor/UnrealEd/Private/Cooker/CookRequestCluster.cpp
--- a/Engine/Source/Editor/UnrealEd/Private/Cooker/CookRequestCluster.cpp
+++ b/Engine/Source/Editor/UnrealEd/Private/Cooker/CookRequestCluster.cpp
@@ -31,6 +31,7 @@
 #include "Misc/ScopeExit.h"
 #include "Misc/StringBuilder.h"
 #include "String/Find.h"
+#include "TargetDomain/TargetDomainUtils.h"
 #include "UObject/CoreRedirects.h"
 #include "UObject/SavePackage.h"
 
@@ -1549,7 +1550,7 @@ bool FRequestCluster::FGraphSearch::FExploreEdgesContext::TryCalculateIterativel
 		const ITargetPlatform* TargetPlatform = FetchPlatformData.Platform;
 		FPackagePlatformData& PackagePlatformData = PackageData->FindOrAddPlatformData(TargetPlatform);
 
-		UE::TargetDomain::FCookDependencies& CookDependencies = QueryPlatformData.CookAttachments.Dependencies;
+		UE::Cook::FPackageArtifacts& Artifacts = QueryPlatformData.CookAttachments.Artifacts;
 		const FAssetPackageData* OverrideAssetPackageData = nullptr;
 		FPackageData* ParentPackageData = nullptr;
 		if (PackageData->IsGenerated())
@@ -1581,7 +1582,7 @@ bool FRequestCluster::FGraphSearch::FExploreEdgesContext::TryCalculateIterativel
 				continue;
 			}
 		}
-		if (!CookDependencies.HasKeyMatch(OverrideAssetPackageData))
+		if (!Artifacts.HasKeyMatch(OverrideAssetPackageData))
 		{
 			SetIsIterativelyUnmodified(PlatformIndex, false, PackagePlatformData);
 			continue;
@@ -1607,7 +1608,7 @@ bool FRequestCluster::FGraphSearch::FExploreEdgesContext::TryCalculateIterativel
 			bool bAnyTransitiveBuildDependencyIsModified = false;
 			UnreadyTransitiveBuildVertices.Reset();
 			for (const UE::Cook::FCookDependency& TransitiveBuildDependency :
-				CookDependencies.GetTransitiveBuildDependencies())
+				Artifacts.GetTransitiveBuildDependencies())
 			{
 				FName TransitiveBuildPackageName = TransitiveBuildDependency.GetPackageName();
 				FVertexData& TransitiveBuildVertex = GraphSearch.Cluster.FindOrAddVertex(TransitiveBuildPackageName);
@@ -2027,7 +2028,7 @@ void FRequestCluster::FGraphSearch::FExploreEdgesContext::AddPlatformDependencyR
 
 void FRequestCluster::FGraphSearch::FExploreEdgesContext::ProcessPlatformAttachments(int32 PlatformIndex,
 	const ITargetPlatform* TargetPlatform, FFetchPlatformData& FetchPlatformData,
-	FPackagePlatformData& PackagePlatformData, UE::TargetDomain::FIncrementalCookAttachments& PlatformAttachments,
+	FPackagePlatformData& PackagePlatformData, FIncrementalCookAttachments& PlatformAttachments,
 	bool bExploreRuntimeDependencies)
 {
 	bool bFoundBuildDefinitions = false;
@@ -2042,10 +2043,10 @@ void FRequestCluster::FGraphSearch::FExploreEdgesContext::ProcessPlatformAttachm
 		{
 			// Queue runtime dependencies if we are exploring runtime dependencies, build definitions if
 			// we are queuing those, and always queue build dependencies
-			UE::TargetDomain::FCookDependencies& CookDependencies = PlatformAttachments.Dependencies;
+			UE::Cook::FPackageArtifacts& Artifacts = PlatformAttachments.Artifacts;
 			if (bExploreRuntimeDependencies && PackagePlatformData.IsCookable() && Cluster.bAllowSoftDependencies)
 			{
-				AddPlatformDependencyRange(CookDependencies.GetRuntimePackageDependencies(), PlatformIndex,
+				AddPlatformDependencyRange(Artifacts.GetRuntimePackageDependencies(), PlatformIndex,
 					EInstigator::SoftDependency);
 			}
 
@@ -2057,7 +2058,7 @@ void FRequestCluster::FGraphSearch::FExploreEdgesContext::ProcessPlatformAttachm
 			}
 
 			for (const UE::Cook::FCookDependency& TransitiveBuildDependency : 
-				CookDependencies.GetTransitiveBuildDependencies())
+				Artifacts.GetTransitiveBuildDependencies())
 			{
 				AddPlatformDependency(TransitiveBuildDependency.GetPackageName(), PlatformIndex,
 					EInstigator::BuildDependency);
@@ -2078,7 +2079,7 @@ void FRequestCluster::FGraphSearch::FExploreEdgesContext::ProcessPlatformAttachm
 			// and will require care to avoid causing a softlock in the presence of a transitive build dependency
 			// cycle.
 		}
-		if (!QueryPlatformData.CookAttachments.Dependencies.HasSaveResults() && Cluster.TraversalMarkCookable())
+		if (!QueryPlatformData.CookAttachments.Artifacts.HasSaveResults() && Cluster.TraversalMarkCookable())
 		{
 			// The package was previously committed only as a build dependency, with no attempt made to save its data
 			// because it was not marked reachable in the previous cook. But during this cook, it is marked reachable
@@ -2158,7 +2159,7 @@ void FRequestCluster::FGraphSearch::FExploreEdgesContext::ProcessPlatformAttachm
 				}
 			}
 			if (bCanCheckHasKeyMatch &&
-				PlatformAgnosticQueryData.CookAttachments.Dependencies.HasKeyMatch(OverrideAssetPackageData))
+				PlatformAgnosticQueryData.CookAttachments.Artifacts.HasKeyMatch(OverrideAssetPackageData))
 			{
 				Cluster.BuildDefinitions.AddBuildDefinitionList(PackageName, TargetPlatform,
 					PlatformAgnosticQueryData.CookAttachments.BuildDefinitions.Definitions);
@@ -2408,12 +2409,12 @@ void FRequestCluster::FQueryVertexBatch::Send()
 															// always use empty
 			)
 		{
-			TUniqueFunction<void(FName PackageName, UE::TargetDomain::FIncrementalCookAttachments&& Result)> Callback =
-				[this, PlatformIndex](FName PackageName, UE::TargetDomain::FIncrementalCookAttachments&& Attachments)
+			TUniqueFunction<void(FName PackageName, FIncrementalCookAttachments&& Result)> Callback =
+				[this, PlatformIndex](FName PackageName, FIncrementalCookAttachments&& Attachments)
 			{
 				RecordCacheResults(PackageName, PlatformIndex, MoveTemp(Attachments));
 			};
-			UE::TargetDomain::FIncrementalCookAttachments::Fetch(PlatformData.PackageNames, FetchPlatformData.Platform,
+			FIncrementalCookAttachments::Fetch(PlatformData.PackageNames, FetchPlatformData.Platform,
 				FetchPlatformData.Writer, MoveTemp(Callback));
 		}
 		else
@@ -2428,7 +2429,7 @@ void FRequestCluster::FQueryVertexBatch::Send()
 			for (int32 PackageNameIndex = 0; PackageNameIndex < NumPackageNames; ++PackageNameIndex)
 			{
 				FName PackageName = PackageNamesData[PackageNameIndex];
-				UE::TargetDomain::FIncrementalCookAttachments Attachments;
+				FIncrementalCookAttachments Attachments;
 				RecordCacheResults(PackageName, PlatformIndex, MoveTemp(Attachments));
 			}
 		}
@@ -2436,7 +2437,7 @@ void FRequestCluster::FQueryVertexBatch::Send()
 }
 
 void FRequestCluster::FQueryVertexBatch::RecordCacheResults(FName PackageName, int32 PlatformIndex,
-	UE::TargetDomain::FIncrementalCookAttachments&& CookAttachments)
+	FIncrementalCookAttachments&& CookAttachments)
 {
 	FVertexData* Vertex = Vertices.FindChecked(PackageName);
 	FQueryPlatformData& PlatformData = Vertex->GetPlatformData()[PlatformIndex];
diff --git a/Engine/Source/Editor/UnrealEd/Private/Cooker/CookRequestCluster.h b/Engine/Source/Editor/UnrealEd/Private/Cooker/CookRequestCluster.h
--- a/Engine/Source/Editor/UnrealEd/Private/Cooker/CookRequestCluster.h
+++ b/Engine/Source/Editor/UnrealEd/Private/Cooker/CookRequestCluster.h
@@ -15,10 +15,10 @@
 #include "Containers/Map.h"
 #include "Containers/MpscQueue.h"
 #include "Containers/RingBuffer.h"
+#include "Cooker/CookPackageArtifacts.h"
 #include "Cooker/CookTypes.h"
 #include "Cooker/TypedBlockAllocator.h"
 #include "HAL/CriticalSection.h"
-#include "TargetDomain/TargetDomainUtils.h"
 #include "Templates/UniquePtr.h"
 #include "UObject/ICookInfo.h"
 #include "UObject/NameTypes.h"
@@ -131,7 +131,7 @@ private:
 		 * Thread synchronization: this field is write-once from the async thread and is not readable until
 		 * bSchedulerThreadFetchCompleted.
 		 */
-		UE::TargetDomain::FIncrementalCookAttachments CookAttachments;
+		FIncrementalCookAttachments CookAttachments;
 		bool bSchedulerThreadFetchCompleted = false;
 		bool bExploreRequested = false;
 		bool bExploreCompleted = false;
@@ -253,7 +253,7 @@ private:
 		void Send();
 
 		void RecordCacheResults(FName PackageName, int32 PlatformIndex,
-			UE::TargetDomain::FIncrementalCookAttachments&& CookAttachments);
+			FIncrementalCookAttachments&& CookAttachments);
 
 		struct FPlatformData
 		{
@@ -389,7 +389,7 @@ private:
 			void AddPlatformDependencyRange(TConstArrayView<FName> Range, int32 PlatformIndex, EInstigator InstigatorType);
 			void ProcessPlatformAttachments(int32 PlatformIndex, const ITargetPlatform* TargetPlatform,
 				FFetchPlatformData& FetchPlatformData, FPackagePlatformData& PackagePlatformData,
-				UE::TargetDomain::FIncrementalCookAttachments& PlatformAttachments, bool bExploreDependencies);
+				FIncrementalCookAttachments& PlatformAttachments, bool bExploreDependencies);
 
 			void SetIsIterativelyUnmodified(int32 PlatformIndex, bool bIterativelyUnmodified,
 				FPackagePlatformData& PackagePlatformData);
diff --git a/Engine/Source/Editor/UnrealEd/Private/Cooker/CookSavePackage.cpp b/Engine/Source/Editor/UnrealEd/Private/Cooker/CookSavePackage.cpp
--- a/Engine/Source/Editor/UnrealEd/Private/Cooker/CookSavePackage.cpp
+++ b/Engine/Source/Editor/UnrealEd/Private/Cooker/CookSavePackage.cpp
@@ -10,6 +10,7 @@
 #include "Cooker/CookDiagnostics.h"
 #include "Cooker/CookGenerationHelper.h"
 #include "Cooker/CookOnTheFlyServerInterface.h"
+#include "Cooker/CookPackageArtifacts.h"
 #include "Cooker/CookPackageData.h"
 #include "Cooker/CookPlatformManager.h"
 #include "Cooker/CookProfiling.h"
@@ -414,7 +415,7 @@ void FSaveCookedPackageContext::FinishPlatform()
 
 	// Calculate up-to-date assetregistry data for Generator and Generated packages
 	TOptional<TArray<FAssetDependency>> OverridePackageDependencies;
-	TOptional<UE::TargetDomain::FGeneratedPackageResultStruct> GeneratedResult;
+	TOptional<FGeneratedPackageResultStruct> GeneratedResult;
 	TOptional<FAssetPackageData> AssetPackageDataBuffer;
 	TOptional<FAssetPackageData> OverrideAssetPackageData;
 	const FAssetPackageData* AssetPackageData = nullptr;
@@ -735,7 +736,7 @@ TArray<FName> FSaveCookedPackageContext::GetPlatformRuntimeDependencies() const
 	return PlatformDependencyNames;
 }
 
-void FSaveCookedPackageContext::CalculateCookDependencies(UE::TargetDomain::FGeneratedPackageResultStruct* GeneratedResult,
+void FSaveCookedPackageContext::CalculateCookDependencies(FGeneratedPackageResultStruct* GeneratedResult,
 	TArray<FAssetDependency>* OverrideDependencies)
 {
 	if (COTFS.bHybridIterativeEnabled)
@@ -750,7 +751,7 @@ void FSaveCookedPackageContext::CalculateCookDependencies(UE::TargetDomain::FGen
 				PlatformRuntimeDependencies.Add(Dependency.AssetId.PackageName);
 			}
 		}
-		PlatformCookAttachments = UE::TargetDomain::FIncrementalCookAttachments::Collect(Package, TargetPlatform,
+		PlatformCookAttachments = FIncrementalCookAttachments::Collect(Package, TargetPlatform,
 			HasSavePackageResult() ? &SavePackageResult : nullptr, GeneratedResult,
 			MoveTemp(PlatformRuntimeDependencies));
 	}
@@ -762,7 +763,7 @@ void FSaveCookedPackageContext::RecordPlatformBuildDependencies()
 	{
 		TSet<FPackageData*>& CurrentDependencies = PlatformDependencies[PlatformIndex].BuildDependencies;
 		for (const FCookDependency& BuildDependency :
-			PlatformCookAttachments.Dependencies.GetTransitiveBuildDependencies())
+			PlatformCookAttachments.Artifacts.GetTransitiveBuildDependencies())
 		{
 			FPackageData* BuildPackageData = COTFS.PackageDatas->TryAddPackageDataByPackageName(
 				BuildDependency.GetPackageName());
diff --git a/Engine/Source/Editor/UnrealEd/Private/Cooker/CookSavePackage.h b/Engine/Source/Editor/UnrealEd/Private/Cooker/CookSavePackage.h
--- a/Engine/Source/Editor/UnrealEd/Private/Cooker/CookSavePackage.h
+++ b/Engine/Source/Editor/UnrealEd/Private/Cooker/CookSavePackage.h
@@ -6,6 +6,7 @@
 #include "Containers/Array.h"
 #include "Containers/Map.h"
 #include "Containers/UnrealString.h"
+#include "Cooker/CookPackageArtifacts.h"
 #include "CoreTypes.h"
 #include "Misc/Optional.h"
 #include "Serialization/PackageWriter.h"
@@ -49,7 +50,7 @@ private: // Used only by UCookOnTheFlyServer, which has private access
 	void CalculatePlatformAgnosticRuntimeDependencies();
 	void CalculatePlatformRuntimeDependencies();
 	TArray<FName> GetPlatformRuntimeDependencies() const;
-	void CalculateCookDependencies(UE::TargetDomain::FGeneratedPackageResultStruct* GeneratedResult,
+	void CalculateCookDependencies(FGeneratedPackageResultStruct* GeneratedResult,
 		TArray<FAssetDependency>* OverrideDependencies);
 	void RecordPlatformBuildDependencies();
 	bool HasSavePackageResult() const;
@@ -69,7 +70,7 @@ private: // Used only by UCookOnTheFlyServer, which has private access
 		TSet<FPackageData*> BuildDependencies;
 	};
 	TArray<FPlatformDiscoveryData, TInlineAllocator<1>> PlatformDependencies;
-	UE::TargetDomain::FIncrementalCookAttachments PlatformCookAttachments;
+	FIncrementalCookAttachments PlatformCookAttachments;
 	FTickStackData& StackData;
 	UPackage* Package;
 	const FString PackageName;
diff --git a/Engine/Source/Editor/UnrealEd/Private/EditorDomain/EditorDomainUtils.cpp b/Engine/Source/Editor/UnrealEd/Private/EditorDomain/EditorDomainUtils.cpp
--- a/Engine/Source/Editor/UnrealEd/Private/EditorDomain/EditorDomainUtils.cpp
+++ b/Engine/Source/Editor/UnrealEd/Private/EditorDomain/EditorDomainUtils.cpp
@@ -10,6 +10,7 @@
 #include "AssetRegistry/IAssetRegistry.h"
 #include "AssetRegistry/PackageReader.h"
 #include "Containers/DirectoryTree.h"
+#include "Cooker/CookPackageArtifacts.h"
 #include "DerivedDataBuildDefinition.h"
 #include "DerivedDataCache.h"
 #include "DerivedDataCacheKey.h"
@@ -1979,8 +1980,8 @@ bool TrySavePackage(UPackage* Package)
 
 	if (bStorageResultValid)
 	{
-		UE::TargetDomain::FIncrementalCookAttachments CookAttachments =
-			UE::TargetDomain::FIncrementalCookAttachments::Collect(
+		UE::Cook::FIncrementalCookAttachments CookAttachments =
+			UE::Cook::FIncrementalCookAttachments::Collect(
 			Package, nullptr, nullptr, nullptr, TArray<FName>());
 		TArray<IPackageWriter::FCommitAttachmentInfo> Attachments;
 		CookAttachments.AppendCommitAttachments(Attachments);
diff --git a/Engine/Source/Editor/UnrealEd/Private/TargetDomain/TargetDomainUtils.cpp b/Engine/Source/Editor/UnrealEd/Private/TargetDomain/TargetDomainUtils.cpp
--- a/Engine/Source/Editor/UnrealEd/Private/TargetDomain/TargetDomainUtils.cpp
+++ b/Engine/Source/Editor/UnrealEd/Private/TargetDomain/TargetDomainUtils.cpp
@@ -5,1086 +5,27 @@
 #include "Algo/BinarySearch.h"
 #include "Algo/IsSorted.h"
 #include "Algo/Sort.h"
-#include "Algo/Unique.h"
+#include "AssetRegistry/AssetData.h"
 #include "AssetRegistry/IAssetRegistry.h"
-#include "Containers/Set.h"
-#include "Containers/SharedString.h"
-#include "Containers/UnrealString.h"
-#include "Cooker/CookConfigAccessTracker.h"
-#include "Cooker/CookDependency.h"
-#include "Cooker/PackageBuildDependencyTracker.h"
-#include "CookOnTheSide/CookLog.h"
-#include "CookPackageSplitter.h"
-#include "DerivedDataBuildDefinition.h"
-#include "DerivedDataBuildKey.h"
+#include "Async/Future.h"
+#include "Compression/CompressedBuffer.h"
 #include "EditorDomain/EditorDomain.h"
 #include "EditorDomain/EditorDomainUtils.h"
-#include "HAL/PlatformFile.h"
 #include "HAL/PlatformFileManager.h"
-#include "Hash/Blake3.h"
-#include "IO/IoDispatcher.h"
-#include "IO/IoHash.h"
+#include "IO/IoStatus.h"
+#include "Logging/LogMacros.h"
 #include "Misc/App.h"
-#include "Misc/Guid.h"
+#include "Misc/AssertionMacros.h"
+#include "Misc/Optional.h"
+#include "Misc/Paths.h"
 #include "Misc/ScopeRWLock.h"
-#include "Misc/StringBuilder.h"
-#include "Serialization/CompactBinary.h"
 #include "Serialization/CompactBinarySerialization.h"
 #include "Serialization/CompactBinaryWriter.h"
-#include "Serialization/PackageWriter.h"
-#include "UObject/CoreRedirects.h"
-#include "ZenStoreHttpClient.h"
+#include "UObject/TopLevelAssetPath.h"
 
 namespace UE::TargetDomain
 {
 
-// Change TargetDomainVersion to a new guid when all TargetDomain keys in an incremental cook need to be invalidated
-FGuid TargetDomainVersion(0xD5476CB0, 0xD7D4482F, 0x966DC671, 0xBBA1BC40);
-
-// Bump CookDependenciesVersion version when the serialization of CookDependencies has changed and we want to add
-// backwards compatibility rather than invalidating everything.
-constexpr uint32 CookDependenciesVersion = 0x00000002;
-
-static const FUtf8StringView CookDependenciesAttachmentKey = UTF8TEXTVIEW("CookDependencies");
-static const FUtf8StringView BuildDefinitionsAttachmentKey = UTF8TEXTVIEW("BuildDefinitionsAttachmentKey");
-/**
- * Reads / writes an oplog for EditorDomain BuildDefinitionLists.
- * TODO: Reduce duplication between this class and FZenStoreWriter
- */
-class FEditorDomainOplog
-{
-public:
-	FEditorDomainOplog();
-
-	bool IsValid() const;
-	void CommitPackage(FName PackageName, TArrayView<IPackageWriter::FCommitAttachmentInfo> Attachments);
-	void GetOplogAttachments(TArrayView<FName> PackageName,
-		TArrayView<FUtf8StringView> AttachmentKeys,
-		TFunctionRef<void(FName PackageName, FUtf8StringView AttachmentKey, FCbObject&& Attachment)> Callback);
-
-private:
-	struct FOplogEntry
-	{
-		struct FAttachment
-		{
-			const UTF8CHAR* Key;
-			FIoHash Hash;
-		};
-
-		TArray<FAttachment> Attachments;
-	};
-
-	void InitializeRead();
-	
-	FCbAttachment CreateAttachment(FSharedBuffer AttachmentData);
-	FCbAttachment CreateAttachment(FCbObject AttachmentData)
-	{
-		return CreateAttachment(AttachmentData.GetBuffer().ToShared());
-	}
-
-	static void StaticInit();
-	static bool IsReservedOplogKey(FUtf8StringView Key);
-
-	UE::FZenStoreHttpClient HttpClient;
-	FCriticalSection Lock;
-	TMap<FName, FOplogEntry> Entries;
-	bool bConnectSuccessful = false;
-	bool bInitializedRead = false;
-
-	static TArray<const UTF8CHAR*> ReservedOplogKeys;
-};
-TUniquePtr<FEditorDomainOplog> GEditorDomainOplog;
-
-// Constructor/Destructor defined here in cpp rather than header so we can 
-// avoid needing the definition of FCookDependency in the header; it is needed
-// for construct/destruct of TArray<FCookDependency>.
-FCookDependencies::FCookDependencies() = default;
-FCookDependencies::~FCookDependencies() = default;
-FCookDependencies::FCookDependencies(const FCookDependencies&) = default;
-FCookDependencies::FCookDependencies(FCookDependencies&&) = default;
-FCookDependencies& FCookDependencies::operator=(const FCookDependencies&) = default;
-FCookDependencies& FCookDependencies::operator=(FCookDependencies&&) = default;
-
-
-bool FCookDependencies::IsValid() const
-{
-	return bValid;
-}
-
-bool FCookDependencies::HasSaveResults() const
-{
-	return bHasSaveResults;
-}
-
-bool FCookDependencies::HasKeyMatch(const FAssetPackageData* OverrideAssetPackageData)
-{
-	if (!bValid)
-	{
-		return false;
-	}
-	if (StoredKey.IsZero())
-	{
-		return false;
-	}
-	if (CurrentKey.IsZero())
-	{
-		ECurrentKeyResult Result = TryCalculateCurrentKey(OverrideAssetPackageData);
-		if (Result != ECurrentKeyResult::Success)
-		{
-			CurrentKey.Reset();
-			return false;
-		}
-	}
-	return CurrentKey == StoredKey;
-}
-
-FCookDependencies::ECurrentKeyResult FCookDependencies::TryCalculateCurrentKey(
-	const FAssetPackageData* OverrideAssetPackageData,
-	TArray<TPair<ELogVerbosity::Type, FString>>* OutMessages)
-{
-	IAssetRegistry* AssetRegistry = IAssetRegistry::Get();
-	if (PackageName.IsNone())
-	{
-		if (OutMessages) OutMessages->Emplace(ELogVerbosity::Error, TEXT("PackageName is not set."));
-		return ECurrentKeyResult::Error;
-	}
-	if (!AssetRegistry)
-	{
-		if (OutMessages) OutMessages->Emplace(ELogVerbosity::Error, TEXT("AssetRegistry is unavailable."));
-		return ECurrentKeyResult::Error;
-	}
-	FEditorDomain* EditorDomain = FEditorDomain::Get();
-	if (!EditorDomain)
-	{
-		if (OutMessages) OutMessages->Emplace(ELogVerbosity::Error, TEXT("EditorDomain is unavailable."));
-		return ECurrentKeyResult::Error;
-	}
-	FBlake3 KeyBuilder;
-	KeyBuilder.Update(&TargetDomainVersion, sizeof(TargetDomainVersion));
-
-	UE::EditorDomain::FPackageDigest PackageDigest = OverrideAssetPackageData ?
-		UE::EditorDomain::CalculatePackageDigest(*OverrideAssetPackageData, PackageName) :
-		EditorDomain->GetPackageDigest(PackageName);
-	if (!PackageDigest.IsSuccessful())
-	{
-		if (OutMessages) OutMessages->Emplace(ELogVerbosity::Error, PackageDigest.GetStatusString());
-		return ECurrentKeyResult::Error;
-	}
-	KeyBuilder.Update(&PackageDigest.Hash, sizeof(PackageDigest.Hash));
-
-	if (!ClassDependencies.IsEmpty())
-	{
-		TArray<FTopLevelAssetPath, TInlineAllocator<1>> ClassPaths;
-		ClassPaths.Reserve(ClassDependencies.Num());
-		for (const FString& ClassPathStr : ClassDependencies)
-		{
-			FTopLevelAssetPath& ClassPath = ClassPaths.Emplace_GetRef(ClassPathStr);
-			if (!ClassPath.IsValid())
-			{
-				if (OutMessages)
-				{
-					OutMessages->Emplace(ELogVerbosity::Error, FString::Printf(
-						TEXT("ClassDependency failed: %s is not a valid TopLevelAssetPath."), *ClassPathStr));
-				}
-				return ECurrentKeyResult::Error;
-			}
-		}
-		FString AppendDigestError;
-		if (!UE::EditorDomain::TryAppendClassDigests(KeyBuilder, ClassPaths, &AppendDigestError))
-		{
-			if (OutMessages)
-			{
-				OutMessages->Emplace(ELogVerbosity::Error, FString::Printf(
-					TEXT("ClassDependency failed: %s"), *AppendDigestError));
-			}
-			return ECurrentKeyResult::Error;
-		}
-	}
-
-	for (FName PackageDependency : BuildPackageDependencies)
-	{
-		PackageDigest = EditorDomain->GetPackageDigest(PackageDependency);
-		if (!PackageDigest.IsSuccessful())
-		{
-			if (OutMessages)
-			{
-				OutMessages->Emplace(ELogVerbosity::Error, FString::Printf(TEXT("Could not create PackageDigest for %s: %s"),
-					*PackageDependency.ToString(), *PackageDigest.GetStatusString()));
-			}
-			return ECurrentKeyResult::Error;
-		}
-		KeyBuilder.Update(&PackageDigest.Hash, sizeof(PackageDigest.Hash));
-	}
-
-	FCoreRedirects::AppendHashOfRedirectsAffectingPackages(KeyBuilder, RuntimePackageDependencies);
-	FCoreRedirects::AppendHashOfRedirectsAffectingPackages(KeyBuilder, ScriptPackageDependencies);
-	FCoreRedirects::AppendHashOfRedirectsAffectingPackages(KeyBuilder, RedirectorsAffectingPackageDependencies);
-
-	if (!ConfigDependencies.IsEmpty())
-	{
-#if UE_WITH_CONFIG_TRACKING
-		using namespace UE::ConfigAccessTracking;
-		FCookConfigAccessTracker& ConfigTracker = FCookConfigAccessTracker::Get();
-#endif
-		for (const FString& ConfigDependency : ConfigDependencies)
-		{
-			FString Value;
-#if UE_WITH_CONFIG_TRACKING
-			Value = ConfigTracker.GetValue(ConfigDependency);
-#endif
-			uint8 Marker = 0;
-			KeyBuilder.Update(&Marker, sizeof(Marker));
-			if (!Value.IsEmpty())
-			{
-				KeyBuilder.Update(*Value, Value.Len() * sizeof(Value[0]));
-			}
-		}
-	}
-
-	ECurrentKeyResult Result = ECurrentKeyResult::Success;
-	if (!CookDependencies.IsEmpty())
-	{
-		UE::Cook::FCookDependencyContext Context(&KeyBuilder,
-			[&Result, OutMessages](ELogVerbosity::Type Verbosity, FString&& Message)
-			{
-				if (OutMessages)
-				{
-					OutMessages->Emplace(Verbosity, MoveTemp(Message));
-				}
-			},
-			[&Result](ELogVerbosity::Type Verbosity)
-			{
-				ECurrentKeyResult NewResult = Verbosity <= ELogVerbosity::Error
-					? ECurrentKeyResult::Error : ECurrentKeyResult::Invalidated;
-				Result = FMath::Max(Result, NewResult);
-			},
-			PackageName);
-
-		for (UE::Cook::FCookDependency& CookDependency : CookDependencies)
-		{
-			CookDependency.UpdateHash(Context);
-		}
-	}
-
-	if (Result != ECurrentKeyResult::Error) // -V547
-	{
-		CurrentKey = KeyBuilder.Finalize();
-	}
-	return Result;
-}
-
-void FCookDependencies::Reset()
-{
-	BuildPackageDependencies.Reset();
-	ConfigDependencies.Reset();
-	RuntimePackageDependencies.Reset();
-	ScriptPackageDependencies.Reset();
-	RedirectorsAffectingPackageDependencies.Reset();
-	PackageName = FName();
-	StoredKey = FIoHash::Zero;
-	CurrentKey = FIoHash::Zero;
-	bValid = false;
-}
-
-void FCookDependencies::Empty()
-{
-	Reset();
-	BuildPackageDependencies.Empty();
-	ConfigDependencies.Empty();
-	RuntimePackageDependencies.Empty();
-	RedirectorsAffectingPackageDependencies.Empty();
-}
-
-enum class EPackageMountPoint
-{
-	Transient,
-	Script,
-	Content,
-	GeneratedContent,
-};
-EPackageMountPoint GetPackageMountPoint(FName PackageName)
-{
-	TStringBuilder<256> StringBuffer;
-	PackageName.ToString(StringBuffer);
-	if (
-		// Some packages get renamed to "TrashedPackage" during blueprint compilation and are no longer valid for
-		// saving but might have been dereferenced by TObjectPtr during PostLoad/PreSave. We need to discard these
-		// packages, which we can do by requiring a valid package name; all valid packages start with /MountPoint/.
-		!StringBuffer.ToView().StartsWith('/') ||
-		// Ignore /Memory and /Temp packages
-		FPackageName::IsMemoryPackage(StringBuffer) ||
-		FPackageName::IsTempPackage(StringBuffer) ||
-		FPackageName::IsInEngineTransientPackages(StringBuffer))
-	{
-		return EPackageMountPoint::Transient;
-	}
-	if (FPackageName::IsScriptPackage(StringBuffer))
-	{
-		return EPackageMountPoint::Script;
-	}
-	if (ICookPackageSplitter::IsUnderGeneratedPackageSubPath(StringBuffer))
-	{
-		return EPackageMountPoint::GeneratedContent;
-	}
-	return EPackageMountPoint::Content;
-}
-
-static void JoinMessagesIntoErrorReason(FStringBuilderBase& OutText,
-	TArray<TPair<ELogVerbosity::Type, FString>>& Messages)
-{
-	if (Messages.IsEmpty())
-	{
-		OutText << TEXT(".");
-	}
-	else
-	{
-		OutText << TEXT(":");
-		for (TPair<ELogVerbosity::Type, FString>& MessagePair : Messages)
-		{
-			OutText << TEXT("\n\t") << MessagePair.Value;
-		}
-	}
-}
-
-FCookDependencies FCookDependencies::Collect(UPackage* Package, const ITargetPlatform* TargetPlatform,
-	FSavePackageResultStruct* SaveResult, const FGeneratedPackageResultStruct* GeneratedResult, 
-	TArray<FName>&& RuntimeDependencies, TArray<TPair<ELogVerbosity::Type, FString>>* OutMessages)
-{
-	if (!Package)
-	{
-		if (OutMessages) OutMessages->Emplace(ELogVerbosity::Error, TEXT("Invalid null package."));
-		return FCookDependencies();
-	}
-	IAssetRegistry* AssetRegistry = IAssetRegistry::Get();
-	if (!AssetRegistry)
-	{
-		if (OutMessages) OutMessages->Emplace(ELogVerbosity::Error, TEXT("AssetRegistry is unavailable."));
-		return FCookDependencies();
-	}
-	FEditorDomain* EditorDomain = FEditorDomain::Get();
-	if (!EditorDomain)
-	{
-		if (OutMessages) OutMessages->Emplace(ELogVerbosity::Error, TEXT("EditorDomain is unavailable."));
-		return FCookDependencies();
-	}
-
-	FCookDependencies Result;
-	Result.PackageName = Package->GetFName();
-	Result.bHasSaveResults = SaveResult != nullptr;
-	enum class EBuildDependency
-	{
-		PackageOnly,
-		Transitive,
-	};
-
-	TMap<FName, EBuildDependency> PackageBuildDependencies;
-
-	TArray<FName> AssetDependencies;
-	if (!GeneratedResult)
-	{
-		AssetRegistry->GetDependencies(Result.PackageName, AssetDependencies,
-			UE::AssetRegistry::EDependencyCategory::Package, UE::AssetRegistry::EDependencyQuery::Game);
-		RuntimeDependencies.Append(MoveTemp(AssetDependencies));
-	}
-	else
-	{
-		// GeneratedResult->PackageDependencies are already incorporated into the PackageHash of the generated
-		// package by FCookGenerationInfo::CreatePackageHash, so we do not need to add them into the BuildDependencies
-	}
-
-#if UE_WITH_PACKAGE_ACCESS_TRACKING
-	FPackageBuildDependencyTracker& Tracker = FPackageBuildDependencyTracker::Get();
-
-	if (Tracker.IsEnabled())
-	{
-		TArray<TPair<FBuildDependencyAccessData, FResultProjectionList>> AccessDatas
-			= Tracker.GetAccessDatas(Result.PackageName);
-
-		PackageBuildDependencies.Reserve(AccessDatas.Num());
-		for (TPair<FBuildDependencyAccessData, FResultProjectionList>& Pair : AccessDatas)
-		{
-			FBuildDependencyAccessData& AccessData = Pair.Key;
-			if (AccessData.TargetPlatform == TargetPlatform || AccessData.TargetPlatform == nullptr)
-			{
-				FResultProjectionList& ProjectionList = Pair.Value;
-				constexpr bool bAutoTransitiveDependenciesEnabled = false;
-				if (!bAutoTransitiveDependenciesEnabled)
-				{
-					// We have not yet enabled marking the auto-added dependencies from TObjectPtr resolve
-					// as transitive, because it causes a performance regression and we are still working on fixing the
-					// regression.
-					PackageBuildDependencies.FindOrAdd(AccessData.ReferencedPackage, EBuildDependency::PackageOnly);
-				}
-				else if (ProjectionList.bHasAll)
-				{
-					// TObjectPtr UE::BuildProjection::All dependencies are added as transitive build dependencies.
-					// We have to do this to be conservative, since we do not know which bytes from the target are
-					// dependended upon and which of the target's build dependencies influence those bytes.
-					PackageBuildDependencies.FindOrAdd(AccessData.ReferencedPackage) = EBuildDependency::Transitive;
-				}
-				else
-				{
-					for (FTopLevelAssetPath ClassPath : ProjectionList.Classes)
-					{
-						Result.ClassDependencies.Add(ClassPath.ToString());
-					}
-
-					for (FName ResultProjection : ProjectionList.ResultProjections)
-					{
-						if (ResultProjection == UE::Cook::ResultProjection::PackageAndClass)
-						{
-							PackageBuildDependencies.FindOrAdd(AccessData.ReferencedPackage, EBuildDependency::PackageOnly);
-						}
-						else
-						{
-							UE_LOG(LogCook, Error,
-								TEXT("When saving %s, found ResultProjection %s, which is a system-specific ResultProjection, and this is not yet implemented. ")
-								TEXT("Find the call to UE_COOK_RESULTPROJECTION_SCOPED passing in this name and remove it."),
-								*Result.PackageName.ToString(), *ResultProjection.ToString());
-						}
-					}
-				}
-			}
-		}
-	}
-	else
-#endif
-	{
-		// When PackageAccessTracking is disabled, defensively treat all asset dependencies as transitive build
-		// dependencies.
-		for (FName AssetDependency : AssetDependencies)
-		{
-			PackageBuildDependencies.FindOrAdd(AssetDependency) = EBuildDependency::Transitive;
-		}
-	}
-
-#if UE_WITH_CONFIG_TRACKING
-	{
-		using namespace UE::ConfigAccessTracking;
-		FCookConfigAccessTracker& ConfigTracker = FCookConfigAccessTracker::Get();
-		if (ConfigTracker.IsEnabled())
-		{
-			TArray<FConfigAccessData> ConfigKeys = ConfigTracker.GetPackageRecords(Result.PackageName, TargetPlatform);
-			Result.ConfigDependencies.Reserve(ConfigKeys.Num());
-			for (const FConfigAccessData& ConfigKey : ConfigKeys)
-			{
-				Result.ConfigDependencies.Add(ConfigKey.FullPathToString());
-			}
-		}
-	}
-#endif
-	if (SaveResult)
-	{
-		// Convert some FCookDependency types into our specially-handled dependencies, and put the more
-		// generic ones into this->CookDependencies
-		Result.CookDependencies.Reserve(SaveResult->CookDependencies.Num());
-
-		// Additional dependencies can be added from SettingsObject dependencies. Add these included dependencies
-		// to a followup list that we process in the second round, if they exist. Handle detecting cycles in the
-		// included dependencies; cycle detection is currently simple: nested includes are not yet required
-		// so give an error if we ever make it to a third round.
-		TArray<UE::Cook::FCookDependency> NextIncludedDependenciesBuffer[2];
-		TArray<UE::Cook::FCookDependency>* RemainingDependencies = &SaveResult->CookDependencies;
-		int32 IncludedDependenciesRound = 0;
-		constexpr int32 MaxNumberOfIncludedDependenciesRounds = 2;
-		while (!RemainingDependencies->IsEmpty())
-		{
-			TArray<UE::Cook::FCookDependency>* NextIncludedDependencies =
-				&NextIncludedDependenciesBuffer[IncludedDependenciesRound % UE_ARRAY_COUNT(NextIncludedDependenciesBuffer)];
-			if (IncludedDependenciesRound++ > 0)
-			{
-				NextIncludedDependencies->Reset();
-				if (IncludedDependenciesRound > MaxNumberOfIncludedDependenciesRounds)
-				{
-					if (OutMessages)
-					{
-						OutMessages->Emplace(ELogVerbosity::Error,
-							TEXT("More rounds than expected when handling included dependencies."));
-					}
-					return FCookDependencies();
-				}
-			}
-
-			for (UE::Cook::FCookDependency& CookDependency : *RemainingDependencies)
-			{
-				switch (CookDependency.GetType())
-				{
-				case UE::Cook::ECookDependency::TransitiveBuild:
-				{
-					// Build dependencies from a package to itself have a performance cost and serve no
-					// purpose, so remove them. They can occur in some systems that naively add a build dependency
-					// from one object to another without checking whether the second object is in the same package.
-					FName DependencyPackageName = CookDependency.GetPackageName();
-					if (DependencyPackageName == Result.PackageName)
-					{
-						continue;
-					}
-					if (CookDependency.IsAlsoAddRuntimeDependency())
-					{
-						RuntimeDependencies.Add(DependencyPackageName);
-					}
-					PackageBuildDependencies.FindOrAdd(DependencyPackageName) = EBuildDependency::Transitive;
-					break;
-				}
-				case UE::Cook::ECookDependency::Package:
-					// Add it as a PackageOnly build dependency (lowest level) unless a higher level already exists.
-					PackageBuildDependencies.FindOrAdd(CookDependency.GetPackageName(), EBuildDependency::PackageOnly);
-					break;
-				case UE::Cook::ECookDependency::Config:
-					Result.ConfigDependencies.Add(FString(CookDependency.GetConfigPath()));
-					break;
-				case UE::Cook::ECookDependency::SettingsObject:
-				{
-					const UObject* SettingsObject = CookDependency.GetSettingsObject();
-					if (SettingsObject)
-					{
-						// We rely on the object to be rooted because we use its pointer as a key for the lifetime of
-						// the cook process, so it being garbage collected and something else allocated on the same
-						// pointer would break our key. IsRooted should have been validated by
-						// FCookDependency::SettingsObject.
-						check(SettingsObject->IsRooted());
-						FCookDependencyGroups::FRecordedDependencies& IncludeDependencies =
-							FCookDependencyGroups::Get().FindOrCreate(reinterpret_cast<UPTRINT>(SettingsObject));
-						if (!IncludeDependencies.bInitialized)
-						{
-							IncludeDependencies.Dependencies = CollectSettingsObject(SettingsObject,
-								&IncludeDependencies.Messages);
-							IncludeDependencies.bInitialized = true;
-						}
-						if (!IncludeDependencies.Dependencies.IsValid())
-						{
-							if (OutMessages)
-							{
-								TStringBuilder<256> ErrorText;
-								ErrorText << TEXT("Dependencies for SettingsObject ") << SettingsObject->GetPathName()
-									<< TEXT(" are unavailable");
-								JoinMessagesIntoErrorReason(ErrorText, IncludeDependencies.Messages);
-								OutMessages->Emplace(ELogVerbosity::Error, FString(*ErrorText));
-							}
-							return FCookDependencies();
-						}
-
-						for (const UE::Cook::FCookDependency& IncludeDependency :
-							IncludeDependencies.Dependencies.GetCookDependencies())
-						{
-							NextIncludedDependencies->Add(IncludeDependency);
-						}
-					}
-					break;
-				}
-				case UE::Cook::ECookDependency::NativeClass:
-					Result.ClassDependencies.Add(FString(CookDependency.GetClassPath()));
-					break;
-				default:
-					Result.CookDependencies.Add(MoveTemp(CookDependency));
-					break;
-				}
-			}
-			RemainingDependencies = NextIncludedDependencies;
-		}
-
-		Result.RedirectorsAffectingPackageDependencies = SaveResult->UntrackedSoftPackageReferences;
-		Result.RedirectorsAffectingPackageDependencies.Sort(FNameLexicalLess());
-		Result.RedirectorsAffectingPackageDependencies.SetNum(Algo::Unique(Result.RedirectorsAffectingPackageDependencies));
-	}
-
-	Algo::Sort(Result.CookDependencies);
-	Result.CookDependencies.SetNum(Algo::Unique(Result.CookDependencies));
-	for (TPair<FName, EBuildDependency>& Pair : PackageBuildDependencies)
-	{
-		if (Pair.Key == Result.PackageName)
-		{
-			// Skip dependencies to self, for performance
-			continue;
-		}
-		if (GetPackageMountPoint(Pair.Key) != EPackageMountPoint::Content)
-		{
-			// Skip script and temp packages
-			continue;
-		}
-		if (Pair.Value == EBuildDependency::Transitive)
-		{
-			Result.TransitiveBuildDependencies.Add(UE::Cook::FCookDependency::Package(Pair.Key));
-		}
-		else
-		{
-			Result.BuildPackageDependencies.Add(Pair.Key);
-		}
-	}
-	Algo::Sort(Result.TransitiveBuildDependencies);
-	Result.BuildPackageDependencies.Sort(FNameLexicalLess());
-	Algo::Sort(Result.ConfigDependencies);
-	Result.ConfigDependencies.SetNum(Algo::Unique(Result.ConfigDependencies));
-	Algo::Sort(Result.ClassDependencies);
-	Result.ClassDependencies.SetNum(Algo::Unique(Result.ClassDependencies));
-
-
-	for (TArray<FName>::TIterator Iter(RuntimeDependencies); Iter; ++Iter)
-	{
-		FName PackageName = *Iter;
-		EPackageMountPoint MountPoint = GetPackageMountPoint(PackageName);
-		switch (MountPoint)
-		{
-		case EPackageMountPoint::GeneratedContent:
-		case EPackageMountPoint::Content:
-			Result.RuntimePackageDependencies.Add(PackageName);
-			break;
-		case EPackageMountPoint::Script:
-			Result.ScriptPackageDependencies.Add(PackageName);
-			break;
-		default:
-			break;
-		}
-	}
-	Result.RuntimePackageDependencies.Sort(FNameLexicalLess());
-	Result.RuntimePackageDependencies.SetNum(Algo::Unique(Result.RuntimePackageDependencies), EAllowShrinking::Yes);
-	Result.ScriptPackageDependencies.Sort(FNameLexicalLess());
-	Result.ScriptPackageDependencies.SetNum(Algo::Unique(Result.ScriptPackageDependencies), EAllowShrinking::Yes);
-
-	const FAssetPackageData* AssetPackageData = GeneratedResult ? &GeneratedResult->AssetPackageData : nullptr;
-	ECurrentKeyResult CurrentKeyResult = Result.TryCalculateCurrentKey(AssetPackageData, OutMessages);
-	if (CurrentKeyResult == ECurrentKeyResult::Error)
-	{
-		return FCookDependencies();
-	}
-	Result.StoredKey = Result.CurrentKey;
-	Result.bValid = true;
-
-	return Result;
-}
-
-FCookDependencies FCookDependencies::CollectSettingsObject(const UObject* Object,
-	TArray<TPair<ELogVerbosity::Type, FString>>* OutMessages)
-{
-	if (!Object)
-	{
-		if (OutMessages)
-		{
-			OutMessages->Emplace(ELogVerbosity::Error, TEXT("Invalid null Object."));
-		}
-		return FCookDependencies();
-	}
-
-	UClass* Class = Object->GetClass();
-	if (!Class->HasAnyClassFlags(CLASS_Config | CLASS_PerObjectConfig))
-	{
-		if (OutMessages)
-		{
-			OutMessages->Emplace(ELogVerbosity::Error, FString::Printf(
-				TEXT("Class %s is not a config class."), *Class->GetPathName()));
-		}
-		return FCookDependencies();
-	}
-	if (!Class->HasAnyClassFlags(CLASS_PerObjectConfig) && Object != Class->GetDefaultObject())
-	{
-		if (OutMessages)
-		{
-			OutMessages->Emplace(ELogVerbosity::Error, FString::Printf(
-				TEXT("Class %s is not a per-object-config class."), *Class->GetPathName()));
-		}
-		return FCookDependencies();
-	}
-
-	FCookDependencies Result;
-	TArray<UE::ConfigAccessTracking::FConfigAccessData> ConfigDatas;
-	const_cast<UObject*>(Object)->LoadConfig(nullptr /* ConfigClass */, nullptr /* Filename */,
-		UE::LCPF_None /* PropagationFlags */, nullptr /* PropertyToLoad */, &ConfigDatas);
-	Result.CookDependencies.Reserve(ConfigDatas.Num() + 1);
-	for (const UE::ConfigAccessTracking::FConfigAccessData& ConfigData : ConfigDatas)
-	{
-		Result.CookDependencies.Add(UE::Cook::FCookDependency::Config(ConfigData));
-	}
-
-	// In addition to adding the config dependencies, add a dependency on the class schema. If the current class has
-	// config fields A,B,C, we add dependencies on those config values. But if the class header is modified to have
-	// additional config field D then we need to rebuild packages that depend on it to record the new dependency on D.
-	UClass* NativeClass = Class;
-	while (NativeClass && !NativeClass->IsNative())
-	{
-		NativeClass = NativeClass->GetSuperClass();
-	}
-	if (NativeClass)
-	{
-		Result.CookDependencies.Add(UE::Cook::FCookDependency::NativeClass(NativeClass));
-	}
-
-	Algo::Sort(Result.CookDependencies);
-	Result.CookDependencies.SetNum(Algo::Unique(Result.CookDependencies));
-	Result.bValid = true;
-	return Result;
-}
-
-}
-
-bool LoadFromCompactBinary(FCbObjectView ObjectView, UE::TargetDomain::FCookDependencies& Dependencies)
-{
-	using namespace UE::TargetDomain;
-
-	Dependencies.Reset();
-	int32 Version = -1;
-
-	for (FCbFieldViewIterator FieldView(ObjectView.CreateViewIterator()); FieldView; )
-	{
-		const FCbFieldViewIterator Last = FieldView;
-		if (FieldView.GetName().Equals(UTF8TEXTVIEW("Version")))
-		{
-			Version = FieldView.AsInt32();
-			if ((FieldView++).HasError() || Version != CookDependenciesVersion)
-			{
-				return false;
-			}
-		}
-		if (FieldView.GetName().Equals(UTF8TEXTVIEW("StoredKey")))
-		{
-			if (!LoadFromCompactBinary(FieldView++, Dependencies.StoredKey))
-			{
-				return false;
-			}
-		}
-		if (FieldView.GetName().Equals(UTF8TEXTVIEW("HasSaveResults")))
-		{
-			if (!LoadFromCompactBinary(FieldView++, Dependencies.bHasSaveResults))
-			{
-				return false;
-			}
-		}
-		if (FieldView.GetName().Equals(UTF8TEXTVIEW("BuildPackageDependencies")))
-		{
-			if (!LoadFromCompactBinary(FieldView++, Dependencies.BuildPackageDependencies))
-			{
-				return false;
-			}
-		}
-		if (FieldView.GetName().Equals(UTF8TEXTVIEW("ConfigDependencies")))
-		{
-			if (!LoadFromCompactBinary(FieldView++, Dependencies.ConfigDependencies))
-			{
-				return false;
-			}
-		}
-		if (FieldView.GetName().Equals(UTF8TEXTVIEW("RuntimePackageDependencies")))
-		{
-			if (!LoadFromCompactBinary(FieldView++, Dependencies.RuntimePackageDependencies))
-			{
-				return false;
-			}
-		}
-		if (FieldView->GetName().Equals(UTF8TEXTVIEW("RedirectorsAffectingPackageDependencies")))
-		{
-			if (!LoadFromCompactBinary(FieldView++, Dependencies.RedirectorsAffectingPackageDependencies))
-			{
-				return false;
-			}
-		}
-		if (FieldView.GetName().Equals(UTF8TEXTVIEW("ScriptPackageDependencies")))
-		{
-			if (!LoadFromCompactBinary(FieldView++, Dependencies.ScriptPackageDependencies))
-			{
-				return false;
-			}
-		}
-		if (FieldView.GetName().Equals(UTF8TEXTVIEW("CookDependencies")))
-		{
-			if (!LoadFromCompactBinary(FieldView++, Dependencies.CookDependencies))
-			{
-				return false;
-			}
-		}
-		if (FieldView.GetName().Equals(UTF8TEXTVIEW("TransitiveBuildDependencies")))
-		{
-			if (!LoadFromCompactBinary(FieldView++, Dependencies.TransitiveBuildDependencies))
-			{
-				return false;
-			}
-		}
-		if (FieldView.GetName().Equals(UTF8TEXTVIEW("ClassDependencies")))
-		{
-			if (!LoadFromCompactBinary(FieldView++, Dependencies.ClassDependencies))
-			{
-				return false;
-			}
-		}
-		if (FieldView == Last)
-		{
-			++FieldView;
-		}
-	}
-	if (Version == -1)
-	{
-		return false;
-	}
-	Dependencies.bValid = true;
-	return true;
-}
-
-FCbWriter& operator<<(FCbWriter& Writer, const UE::TargetDomain::FCookDependencies& CookDependencies)
-{
-	using namespace UE::TargetDomain;
-
-	Writer.BeginObject();
-	Writer << "Version" << CookDependenciesVersion;
-	Writer << "StoredKey" << CookDependencies.StoredKey;
-	Writer << "HasSaveResults" << CookDependencies.bHasSaveResults;
-	if (!CookDependencies.BuildPackageDependencies.IsEmpty())
-	{
-		Writer << "BuildPackageDependencies" << CookDependencies.BuildPackageDependencies;
-	}
-	if (!CookDependencies.ConfigDependencies.IsEmpty())
-	{
-		Writer << "ConfigDependencies" << CookDependencies.ConfigDependencies;
-	}
-	if (!CookDependencies.RuntimePackageDependencies.IsEmpty())
-	{
-		Writer << "RuntimePackageDependencies" << CookDependencies.RuntimePackageDependencies;
-	}
-	if (!CookDependencies.RedirectorsAffectingPackageDependencies.IsEmpty())
-	{
-		Writer << "RedirectorsAffectingPackageDependencies" << CookDependencies.RedirectorsAffectingPackageDependencies;
-	}
-	if (!CookDependencies.ScriptPackageDependencies.IsEmpty())
-	{
-		Writer << "ScriptPackageDependencies" << CookDependencies.ScriptPackageDependencies;
-	}
-	if (!CookDependencies.CookDependencies.IsEmpty())
-	{
-		Writer << "CookDependencies" << CookDependencies.CookDependencies;
-	}
-	if (!CookDependencies.TransitiveBuildDependencies.IsEmpty())
-	{
-		Writer << "TransitiveBuildDependencies" << CookDependencies.TransitiveBuildDependencies;
-	}
-	if (!CookDependencies.ClassDependencies.IsEmpty())
-	{
-		Writer << "ClassDependencies" << CookDependencies.ClassDependencies;
-	}
-
-	Writer.EndObject();
-	return Writer;
-}
-
-namespace UE::TargetDomain
-{
-
-FCookDependencyGroups& FCookDependencyGroups::Get()
-{
-	static FCookDependencyGroups Singleton;
-	return Singleton;
-}
-
-FCookDependencyGroups::FRecordedDependencies& FCookDependencyGroups::FindOrCreate(UPTRINT Key)
-{
-	return Groups.FindOrAdd(Key);
-}
-
-FBuildDefinitionList FBuildDefinitionList::Collect(UPackage* Package, const ITargetPlatform* TargetPlatform,
-	TArray<TPair<ELogVerbosity::Type, FString>>* OutMessages)
-{
-	using namespace UE::DerivedData;
-
-	FBuildDefinitionList Result;
-
-	// TODO_BuildDefinitionList: Calculate and store BuildDefinitionList on the PackageData, or collect it here from some other source.
-	if (Result.Definitions.IsEmpty())
-	{
-		if (OutMessages) OutMessages->Emplace(ELogVerbosity::Error, TEXT("Not yet implemented"));
-		return FBuildDefinitionList();
-	}
-
-	TArray<FBuildDefinition>& Defs = Result.Definitions;
-	Algo::Sort(Defs, [](const FBuildDefinition& A, const FBuildDefinition& B)
-		{
-			return A.GetKey().Hash < B.GetKey().Hash;
-		});
-
-	return Result;
-}
-
-void FBuildDefinitionList::Reset()
-{
-	Definitions.Reset();
-}
-
-void FBuildDefinitionList::Empty()
-{
-	Definitions.Empty();
-}
-
-}
-
-bool LoadFromCompactBinary(FCbObject&& Object, UE::TargetDomain::FBuildDefinitionList& Definitions)
-{
-	using namespace UE::DerivedData;
-
-	FCbField DefinitionsField = Object["BuildDefinitions"];
-	FCbArray DefinitionsArrayField = DefinitionsField.AsArray();
-	if (DefinitionsField.HasError())
-	{
-		return false;
-	}
-	TArray<FBuildDefinition>& Defs = Definitions.Definitions;
-	Defs.Empty(DefinitionsArrayField.Num());
-	for (FCbField& BuildDefinitionObj : DefinitionsArrayField)
-	{
-		FOptionalBuildDefinition BuildDefinition = FBuildDefinition::Load(TEXTVIEW("TargetDomainBuildDefinitionList"),
-			BuildDefinitionObj.AsObject());
-		if (!BuildDefinition)
-		{
-			Defs.Empty();
-			return false;
-		}
-		Defs.Add(MoveTemp(BuildDefinition).Get());
-	}
-
-	return true;
-}
-
-FCbWriter& operator<<(FCbWriter& Writer, const UE::TargetDomain::FBuildDefinitionList& Definitions)
-{
-	using namespace UE::DerivedData;
-
-	Writer.BeginObject();
-	Writer.BeginArray("BuildDefinitions");
-	for (const FBuildDefinition& BuildDefinition : Definitions.Definitions)
-	{
-		BuildDefinition.Save(Writer);
-	}
-	Writer.EndArray();
-	return Writer;
-}
-
-namespace UE::TargetDomain
-{
-
-void FIncrementalCookAttachments::Reset()
-{
-	Dependencies.Reset();
-	BuildDefinitions.Reset();
-}
-
-void FIncrementalCookAttachments::Empty()
-{
-	Dependencies.Empty();
-	BuildDefinitions.Empty();
-}
-
-void FIncrementalCookAttachments::AppendCommitAttachments(TArray<IPackageWriter::FCommitAttachmentInfo>& OutAttachments)
-{
-	if (Dependencies.IsValid())
-	{
-		FCbWriter Writer;
-		Writer << Dependencies;
-		IPackageWriter::FCommitAttachmentInfo& Attachment = OutAttachments.Emplace_GetRef();
-		Attachment.Key = CookDependenciesAttachmentKey;
-		Attachment.Value = Writer.Save().AsObject();
-	}
-
-	if (!BuildDefinitions.Definitions.IsEmpty())
-	{
-		FCbWriter Writer;
-		Writer << BuildDefinitions;
-		IPackageWriter::FCommitAttachmentInfo& Attachment = OutAttachments.Emplace_GetRef();
-		Attachment.Key = BuildDefinitionsAttachmentKey;
-		Attachment.Value = Writer.Save().AsObject();
-	}
-}
-
-FIncrementalCookAttachments FIncrementalCookAttachments::Collect(UPackage* Package,
-	const ITargetPlatform* TargetPlatform, FSavePackageResultStruct* SaveResult,
-	const FGeneratedPackageResultStruct* GeneratedResult, TArray<FName>&& RuntimeDependencies)
-{
-	FIncrementalCookAttachments Result;
-
-	Result.CommitStatus = IPackageWriter::ECommitStatus::NotCommitted;
-
-	TArray<TPair<ELogVerbosity::Type, FString>> Messages;
-	Result.Dependencies = FCookDependencies::Collect(Package, TargetPlatform, SaveResult,
-		GeneratedResult, MoveTemp(RuntimeDependencies), &Messages);
-	if (!Result.Dependencies.IsValid())
-	{
-		TStringBuilder<256> LogText;
-		LogText << TEXT("Could not collect CookDependencies for package '") << Package->GetFName() << TEXT("'");
-		JoinMessagesIntoErrorReason(LogText, Messages);
-
-		// IterativeTODO: This error occurs due to dependencies on _Verse.
-		// Raise Verbosity to Error once that is fixed.
-		UE_LOG(LogCook, Verbose, TEXT("%s"), *LogText);
-	}
-
-	Result.BuildDefinitions = FBuildDefinitionList::Collect(Package, TargetPlatform);
-
-	return Result;
-}
-
-void FIncrementalCookAttachments::Fetch(TArrayView<FName> PackageNames, const ITargetPlatform* TargetPlatform,
-	ICookedPackageWriter* PackageWriter,
-	TUniqueFunction<void(FName PackageName, FIncrementalCookAttachments&& Result)>&& Callback)
-{
-	struct FInProgressResult
-	{
-		FIncrementalCookAttachments Result;
-		int ReceivedAttachmentCount = 0;
-	};
-	TMap<FName, FInProgressResult> InProgressResults;
-	TArray<FUtf8StringView> RequestedAttachments = {CookDependenciesAttachmentKey, BuildDefinitionsAttachmentKey};
-
-	if (TargetPlatform || GEditorDomainOplog)
-	{
-		auto OnOplogAttachment = 
-		[PackageWriter, Callback = MoveTemp(Callback), RequestedAttachmentNum = RequestedAttachments.Num(), &InProgressResults]
-		(FName PackageName, FUtf8StringView AttachmentKey, FCbObject&& Attachment)
-		{
-			FInProgressResult& InProgressResult = InProgressResults.FindOrAdd(PackageName);
-			InProgressResult.ReceivedAttachmentCount++;
-			if (AttachmentKey == CookDependenciesAttachmentKey)
-			{
-				if (PackageWriter)
-				{
-					InProgressResult.Result.CommitStatus = PackageWriter->GetCommitStatus(PackageName);
-				}
-				else
-				{
-					InProgressResult.Result.CommitStatus = Attachment ? IPackageWriter::ECommitStatus::Success
-					: IPackageWriter::ECommitStatus::NotCommitted;
-				}
-
-				if (LoadFromCompactBinary(MoveTemp(Attachment), InProgressResult.Result.Dependencies))
-				{
-					InProgressResult.Result.Dependencies.PackageName = PackageName;
-				}
-			}
-			else if (AttachmentKey == BuildDefinitionsAttachmentKey)
-			{
-				LoadFromCompactBinary(MoveTemp(Attachment), InProgressResult.Result.BuildDefinitions);
-			}
-
-			if (InProgressResult.ReceivedAttachmentCount == RequestedAttachmentNum)
-			{
-				Callback(PackageName, MoveTemp(InProgressResult.Result));
-				InProgressResults.Remove(PackageName);
-			}
-		};
-
-		if (PackageWriter)
-		{
-			PackageWriter->GetOplogAttachments(PackageNames,
-				RequestedAttachments,OnOplogAttachment);
-		}
-		else
-		{
-			GEditorDomainOplog->GetOplogAttachments(PackageNames,
-				RequestedAttachments,OnOplogAttachment);
-		}
-	}
-	else
-	{
-		for (FName PackageName : PackageNames)
-		{
-			Callback(PackageName, FIncrementalCookAttachments());
-		}
-	}
-	check(InProgressResults.Num() == 0);
-}
-
 bool IsIterativeEnabled(FName PackageName, bool bAllowAllClasses, const FAssetPackageData* OverrideAssetPackageData)
 {
 	IAssetRegistry* AssetRegistry = IAssetRegistry::Get();
@@ -1159,6 +100,7 @@ bool IsIterativeEnabled(FName PackageName, bool bAllowAllClasses, const FAssetPa
 	return true;
 }
 
+TUniquePtr<FEditorDomainOplog> GEditorDomainOplog;
 TArray<const UTF8CHAR*> FEditorDomainOplog::ReservedOplogKeys;
 
 FEditorDomainOplog::FEditorDomainOplog()
@@ -1489,5 +431,4 @@ void CookInitialize()
 	}
 }
 
-
 } // namespace UE::TargetDomain
diff --git a/Engine/Source/Editor/UnrealEd/Private/TargetDomain/TargetDomainUtils.h b/Engine/Source/Editor/UnrealEd/Private/TargetDomain/TargetDomainUtils.h
--- a/Engine/Source/Editor/UnrealEd/Private/TargetDomain/TargetDomainUtils.h
+++ b/Engine/Source/Editor/UnrealEd/Private/TargetDomain/TargetDomainUtils.h
@@ -2,36 +2,20 @@
 
 #pragma once
 
-#include "AssetRegistry/AssetData.h"
-#include "AssetRegistry/IAssetRegistry.h"
 #include "Containers/Array.h"
 #include "Containers/ArrayView.h"
-#include "Containers/ContainersFwd.h"
-#include "Containers/UnrealString.h"
-#include "DerivedDataBuildDefinition.h"
+#include "Containers/StringView.h"
+#include "HAL/CriticalSection.h"
 #include "IO/IoHash.h"
-#include "Logging/LogVerbosity.h"
-#include "Misc/Optional.h"
+#include "Serialization/CompactBinary.h"
 #include "Serialization/PackageWriter.h"
+#include "Templates/Function.h"
+#include "Templates/UniquePtr.h"
 #include "UObject/NameTypes.h"
+#include "ZenStoreHttpClient.h"
 
-class FCbObject;
-class FCbObjectView;
-class FCbWriter;
-class ITargetPlatform;
-class UPackage;
-struct FSavePackageResultStruct;
-template <typename FuncType> class TUniqueFunction;
-
-namespace UE::Cook { class FCookDependency; }
-namespace UE::TargetDomain { class FCookDependencies; }
-namespace UE::TargetDomain { struct FBuildDefinitionList; }
-namespace UE::TargetDomain { struct FIncrementalCookAttachments; }
-
-bool LoadFromCompactBinary(FCbObjectView ObjectView, UE::TargetDomain::FCookDependencies& CookAttachments);
-FCbWriter& operator<<(FCbWriter& Writer, const UE::TargetDomain::FCookDependencies& CookAttachments);
-bool LoadFromCompactBinary(FCbObject&& Object, UE::TargetDomain::FBuildDefinitionList& Definitions);
-FCbWriter& operator<<(FCbWriter& Writer, const UE::TargetDomain::FBuildDefinitionList& Definitions);
+class FAssetPackageData;
+class FSharedBuffer;
 
 namespace UE::TargetDomain
 {
@@ -39,222 +23,59 @@ namespace UE::TargetDomain
 /** Call during Startup to initialize global data used by TargetDomain functions. */
 void CookInitialize();
 
-/**
- * Information collected from a CookPackageSplitter Generated package after it is saved; this information is needed to collect
- * the CookDependencies for the generated package.
- */
-struct FGeneratedPackageResultStruct
-{
-	FAssetPackageData AssetPackageData;
-	TArray<FAssetDependency> PackageDependencies;
-};
+/** Return whether iterative cook is enabled for the given packagename, based on used-class allowlist/blocklist. */
+bool IsIterativeEnabled(FName PackageName, bool bAllowAllClasses,
+	const FAssetPackageData* OverrideAssetPackageData = nullptr);
+
+/** Store extra information derived during save and used by the cooker for the given EditorDomain package. */
+void CommitEditorDomainCookAttachments(FName PackageName, TArrayView<IPackageWriter::FCommitAttachmentInfo> Attachments);
 
 /**
- * Recording of the dependencies of a package discovered during cook, used in incremental cooks. All dependencies
- * except for those marked Runtime contribute to the packages TargetDomain Key. If HasKeyMatch returns false after
- * fetching this structure for a package at the beginning of cook, then the package is not iteratively skippable and
- * needs to be recooked, and this structure needs to be recalculated for the package.
- * 
- * Runtime fields on the structure are used to inform the cook of discovered softreferences that need to be added to
- * the cook when the package is cooked.
+ * Reads / writes an oplog for EditorDomain BuildDefinitionLists.
+ * TODO: Reduce duplication between this class and FZenStoreWriter
  */
-class FCookDependencies
+class FEditorDomainOplog
 {
 public:
-	FCookDependencies();
-	~FCookDependencies();
-	FCookDependencies(const FCookDependencies&);
-	FCookDependencies(FCookDependencies&&);
-	FCookDependencies& operator=(const FCookDependencies&);
-	FCookDependencies& operator=(FCookDependencies&&);
-	// Build Dependencies
-	const TArray<FName>& GetBuildPackageDependencies() const;
-	const TArray<FString>& GetConfigDependencies() const;
-
-	// Runtime Dependencies
-	const TArray<FName>& GetRuntimePackageDependencies() const;
-	// Script dependencies, needed for TryCalculateCurrentKey
-	const TArray<FName>& GetScriptPackageDependencies() const;
+	FEditorDomainOplog();
 
-	// Cook Dependencies
-	const TArray<UE::Cook::FCookDependency>& GetCookDependencies() const;
-	const TArray<UE::Cook::FCookDependency>& GetTransitiveBuildDependencies() const;
-
-	// Data about the Key and Package that are not dependencies
-	/**
-	 * True if the structure has been calculated or fetched and accurately reports dependencies and
-	 * key for the package. False if the stucture is default, has been reset, or was marked invalid.
-	 */
 	bool IsValid() const;
-	bool HasSaveResults() const;
-	FName GetPackageName() const;
-	const FIoHash& GetStoredKey() const;
-	const FIoHash& GetCurrentKey() const;
-	bool HasKeyMatch(const FAssetPackageData* OverrideAssetPackageData);
-
-	// Modifying the structure data
-	enum class ECurrentKeyResult
-	{
-		Success,
-		Invalidated,
-		Error,
-	};
-	ECurrentKeyResult TryCalculateCurrentKey(const FAssetPackageData* OverrideAssetPackageData,
-		TArray<TPair<ELogVerbosity::Type, FString>>* OutMessages = nullptr);
-	void Reset();
-	void Empty();
-
-	// Construction functions
-	/**
-	 * Read dependencies for the given targetplatform of the given package out of global dependency trackers
-	 * that have recorded its data during the package's load/other/save operations in the current cook session.
-	 */
-	static FCookDependencies Collect(UPackage* Package, const ITargetPlatform* TargetPlatform,
-		FSavePackageResultStruct* SaveResult, const FGeneratedPackageResultStruct* GeneratedResult,
-		TArray<FName>&& RuntimeDependencies, TArray<TPair<ELogVerbosity::Type, FString>>* OutMessages = nullptr);
-	static FCookDependencies CollectSettingsObject(const UObject* Object,
-		TArray<TPair<ELogVerbosity::Type, FString>>* OutMessages);
-
-	// Fetch function to load the dependencies from a PackageStore is not yet implemented independently for
-	// this structure. Use FIncrementalCookAttachments instead.
+	void CommitPackage(FName PackageName, TArrayView<IPackageWriter::FCommitAttachmentInfo> Attachments);
+	void GetOplogAttachments(TArrayView<FName> PackageName,
+		TArrayView<FUtf8StringView> AttachmentKeys,
+		TFunctionRef<void(FName PackageName, FUtf8StringView AttachmentKey, FCbObject&& Attachment)> Callback);
 
 private:
-	TArray<FName> BuildPackageDependencies;
-	TArray<FString> ConfigDependencies;
-	TArray<FName> RuntimePackageDependencies;
-	TArray<FName> ScriptPackageDependencies;
-	TArray<FName> RedirectorsAffectingPackageDependencies;
-	TArray<UE::Cook::FCookDependency> CookDependencies;
-	TArray<UE::Cook::FCookDependency> TransitiveBuildDependencies;
-	TArray<FString> ClassDependencies;
-	FName PackageName;
-	FIoHash StoredKey;
-	FIoHash CurrentKey;
-	bool bHasSaveResults = false;
-	bool bValid = false;
-
-	friend bool ::LoadFromCompactBinary(FCbObjectView ObjectView, FCookDependencies& CookAttachments);
-	friend FCbWriter& ::operator<<(FCbWriter& Writer, const FCookDependencies& CookAttachments);
-	friend FIncrementalCookAttachments;
-};
-
-/**
- * Non-persistent cache of groups of cookdependencies. Dependencies to a CookDependencyGroup are not persistently
- * recorded into the oplog, instead we make a copy of all of their dependencies and append those dependencies onto
- * the CookDependencies that are written for a package.
- *
- * Example: The cookdependencies used by the CDO of a settings object that itself is configured by config values.
- *          The settings object's class's schema and the list of config settings are included in the cookdependencies.
- */
-class FCookDependencyGroups
-{
-public:
-	struct FRecordedDependencies
+	struct FOplogEntry
 	{
-		FCookDependencies Dependencies;
-		TArray<TPair<ELogVerbosity::Type,FString>> Messages;
-		bool bInitialized = false;
-	};
-
-	static FCookDependencyGroups& Get();
-	FRecordedDependencies& FindOrCreate(UPTRINT Key);
-
-private:
-	TMap<UPTRINT, FRecordedDependencies> Groups;
-};
-
-/** Wrapper around TArray<FBuildDefinition>, used to provide custom functions for compactbinary, collection, and fetch */
-struct FBuildDefinitionList
-{
-public:
-	TArray<UE::DerivedData::FBuildDefinition> Definitions;
+		struct FAttachment
+		{
+			const UTF8CHAR* Key;
+			FIoHash Hash;
+		};
 
-	void Reset();
-	void Empty();
+		TArray<FAttachment> Attachments;
+	};
 
-	/** Collect DDC BuildDefinitions that were issued from the load/save of the given package and platform. */
-	static FBuildDefinitionList Collect(UPackage* Package, const ITargetPlatform* TargetPlatform,
-		TArray<TPair<ELogVerbosity::Type, FString>>* OutMessages = nullptr);
+	void InitializeRead();
 
-private:
-	friend bool ::LoadFromCompactBinary(FCbObject&& ObjectView, FBuildDefinitionList& Definitions);
-	friend FCbWriter& ::operator<<(FCbWriter& Writer, const FBuildDefinitionList& Definitions);
-};
+	FCbAttachment CreateAttachment(FSharedBuffer AttachmentData);
+	FCbAttachment CreateAttachment(FCbObject AttachmentData)
+	{
+		return CreateAttachment(AttachmentData.GetBuffer().ToShared());
+	}
 
-/** All of the metadata that is written/read to the oplog for the incremental cook of a package. */
-struct FIncrementalCookAttachments
-{
-	FCookDependencies Dependencies;
-	FBuildDefinitionList BuildDefinitions;
-	IPackageWriter::ECommitStatus CommitStatus = IPackageWriter::ECommitStatus::NotCommitted;
+	static void StaticInit();
+	static bool IsReservedOplogKey(FUtf8StringView Key);
 
-	void Reset();
-	void Empty();
-	void AppendCommitAttachments(TArray<IPackageWriter::FCommitAttachmentInfo>& OutAttachments);
+	UE::FZenStoreHttpClient HttpClient;
+	FCriticalSection Lock;
+	TMap<FName, FOplogEntry> Entries;
+	bool bConnectSuccessful = false;
+	bool bInitializedRead = false;
 
-	static void Fetch(TArrayView<FName> PackageNames, const ITargetPlatform* TargetPlatform,
-		ICookedPackageWriter* PackageWriter,
-		TUniqueFunction<void(FName PackageName, FIncrementalCookAttachments&& Result)>&& Callback);
-	static FIncrementalCookAttachments Collect(UPackage* Package,
-		const ITargetPlatform* TargetPlatform, FSavePackageResultStruct* SaveResult,
-		const FGeneratedPackageResultStruct* GeneratedResult, TArray<FName>&& RuntimeDependencies);
+	static TArray<const UTF8CHAR*> ReservedOplogKeys;
 };
-
-/** Return whether iterative cook is enabled for the given packagename, based on used-class allowlist/blocklist. */
-bool IsIterativeEnabled(FName PackageName, bool bAllowAllClasses,
-	const FAssetPackageData* OverrideAssetPackageData = nullptr);
-
-/** Store extra information derived during save and used by the cooker for the given EditorDomain package. */
-void CommitEditorDomainCookAttachments(FName PackageName, TArrayView<IPackageWriter::FCommitAttachmentInfo> Attachments);
-
-
-////////////////////////////////
-// Inline Implementations
-////////////////////////////////
-
-inline const TArray<FName>& FCookDependencies::GetBuildPackageDependencies() const
-{
-	return BuildPackageDependencies;
-}
-
-inline const TArray<FString>& FCookDependencies::GetConfigDependencies() const
-{
-	return ConfigDependencies;
-}
-
-inline const TArray<FName>& FCookDependencies::GetRuntimePackageDependencies() const
-{
-	return RuntimePackageDependencies;
-}
-
-inline const TArray<FName>& FCookDependencies::GetScriptPackageDependencies() const
-{
-	return ScriptPackageDependencies;
-}
-
-inline const TArray<UE::Cook::FCookDependency>& FCookDependencies::GetCookDependencies() const
-{
-	return CookDependencies;
-}
-
-inline const TArray<UE::Cook::FCookDependency>& FCookDependencies::GetTransitiveBuildDependencies() const
-{
-	return TransitiveBuildDependencies;
-}
-
-inline FName FCookDependencies::GetPackageName() const
-{
-	return PackageName;
-}
-
-inline const FIoHash& FCookDependencies::GetStoredKey() const
-{
-	return StoredKey;
-}
-
-inline const FIoHash& FCookDependencies::GetCurrentKey() const
-{
-	return CurrentKey;
-}
+extern TUniquePtr<FEditorDomainOplog> GEditorDomainOplog;
 
 } // namespace UE::TargetDomain
