ITF frustum select drag tool: fixed crash when initializing the drag tool

#rnx

diff --git a/Engine/Source/Editor/UnrealEd/Private/EditorViewportClient.cpp b/Engine/Source/Editor/UnrealEd/Private/EditorViewportClient.cpp
--- a/Engine/Source/Editor/UnrealEd/Private/EditorViewportClient.cpp
+++ b/Engine/Source/Editor/UnrealEd/Private/EditorViewportClient.cpp
@@ -6764,9 +6764,12 @@ void FEditorViewportClient::RegisterDragTools()
 					InteractiveToolContext->SetBypassAltModifier(false);
 
 					DragToolsBehaviorSource = NewObject<UDragToolsBehaviorSource>();
-					DragToolsBehaviorSource->Initialize(InteractiveToolContext);
+					DragToolsBehaviorSource->Initialize(InteractiveToolContext, this);
 
-					InteractiveToolContext->InputRouter->RegisterSource(DragToolsBehaviorSource);
+					if (ensure(DragToolsBehaviorSource->GetInputBehaviors()))
+					{
+						InteractiveToolContext->InputRouter->RegisterSource(DragToolsBehaviorSource);
+					}
 				}
 			}
 		}
diff --git a/Engine/Source/Editor/UnrealEd/Private/Tools/DragToolsBehaviourSource.cpp b/Engine/Source/Editor/UnrealEd/Private/Tools/DragToolsBehaviourSource.cpp
--- a/Engine/Source/Editor/UnrealEd/Private/Tools/DragToolsBehaviourSource.cpp
+++ b/Engine/Source/Editor/UnrealEd/Private/Tools/DragToolsBehaviourSource.cpp
@@ -6,7 +6,7 @@
 #include "DragTool_FrustumSelect.h"
 #include "EdModeInteractiveToolsContext.h"
 
-void UDragToolsBehaviorSource::Initialize(const UEditorInteractiveToolsContext* InInteractiveToolsContext)
+void UDragToolsBehaviorSource::Initialize(const UEditorInteractiveToolsContext* InInteractiveToolsContext, FEditorViewportClient* InViewportClient)
 {
 	if (!InInteractiveToolsContext)
 	{
@@ -25,21 +25,21 @@ void UDragToolsBehaviorSource::Initialize(const UEditorInteractiveToolsContext*
 		return;
 	}
 
-	FViewport* Viewport = ContextQueriesAPI->GetFocusedViewport();
-	if (!Viewport)
-	{
-		return;
-	}
-
-	FEditorViewportClient* EditorViewportClient = static_cast<FEditorViewportClient*>(Viewport->GetClient());
+	FEditorViewportClient* EditorViewportClient = InViewportClient;
 	if (!EditorViewportClient)
 	{
-		return;
-	}
+		FViewport* Viewport = ContextQueriesAPI->GetFocusedViewport();
+		if (!Viewport)
+		{
+			return;
+		}
 
-	// Some Drag Tools require or can make use specifically of Level Viewport Client.
-	// This can be nullptr
-	FLevelEditorViewportClient* LevelViewportClient = static_cast<FLevelEditorViewportClient*>(EditorViewportClient);
+		EditorViewportClient = static_cast<FEditorViewportClient*>(Viewport->GetClient());
+		if (!EditorViewportClient)
+		{
+			return;
+		}
+	}
 
 	BehaviorSet = NewObject<UInputBehaviorSet>();
 
@@ -57,9 +57,11 @@ void UDragToolsBehaviorSource::Initialize(const UEditorInteractiveToolsContext*
 	{
 		TSharedPtr<FDragTool> FrustumSelectTool;
 
-		if (LevelViewportClient)
+		// Some Drag Tools require or can make use specifically of Level Viewport Client.
+		if (EditorViewportClient->IsLevelEditorClient())
 		{
 			// Initialize using Level Viewport Client if available
+			FLevelEditorViewportClient* LevelViewportClient = static_cast<FLevelEditorViewportClient*>(EditorViewportClient);
 			FrustumSelectTool = MakeShared<FDragTool_ActorFrustumSelect>(LevelViewportClient);
 		}
 		else
diff --git a/Engine/Source/Editor/UnrealEd/Private/Tools/DragToolsBehaviourSource.h b/Engine/Source/Editor/UnrealEd/Private/Tools/DragToolsBehaviourSource.h
--- a/Engine/Source/Editor/UnrealEd/Private/Tools/DragToolsBehaviourSource.h
+++ b/Engine/Source/Editor/UnrealEd/Private/Tools/DragToolsBehaviourSource.h
@@ -8,6 +8,7 @@
 #include "DragToolsBehaviourSource.generated.h"
 
 class UEditorInteractiveToolsContext;
+class FEditorViewportClient;
 
 /**
  * Hosts Drag Tools and needed behaviors.
@@ -28,9 +29,10 @@ public:
 	 * Instantiates the BehaviorSet hosting all behaviors required by drag tools.
 	 *
 	 * @param InInteractiveToolsContext: used to retrieve the viewport client used by drag tools. Can be gathered e.g.
+	 * @param InViewportClient: the viewport client used by drag tools. (the focussed viewport will be retrieved from the context if set to null).
 	 * from UEdMode::GetInteractiveToolsContext.
 	 */
-	void Initialize(const UEditorInteractiveToolsContext* InInteractiveToolsContext);
+	void Initialize(const UEditorInteractiveToolsContext* InInteractiveToolsContext, FEditorViewportClient* InViewportClient = nullptr);
 
 	/**
 	 * Returns the drag tool currently being used, if any
