Removed Engine API for Blueprint sanitization

#rnx

diff --git a/Engine/Source/Editor/UnrealEd/Private/Kismet2/BlueprintEditorUtils.cpp b/Engine/Source/Editor/UnrealEd/Private/Kismet2/BlueprintEditorUtils.cpp
--- a/Engine/Source/Editor/UnrealEd/Private/Kismet2/BlueprintEditorUtils.cpp
+++ b/Engine/Source/Editor/UnrealEd/Private/Kismet2/BlueprintEditorUtils.cpp
@@ -10520,133 +10520,6 @@ bool FBlueprintEditorUtils::HasFunctionBlueprintThreadSafeMetaData(const UFuncti
 	return false;
 }
 
-bool FBlueprintEditorUtils::HasRestrictedNodes(const UBlueprint* BP, TArray<UEdGraphNode*>* OutRestrictedNodes /* = nullptr */)
-{
-	if (!BP)
-	{
-		return false;
-	}
-
-	UBlueprintEditorSettings* BlueprintEditorSettings = GetMutableDefault<UBlueprintEditorSettings>();
-	check(BlueprintEditorSettings);
-
-	// No need to continue if we're not restricting anything.
-	if (!BlueprintEditorSettings->HasClassFiltering()
-		&& !BlueprintEditorSettings->GetFunctionPermissions().HasFiltering())
-	{
-		return false;
-	}
-
-	if (OutRestrictedNodes)
-	{
-		OutRestrictedNodes->Reset();
-	}
-
-	// Check for any disallowed node types.
-	auto IsRestrictedNodeType = [BlueprintEditorSettings](const UEdGraphNode* InNode) -> bool
-	{
-		return !BlueprintEditorSettings->IsClassAllowed(InNode->GetClass());
-	};
-
-	// Check for any disallowed event graph implementations. The event node is considered restricted in this case.
-	auto IsRestrictedEventGraph = [BP, BlueprintEditorSettings](const UEdGraphNode* InNode) -> bool
-	{
-		const UK2Node_Event* EventNode = Cast<UK2Node_Event>(InNode);
-		return EventNode && !BlueprintEditorSettings->IsFunctionAllowed(BP, EventNode->EventReference.GetMemberName());
-	};
-
-	// Check for any disallowed function override graphs. Its entry node is considered to be restricted in this case.
-	auto IsRestrictedFunctionGraph = [BP, BlueprintEditorSettings](const UEdGraphNode* InNode) -> bool
-	{
-		const UK2Node_FunctionEntry* EntryNode = Cast<UK2Node_FunctionEntry>(InNode);
-		return EntryNode && !BlueprintEditorSettings->IsFunctionAllowed(BP, EntryNode->FunctionReference.GetMemberName());
-	};
-
-	TArray<UEdGraphNode*> AllNodes;
-	GetAllNodesOfClass(BP, AllNodes);
-	for (UEdGraphNode* Node : AllNodes)
-	{
-		if (!Node)
-		{
-			continue;
-		}
-
-		if (IsRestrictedNodeType(Node)
-			|| IsRestrictedEventGraph(Node)
-			|| IsRestrictedFunctionGraph(Node))
-		{
-			if (OutRestrictedNodes)
-			{
-				OutRestrictedNodes->Add(Node);
-			}
-			else
-			{
-				// We found a restricted node and/or graph entry point, so we can return immediately since we're not populating the output list.
-				return true;
-			}
-		}
-	}
-
-	if (OutRestrictedNodes && OutRestrictedNodes->Num() > 0)
-	{
-		return true;
-	}
-
-	return false;
-}
-
-void FBlueprintEditorUtils::SanitizeRestrictedContent(UBlueprint* BP)
-{
-	TArray<UEdGraphNode*> RestrictedNodes;
-	if (HasRestrictedNodes(BP, &RestrictedNodes))
-	{
-		// Remove any restricted nodes and/or graphs (note: in restricted mode, it is implied that
-		// restricted items are not also user-facing, so any removals here should be opaque to users).
-		TSet<UEdGraph*> NodeSubgraphs;
-		TSet<UEdGraph*> GraphsToRemove;
-		for (UEdGraphNode* Node : RestrictedNodes)
-		{
-			// If we've already handled the node's graph/subgraph, no need to remove it again here.
-			const UEdGraph* NodeGraph = Node->GetGraph();
-			if (GraphsToRemove.Contains(NodeGraph)
-				|| NodeSubgraphs.Contains(NodeGraph))
-			{
-				continue;
-			}
-
-			// Remove the entire graph if this node represents an entry point for an event/function.
-			if (Node->IsA<UK2Node_Event>()
-				|| Node->IsA<UK2Node_FunctionEntry>())
-			{
-				GraphsToRemove.Add(Node->GetGraph());
-			}
-			else
-			{
-				// This may remove a subset of restricted nodes, so keep track of removed subgraphs.
-				NodeSubgraphs.Append(Node->GetSubGraphs());
-
-				// Defer updating the skeleton class since we're going to do a full recompile below.
-				RemoveNode(BP, Node, /*bDontRecompile =*/ true);
-
-				// Exclude any subgraphs that were not removed above from the set that we're tracking.
-				for (UEdGraph* SubGraph : Node->GetSubGraphs())
-				{
-					NodeSubgraphs.Remove(SubGraph);
-				}
-			}
-		}
-
-		// Note: On removal, graphs will be renamed and outered to the package containing the Blueprint.
-		for (UEdGraph* Graph : GraphsToRemove)
-		{
-			RemoveGraph(BP, Graph, EGraphRemoveFlags::MarkTransient);
-		}
-
-		// Do a full recompile to bring everything back up-to-date and to sanitize any generated bytecode.
-		FKismetEditorUtilities::CompileBlueprint(BP, EBlueprintCompileOptions::SkipGarbageCollection);
-	}
-}
-
 bool FBlueprintEditorUtils::ShouldOpenWithDataOnlyEditor(const UBlueprint* Blueprint)
 {
 	return FBlueprintEditorUtils::IsDataOnlyBlueprint(Blueprint)
diff --git a/Engine/Source/Editor/UnrealEd/Public/Kismet2/BlueprintEditorUtils.h b/Engine/Source/Editor/UnrealEd/Public/Kismet2/BlueprintEditorUtils.h
--- a/Engine/Source/Editor/UnrealEd/Public/Kismet2/BlueprintEditorUtils.h
+++ b/Engine/Source/Editor/UnrealEd/Public/Kismet2/BlueprintEditorUtils.h
@@ -1974,21 +1974,6 @@ public:
 	 */
 	static UNREALED_API bool HasFunctionBlueprintThreadSafeMetaData(const UFunction* InFunction);
 
-	/**
-	 * Returns true if the given Blueprint is found to contain one or more nodes that are disallowed (restricted) within
-	 * the current editor mode context.
-	 * 
-	 * @param OutRestrictedNodes	(Optional) If non-NULL, will be populated with references to any nodes that are restricted.
-	 */
-	static UNREALED_API bool HasRestrictedNodes(const UBlueprint* BP, TArray<UEdGraphNode*>* OutRestrictedNodes = nullptr);
-
-	/**
-	 * Checks the given Blueprint for any restricted content within the current editor context and sanitizes it away.
-	 * 
-	 * Note: If any restricted content is removed, this will also recompile the Blueprint.
-	 */
-	static UNREALED_API void SanitizeRestrictedContent(UBlueprint* BP);
-
 public:
 	static UNREALED_API bool ShouldOpenWithDataOnlyEditor(const UBlueprint* Blueprint);
 };
