Include InOldToNewClassMap in the OldToNewInstanceMap in order to prevent mutation of old classes by reference replacement when running with EBlueprintCompileOptions::IncludeCDOInReferenceReplacement, plus some minor cleanup of the new ReinstanceWithVersioningData codepath

#rnx

diff --git a/Engine/Source/Editor/UnrealEd/Private/Kismet2/KismetReinstanceUtilities.cpp b/Engine/Source/Editor/UnrealEd/Private/Kismet2/KismetReinstanceUtilities.cpp
--- a/Engine/Source/Editor/UnrealEd/Private/Kismet2/KismetReinstanceUtilities.cpp
+++ b/Engine/Source/Editor/UnrealEd/Private/Kismet2/KismetReinstanceUtilities.cpp
@@ -2661,22 +2661,14 @@ void FBlueprintCompileReinstancer::ReplaceInstancesOfClass_Inner(const TMap<UCla
 			check(bArchetypesAreUpToDate); // not supported, can't be
 			check(Params.InstancesThatShouldUseOldClass == nullptr); // not supported, could be.. but why?
 
-			// if an OldToNewTemplates map has been provided, include all the instances in the replacement map
-			// this is used to make reinstancing during PIE more robust:
-			if (Params.OldToNewTemplates)
-			{
-				for(const TPair<UClass*, TMap<UObject*, UObject*>>& Mapping : *Params.OldToNewTemplates)
-				{
-					OldToNewInstanceMap.Append(Mapping.Value);
-				}
-			}
-			ReinstanceWithVersioningData(InOldToNewClassMap, 
+			ReinstanceWithVersioningData(
 				{
 					.ObjectsToReplaceScratch = ObjectsToReplace,
 					.CachedArchetypeObjects = CachedArchetypeObjects,
 					.OldToNewInstanceMap = OldToNewInstanceMap,
 					.InOldToNewClassMap = InOldToNewClassMap,
 					.OldToNewNameMap = OldToNewNameMap,
+					.OldToNewTemplates = Params.OldToNewTemplates,
 					.PotentialEditorsForRefreshing = PotentialEditorsForRefreshing,
 					.OwnersToRerunConstructionScript = OwnersToRerunConstructionScript,
 					.ReplacementActors = ReplacementActors,
@@ -3436,13 +3428,30 @@ FBlake3Hash FBlueprintCompileReinstancer::CalculateConstructionHash(const UClass
 	return Hash.Finalize();
 }
 
-void FBlueprintCompileReinstancer::ReinstanceWithVersioningData(const TMap<UClass*, UClass*>& InOldToNewClassMap, const FReinstancingState& ReinstancingState)
+void FBlueprintCompileReinstancer::ReinstanceWithVersioningData(const FReinstancingState& ReinstancingState)
 {
+	// if an OldToNewTemplates map has been provided, include all the instances in the replacement map
+	// this is used to make reinstancing during PIE more robust:
+	if (ReinstancingState.OldToNewTemplates)
+	{
+		for (const TPair<UClass*, TMap<UObject*, UObject*>>& Mapping : *ReinstancingState.OldToNewTemplates)
+		{
+			ReinstancingState.OldToNewInstanceMap.Append(Mapping.Value);
+		}
+	}
+	// Include a mapping from the old class to the new class as well - this will
+	// prevent mutation of the old classes by reference replacement, and aslo patch
+	// leaks to REINST classes:
+	for (const TPair<UClass*, UClass*>& Mapping : ReinstancingState.InOldToNewClassMap)
+	{
+		ReinstancingState.OldToNewInstanceMap.Add(Mapping.Key, Mapping.Value);
+	}
+
 	TMap<UClass*, UClass*> UnchangedClasses;
 	TMap<UClass*, UClass*> NormalReinstancingClasses;
 	TMap<UClass*, UClass*> ActorReinstancingClasses;
 	CategorizeClassesAndInstances(
-		InOldToNewClassMap, 
+		ReinstancingState.InOldToNewClassMap, 
 		UnchangedClasses, 
 		NormalReinstancingClasses, 
 		ActorReinstancingClasses, 
diff --git a/Engine/Source/Editor/UnrealEd/Public/Kismet2/KismetReinstanceUtilities.h b/Engine/Source/Editor/UnrealEd/Public/Kismet2/KismetReinstanceUtilities.h
--- a/Engine/Source/Editor/UnrealEd/Public/Kismet2/KismetReinstanceUtilities.h
+++ b/Engine/Source/Editor/UnrealEd/Public/Kismet2/KismetReinstanceUtilities.h
@@ -353,6 +353,7 @@ private:
 		TMap<UObject*, UObject*>& OldToNewInstanceMap;
 		const TMap<UClass*, UClass*>& InOldToNewClassMap;
 		TMap<UObject*, FName>& OldToNewNameMap;
+		TMap<UClass*, TMap<UObject*, UObject*>>* OldToNewTemplates = nullptr;
 		TArray<UObject*>& PotentialEditorsForRefreshing;
 		TSet<AActor*>& OwnersToRerunConstructionScript;
 		TArray<struct FActorReplacementHelper>& ReplacementActors;
@@ -364,8 +365,7 @@ private:
 		bool& bFixupSCS;
 	};
 
-	static void ReinstanceWithVersioningData(
-		const TMap<UClass*, UClass*>& InOldToNewClassMap, const FReinstancingState& ReinstancingState);
+	static void ReinstanceWithVersioningData(const FReinstancingState& ReinstancingState);
 	static void CategorizeClassesAndInstances(
 		const TMap<UClass*, UClass*>& InOldToNewClassMap
 		, TMap<UClass*, UClass*>& UnchangedClasses
