Adds bRequireAssetRegistryLoadedBeforeIteration. Assest registry scan is required so that EPackageLocationFilter::FileSystem is set for packages, otherwise the FPackageName::DoesPackageExistEx optimisation to avoid disc scanning fails.

diff --git a/Engine/Source/Editor/UnrealEd/Classes/Commandlets/BaseIteratePackagesCommandlet.h b/Engine/Source/Editor/UnrealEd/Classes/Commandlets/BaseIteratePackagesCommandlet.h
--- a/Engine/Source/Editor/UnrealEd/Classes/Commandlets/BaseIteratePackagesCommandlet.h
+++ b/Engine/Source/Editor/UnrealEd/Classes/Commandlets/BaseIteratePackagesCommandlet.h
@@ -112,6 +112,9 @@ protected:
 	/** Similar to above, but applies to all packages rather than just maps **/
 	//bool bShouldBuildTextureStreamingForAll;
 
+	/** Should the asset registry have been scanned before iterating packages */
+	bool bRequireAssetRegistryLoadedBeforeIteration;
+
 	bool bForceFinishAllCompilationBeforeGC;
 
 	/** only process packages containing materials **/
diff --git a/Engine/Source/Editor/UnrealEd/Private/Commandlets/BaseIteratePackagesCommandlet.cpp b/Engine/Source/Editor/UnrealEd/Private/Commandlets/BaseIteratePackagesCommandlet.cpp
--- a/Engine/Source/Editor/UnrealEd/Private/Commandlets/BaseIteratePackagesCommandlet.cpp
+++ b/Engine/Source/Editor/UnrealEd/Private/Commandlets/BaseIteratePackagesCommandlet.cpp
@@ -35,6 +35,7 @@
 #include "EngineGlobals.h"
 #include "GameFramework/WorldSettings.h"
 #include "AssetRegistry/AssetData.h"
+#include "AssetRegistry/IAssetRegistry.h"
 #include "Editor.h"
 #include "EditorWorldUtils.h"
 #include "FileHelpers.h"
@@ -91,7 +92,7 @@ UBaseIteratePackagesCommandlet::UBaseIteratePackagesCommandlet(const FObjectInit
 
 int32 UBaseIteratePackagesCommandlet::InitializeParameters( const TArray<FString>& Tokens, TArray<FString>& PackageNames )
 {
-	TRACE_CPUPROFILER_EVENT_SCOPE(UBaseIteratePackagesCommandlet::InitializeResaveParameters);
+	TRACE_CPUPROFILER_EVENT_SCOPE(UBaseIteratePackagesCommandlet::InitializeParameters);
 
 	Verbosity = VERY_VERBOSE;
 
@@ -1070,7 +1071,13 @@ int32 UBaseIteratePackagesCommandlet::Main( const FString& Params )
 	/** only process packages containing materials */
 	bOnlyMaterials = Switches.Contains(TEXT("onlymaterials"));
 
+	// Start an async scan of the asset registry. The asset registry being loaded is required to set EPackageLocationFilter::FileSystem to support loading uncooked.
+	FAssetRegistryModule& AssetRegistryModule = FModuleManager::LoadModuleChecked<FAssetRegistryModule>(TEXT("AssetRegistry"));
+	IAssetRegistry& AssetRegistry = AssetRegistryModule.Get();
+	AssetRegistry.SearchAllAssets(false);
+
 	bForceFinishAllCompilationBeforeGC = false;
+	bRequireAssetRegistryLoadedBeforeIteration = true;
 
 	bUseWorldPartitionBuilder = false;
 	WorldInitialisationValues.RequiresHitProxies(false);
@@ -1173,6 +1180,13 @@ int32 UBaseIteratePackagesCommandlet::Main( const FString& Params )
 	// Make sure any remaining temp files from previous runs are removed
 	CleanTempFiles();
 
+	if(bRequireAssetRegistryLoadedBeforeIteration)
+	{
+		// Ensure the asset registry has loaded, before iterating.
+		TRACE_CPUPROFILER_EVENT_SCOPE(AssetRegistry.WaitForCompletion);
+		AssetRegistry.WaitForCompletion();
+	}
+	
 	// Iterate over all packages.
 	for( int32 PackageIndex = 0; PackageIndex < PackageNames.Num(); PackageIndex++ )
 	{
@@ -1239,8 +1253,6 @@ int32 UBaseIteratePackagesCommandlet::Main( const FString& Params )
 	FDirectoryWatcherModule& DirectoryWatcherModule = FModuleManager::Get().LoadModuleChecked<FDirectoryWatcherModule>(TEXT("DirectoryWatcher"));
 	DirectoryWatcherModule.Get()->Tick(-1.0f);
 
-	FAssetRegistryModule& AssetRegistryModule = FModuleManager::LoadModuleChecked<FAssetRegistryModule>(TEXT("AssetRegistry"));
-	IAssetRegistry& AssetRegistry = AssetRegistryModule.Get();
 	AssetRegistry.Tick(-1.0f);
 
 	PostProcessPackages();
