Cooker: Add ExternalObjects to the NeverCook paths in the same way ExternalActors are in NeverCook. Loading ExternalObjects or ExternalActors in the cooker can cause the reload of packages that have been garbage collected after their objects were moved out into a WorldPartition map and the map has not yet been saved, and this reload of the external package can create a duplicate object that gets found during the save of the map and creates a cook error. Currently the cooker loading the package deliberately to try to save it is the only source of those reloads, so stopping the cooker from doing that prevents the cook errors.
#rnx

diff --git a/Engine/Source/Editor/UnrealEd/Private/Commandlets/AssetRegistryGenerator.cpp b/Engine/Source/Editor/UnrealEd/Private/Commandlets/AssetRegistryGenerator.cpp
--- a/Engine/Source/Editor/UnrealEd/Private/Commandlets/AssetRegistryGenerator.cpp
+++ b/Engine/Source/Editor/UnrealEd/Private/Commandlets/AssetRegistryGenerator.cpp
@@ -1374,6 +1374,7 @@ void FAssetRegistryGenerator::ComputePackageDifferences(const FComputeDifference
 	if (Options.bRecurseModifications)
 	{
 		const FStringView ExternalActorsFolderName(ULevel::GetExternalActorsFolderName());
+		const FStringView ExternalObjectsFolderName = FPackagePath::GetExternalObjectsFolderName();
 
 		// Recurse modified packages to their dependencies. This is needed because we only compare package guids
 		TArray<FName> VisitStack;
@@ -1416,20 +1417,29 @@ void FAssetRegistryGenerator::ComputePackageDifferences(const FComputeDifference
 				// Workaround: Detect external actors by naming convention and suppress their reference to the map package.
 				// See also UAssetManager::ShouldSetManager
 				TStringBuilder<256> ReferencerPackageNameStr(InPlace, ReferencerPackageName);
-				int32 ExternalActorsFolderIndex = UE::String::FindFirst(ReferencerPackageNameStr, ExternalActorsFolderName, ESearchCase::IgnoreCase);
-				if (ExternalActorsFolderIndex != INDEX_NONE)
+				bool bReferencerIsMapExternalPackage = false;
+				for (FStringView ExternalFolderName : { ExternalActorsFolderName, ExternalObjectsFolderName })
 				{
-					if (ModifiedPackageLeafName.IsEmpty())
+					int32 ExternalFolderIndex = UE::String::FindFirst(ReferencerPackageNameStr, ExternalFolderName, ESearchCase::IgnoreCase);
+					if (ExternalFolderIndex != INDEX_NONE)
 					{
-						ModifiedPackageLeafName = FPathViews::GetCleanFilename(WriteToString<256>(ModifiedPackage));
-					}
-					FStringView GeneratedRelativePath = ReferencerPackageNameStr.ToView().RightChop(ExternalActorsFolderIndex + ExternalActorsFolderName.Len());
-					if (GeneratedRelativePath.Contains(ModifiedPackageLeafName))
-					{
-						// Suppress this reference; External actors should not be marked dirty if only their map package is dirty
-						continue;
+						if (ModifiedPackageLeafName.IsEmpty())
+						{
+							ModifiedPackageLeafName = FPathViews::GetCleanFilename(WriteToString<256>(ModifiedPackage));
+						}
+						FStringView GeneratedRelativePath = ReferencerPackageNameStr.ToView().RightChop(ExternalFolderIndex + ExternalFolderName.Len());
+						if (GeneratedRelativePath.Contains(ModifiedPackageLeafName))
+						{
+							bReferencerIsMapExternalPackage = true;
+							break;
+						}
 					}
 				}
+				if (bReferencerIsMapExternalPackage)
+				{
+					// Suppress this reference; External actors should not be marked dirty if only their map package is dirty
+					continue;
+				}
 
 				int32 ReferencerPackageHash = GetTypeHash(ReferencerPackageName);
 				bool bAlreadyVisited;
diff --git a/Engine/Source/Editor/UnrealEd/Private/CookOnTheFlyServer.cpp b/Engine/Source/Editor/UnrealEd/Private/CookOnTheFlyServer.cpp
--- a/Engine/Source/Editor/UnrealEd/Private/CookOnTheFlyServer.cpp
+++ b/Engine/Source/Editor/UnrealEd/Private/CookOnTheFlyServer.cpp
@@ -1448,7 +1448,9 @@ void UCookOnTheFlyServer::RunCookList(ECookListOptions CookListOptions)
 			if (bIncludePackage)
 			{
 				// Skip printing out a message for external actors
-				if (INDEX_NONE != UE::String::FindFirst(WriteToString<256>(PackageData->GetPackageName()), ULevel::GetExternalActorsFolderName()))
+				TStringBuilder<256> PackageNameStr(InPlace, PackageData->GetPackageName());
+				if (INDEX_NONE != UE::String::FindFirst(PackageNameStr, ULevel::GetExternalActorsFolderName()) ||
+					INDEX_NONE != UE::String::FindFirst(PackageNameStr, FPackagePath::GetExternalObjectsFolderName()))
 				{
 					bIncludePackage = false;
 				}
@@ -2910,6 +2912,8 @@ void UCookOnTheFlyServer::DemoteToIdle(UE::Cook::FPackageData& PackageData, UE::
 			{
 				bPrintDiagnostic &= UE::String::FindFirst(PackageNameStr.ToView(),
 					ULevel::GetExternalActorsFolderName(), ESearchCase::IgnoreCase) == INDEX_NONE;
+				bPrintDiagnostic &= UE::String::FindFirst(PackageNameStr.ToView(),
+					FPackagePath::GetExternalObjectsFolderName(), ESearchCase::IgnoreCase) == INDEX_NONE;
 			}
 
 			// Reachability: Suppress the diagnostic that were found via cookload reference traversal but are not reachable on the target platforms
@@ -4632,12 +4636,15 @@ void UCookOnTheFlyServer::ProcessUnsolicitedPackages(TArray<FName>* OutDiscovere
 		}
 
 		AddToOutDiscovered(PackageData, Instigator);
+		TStringBuilder<256> PackageNameStr(InPlace, PackageData->GetPackageName());
 
 		if (bSkipOnlyEditorOnly &&
 			(Instigator.Category == EInstigator::ForceExplorableSaveTimeSoftDependency ||
 				(PackageTracker->NeverCookPackageList.Contains(PackageData->GetPackageName()) &&
-					INDEX_NONE != UE::String::FindFirst(WriteToString<256>(PackageData->GetPackageName()),
-						ULevel::GetExternalActorsFolderName(), ESearchCase::IgnoreCase))))
+					(INDEX_NONE != UE::String::FindFirst(PackageNameStr,
+						ULevel::GetExternalActorsFolderName(), ESearchCase::IgnoreCase) ||
+						INDEX_NONE != UE::String::FindFirst(PackageNameStr,
+							FPackagePath::GetExternalObjectsFolderName(), ESearchCase::IgnoreCase)))))
 		{
 			// ONLYEDITORONLY_TODO: WorldPartition should mark up these loads with EInstigator::ForceExplorableSaveTimeSoftDependency rather than
 			// needing to use a naming convention. We should also mark them once, during Save, rather than marking them
@@ -11807,14 +11814,23 @@ TArray<FName> UCookOnTheFlyServer::GetNeverCookPackageNames(TArrayView<const FSt
 
 	// For all modes, never cook External Actors; they are handled by the parent map
 	FString ExternalActorsFolderName = ULevel::GetExternalActorsFolderName();
-	FString FullExternalActorsPath = FPaths::Combine(TEXT("/Game/"), ExternalActorsFolderName);
-	NeverCookDirectories.Add(MoveTemp(FullExternalActorsPath));
-	FullExternalActorsPath = FPaths::Combine(TEXT("/Engine/"), ExternalActorsFolderName);
-	NeverCookDirectories.Add(MoveTemp(FullExternalActorsPath));
+	FString ExternalObjectsFolderName = FPackagePath::GetExternalObjectsFolderName();
+	FString FullExternalPath;
+	for (const TCHAR* ProjectFolder : { TEXT("/Game/"), TEXT("/Engine/") })
+	{
+		for (const FString* ExternalFolderName : { &ExternalActorsFolderName, &ExternalObjectsFolderName })
+		{
+			FullExternalPath = FPaths::Combine(FString(ProjectFolder), *ExternalFolderName);
+			NeverCookDirectories.Add(MoveTemp(FullExternalPath));
+		}
+	}
 	for (TSharedRef<IPlugin>& Plugin : IPluginManager::Get().GetEnabledPluginsWithContent())
 	{
-		FullExternalActorsPath = FPaths::Combine(Plugin->GetMountedAssetPath(), ExternalActorsFolderName);
-		NeverCookDirectories.Add(MoveTemp(FullExternalActorsPath));
+		for (const FString* ExternalFolderName : { &ExternalActorsFolderName, &ExternalObjectsFolderName })
+		{
+			FullExternalPath = FPaths::Combine(Plugin->GetMountedAssetPath(), *ExternalFolderName);
+			NeverCookDirectories.Add(MoveTemp(FullExternalPath));
+		}
 	}
 
 	TArray<FName> NeverCookPackages;
diff --git a/Engine/Source/Editor/UnrealEd/Private/Cooker/CookRequestCluster.cpp b/Engine/Source/Editor/UnrealEd/Private/Cooker/CookRequestCluster.cpp
--- a/Engine/Source/Editor/UnrealEd/Private/Cooker/CookRequestCluster.cpp
+++ b/Engine/Source/Editor/UnrealEd/Private/Cooker/CookRequestCluster.cpp
@@ -2544,6 +2544,8 @@ void FRequestCluster::IsRequestCookable(const ITargetPlatform* Platform, FPackag
 	if (InCOTFS.PackageTracker->NeverCookPackageList.Contains(PackageName))
 	{
 		if (INDEX_NONE != UE::String::FindFirst(NameBuffer, ULevel::GetExternalActorsFolderName(), 
+			ESearchCase::IgnoreCase) ||
+			INDEX_NONE != UE::String::FindFirst(NameBuffer, FPackagePath::GetExternalObjectsFolderName(),
 			ESearchCase::IgnoreCase))
 		{
 			// EXTERNALACTOR_TODO: Add a separate category for ExternalActors rather than putting them in
