Expose control of the Near View Plane in the level editor viewport options.

diff --git a/Engine/Source/Editor/EditorStyle/Private/StarshipStyle.cpp b/Engine/Source/Editor/EditorStyle/Private/StarshipStyle.cpp
--- a/Engine/Source/Editor/EditorStyle/Private/StarshipStyle.cpp
+++ b/Engine/Source/Editor/EditorStyle/Private/StarshipStyle.cpp
@@ -2686,6 +2686,7 @@ void FStarshipEditorStyle::FStyle::SetupViewportStyles()
 			Set("ViewportToolbar.FarViewPlane", new IMAGE_BRUSH_SVG("Starship/Common/FarViewPlane_16", Icon16x16));
 			Set("ViewportToolbar.FieldOfView", new IMAGE_BRUSH_SVG("Starship/Common/FieldOfView_16", Icon16x16));
 			Set("ViewportToolbar.GizmoScale", new IMAGE_BRUSH_SVG("Starship/Common/GizmoScale_16", Icon16x16));
+			Set("ViewportToolbar.NearViewPlane", new IMAGE_BRUSH_SVG("Starship/Common/NearViewPlane_16", Icon16x16));
 			Set("ViewportToolbar.PreviewSceneSettings", new IMAGE_BRUSH_SVG("Starship/AssetEditors/PreviewSceneSettings_16", Icon16x16));
 			Set("ViewportToolbar.SetShowGrid", new IMAGE_BRUSH_SVG("Starship/Common/Grid", Icon16x16));
 			Set("ViewportToolbar.Snap", new IMAGE_BRUSH_SVG("Starship/Common/Snap_16", Icon16x16));
diff --git a/Engine/Source/Editor/UnrealEd/Private/ViewportToolbar/UnrealEdViewportToolbar.cpp b/Engine/Source/Editor/UnrealEd/Private/ViewportToolbar/UnrealEdViewportToolbar.cpp
--- a/Engine/Source/Editor/UnrealEd/Private/ViewportToolbar/UnrealEdViewportToolbar.cpp
+++ b/Engine/Source/Editor/UnrealEd/Private/ViewportToolbar/UnrealEdViewportToolbar.cpp
@@ -2292,6 +2292,52 @@ TSharedRef<SWidget> CreateFOVMenuWidget(const TSharedRef<SEditorViewport>& InVie
 	// clang-format on
 }
 
+TSharedRef<SWidget> CreateNearViewPlaneMenuWidget(const TSharedRef<SEditorViewport>& InViewport)
+{
+	TWeakPtr<FEditorViewportClient> ViewportClientWeak = InViewport->GetViewportClient();
+
+	return
+		// clang-format off
+		SNew(SBox)
+		.HAlign(HAlign_Right)
+		[
+			SNew(SBox)
+			.Padding(FMargin(4.0f, 0.0f, 0.0f, 0.0f))
+			.WidthOverride(100.0f)
+			[
+				SNew(SBorder)
+				.BorderImage(FAppStyle::Get().GetBrush("Menu.WidgetBorder"))
+				.Padding(FMargin(1.0f))
+				[
+					SNew(SSpinBox<float>)
+					.Style(&FAppStyle::Get(), "Menu.SpinBox")
+					.ToolTipText(LOCTEXT("NearViewPlaneTooltip", "Distance to use as the near view plane"))
+					.MinValue(0.001f)
+					.MaxValue(100.0f)
+					.Font(FAppStyle::GetFontStyle(TEXT("MenuItem.Font")))
+					.Value_Lambda([ViewportClientWeak]()
+					{
+						if (TSharedPtr<FEditorViewportClient> ViewportClient = ViewportClientWeak.Pin())
+						{
+							return ViewportClient->GetNearClipPlane();
+						}
+
+						return 1.0f;
+					})
+					.OnValueChanged_Lambda([ViewportClientWeak](float InNewValue)
+					{
+						if (TSharedPtr<FEditorViewportClient> ViewportClient = ViewportClientWeak.Pin())
+						{
+							ViewportClient->OverrideNearClipPlane(InNewValue);
+							ViewportClient->Invalidate();
+						}
+					})
+				]
+			]
+		];
+	// clang-format on
+}
+
 TSharedRef<SWidget> CreateFarViewPlaneMenuWidget(const TSharedRef<SEditorViewport>& InViewport)
 {
 	TWeakPtr<FEditorViewportClient> ViewportClientWeak = InViewport->GetViewportClient();
@@ -2820,14 +2866,14 @@ void PopulateCameraMenu(UToolMenu* InMenu)
 	}
 }
 
-void ExtendCameraSubmenu(FName InCameraOptionsSubmenuName, bool bInShowFarViewPlaneEntry /* = true */)
+void ExtendCameraSubmenu(FName InCameraOptionsSubmenuName, bool bInShowViewPlaneEntries /* = true */)
 {
 	UToolMenu* const Submenu = UToolMenus::Get()->ExtendMenu(InCameraOptionsSubmenuName);
 
 	Submenu->AddDynamicSection(
 		"EditorCameraExtensionDynamicSection",
 		FNewToolMenuDelegate::CreateLambda(
-			[bInShowFarViewPlaneEntry](UToolMenu* InDynamicMenu)
+			[bInShowViewPlaneEntries](UToolMenu* InDynamicMenu)
 			{
 				UUnrealEdViewportToolbarContext* const EditorViewportContext =
 					InDynamicMenu->FindContext<UUnrealEdViewportToolbarContext>();
@@ -2855,9 +2901,17 @@ void ExtendCameraSubmenu(FName InCameraOptionsSubmenuName, bool bInShowFarViewPl
 				CameraFOV.Icon = FSlateIcon(FAppStyle::GetAppStyleSetName(), "ViewportToolbar.FieldOfView");
 				PostOrthoSection.AddEntry(CameraFOV);
 
-				if (bInShowFarViewPlaneEntry)
+				if (bInShowViewPlaneEntries)
 				{
-					FToolMenuEntry CameraFarViewPlane =FToolMenuEntry::InitWidget(
+					FToolMenuEntry CameraNearViewPlane = FToolMenuEntry::InitWidget(
+						"CameraNearViewPlane",
+						UE::UnrealEd::CreateNearViewPlaneMenuWidget(EditorViewport.ToSharedRef()),
+						LOCTEXT("CameraSubmenu_NearViewPlaneLabel", "Near View Plane")
+					);
+					CameraNearViewPlane.Icon = FSlateIcon(FAppStyle::GetAppStyleSetName(), "ViewportToolbar.NearViewPlane");
+					PostOrthoSection.AddEntry(CameraNearViewPlane);
+				
+					FToolMenuEntry CameraFarViewPlane = FToolMenuEntry::InitWidget(
 						"CameraFarViewPlane",
 						UE::UnrealEd::CreateFarViewPlaneMenuWidget(EditorViewport.ToSharedRef()),
 						LOCTEXT("CameraSubmenu_FarViewPlaneLabel", "Far View Plane")
diff --git a/Engine/Source/Editor/UnrealEd/Public/ViewportToolbar/UnrealEdViewportToolbar.h b/Engine/Source/Editor/UnrealEd/Public/ViewportToolbar/UnrealEdViewportToolbar.h
--- a/Engine/Source/Editor/UnrealEd/Public/ViewportToolbar/UnrealEdViewportToolbar.h
+++ b/Engine/Source/Editor/UnrealEd/Public/ViewportToolbar/UnrealEdViewportToolbar.h
@@ -163,7 +163,7 @@ UNREALED_API void PopulateCameraMenu(UToolMenu* InMenu);
 /**
  * Adds Field of View and Far View Plane entries to the specified Camera Submenu
  */
-UNREALED_API void ExtendCameraSubmenu(FName InCameraOptionsSubmenuName, bool bInShowFarViewPlaneEntry = true);
+UNREALED_API void ExtendCameraSubmenu(FName InCameraOptionsSubmenuName, bool bInShowViewPlaneEntries = true);
 
 UNREALED_API void GenerateViewportTypeMenu(UToolMenu* InMenu);
 
@@ -208,6 +208,7 @@ UNREALED_API FText GetCameraSpeedTooltip();
 // Camera Menu Widgets
 TSharedRef<SWidget> CreateCameraMenuWidget(const TSharedRef<SEditorViewport>& InViewport, bool bInShowExposureSettings = false);
 TSharedRef<SWidget> CreateFOVMenuWidget(const TSharedRef<SEditorViewport>& InViewport);
+TSharedRef<SWidget> CreateNearViewPlaneMenuWidget(const TSharedRef<SEditorViewport>& InViewport);
 TSharedRef<SWidget> CreateFarViewPlaneMenuWidget(const TSharedRef<SEditorViewport>& InViewport);
 
 // Screen Percentage Submenu Widgets
