Do not register dependencies for in-memory packages.

#rnx
#incrementalcook

diff --git a/Engine/Source/Editor/UnrealEd/Private/Cooker/CookPackageArtifacts.cpp b/Engine/Source/Editor/UnrealEd/Private/Cooker/CookPackageArtifacts.cpp
--- a/Engine/Source/Editor/UnrealEd/Private/Cooker/CookPackageArtifacts.cpp
+++ b/Engine/Source/Editor/UnrealEd/Private/Cooker/CookPackageArtifacts.cpp
@@ -147,16 +147,28 @@ FPackageArtifacts::ECurrentKeyResult FPackageArtifacts::TryCalculateCurrentKey(
 	for (FName PackageDependency : BuildPackageDependencies)
 	{
 		PackageDigest = EditorDomain->GetPackageDigest(PackageDependency);
-		if (!PackageDigest.IsSuccessful())
+		if (PackageDigest.IsSuccessful())
 		{
-			if (OutMessages)
-			{
-				OutMessages->Emplace(ELogVerbosity::Error, FString::Printf(TEXT("Could not create PackageDigest for %s: %s"),
-					*PackageDependency.ToString(), *PackageDigest.GetStatusString()));
-			}
-			return ECurrentKeyResult::Error;
+			KeyBuilder.Update(&PackageDigest.Hash, sizeof(PackageDigest.Hash));
+			continue;
+		}
+
+		// The digest was not found, it could be because the package was created during the cook and does not exist on disk (in-memory package).
+		// Lets ignore those packages. Whoever created this package is responsible of saving dependencies if any.
+		TStringBuilder<256> PackageDependencyStr(EInPlace::InPlace, PackageDependency);
+		const UPackage* PackageObject = FindPackage(nullptr, *PackageDependencyStr);
+		if (PackageObject && PackageObject->HasAnyPackageFlags(PKG_NewlyCreated))
+		{
+			continue;
 		}
-		KeyBuilder.Update(&PackageDigest.Hash, sizeof(PackageDigest.Hash));
+
+		// The digest can't be found and the package is not in-memory package then error out.
+		if (OutMessages)
+		{
+			OutMessages->Emplace(ELogVerbosity::Error, FString::Printf(TEXT("Could not create PackageDigest for %s: %s"),
+				*PackageDependency.ToString(), *PackageDigest.GetStatusString()));
+		}
+		return ECurrentKeyResult::Error;
 	}
 
 	FCoreRedirects::AppendHashOfRedirectsAffectingPackages(KeyBuilder, RuntimePackageDependencies);
