Cooker: Downgrade error message to display for packages that cannot be cooked on the target platform, if they came from one of the ModifyCook delegates. We are using ModifyCook delegates to replace the automatic inclusion of StartupPackages, and like StartupPackages they can be added by systems that don't need them on every platform but instead just want them cooked for the platforms that support them. The first change that needs this downgrade is CL 3954659, the change to cook MovieRenderPipeline startup packages via ModifyCook..
#rnx

diff --git a/Engine/Source/Editor/UnrealEd/Private/CookOnTheFlyServer.cpp b/Engine/Source/Editor/UnrealEd/Private/CookOnTheFlyServer.cpp
--- a/Engine/Source/Editor/UnrealEd/Private/CookOnTheFlyServer.cpp
+++ b/Engine/Source/Editor/UnrealEd/Private/CookOnTheFlyServer.cpp
@@ -10572,7 +10572,7 @@ UE::Cook::FCookSavePackageContext* UCookOnTheFlyServer::CreateSaveContext(const
 					{
 						FInstigator Instigator = GetInstigator(ValidationContext.Package->GetFName(), EReachability::Runtime);
 						bool bIsError = true;
-						if (Instigator.Category == EInstigator::StartupPackage)
+						if (Instigator.Category == EInstigator::StartupPackage || Instigator.Category == EInstigator::ModifyCookDelegate)
 						{
 							// StartupPackages might be around just because of the editor;
 							// if they're not available on client, ignore them without error
