Avoid doing mount point write access validatation in bulk at boot and let them be lazy instead
- Saves 3s at editor boot and lots of IO noise / NTFS journal snapshots / directory watcher activity.

diff --git a/Engine/Source/Editor/UnrealEd/Private/UnrealEdEngine.cpp b/Engine/Source/Editor/UnrealEd/Private/UnrealEdEngine.cpp
--- a/Engine/Source/Editor/UnrealEd/Private/UnrealEdEngine.cpp
+++ b/Engine/Source/Editor/UnrealEd/Private/UnrealEdEngine.cpp
@@ -88,7 +88,6 @@
 
 DEFINE_LOG_CATEGORY_STATIC(LogUnrealEdEngine, Log, All);
 
-
 void UUnrealEdEngine::Init(IEngineLoop* InEngineLoop)
 {
 	Super::Init(InEngineLoop);
@@ -117,23 +116,6 @@ void UUnrealEdEngine::Init(IEngineLoop* InEngineLoop)
 	// Register for the package dirty state updated callback to catch packages that have been modified and need to be checked out.
 	UPackage::PackageDirtyStateChangedEvent.AddUObject(this, &UUnrealEdEngine::OnPackageDirtyStateUpdated);
 
-	// Set-up the initial set of content mount paths to check for write permision
-	TArray<FString> RootPaths;
-	FPackageName::QueryRootContentPaths(RootPaths);
-	int32 RootPathIndex = 0;
-	FDirectoryWatcherModule* DirectoryWatcherModule = FModuleManager::GetModulePtr<FDirectoryWatcherModule>(TEXT("DirectoryWatcher"));
-	IDirectoryWatcher* DirectoryWatcher = DirectoryWatcherModule ? DirectoryWatcherModule->Get() : nullptr;
-
-	for (const FString& RootPath : RootPaths)
-	{
-		VerifyMountPointWritePermission(*RootPath);
-		// We're writing files to test the write permission, and this will send file events to the directory watcher. Tick
-		// every few writes so that we don't overflow the directory watcher's buffer and cause a FCA_RescanRequired event.
-		if (DirectoryWatcher && ((++RootPathIndex) % 10) == 0)
-		{
-			DirectoryWatcher->Tick(-1.0f);
-		}
-	}
 	// Watch for new content mount paths
 	FPackageName::OnContentPathMounted().AddUObject(this, &UUnrealEdEngine::OnContentPathMounted);
 	FPackageName::OnContentPathDismounted().AddUObject(this, &UUnrealEdEngine::OnContentPathDismounted);
@@ -1547,6 +1529,8 @@ bool UUnrealEdEngine::VerifyMountPointWritePermission(FName MountPoint)
 		return true;
 	}
 
+	TRACE_CPUPROFILER_EVENT_SCOPE(VerifyMountPointWritePermission);
+
 	// Get a unique temp filename under the mount point
 	FString WriteCheckPath = FPackageName::LongPackageNameToFilename(MountPoint.ToString()) + TEXT("WritePermissions.") + FGuid::NewGuid().ToString() + TEXT(".temp");
 
@@ -1566,7 +1550,8 @@ bool UUnrealEdEngine::VerifyMountPointWritePermission(FName MountPoint)
 
 void UUnrealEdEngine::OnContentPathMounted(const FString& AssetPath, const FString& /*FileSystemPath*/)
 {
-	VerifyMountPointWritePermission(*AssetPath);
+	// Will be verified lazily if needed. Just invalidate the entry if one already existed.
+	MountPointCheckedForWritePermission.Remove(*AssetPath);
 }
 
 void UUnrealEdEngine::OnContentPathDismounted(const FString& AssetPath, const FString& /*FileSystemPath*/)
