Add ability to have named pak chunks.

New Function - GetPakChunkIdToStringMapping.
This allows the AssetRegistryGenerator to map from a chunkId to the string name you would like to give a pakchunk.
This function will be used at runtime and cooktime to map from a chunkId to the name of the expected pakchunk file.
ChunkInstall and looking for missing assets needs to be able to refer back to the pakchunk needed.

Added support so Pakchunks can request a refresh of ChunkIndex after the AssetManager has initialized given the mapping of ChunkIndex -> String override is config driven.

Add the ability for a config driven or gameplay driven method to define a pakchunk integer to a string.

diff --git a/Engine/Source/Developer/PakFileUtilities/Private/PakFileUtilities.cpp b/Engine/Source/Developer/PakFileUtilities/Private/PakFileUtilities.cpp
--- a/Engine/Source/Developer/PakFileUtilities/Private/PakFileUtilities.cpp
+++ b/Engine/Source/Developer/PakFileUtilities/Private/PakFileUtilities.cpp
@@ -3471,7 +3471,7 @@ int32 GetPakPriorityFromFilename( const FString& PakFilename )
 
 int32 GetPakChunkIndexFromFilename( const FString& PakFilePath )
 {
-	const TCHAR* PakChunkPrefix = TEXT("pakchunk");
+	FStringView PakChunkPrefix(FGenericPlatformMisc::GetPakFilenamePrefix());
 	const int32 PakChunkPrefixLength = 8;//FCString::Strlen(PakChunkPrefix);
 
 	int32 PakChunkIndex = -1;
diff --git a/Engine/Source/Editor/UnrealEd/Private/Commandlets/AssetRegistryGenerator.cpp b/Engine/Source/Editor/UnrealEd/Private/Commandlets/AssetRegistryGenerator.cpp
--- a/Engine/Source/Editor/UnrealEd/Private/Commandlets/AssetRegistryGenerator.cpp
+++ b/Engine/Source/Editor/UnrealEd/Private/Commandlets/AssetRegistryGenerator.cpp
@@ -30,6 +30,7 @@
 #include "ICollectionManager.h"
 #include "Interfaces/ITargetPlatform.h"
 #include "Misc/App.h"
+#include "Misc/AsciiSet.h"
 #include "Misc/ConfigCacheIni.h"
 #include "Misc/DataDrivenPlatformInfoRegistry.h"
 #include "Misc/FileHelper.h"
@@ -201,7 +202,8 @@ public:
 	bool IsChunkReferenced(int32 PakchunkIndex)
 	{
 		TStringBuilder<64> ChunkFileName;
-		ChunkFileName.Appendf(TEXT("pakchunk%d"), PakchunkIndex);
+		ChunkFileName.Append(FGenericPlatformMisc::GetPakFilenamePrefix());
+		ChunkFileName.Appendf(TEXT("%d"), PakchunkIndex);
 		FStringView ChunkFileNameView(ChunkFileName);
 		return ReferencedPaks.ContainsByHash(GetTypeHash(ChunkFileNameView), ChunkFileNameView);
 	}
@@ -575,6 +577,8 @@ bool FAssetRegistryGenerator::GenerateStreamingInstallManifest(int64 InOverrideC
 	TArray<FString> AllowPerChunkCompressionWildcards;
 	PlatformIniFile.GetArray(TEXT("/Script/UnrealEd.ProjectPackagingSettings"), TEXT("AllowPerChunkCompressionWildcard"), AllowPerChunkCompressionWildcards);
 
+	const TMap<int32, FString> PakChunkIdToStringOverride = GetPakChunkIdToStringOverrideMap();
+
 	// generate per-chunk pak list files
 	FDefaultPakFileRules DefaultPakFileRules;
 	bool bSucceeded = true;
@@ -628,12 +632,13 @@ bool FAssetRegistryGenerator::GenerateStreamingInstallManifest(int64 InOverrideC
 			}
 		}
 
+		const FString PakChunkName = GetPakChunkNameForId(PakChunkIdToStringOverride, PakchunkIndex);
 		bool bFinishedAllFiles = false;
 		for (int32 SubChunkIndex = 0; !bFinishedAllFiles; ++SubChunkIndex)
 		{
 			const FString PakChunkFilename = (SubChunkIndex > 0)
-				? FString::Printf(TEXT("pakchunk%d_s%d.txt"), PakchunkIndex, SubChunkIndex)
-				: FString::Printf(TEXT("pakchunk%d.txt"), PakchunkIndex);
+				? FString(FGenericPlatformMisc::GetPakFilenamePrefix()) + FString::Printf(TEXT("%s_s%d.txt"), *PakChunkName, SubChunkIndex)
+				: FString(FGenericPlatformMisc::GetPakFilenamePrefix()) + FString::Printf(TEXT("%s.txt"), *PakChunkName);
 
 			const FString PakListFilename = FString::Printf(TEXT("%s/%s"), *ChunkManifestDir, *PakChunkFilename);
 			TUniquePtr<FArchive> PakListFile(IFileManager::Get().CreateFileWriter(*PakListFilename));
@@ -1098,6 +1103,46 @@ void FAssetRegistryGenerator::UpdateKeptPackages()
 	PreviousPackagesToUpdate.Empty();
 }
 
+TMap<int32, FString> FAssetRegistryGenerator::GetPakChunkIdToStringOverrideMap()
+{
+	TMap<int32, FString> ChunkIdStringOverride;
+	UAssetManager::Get().GetPakChunkIdToStringMapping(ChunkIdStringOverride);
+	constexpr FAsciiSet ValidCharacters(VALID_SAVEDDIRSUFFIX_CHARACTERS);
+
+	TArray<int32> InvalidOverrides;
+	for (const TPair<int32, FString>& PakChunkToName : ChunkIdStringOverride)
+	{
+		const FString& PakChunkName = PakChunkToName.Value;
+		for (TCHAR Char : PakChunkName.GetCharArray())
+		{
+			if (!ValidCharacters.Contains(Char))
+			{
+				UE_LOG(LogCook, Error, TEXT("GetPakChunkIdToStringMapping contains invalid string for chunk mapping. Character '%c' within %s. Only the following character are allowed:\n%s"), Char, *PakChunkName, VALID_SAVEDDIRSUFFIX_CHARACTERS);
+				InvalidOverrides.Add(PakChunkToName.Key);
+				break;
+			}
+		}
+	}
+	for (const int32 InvalidOverride : InvalidOverrides)
+	{
+		ChunkIdStringOverride.Remove(InvalidOverride);
+		UE_LOG(LogCook, Log, TEXT("Removing PakChunk name override for chunk - %d due to invalid character."), InvalidOverride);
+	}
+	return ChunkIdStringOverride;
+}
+
+FString FAssetRegistryGenerator::GetPakChunkNameForId(const int32 PakChunkId)
+{
+	return GetPakChunkNameForId(GetPakChunkIdToStringOverrideMap(), PakChunkId);
+}
+
+FString FAssetRegistryGenerator::GetPakChunkNameForId(const TMap<int32, FString>& ChunkIdToStringOveride, const int32 PakChunkId)
+{
+	return ChunkIdToStringOveride.Contains(PakChunkId) ?
+		FString::Printf(TEXT("%c%s%c"), NAMED_PAK_CHUNK_DELIMITER_CHAR, *ChunkIdToStringOveride[PakChunkId], NAMED_PAK_CHUNK_DELIMITER_CHAR) :
+		FString::Printf(TEXT("%d"), PakChunkId);
+}
+
 void FAssetRegistryGenerator::UpdateCollectionAssetData()
 {
 	// Read out the per-platform settings use to build the list of collections to tag
@@ -2263,13 +2308,15 @@ bool FAssetRegistryGenerator::GenerateAssetChunkInformationCSV(const FString& Ou
 
 	AllChunksFile->Serialize(HeaderText, sizeof(HeaderText)-1);
 
+	const TMap<int32, FString> PakChunkIdToStringOverride = GetPakChunkIdToStringOverrideMap();
 	// Create file for each chunk if needed
 	TArray<TUniquePtr<FArchive>> ChunkFiles;
 	if (bWriteIndividualFiles)
 	{
 		for (int32 PakchunkIndex = 0; PakchunkIndex < FinalChunkManifests.Num(); ++PakchunkIndex)
 		{
-			FArchive* ChunkFile = IFileManager::Get().CreateFileWriter(*FPaths::Combine(*OutputPath, *FString::Printf(TEXT("Chunks%dInfo.csv"), PakchunkIndex)));
+			const FString PakChunkName = GetPakChunkNameForId(PakChunkIdToStringOverride, PakchunkIndex);
+			FArchive* ChunkFile = IFileManager::Get().CreateFileWriter(*FPaths::Combine(*OutputPath, *FString::Printf(TEXT("Chunks%sInfo.csv"), *PakChunkName)));
 			if (ChunkFile == nullptr)
 			{
 				return false;
@@ -2313,13 +2360,15 @@ bool FAssetRegistryGenerator::GenerateAssetChunkInformationCSV(const FString& Ou
 				{
 					if (OtherChunk != PakchunkIndex)
 					{
-						TmpString = FString::Printf(TEXT("%d "), OtherChunk);
+						const FString OtherPakChunkName = GetPakChunkNameForId(PakChunkIdToStringOverride, OtherChunk);
+						TmpString = FString::Printf(TEXT("%s "), *OtherPakChunkName);
 					}
 				}
 
 				// Build csv line
-				TmpString = FString::Printf(TEXT("%d,%s,%s,%s,%lld,%s\n"),
-					PakchunkIndex,
+				const FString PakChunkName = GetPakChunkNameForId(PakChunkIdToStringOverride, PakchunkIndex);
+				TmpString = FString::Printf(TEXT("%s,%s,%s,%s,%lld,%s\n"),
+					*PakChunkName,
 					*AssetData.PackageName.ToString(),
 					*AssetData.AssetClassPath.ToString(),
 					bHardChunk ? TEXT("Hard") : *SoftChain,
diff --git a/Engine/Source/Editor/UnrealEd/Public/Commandlets/AssetRegistryGenerator.h b/Engine/Source/Editor/UnrealEd/Public/Commandlets/AssetRegistryGenerator.h
--- a/Engine/Source/Editor/UnrealEd/Public/Commandlets/AssetRegistryGenerator.h
+++ b/Engine/Source/Editor/UnrealEd/Public/Commandlets/AssetRegistryGenerator.h
@@ -264,6 +264,12 @@ public:
 	 */
 	static void UpdateAssetManagerDatabase();
 
+	/** Sets the generated pak chunk to be named the given OverrideName rather than the chunkId. This does not update the internal asset registry number only the generated pak chunk file name. */
+	static UNREALED_API void AddChunkIdStringOverride(const int32 ChunkId, const FString& OverrideName);
+
+	/** Removes the String Override for the given ChunkId. Returns true if an override was removed. */
+	static UNREALED_API bool ClearChunkIdOverride(const int32 ChunkId);
+
 private:
 	class FGetShortestReferenceChain;
 
@@ -356,6 +362,22 @@ private:
 	/** Mapping from chunk id to pakchunk file index. If not defined, Pakchunk index will be the same as chunk id by default */
 	TMap<int32, int32> ChunkIdPakchunkIndexMapping;
 
+	/**
+	 * Gets the chunkId -> String map from the asset manager
+	 * Wrapper function to confirm and error when there's an invalid characters used.
+	 */
+	TMap<int32, FString> GetPakChunkIdToStringOverrideMap();
+
+	/**
+	 * Returns the PakChunkName for the given chunk id.
+	 * When a PakChunkId is overridden name will be wrapped by the Delimiter - so that there is an agreed format for regexing on later when loading these.
+	 * Paks will fall into the format of:
+	 * PakChunk0.pak
+	 * PakChunk-MyName-.pak
+	 */
+	FString GetPakChunkNameForId(const TMap<int32, FString>& ChunkIdToStringOveride, const int32 PakChunkId);
+	FString GetPakChunkNameForId(const int32 PakChunkId);
+
 	/**
 	 * Updates AssetData with TagsAndValues corresponding to any collections 
 	 * flagged for inclusion as asset registry tags.
diff --git a/Engine/Source/Runtime/Core/Private/GenericPlatform/GenericPlatformMisc.cpp b/Engine/Source/Runtime/Core/Private/GenericPlatform/GenericPlatformMisc.cpp
--- a/Engine/Source/Runtime/Core/Private/GenericPlatform/GenericPlatformMisc.cpp
+++ b/Engine/Source/Runtime/Core/Private/GenericPlatform/GenericPlatformMisc.cpp
@@ -1987,13 +1987,35 @@ void FGenericPlatformMisc::ParseChunkIdPakchunkIndexMapping(TArray<FString> Chun
 	}
 }
 
+FStringView FGenericPlatformMisc::GetPakFilenamePrefix()
+{
+	return TEXTVIEW("pakchunk");
+}
+
 int32 FGenericPlatformMisc::GetPakchunkIndexFromPakFile(FStringView InFilename)
 {
-	FStringView ChunkIdentifier(TEXTVIEW("pakchunk"));
+	FStringView ChunkIdentifier(GetPakFilenamePrefix());
+	const FString NamedChunkIdentifier(FString(GetPakFilenamePrefix()) + NAMED_PAK_CHUNK_DELIMITER_CHAR);
 	FStringView BaseFilename = FPathViews::GetBaseFilename(InFilename);
 	int32 ChunkNumber = INDEX_NONE;
 
-	if (BaseFilename.StartsWith(ChunkIdentifier))
+	if (FCoreDelegates::PakChunkIdOverrideDelegate.IsBound() && BaseFilename.StartsWith(*NamedChunkIdentifier))
+	{
+		const int32 StartOfChunkName = NamedChunkIdentifier.Len();
+		int32 CharCount = 0;
+		while ((CharCount + StartOfChunkName) < BaseFilename.Len() && BaseFilename[StartOfChunkName + CharCount] != NAMED_PAK_CHUNK_DELIMITER_CHAR)
+		{
+			CharCount++;
+		}
+
+		FString ChunkName = FString(BaseFilename.Mid(StartOfChunkName, CharCount));
+		int32 ChunkId = FCoreDelegates::PakChunkIdOverrideDelegate.Execute(ChunkName);
+		if (ChunkId != INDEX_NONE)
+		{
+			return ChunkId;
+		}
+	}
+	else if (BaseFilename.StartsWith(ChunkIdentifier))
 	{
 		const int32 StartOfNumber = ChunkIdentifier.Len();
 		int32 DigitCount = 0;
diff --git a/Engine/Source/Runtime/Core/Private/Misc/CoreDelegates.cpp b/Engine/Source/Runtime/Core/Private/Misc/CoreDelegates.cpp
--- a/Engine/Source/Runtime/Core/Private/Misc/CoreDelegates.cpp
+++ b/Engine/Source/Runtime/Core/Private/Misc/CoreDelegates.cpp
@@ -30,6 +30,8 @@ TMulticastDelegate<void()>& FCoreDelegates::GetPreMainInitDelegate()
 TDelegate<int32(const TArray<FString>&)> FCoreDelegates::OnMountAllPakFiles;
 TDelegate<IPakFile*(const FString&, int32)> FCoreDelegates::MountPak;
 TDelegate<bool(const FString&)> FCoreDelegates::OnUnmountPak;
+TDelegate<int32(const FString&)> FCoreDelegates::PakChunkIdOverrideDelegate;
+
 TDelegate<void()> FCoreDelegates::OnOptimizeMemoryUsageForMountedPaks;
 
 PRAGMA_DISABLE_DEPRECATION_WARNINGS
diff --git a/Engine/Source/Runtime/Core/Public/GenericPlatform/GenericPlatformMisc.h b/Engine/Source/Runtime/Core/Public/GenericPlatform/GenericPlatformMisc.h
--- a/Engine/Source/Runtime/Core/Public/GenericPlatform/GenericPlatformMisc.h
+++ b/Engine/Source/Runtime/Core/Public/GenericPlatform/GenericPlatformMisc.h
@@ -1826,6 +1826,8 @@ public:
 		return PakchunkIndex;
 	}
 
+	static CORE_API FStringView GetPakFilenamePrefix();
+
 	static CORE_API int32 GetPakchunkIndexFromPakFile(FStringView InFilename);
 
 	FORCEINLINE static bool Expand16BitIndicesTo32BitOnLoad()
diff --git a/Engine/Source/Runtime/Core/Public/Misc/CoreDelegates.h b/Engine/Source/Runtime/Core/Public/Misc/CoreDelegates.h
--- a/Engine/Source/Runtime/Core/Public/Misc/CoreDelegates.h
+++ b/Engine/Source/Runtime/Core/Public/Misc/CoreDelegates.h
@@ -125,6 +125,9 @@ public:
 	// Callback to prompt the pak system to unmount a pak file.
 	static CORE_API TDelegate<bool(const FString&)> OnUnmountPak;
 
+	// Used by asset manager to translate from a pakchunk name to the chunkId it should have.
+	static CORE_API TDelegate<int32(const FString& /*InPakchunkName*/)> PakChunkIdOverrideDelegate;
+
 	// Callback to optimize memeory for currently mounted paks
 	static CORE_API TDelegate<void()> OnOptimizeMemoryUsageForMountedPaks;
 
diff --git a/Engine/Source/Runtime/Core/Public/UObject/NameTypes.h b/Engine/Source/Runtime/Core/Public/UObject/NameTypes.h
--- a/Engine/Source/Runtime/Core/Public/UObject/NameTypes.h
+++ b/Engine/Source/Runtime/Core/Public/UObject/NameTypes.h
@@ -176,6 +176,9 @@ typedef FNameEntryId NAME_INDEX;
 /** These characters can be used in relative directory names (lowercase versions as well) */
 #define VALID_SAVEDDIRSUFFIX_CHARACTERS	TEXT("_0123456789ABCDEFGHIJKLMNOPQRSTUVWXYZabcdefghijklmnopqrstuvwxyz")
 
+/** Delimiter used to distinguish when a named pak chunk starts and ends. When using named pak chunk this will be the character wrapping the name. */
+#define NAMED_PAK_CHUNK_DELIMITER_CHAR	TEXT('-')
+
 enum class ENameCase : uint8
 {
 	CaseSensitive,
diff --git a/Engine/Source/Runtime/Core/Tests/GenericPlatform/GenericPlatformMiscTest.cpp b/Engine/Source/Runtime/Core/Tests/GenericPlatform/GenericPlatformMiscTest.cpp
--- a/Engine/Source/Runtime/Core/Tests/GenericPlatform/GenericPlatformMiscTest.cpp
+++ b/Engine/Source/Runtime/Core/Tests/GenericPlatform/GenericPlatformMiscTest.cpp
@@ -9,10 +9,10 @@ TEST_CASE_NAMED(FGenericPlatformMiscTest, "System::Core::GenericPlatform", "[Cor
 {
 	SECTION("Parse pakchunk index")
 	{
-		CHECK(1 == FGenericPlatformMisc::GetPakchunkIndexFromPakFile(FString(TEXT("pakchunk1-Windows"))));
-		CHECK(12 == FGenericPlatformMisc::GetPakchunkIndexFromPakFile(FString(TEXT("../../../Content/Paks/pakchunk12-Windows"))));
-		CHECK(42 == FGenericPlatformMisc::GetPakchunkIndexFromPakFile(FString(TEXT("pakchunk42"))));
-		CHECK(INDEX_NONE == FGenericPlatformMisc::GetPakchunkIndexFromPakFile(FString(TEXT("pakchunk"))));
+		CHECK(1 == FGenericPlatformMisc::GetPakchunkIndexFromPakFile(FString(FGenericPlatformMisc::GetPakFilenamePrefix()) + TEXT("1-Windows")));
+		CHECK(12 == FGenericPlatformMisc::GetPakchunkIndexFromPakFile(TEXT("../../../Content/Paks/") + FString(FGenericPlatformMisc::GetPakFilenamePrefix()) + TEXT("12-Windows")));
+		CHECK(42 == FGenericPlatformMisc::GetPakchunkIndexFromPakFile(FString(FGenericPlatformMisc::GetPakFilenamePrefix()) + TEXT("42")));
+		CHECK(INDEX_NONE == FGenericPlatformMisc::GetPakchunkIndexFromPakFile(FString(FGenericPlatformMisc::GetPakFilenamePrefix())));
 	}
 }
 
diff --git a/Engine/Source/Runtime/Engine/Classes/Engine/AssetManager.h b/Engine/Source/Runtime/Engine/Classes/Engine/AssetManager.h
--- a/Engine/Source/Runtime/Engine/Classes/Engine/AssetManager.h
+++ b/Engine/Source/Runtime/Engine/Classes/Engine/AssetManager.h
@@ -692,6 +692,21 @@ public:
 	}
 #endif
 
+	/**
+	 * Allows the asset manager to override the chunkID to be a string rather than integer value.
+	 * Used at cook time and runtime to determine how to map chunkids to the potential string override.
+	 * This does not update the internal asset registry number only the generated pak chunk file name.
+	 * The AssetRegistryGenerator will request this mapping during Chunk Generation time.
+	 * All strings must be be alpha numeric with no spaces.
+	 *
+	 * When overriding it's recommended to setup PlatformChunk mapping via the InstallBundle rules.
+	 * When a pak chunk does not have a parseable number the default platform chunk will be 0 unless InstallBundle rules override that.
+	 *
+	 * 
+	 * @param ChunkIdStringOverride - Out parameter for the override information of a chunkID to string.
+	 */
+	virtual void GetPakChunkIdToStringMapping(TMap<int32, FString>& ChunkIdStringOverride) {}
+
 protected:
 	friend class FAssetManagerEditorModule;
 
diff --git a/Engine/Source/Runtime/Engine/Private/AssetManager.cpp b/Engine/Source/Runtime/Engine/Private/AssetManager.cpp
--- a/Engine/Source/Runtime/Engine/Private/AssetManager.cpp
+++ b/Engine/Source/Runtime/Engine/Private/AssetManager.cpp
@@ -450,6 +450,29 @@ void UAssetManager::PostInitProperties()
 		AllAssetSearchRoots.Add(TEXT("/Game"));
 
 		FCoreUObjectDelegates::GetPreGarbageCollectDelegate().AddUObject(this, &UAssetManager::CallPreGarbageCollect);
+		
+		TMap<FString, int32> PakChunkNameToChunkId;
+		{
+			TMap<int32, FString> ChunkIdStringOverride;
+			GetPakChunkIdToStringMapping(ChunkIdStringOverride);
+			PakChunkNameToChunkId.Reserve(ChunkIdStringOverride.Num());
+			for (const TPair<int32, FString>& ChunkToString : ChunkIdStringOverride)
+			{
+				PakChunkNameToChunkId.Add(ChunkToString.Value, ChunkToString.Key);
+			}
+		}
+		if (PakChunkNameToChunkId.Num())
+		{
+			UE_LOG(LogAssetManager, Display, TEXT("Detected PakchunkId override config. Binding PakChunkId override delegate."));
+			FCoreDelegates::PakChunkIdOverrideDelegate.BindLambda([LocalPakChunkNameToChunkId = MoveTemp(PakChunkNameToChunkId)](const FString& InPakchunkName)
+				{
+					return LocalPakChunkNameToChunkId.Contains(InPakchunkName) ? LocalPakChunkNameToChunkId[InPakchunkName] : INDEX_NONE;
+				});
+			if (FPakPlatformFile* Pak = (FPakPlatformFile*)FPlatformFileManager::Get().FindPlatformFile(TEXT("PakFile")))
+			{
+				Pak->RefreshPakChunkIndicies();
+			}
+		}
 	}
 }
 
diff --git a/Engine/Source/Runtime/PakFile/Private/IPlatformFilePak.cpp b/Engine/Source/Runtime/PakFile/Private/IPlatformFilePak.cpp
--- a/Engine/Source/Runtime/PakFile/Private/IPlatformFilePak.cpp
+++ b/Engine/Source/Runtime/PakFile/Private/IPlatformFilePak.cpp
@@ -5180,6 +5180,20 @@ bool FPakPlatformFile::IterateDirectoryRecursively(const TCHAR* Directory, IPlat
 	return IterateDirectoryInternal(Directory, Visitor, true /* bRecursive */);
 }
 
+void FPakPlatformFile::RefreshPakChunkIndicies()
+{
+	UE::TScopeLock ScopedLock(PakListCritical);
+	for (FPakListEntry& PakEntry : PakFiles)
+	{
+		PakEntry.PakFile->PakchunkIndex = GetPakchunkIndexFromPakFile(PakEntry.PakFile->PakFilename);
+	}
+
+	for (FPakListDeferredEntry& Entry : PendingEncryptedPakFiles)
+	{
+		Entry.PakchunkIndex = GetPakchunkIndexFromPakFile(Entry.Filename);
+	}
+}
+
 FFilenameSecurityDelegate& FPakPlatformFile::GetFilenameSecurityDelegate()
 {
 	static FFilenameSecurityDelegate Delegate;
diff --git a/Engine/Source/Runtime/PakFile/Public/IPlatformFilePak.h b/Engine/Source/Runtime/PakFile/Public/IPlatformFilePak.h
--- a/Engine/Source/Runtime/PakFile/Public/IPlatformFilePak.h
+++ b/Engine/Source/Runtime/PakFile/Public/IPlatformFilePak.h
@@ -2557,6 +2557,12 @@ public:
 	}
 	//~ End IPlatformFile Interface
 
+	/**
+	 * Requests that all the known paks re-evaluate and set their PakchunkIndex value.
+	 * Used to allow the AssetManager inform the pak file system about potential name pak chunk overrides.
+	 */
+	PAKFILE_API void RefreshPakChunkIndicies();
+
 	// Access static delegate for loose file security
 	PAKFILE_API static FFilenameSecurityDelegate& GetFilenameSecurityDelegate();
 
