[Editor] DeleteActor should not warn user about SoftReferences that are outered to the actor being deleted (as those objects are also being deleted)

#rnx

diff --git a/Engine/Source/Editor/UnrealEd/Private/EditorActor.cpp b/Engine/Source/Editor/UnrealEd/Private/EditorActor.cpp
--- a/Engine/Source/Editor/UnrealEd/Private/EditorActor.cpp
+++ b/Engine/Source/Editor/UnrealEd/Private/EditorActor.cpp
@@ -967,6 +967,7 @@ bool UUnrealEdEngine::DeleteActors(const TArray<AActor*>& InActorsToDelete, UWor
 		bool bReferencedByGroupActor = false;
 		bool bReferencedBySoftReference = false;
 		TArray<UObject*>* SoftReferencingObjects = nullptr;
+		const UPackage* ActorPackage = Actor->GetPackage();
 
 		if (bWarnAboutSoftReferences)
 		{
@@ -979,6 +980,10 @@ bool UUnrealEdEngine::DeleteActors(const TArray<AActor*>& InActorsToDelete, UWor
 				{
 					const UObject* SoftReferencingObject = (*SoftReferencingObjects)[i];
 					if (const AActor* ReferencingActor = Cast<AActor>(SoftReferencingObject); ReferencingActor && ActorsToDelete.Contains(ReferencingActor))
+					{
+						SoftReferencingObjects->RemoveAt(i);
+					} // Object is inside the actor being deleted, so reference should be ignored (check that they share the same package because that could mean we are in a special case where warnings are needed)
+					else if (SoftReferencingObject->IsInOuter(Actor) && SoftReferencingObject->GetPackage() == ActorPackage)
 					{
 						SoftReferencingObjects->RemoveAt(i);
 					}
