PR #12793: Remove roll applied to player view when returning to editor from PIE

diff --git a/Engine/Source/Editor/UnrealEd/Private/PlayLevel.cpp b/Engine/Source/Editor/UnrealEd/Private/PlayLevel.cpp
--- a/Engine/Source/Editor/UnrealEd/Private/PlayLevel.cpp
+++ b/Engine/Source/Editor/UnrealEd/Private/PlayLevel.cpp
@@ -2117,6 +2117,7 @@ void UEditorEngine::ToggleBetweenPIEandSIE( bool bNewSession )
 				FVector ViewLocation;
 				FRotator ViewRotation;
 				SlatePlayInEditorSession.EditorPlayer.Get()->PlayerController->GetPlayerViewPoint( ViewLocation, ViewRotation );
+				ViewRotation.Roll = 0.0; // Clear roll since editor camera controls has no way to adjust roll
 				EditorViewportClient.SetViewLocation( ViewLocation );
 
 				if( EditorViewportClient.IsPerspective() )
