Remove in-memory package from the list of build package dependencies. Error out if the package doesn't exist.
Do not ignore packages in TryCalculateCurrentKey anymore. The list BuildPackageDependencies is already filtered out of any in-memory package.

#rnx
#incrementalcook

diff --git a/Engine/Source/Editor/UnrealEd/Private/Cooker/CookPackageArtifacts.cpp b/Engine/Source/Editor/UnrealEd/Private/Cooker/CookPackageArtifacts.cpp
--- a/Engine/Source/Editor/UnrealEd/Private/Cooker/CookPackageArtifacts.cpp
+++ b/Engine/Source/Editor/UnrealEd/Private/Cooker/CookPackageArtifacts.cpp
@@ -153,16 +153,6 @@ FPackageArtifacts::ECurrentKeyResult FPackageArtifacts::TryCalculateCurrentKey(
 			continue;
 		}
 
-		// The digest was not found, it could be because the package was created during the cook and does not exist on disk (in-memory package).
-		// Lets ignore those packages. Whoever created this package is responsible of saving dependencies if any.
-		TStringBuilder<256> PackageDependencyStr(EInPlace::InPlace, PackageDependency);
-		const UPackage* PackageObject = FindPackage(nullptr, *PackageDependencyStr);
-		if (PackageObject && PackageObject->HasAnyPackageFlags(PKG_NewlyCreated))
-		{
-			continue;
-		}
-
-		// The digest can't be found and the package is not in-memory package then error out.
 		if (OutMessages)
 		{
 			OutMessages->Emplace(ELogVerbosity::Error, FString::Printf(TEXT("Could not create PackageDigest for %s: %s"),
@@ -557,6 +547,28 @@ FPackageArtifacts FPackageArtifacts::Collect(UPackage* Package, const ITargetPla
 			// Skip script and temp packages
 			continue;
 		}
+
+		TStringBuilder<256> PackagePath(EInPlace::InPlace, Pair.Key);
+		UPackage* PackageObject = FindPackage(nullptr, *PackagePath);
+		if (!PackageObject)
+		{
+			if (OutMessages)
+			{
+				TStringBuilder<256> ErrorText;
+				ErrorText << TEXT("Package ") << PackagePath << TEXT(" does not exist.");
+				OutMessages->Emplace(ELogVerbosity::Error, FString(*ErrorText));
+			}
+
+			return FPackageArtifacts();
+		}
+
+		if(PackageObject->HasAnyPackageFlags(PKG_NewlyCreated))
+		{
+			// If the package is a newly created package (in-memory package) then ignore it.
+			// In-memory packages are ignored because we can't compute their digest. Only packages on disk have a digest.
+			continue;
+		}
+
 		if (Pair.Value == EBuildDependency::Transitive)
 		{
 			Result.TransitiveBuildDependencies.Add(UE::Cook::FCookDependency::Package(Pair.Key));
