[Viewport Toolbar] Be able to enable snapping by selecting a new value from the value dropdown.

Entries for snapping submenus are always active, and will enable their respective snapping when edited:

- Surface (including the "Rotate to surface normal" flag)
- Location
- Rotation
- Scale

diff --git a/Engine/Source/Editor/UnrealEd/Private/EditorConstraints.cpp b/Engine/Source/Editor/UnrealEd/Private/EditorConstraints.cpp
--- a/Engine/Source/Editor/UnrealEd/Private/EditorConstraints.cpp
+++ b/Engine/Source/Editor/UnrealEd/Private/EditorConstraints.cpp
@@ -10,6 +10,7 @@
 #include "Editor/EditorEngine.h"
 #include "Editor.h"
 #include "EditorSupportDelegates.h"
+#include "LevelEditorActions.h"
 
 float UEditorEngine::GetGridSize() const
 {
@@ -36,6 +37,13 @@ void UEditorEngine::SetGridSize( int32 InIndex )
 
 	ULevelEditorViewportSettings* ViewportSettings = GetMutableDefault<ULevelEditorViewportSettings>();
 	ViewportSettings->CurrentPosGridSize = FMath::Clamp<int32>( InIndex, 0, GetCurrentPositionGridArray().Num() - 1 );
+
+	// If user is editing snapping values, we assume they also want snapping turned on
+	if (!FLevelEditorActionCallbacks::LocationGridSnap_IsChecked())
+	{
+		FLevelEditorActionCallbacks::LocationGridSnap_Clicked();
+	}
+
 	ViewportSettings->PostEditChange();
 
 	FEditorDelegates::OnGridSnappingChanged.Broadcast(GetDefault<ULevelEditorViewportSettings>()->GridEnabled, GetGridSize());
@@ -95,6 +103,13 @@ void UEditorEngine::SetRotGridSize( int32 InIndex, ERotationGridMode InGridMode
 	ULevelEditorViewportSettings* ViewportSettings = GetMutableDefault<ULevelEditorViewportSettings>();
 	ViewportSettings->CurrentRotGridMode = InGridMode;
 	ViewportSettings->CurrentRotGridSize = FMath::Clamp<int32>( InIndex, 0, RotGridSizes.Num() - 1 );
+
+	// If user is editing snapping values, we assume they also want snapping turned on
+	if (!FLevelEditorActionCallbacks::RotationGridSnap_IsChecked())
+	{
+		FLevelEditorActionCallbacks::RotationGridSnap_Clicked();
+	}
+
 	ViewportSettings->PostEditChange();
 	
 	RedrawLevelEditingViewports();
@@ -139,6 +154,13 @@ void UEditorEngine::SetScaleGridSize( int32 InIndex )
 
 	ULevelEditorViewportSettings* ViewportSettings = GetMutableDefault<ULevelEditorViewportSettings>();
 	ViewportSettings->CurrentScalingGridSize = FMath::Clamp<int32>( InIndex, 0, ViewportSettings->ScalingGridSizes.Num() - 1 );
+
+	// If user is editing snapping values, we assume they also want snapping turned on
+	if (!FLevelEditorActionCallbacks::ScaleGridSnap_IsChecked())
+	{
+		FLevelEditorActionCallbacks::ScaleGridSnap_Clicked();
+	}
+
 	ViewportSettings->PostEditChange();
 
 	RedrawLevelEditingViewports();
diff --git a/Engine/Source/Editor/UnrealEd/Private/SEditorViewport.cpp b/Engine/Source/Editor/UnrealEd/Private/SEditorViewport.cpp
--- a/Engine/Source/Editor/UnrealEd/Private/SEditorViewport.cpp
+++ b/Engine/Source/Editor/UnrealEd/Private/SEditorViewport.cpp
@@ -467,7 +467,7 @@ void SEditorViewport::BindCommands()
 	CommandListRef.MapAction(
 		Commands.RotateToSurfaceNormal,
 		FExecuteAction::CreateStatic( &SEditorViewport::OnToggleRotateToSurfaceNormal ),
-		FCanExecuteAction::CreateStatic( &SEditorViewport::OnIsSurfaceSnapEnabled ),
+		FCanExecuteAction(),
 		FIsActionChecked::CreateStatic( &SEditorViewport::IsRotateToSurfaceNormalEnabled ) 
 		);
 
@@ -963,6 +963,12 @@ void SEditorViewport::OnToggleRotateToSurfaceNormal()
 {
 	auto& Settings = GetMutableDefault<ULevelEditorViewportSettings>()->SnapToSurface;
 	Settings.bSnapRotation = !Settings.bSnapRotation;
+
+	// If user is editing snapping settings, we assume they also want snapping turned on
+	if (!Settings.bEnabled)
+	{
+		Settings.bEnabled = true;
+	}
 }
 
 bool SEditorViewport::IsRotateToSurfaceNormalEnabled()
diff --git a/Engine/Source/Editor/UnrealEd/Private/ViewportToolbar/UnrealEdViewportToolbar.cpp b/Engine/Source/Editor/UnrealEd/Private/ViewportToolbar/UnrealEdViewportToolbar.cpp
--- a/Engine/Source/Editor/UnrealEd/Private/ViewportToolbar/UnrealEdViewportToolbar.cpp
+++ b/Engine/Source/Editor/UnrealEd/Private/ViewportToolbar/UnrealEdViewportToolbar.cpp
@@ -82,6 +82,24 @@ FToolUIAction DisabledAction()
 	return Action;
 }
 
+void ToggleSurfaceSnapping()
+{
+	if (ULevelEditorViewportSettings* ViewportSettings = GetMutableDefault<ULevelEditorViewportSettings>())
+	{
+		ViewportSettings->SnapToSurface.bEnabled = !ViewportSettings->SnapToSurface.bEnabled;
+	}
+}
+
+bool IsSurfaceSnappingEnabled()
+{
+	if (const ULevelEditorViewportSettings* const ViewportSettings = GetDefault<ULevelEditorViewportSettings>())
+	{
+		return ViewportSettings->SnapToSurface.bEnabled;
+	}
+
+	return false;
+}
+
 // TODO: Maybe export CreateSurfaceSnapOffsetEntry function, so that it can be used elsewhere, e.g. STransformViewportToolbar.cpp
 FToolMenuEntry CreateSurfaceSnapOffsetEntry()
 {
@@ -92,15 +110,7 @@ FToolMenuEntry CreateSurfaceSnapOffsetEntry()
 
 	FToolMenuEntry SurfaceOffset = FToolMenuEntry::InitMenuEntry(
 		"SurfaceOffset",
-		FUIAction(
-			FExecuteAction(),
-			FCanExecuteAction::CreateLambda(
-				[]()
-				{
-					return GetDefault<ULevelEditorViewportSettings>()->SnapToSurface.bEnabled;
-				}
-			)
-		),
+		FUIAction(),
 		// clang-format off
 		SNew(SHorizontalBox)
 		+ SHorizontalBox::Slot()
@@ -135,6 +145,12 @@ FToolMenuEntry CreateSurfaceSnapOffsetEntry()
 					ULevelEditorViewportSettings* Settings =
 						GetMutableDefault<ULevelEditorViewportSettings>();
 					Settings->SnapToSurface.SnapOffsetExtent = InNewValue;
+
+					// If user is editing surface snapping values, we assume they also want surface snapping turned on
+					if (!UE::UnrealEd::Private::IsSurfaceSnappingEnabled())
+					{
+						UE::UnrealEd::Private::ToggleSurfaceSnapping();
+					}
 				})
 				.Value_Lambda([]()
 				{
@@ -177,16 +193,15 @@ FToolMenuEntry CreateSurfaceSnapCheckboxMenu()
 		FToolMenuExecuteAction::CreateLambda(
 			[](const FToolMenuContext& InContext)
 			{
-				ULevelEditorViewportSettings* Settings = GetMutableDefault<ULevelEditorViewportSettings>();
-				Settings->SnapToSurface.bEnabled = !Settings->SnapToSurface.bEnabled;
+				ToggleSurfaceSnapping();
 			}
 		),
 		FToolMenuCanExecuteAction(),
 		FToolMenuGetActionCheckState::CreateLambda(
 			[](const FToolMenuContext& InContext)
 			{
-				return GetDefault<ULevelEditorViewportSettings>()->SnapToSurface.bEnabled ? ECheckBoxState::Checked
-																						  : ECheckBoxState::Unchecked;
+				return UE::UnrealEd::Private::IsSurfaceSnappingEnabled() ? ECheckBoxState::Checked
+																		 : ECheckBoxState::Unchecked;
 			}
 		),
 		MakeMenuDelegate
@@ -316,13 +331,7 @@ FToolMenuEntry CreateLocationSnapCheckboxMenu()
 						UnrealEd::CreateLocationGridSnapMenu(
 							ExecuteDelegate,
 							IsCheckedDelegate,
-							GridSizes,
-							TAttribute<bool>::CreateLambda(
-								[]()
-								{
-									return FLevelEditorActionCallbacks::LocationGridSnap_IsChecked();
-								}
-							)
+							GridSizes
 						),
 						FText()
 					)
@@ -388,14 +397,10 @@ FToolMenuEntry CreateRotationSnapCheckboxMenu()
 					}
 				);
 
-			const TAttribute<bool> IsEnabledDelegate =
-				TAttribute<bool>::Create(&FLevelEditorActionCallbacks::RotationGridSnap_IsChecked);
-
 			InToolMenu->AddMenuEntry(
 				RotationSnapName,
 				FToolMenuEntry::InitWidget(
-					RotationSnapName,
-					UnrealEd::CreateRotationGridSnapMenu(ExecuteDelegate, IsCheckedDelegate, IsEnabledDelegate),
+					RotationSnapName, UnrealEd::CreateRotationGridSnapMenu(ExecuteDelegate, IsCheckedDelegate),
 					FText()
 				)
 			);
@@ -474,12 +479,7 @@ FToolMenuEntry CreateScaleSnapCheckboxMenu()
 						ExecuteDelegate,
 						IsCheckedDelegate,
 						GridSizes,
-						TAttribute<bool>::CreateLambda(
-							[]()
-							{
-								return FLevelEditorActionCallbacks::ScaleGridSnap_IsChecked();
-							}
-						),
+						TAttribute<bool>(true),
 						{},
 						true,
 						FUIAction(
@@ -637,6 +637,7 @@ void AddShowGridToggle(UToolMenu* InMenu)
 		}
 	}
 }
+
 } // namespace UE::UnrealEd::Private
 
 static FAutoConsoleVariableRef CVarToolMenusViewportToolbars(
