Remove use of FindFirstObject, instead use the UFunction present in the callstack directly. Various const fixes.

[RN] Fixed ensure when navigating to a function via the Blueprint Callstack Viewer when there were multiple classes with the same name

diff --git a/Engine/Source/Editor/Kismet/Private/CallStackViewer.cpp b/Engine/Source/Editor/Kismet/Private/CallStackViewer.cpp
--- a/Engine/Source/Editor/Kismet/Private/CallStackViewer.cpp
+++ b/Engine/Source/Editor/Kismet/Private/CallStackViewer.cpp
@@ -87,6 +87,7 @@ struct FCallStackRow
 {
 	FCallStackRow(
 		UObject* InContextObject,
+		const UFunction* InFunction,
 		const FName& InScopeName, 
 		const FName& InFunctionName, 
 		int32 InScriptOffset, 
@@ -95,6 +96,7 @@ struct FCallStackRow
 		const FText& InFunctionDisplayName
 	)
 		: ContextObject(InContextObject)
+		, Function(InFunction)
 		, ScopeName(InScopeName)
 		, FunctionName(InFunctionName)
 		, ScriptOffset(InScriptOffset)
@@ -104,7 +106,8 @@ struct FCallStackRow
 	{
 	}
 
-	UObject* ContextObject;
+	TWeakObjectPtr<UObject> ContextObject;
+	TWeakObjectPtr<const UFunction> Function;
 
 	FName ScopeName;
 	FName FunctionName;
@@ -497,27 +500,17 @@ void SCallStackViewer::JumpToEntry(TSharedRef< FCallStackRow > Entry)
 {
 	LastFrameNavigatedTo = Entry;
 
-	// Try to find a UClass* source:
-	UBlueprintGeneratedClass* BPGC = FindFirstObject<UBlueprintGeneratedClass>(*Entry->ScopeName.ToString(), EFindFirstObjectOptions::EnsureIfAmbiguous);
-
-	if(BPGC)
+	const UFunction* Function = Entry->Function.Get();
+	if (Function)
 	{
-		UFunction* Function = BPGC->FindFunctionByName(Entry->FunctionName);
-		if(Function)
+		UEdGraphNode* Node = FKismetDebugUtilities::FindSourceNodeForCodeLocation(Entry->ContextObject.Get(), Function, Entry->ScriptOffset, true);
+		if (Node)
 		{
-			UEdGraphNode* Node = FKismetDebugUtilities::FindSourceNodeForCodeLocation(Entry->ContextObject, Function, Entry->ScriptOffset, true);
-			if(Node)
-			{
-				FKismetEditorUtilities::BringKismetToFocusAttentionOnObject(Node);
-			}
-			else
-			{
-				FKismetEditorUtilities::BringKismetToFocusAttentionOnObject(Function);
-			}
+			FKismetEditorUtilities::BringKismetToFocusAttentionOnObject(Node);
 		}
 		else
 		{
-			FKismetEditorUtilities::BringKismetToFocusAttentionOnObject(BPGC);
+			FKismetEditorUtilities::BringKismetToFocusAttentionOnObject(Function);
 		}
 	}
 }
@@ -592,6 +585,7 @@ void CallStackViewer::UpdateDisplayedCallstack(TArrayView<const FFrame* const> S
 			Private_CallStackSource.Add(
 				MakeShared<FCallStackRow>(
 					StackNode->Object,
+					StackNode->Node,
 					StackNode->Node->GetOuter()->GetFName(),
 					StackNode->Node->GetFName(),
 					ScriptOffset,
@@ -607,6 +601,7 @@ void CallStackViewer::UpdateDisplayedCallstack(TArrayView<const FFrame* const> S
 			{
 				Private_CallStackSource.Add(
 					MakeShared<FCallStackRow>(
+						nullptr,
 						nullptr,
 						FName(),
 						FName(),
diff --git a/Engine/Source/Editor/UnrealEd/Private/Kismet2/KismetDebugUtilities.cpp b/Engine/Source/Editor/UnrealEd/Private/Kismet2/KismetDebugUtilities.cpp
--- a/Engine/Source/Editor/UnrealEd/Private/Kismet2/KismetDebugUtilities.cpp
+++ b/Engine/Source/Editor/UnrealEd/Private/Kismet2/KismetDebugUtilities.cpp
@@ -248,7 +248,7 @@ void FKismetDebugUtilities::OnScriptException(const UObject* ActiveObject, const
 
 	checkSlow(ActiveObject != nullptr);
 
-	UClass* ClassContainingCode = FindClassForNode(ActiveObject, StackFrame.Node);
+	const UClass* ClassContainingCode = FindClassForNode(ActiveObject, StackFrame.Node);
 	UBlueprint* BlueprintObj = (ClassContainingCode ? Cast<UBlueprint>(ClassContainingCode->ClassGeneratedBy) : nullptr);
 	if (BlueprintObj)
 	{
@@ -315,7 +315,7 @@ void FKismetDebugUtilities::OnScriptException(const UObject* ActiveObject, const
 				Message->AddToken(FTextToken::Create(FText::Format(LOCTEXT("RuntimeErrorMessageFmt", "Blueprint Runtime Error: \"{0}\"."), Info.GetDescription())));
 
 #if WITH_EDITORONLY_DATA // to protect access to GeneratedClass->DebugData
-				UBlueprintGeneratedClass* GeneratedClass = Cast<UBlueprintGeneratedClass>(ClassContainingCode);
+				const UBlueprintGeneratedClass* GeneratedClass = Cast<UBlueprintGeneratedClass>(ClassContainingCode);
 				if ((GeneratedClass != nullptr) && GeneratedClass->DebugData.IsValid())
 				{
 					UEdGraphNode* BlueprintNode = GeneratedClass->DebugData.FindSourceNodeFromCodeLocation(StackFrame.Node, BreakpointOffset, true);
@@ -560,7 +560,7 @@ bool FKismetDebugUtilities::TracepointBreakAllowedOnOwningWorld(const UObject* O
 	return bAllowTracepointBreak;
 }
 
-UClass* FKismetDebugUtilities::FindClassForNode(const UObject* Object, UFunction* Function)
+const UClass* FKismetDebugUtilities::FindClassForNode(const UObject* Object, const UFunction* Function)
 {
 	if (NULL != Function)
 	{
@@ -580,12 +580,12 @@ const TSimpleRingBuffer<FKismetTraceSample>& FKismetDebugUtilities::GetTraceStac
 	return FKismetDebugUtilitiesData::Get().TraceStackSamples; 
 }
 
-UEdGraphNode* FKismetDebugUtilities::FindSourceNodeForCodeLocation(const UObject* Object, UFunction* Function, int32 DebugOpcodeOffset, bool bAllowImpreciseHit)
+UEdGraphNode* FKismetDebugUtilities::FindSourceNodeForCodeLocation(const UObject* Object, const UFunction* Function, int32 DebugOpcodeOffset, bool bAllowImpreciseHit)
 {
 	if (Object != NULL)
 	{
 		// Find the blueprint that corresponds to the object
-		if (UBlueprintGeneratedClass* Class = Cast<UBlueprintGeneratedClass>(FindClassForNode(Object, Function)))
+		if (const UBlueprintGeneratedClass* Class = Cast<UBlueprintGeneratedClass>(FindClassForNode(Object, Function)))
 		{
 			return Class->GetDebugData().FindSourceNodeFromCodeLocation(Function, DebugOpcodeOffset, bAllowImpreciseHit);
 		}
diff --git a/Engine/Source/Editor/UnrealEd/Public/Kismet2/KismetDebugUtilities.h b/Engine/Source/Editor/UnrealEd/Public/Kismet2/KismetDebugUtilities.h
--- a/Engine/Source/Editor/UnrealEd/Public/Kismet2/KismetDebugUtilities.h
+++ b/Engine/Source/Editor/UnrealEd/Public/Kismet2/KismetDebugUtilities.h
@@ -234,10 +234,10 @@ public:
 	static UNREALED_API const TSimpleRingBuffer<FKismetTraceSample>& GetTraceStack();
 
 	// Find the node that resulted in code at the specified location in the Object, or NULL if there was a problem (e.g., no debugging information was generated)
-	static UNREALED_API class UEdGraphNode* FindSourceNodeForCodeLocation(const UObject* Object, UFunction* Function, int32 DebugOpcodeOffset, bool bAllowImpreciseHit = false);
+	static UNREALED_API class UEdGraphNode* FindSourceNodeForCodeLocation(const UObject* Object, const UFunction* Function, int32 DebugOpcodeOffset, bool bAllowImpreciseHit = false);
 
 	// Return proper class for breakpoint
-	static UNREALED_API UClass* FindClassForNode(const UObject* Object, UFunction* Function);
+	static UNREALED_API const UClass* FindClassForNode(const UObject* Object, const UFunction* Function);
 
 	// Notify the debugger of the start of the game frame
 	static UNREALED_API void NotifyDebuggerOfStartOfGameFrame(UWorld* CurrentWorld);
diff --git a/Engine/Source/Runtime/Engine/Classes/Engine/BlueprintGeneratedClass.h b/Engine/Source/Runtime/Engine/Classes/Engine/BlueprintGeneratedClass.h
--- a/Engine/Source/Runtime/Engine/Classes/Engine/BlueprintGeneratedClass.h
+++ b/Engine/Source/Runtime/Engine/Classes/Engine/BlueprintGeneratedClass.h
@@ -241,7 +241,7 @@ public:
 	ENGINE_API bool IsValid() const;
 
 	// Finds the UEdGraphNode associated with the code location Function+CodeOffset, or nullptr if there isn't one
-	ENGINE_API UEdGraphNode* FindSourceNodeFromCodeLocation(UFunction* Function, int32 CodeOffset, bool bAllowImpreciseHit) const;
+	ENGINE_API UEdGraphNode* FindSourceNodeFromCodeLocation(const UFunction* Function, int32 CodeOffset, bool bAllowImpreciseHit) const;
 
 	// Finds the source pin associated with the code location Function+CodeOffset, or nullptr if there isn't one
 	ENGINE_API UEdGraphPin* FindSourcePinFromCodeLocation(UFunction* Function, int32 CodeOffset) const;
@@ -692,6 +692,11 @@ public:
 	{
 		return DebugData;
 	}
+
+	const FBlueprintDebugData& GetDebugData() const
+	{
+		return DebugData;
+	}
 #endif
 
 	/** Bind functions on supplied actor to delegates */
diff --git a/Engine/Source/Runtime/Engine/Private/BlueprintGeneratedClass.cpp b/Engine/Source/Runtime/Engine/Private/BlueprintGeneratedClass.cpp
--- a/Engine/Source/Runtime/Engine/Private/BlueprintGeneratedClass.cpp
+++ b/Engine/Source/Runtime/Engine/Private/BlueprintGeneratedClass.cpp
@@ -3000,7 +3000,7 @@ bool FBlueprintDebugData::IsValid() const
 	return DebugNodeLineNumbers.Num() > 0;
 }
 
-UEdGraphNode* FBlueprintDebugData::FindSourceNodeFromCodeLocation(UFunction* Function, int32 CodeOffset, bool bAllowImpreciseHit) const
+UEdGraphNode* FBlueprintDebugData::FindSourceNodeFromCodeLocation(const UFunction* Function, int32 CodeOffset, bool bAllowImpreciseHit) const
 {
 	if (const FDebuggingInfoForSingleFunction* pFuncInfo = PerFunctionLineNumbers.Find(Function))
 	{
