IncrementalCook: Refactor CookPackageArtifacts to represent all dependencies as FCookDependency. This will remove the special cases required to implement the upcoming BuildResults features, where each BuildResult has its own set of dependencies.
#rnx
#incrementalcook

diff --git a/Engine/Source/Editor/UnrealEd/Private/Cooker/CookFunctionLibrary.cpp b/Engine/Source/Editor/UnrealEd/Private/Cooker/CookFunctionLibrary.cpp
--- a/Engine/Source/Editor/UnrealEd/Private/Cooker/CookFunctionLibrary.cpp
+++ b/Engine/Source/Editor/UnrealEd/Private/Cooker/CookFunctionLibrary.cpp
@@ -206,15 +206,16 @@ void UCookFunctionLibrary::CookAsset(UObject* Object, const FString& ForPlatform
 		CommitInfo.WriteOptions = IPackageWriter::EWriteOptions::Write;
 		LooseWriter->CommitPackage(MoveTemp(CommitInfo)); // please clap
 
-		// FGeneratedPackageResultStruct is used to support the creation of generated 
+		// FGenerationHelper is used to support the creation of generated 
 		// streamingcells from WorldPartition. Supporting streaming cells would require 
 		// several other changes to this function, and is not the priority of this 
 		// routine at this time.
-		FGeneratedPackageResultStruct* Todo = nullptr;
+		UE::Cook::FGenerationHelper* Todo = nullptr;
+		bool bGenerated = false;
 		TArray<FName> RuntimeDepsTodo; // I don't believe we need valid data in here at this time
 
 		TArray<TPair<ELogVerbosity::Type, FString>> Messages;
-		FPackageArtifacts Deps = FPackageArtifacts::Collect(Package, TargetPlatform, &Result, Todo, MoveTemp(RuntimeDepsTodo), &Messages);
+		FPackageArtifacts Deps = FPackageArtifacts::Collect(Package, TargetPlatform, &Result, Todo, bGenerated, MoveTemp(RuntimeDepsTodo), &Messages);
 		if (!Deps.IsValid())
 		{
 			TStringBuilder<256> ErrorMessage;
diff --git a/Engine/Source/Editor/UnrealEd/Private/Cooker/CookGenerationHelper.cpp b/Engine/Source/Editor/UnrealEd/Private/Cooker/CookGenerationHelper.cpp
--- a/Engine/Source/Editor/UnrealEd/Private/Cooker/CookGenerationHelper.cpp
+++ b/Engine/Source/Editor/UnrealEd/Private/Cooker/CookGenerationHelper.cpp
@@ -11,6 +11,7 @@
 #include "Cooker/CookWorkerServer.h"
 #include "Cooker/IWorkerRequests.h"
 #include "Cooker/PackageTracker.h"
+#include "EditorDomain/EditorDomainUtils.h"
 #include "Interfaces/ITargetPlatform.h"
 #include "Misc/CommandLine.h"
 #include "Misc/ConfigCacheIni.h"
@@ -649,6 +650,22 @@ FCookGenerationInfo* FGenerationHelper::FindInfoNoInitialize(const FPackageData&
 	return nullptr;
 }
 
+FCookGenerationInfo* FGenerationHelper::FindInfoNoInitialize(FName PackageName)
+{
+	if (PackageName == GetOwner().GetPackageName())
+	{
+		return &OwnerInfo;
+	}
+	for (FCookGenerationInfo& Info : PackagesToGenerate)
+	{
+		if (Info.PackageData->GetPackageName() == PackageName)
+		{
+			return &Info;
+		}
+	}
+	return nullptr;
+}
+
 const FCookGenerationInfo* FGenerationHelper::FindInfo(const FPackageData& PackageData) const
 {
 	return const_cast<FGenerationHelper*>(this)->FindInfo(PackageData);
@@ -1280,8 +1297,8 @@ void FGenerationHelper::FinishGeneratorPlatformSave(FPackageData& PackageData, b
 	}
 }
 
-void FGenerationHelper::FinishGeneratedPlatformSave(FPackageData& PackageData,
-	FGeneratedPackageResultStruct& OutGeneratedResult)
+void FGenerationHelper::FinishGeneratedPlatformSave(FPackageData& PackageData, FAssetPackageData& OutAssetPackageData,
+	TArray<FAssetDependency>& OutDependencies)
 {
 	ConditionalInitialize();
 
@@ -1296,7 +1313,7 @@ void FGenerationHelper::FinishGeneratedPlatformSave(FPackageData& PackageData,
 
 	// There should be no package dependencies present for the package from the global assetregistry
 	// because it is newly created. Add on the dependencies declared for it from the CookPackageSplitter.
-	OutGeneratedResult.PackageDependencies = Info->PackageDependencies;
+	OutDependencies = Info->PackageDependencies;
 
 	// Update the AssetPackageData for each requested platform with Guid and ImportedClasses
 	TSet<UClass*> PackageClasses;
@@ -1321,25 +1338,45 @@ void FGenerationHelper::FinishGeneratedPlatformSave(FPackageData& PackageData,
 	}
 	ImportedClasses.Sort(FNameLexicalLess());
 
-	OutGeneratedResult.AssetPackageData.FileVersionUE = GPackageFileUEVersion;
-	OutGeneratedResult.AssetPackageData.FileVersionLicenseeUE = GPackageFileLicenseeUEVersion;
-	OutGeneratedResult.AssetPackageData.SetIsLicenseeVersion(FEngineVersion::Current().IsLicenseeVersion());
+	FAssetPackageData& AssetPackageData = Info->AssetPackageData.Emplace();
+	AssetPackageData.FileVersionUE = GPackageFileUEVersion;
+	AssetPackageData.FileVersionLicenseeUE = GPackageFileLicenseeUEVersion;
+	AssetPackageData.SetIsLicenseeVersion(FEngineVersion::Current().IsLicenseeVersion());
 	// Currently assumes IoDispatcher however should we cook to loose files in the future this would need updating
-	OutGeneratedResult.AssetPackageData.SetPackageLocation(FPackageName::EPackageLocationFilter::IoDispatcher);
-	OutGeneratedResult.AssetPackageData.Extension = FPackagePath::ParseExtension(
+	AssetPackageData.SetPackageLocation(FPackageName::EPackageLocationFilter::IoDispatcher);
+	AssetPackageData.Extension = FPackagePath::ParseExtension(
 		WriteToString<256>(PackageData.GetFileName()));
-	OutGeneratedResult.AssetPackageData.SetPackageSavedHash(Info->PackageHash);
-	OutGeneratedResult.AssetPackageData.ImportedClasses = MoveTemp(ImportedClasses);
+	AssetPackageData.SetPackageSavedHash(Info->PackageHash);
+	AssetPackageData.ImportedClasses = MoveTemp(ImportedClasses);
+
+	OutAssetPackageData = AssetPackageData;
 }
 
-const FAssetPackageData* FGenerationHelper::GetIncrementalCookAssetPackageData(FPackageData& PackageData)
+const FAssetPackageData* FGenerationHelper::GetAssetPackageData(FName PackageName)
 {
-	return PreviousGeneratedPackages.Find(PackageData.GetPackageName());
+	const FAssetPackageData* Result = PreviousGeneratedPackages.Find(PackageName);
+	if (Result)
+	{
+		return Result;
+	}
+
+	FCookGenerationInfo* Info = FindInfoNoInitialize(PackageName);
+	if (!Info || Info->IsGenerator())
+	{
+		return nullptr;
+	}
+
+	return Info->AssetPackageData.GetPtrOrNull();
 }
 
-const FAssetPackageData* FGenerationHelper::GetIncrementalCookAssetPackageData(FName PackageName)
+UE::EditorDomain::FPackageDigest FGenerationHelper::GetPackageDigest(FName PackageName)
 {
-	return PreviousGeneratedPackages.Find(PackageName);
+	const FAssetPackageData* AssetPackageData = GetAssetPackageData(PackageName);
+	if (!AssetPackageData)
+	{
+		return UE::EditorDomain::FPackageDigest();
+	}
+	return UE::EditorDomain::CalculatePackageDigest(*AssetPackageData, PackageName);
 }
 
 void FGenerationHelper::ResetSaveState(FCookGenerationInfo& Info, UPackage* Package,
diff --git a/Engine/Source/Editor/UnrealEd/Private/Cooker/CookGenerationHelper.h b/Engine/Source/Editor/UnrealEd/Private/Cooker/CookGenerationHelper.h
--- a/Engine/Source/Editor/UnrealEd/Private/Cooker/CookGenerationHelper.h
+++ b/Engine/Source/Editor/UnrealEd/Private/Cooker/CookGenerationHelper.h
@@ -26,7 +26,7 @@
 class ITargetPlatform;
 class UCookOnTheFlyServer;
 namespace UE::Cook { class FCookGCDiagnosticContext; }
-namespace UE::Cook { struct FGeneratedPackageResultStruct; }
+namespace UE::EditorDomain { struct FPackageDigest; }
 
 namespace UE::Cook
 {
@@ -133,6 +133,8 @@ public:
 	TArray<TWeakObjectPtr<UPackage>> KeepReferencedPackages;
 	TMap<UObject*, FCachedObjectInOuterGeneratorInfo> CachedObjectsInOuterInfo;
 	FWorkerId SavedOnWorker = FWorkerId::Invalid();
+	TOptional<FAssetPackageData> AssetPackageData;
+
 private:
 	bool bCreateAsMap : 1;
 	bool bHasCreatedPackage : 1;
@@ -299,15 +301,15 @@ public:
 	 * Called when a generated package is in FSaveCookedPackageContext::FinishPlatformSave.
 	 * Calculates AssetRegistryData.
 	 */
-	void FinishGeneratedPlatformSave(FPackageData& PackageData,
-		FGeneratedPackageResultStruct& GeneratedResult);
+	void FinishGeneratedPlatformSave(FPackageData& PackageData, FAssetPackageData& OutAssetPackageData,
+		TArray< FAssetDependency>& OutDependencies);
 	/**
-	 * Return the AssetPackageData that was loaded for the given PackageData from the previous cook's
-	 * AssetRegistry, or nullptr if not an incremental cook or it was not previously cooked.
+	 * Return the AssetPackageData that was created for the given PackageData, or loaded for it from a previous cook.
+	 * Will return nullptr if not loaded from a previous incremental cook and it has not yet been saved.
 	 * Does not call initialize.
 	 */
-	const FAssetPackageData* GetIncrementalCookAssetPackageData(FPackageData& PackageData);
-	const FAssetPackageData* GetIncrementalCookAssetPackageData(FName PackageName);
+	const FAssetPackageData* GetAssetPackageData(FName PackageName);
+	UE::EditorDomain::FPackageDigest GetPackageDigest(FName PackageName);
 
 	/**
 	 * Clear any data that should only be held when an FPackageData is in the save state, for the given Info. The given
@@ -468,6 +470,7 @@ private:
 private:
 	void ConditionalInitialize() const;
 	FCookGenerationInfo* FindInfoNoInitialize(const FPackageData& PackageData);
+	FCookGenerationInfo* FindInfoNoInitialize(FName PackageName);
 	void NotifyCompletion(ICookPackageSplitter::ETeardown Status);
 	void PreGarbageCollectGCLifetimeData();
 	void PostGarbageCollectGCLifetimeData(FCookGCDiagnosticContext& Context);
diff --git a/Engine/Source/Editor/UnrealEd/Private/Cooker/CookPackageArtifacts.cpp b/Engine/Source/Editor/UnrealEd/Private/Cooker/CookPackageArtifacts.cpp
--- a/Engine/Source/Editor/UnrealEd/Private/Cooker/CookPackageArtifacts.cpp
+++ b/Engine/Source/Editor/UnrealEd/Private/Cooker/CookPackageArtifacts.cpp
@@ -7,6 +7,7 @@
 #include "AssetRegistry/IAssetRegistry.h"
 #include "Cooker/CookConfigAccessTracker.h"
 #include "Cooker/CookDependency.h"
+#include "Cooker/CookGenerationHelper.h"
 #include "Cooker/PackageBuildDependencyTracker.h"
 #include "CookOnTheSide/CookLog.h"
 #include "CookPackageSplitter.h"
@@ -33,7 +34,7 @@ FGuid TargetDomainVersion(0x10E1BAE8, 0x6CA04033, 0xA7E57A37, 0xD967F493);
 
 // Bump PackageArtifactsVersion version when the serialization of PackageArtifacts has changed and we want to add
 // backwards compatibility rather than invalidating everything.
-constexpr uint32 PackageArtifactsVersion = 0x00000002;
+constexpr uint32 PackageArtifactsVersion = 0x00000003;
 
 static const FUtf8StringView PackageArtifactsAttachmentKey = UTF8TEXTVIEW("CookPackageArtifacts");
 static const FUtf8StringView BuildDefinitionsAttachmentKey = UTF8TEXTVIEW("BuildDefinitionsAttachmentKey");
@@ -59,7 +60,7 @@ bool FPackageArtifacts::HasSaveResults() const
 	return bHasSaveResults;
 }
 
-bool FPackageArtifacts::HasKeyMatch(const FAssetPackageData* OverrideAssetPackageData)
+bool FPackageArtifacts::HasKeyMatch(FGenerationHelper* GenerationHelper)
 {
 	if (!bValid)
 	{
@@ -71,7 +72,7 @@ bool FPackageArtifacts::HasKeyMatch(const FAssetPackageData* OverrideAssetPackag
 	}
 	if (CurrentKey.IsZero())
 	{
-		ECurrentKeyResult Result = TryCalculateCurrentKey(OverrideAssetPackageData);
+		ECurrentKeyResult Result = TryCalculateCurrentKey(GenerationHelper);
 		if (Result != ECurrentKeyResult::Success)
 		{
 			CurrentKey.Reset();
@@ -82,8 +83,7 @@ bool FPackageArtifacts::HasKeyMatch(const FAssetPackageData* OverrideAssetPackag
 }
 
 FPackageArtifacts::ECurrentKeyResult FPackageArtifacts::TryCalculateCurrentKey(
-	const FAssetPackageData* OverrideAssetPackageData,
-	TArray<TPair<ELogVerbosity::Type, FString>>* OutMessages)
+	FGenerationHelper* GenerationHelper, TArray<TPair<ELogVerbosity::Type, FString>>* OutMessages)
 {
 	IAssetRegistry* AssetRegistry = IAssetRegistry::Get();
 	if (PackageName.IsNone())
@@ -105,109 +105,126 @@ FPackageArtifacts::ECurrentKeyResult FPackageArtifacts::TryCalculateCurrentKey(
 	FBlake3 KeyBuilder;
 	KeyBuilder.Update(&TargetDomainVersion, sizeof(TargetDomainVersion));
 
-	UE::EditorDomain::FPackageDigest PackageDigest = OverrideAssetPackageData ?
-		UE::EditorDomain::CalculatePackageDigest(*OverrideAssetPackageData, PackageName) :
-		EditorDomain->GetPackageDigest(PackageName);
-	if (!PackageDigest.IsSuccessful())
-	{
-		if (OutMessages) OutMessages->Emplace(ELogVerbosity::Error, PackageDigest.GetStatusString());
-		return ECurrentKeyResult::Error;
-	}
-	KeyBuilder.Update(&PackageDigest.Hash, sizeof(PackageDigest.Hash));
-
-	if (!ClassDependencies.IsEmpty())
-	{
-		TArray<FTopLevelAssetPath, TInlineAllocator<1>> ClassPaths;
-		ClassPaths.Reserve(ClassDependencies.Num());
-		for (const FString& ClassPathStr : ClassDependencies)
-		{
-			FTopLevelAssetPath& ClassPath = ClassPaths.Emplace_GetRef(ClassPathStr);
-			if (!ClassPath.IsValid())
-			{
-				if (OutMessages)
-				{
-					OutMessages->Emplace(ELogVerbosity::Error, FString::Printf(
-						TEXT("ClassDependency failed: %s is not a valid TopLevelAssetPath."), *ClassPathStr));
-				}
-				return ECurrentKeyResult::Error;
-			}
-		}
-		FString AppendDigestError;
-		if (!UE::EditorDomain::TryAppendClassDigests(KeyBuilder, ClassPaths, &AppendDigestError))
+	ECurrentKeyResult Result = ECurrentKeyResult::Success;
+	UE::Cook::FCookDependencyContext Context(&KeyBuilder,
+		[&Result, OutMessages](ELogVerbosity::Type Verbosity, FString&& Message)
 		{
 			if (OutMessages)
 			{
-				OutMessages->Emplace(ELogVerbosity::Error, FString::Printf(
-					TEXT("ClassDependency failed: %s"), *AppendDigestError));
+				OutMessages->Emplace(Verbosity, MoveTemp(Message));
 			}
-			return ECurrentKeyResult::Error;
-		}
-	}
+		},
+		[&Result](ELogVerbosity::Type Verbosity)
+		{
+			ECurrentKeyResult NewResult = Verbosity <= ELogVerbosity::Error
+				? ECurrentKeyResult::Error : ECurrentKeyResult::Invalidated;
+			Result = FMath::Max(Result, NewResult);
+		},
+		PackageName);
 
-	for (FName PackageDependency : BuildPackageDependencies)
+	// The BuildDependencies have already been sorted, and the CookDependency sort
+	// function sorts CookDependencies of the same type together. This allows us
+	// to create batches for CookDependencies with Update functions that benefit
+	// from being updated in batches.
+	int32 NumDependencies = BuildDependencies.Num();
+	for (int32 BatchStart = 0; BatchStart < NumDependencies; )
 	{
-		PackageDigest = EditorDomain->GetPackageDigest(PackageDependency);
-		if (PackageDigest.IsSuccessful())
+		ECookDependency BatchCategory = BuildDependencies[BatchStart].GetType();
+		int32 BatchEnd = BatchStart + 1;
+		while (BatchEnd < NumDependencies && BuildDependencies[BatchEnd].GetType() == BatchCategory)
 		{
-			KeyBuilder.Update(&PackageDigest.Hash, sizeof(PackageDigest.Hash));
-			continue;
+			++BatchEnd;
 		}
+		TConstArrayView<FCookDependency> Batch(BuildDependencies.GetData() + BatchStart, BatchEnd - BatchStart);
+		BatchStart += Batch.Num();
 
-		if (OutMessages)
+		// Some CookDependency types can not handle Update being called, because their Update relies on functions
+		// only available outside of the CoreUObject module. Handle those types.
+		switch (BatchCategory)
 		{
-			OutMessages->Emplace(ELogVerbosity::Error, FString::Printf(TEXT("Could not create PackageDigest for %s: %s"),
-				*PackageDependency.ToString(), *PackageDigest.GetStatusString()));
-		}
-		return ECurrentKeyResult::Error;
-	}
-
-	FCoreRedirects::AppendHashOfRedirectsAffectingPackages(KeyBuilder, RuntimePackageDependencies);
-	FCoreRedirects::AppendHashOfRedirectsAffectingPackages(KeyBuilder, ScriptPackageDependencies);
-	FCoreRedirects::AppendHashOfRedirectsAffectingPackages(KeyBuilder, RedirectorsAffectingPackageDependencies);
+		case ECookDependency::Package:
+		{
+			for (const FCookDependency& PackageDependency : Batch)
+			{
+				FName DependencyPackageName = PackageDependency.GetPackageName();
+				UE::EditorDomain::FPackageDigest PackageDigest = EditorDomain->GetPackageDigest(DependencyPackageName);
+				if (!PackageDigest.IsSuccessful() && GenerationHelper)
+				{
+					PackageDigest = GenerationHelper->GetPackageDigest(DependencyPackageName);
+				}
+				if (PackageDigest.IsSuccessful())
+				{
+					KeyBuilder.Update(&PackageDigest.Hash, sizeof(PackageDigest.Hash));
+					continue;
+				}
 
-	if (!ConfigDependencies.IsEmpty())
-	{
-#if UE_WITH_CONFIG_TRACKING
-		using namespace UE::ConfigAccessTracking;
-		FCookConfigAccessTracker& ConfigTracker = FCookConfigAccessTracker::Get();
-#endif
-		for (const FString& ConfigDependency : ConfigDependencies)
+				Context.LogError(FString::Printf(
+					TEXT("PackageDependency failed: Could not create PackageDigest for %s: %s"),
+					*DependencyPackageName.ToString(), *PackageDigest.GetStatusString()));
+			}
+			break;
+		}
+		case ECookDependency::Config:
 		{
-			FString Value;
 #if UE_WITH_CONFIG_TRACKING
-			Value = ConfigTracker.GetValue(ConfigDependency);
+			using namespace UE::ConfigAccessTracking;
+			FCookConfigAccessTracker& ConfigTracker = FCookConfigAccessTracker::Get();
 #endif
-			uint8 Marker = 0;
-			KeyBuilder.Update(&Marker, sizeof(Marker));
-			if (!Value.IsEmpty())
+			for (const FCookDependency& ConfigDependency : Batch)
 			{
-				KeyBuilder.Update(*Value, Value.Len() * sizeof(Value[0]));
+				FString Value;
+#if UE_WITH_CONFIG_TRACKING
+				Value = ConfigTracker.GetValue(ConfigDependency.GetConfigAccessData());
+#endif
+				uint8 Marker = 0;
+				KeyBuilder.Update(&Marker, sizeof(Marker));
+				if (!Value.IsEmpty())
+				{
+					KeyBuilder.Update(*Value, Value.Len() * sizeof(Value[0]));
+				}
 			}
+			break;
 		}
-	}
-
-	ECurrentKeyResult Result = ECurrentKeyResult::Success;
-	if (!CookDependencies.IsEmpty())
-	{
-		UE::Cook::FCookDependencyContext Context(&KeyBuilder,
-			[&Result, OutMessages](ELogVerbosity::Type Verbosity, FString&& Message)
+		case ECookDependency::NativeClass:
+		{
+			TArray<FTopLevelAssetPath, TInlineAllocator<10>> ClassPaths;
+			ClassPaths.Reserve(Batch.Num());
+			for (const FCookDependency& ClassDependency : Batch)
 			{
-				if (OutMessages)
+				FTopLevelAssetPath& ClassPath = ClassPaths.Emplace_GetRef(ClassDependency.GetClassPath());
+				if (!ClassPath.IsValid())
 				{
-					OutMessages->Emplace(Verbosity, MoveTemp(Message));
+					Context.LogError(FString::Printf(
+						TEXT("ClassDependency failed: %s is not a valid TopLevelAssetPath."),
+						*WriteToString<256>(ClassDependency.GetClassPath())));
+					ClassPaths.Pop();
+					continue;
 				}
-			},
-			[&Result](ELogVerbosity::Type Verbosity)
+			}
+			FString AppendDigestError;
+			if (!UE::EditorDomain::TryAppendClassDigests(KeyBuilder, ClassPaths, &AppendDigestError))
 			{
-				ECurrentKeyResult NewResult = Verbosity <= ELogVerbosity::Error
-					? ECurrentKeyResult::Error : ECurrentKeyResult::Invalidated;
-				Result = FMath::Max(Result, NewResult);
-			},
-			PackageName);
-
-		for (UE::Cook::FCookDependency& CookDependency : CookDependencies)
+				Context.LogError(FString::Printf(TEXT("ClassDependency failed: %s"), *AppendDigestError));
+			}
+			break;
+		}
+		case ECookDependency::RedirectionTarget:
 		{
-			CookDependency.UpdateHash(Context);
+			TArray<FName, TInlineAllocator<10>> PackageNames;
+			PackageNames.Reserve(Batch.Num());
+			for (const FCookDependency& RedirectionDependency : Batch)
+			{
+				PackageNames.Add(RedirectionDependency.GetPackageName());
+			}
+			FCoreRedirects::AppendHashOfRedirectsAffectingPackages(KeyBuilder, PackageNames);
+			break;
+		}
+		default:
+			for (const FCookDependency& BatchDependency : Batch)
+			{
+				BatchDependency.UpdateHash(Context);
+			}
+			break;
 		}
 	}
 
@@ -220,11 +237,8 @@ FPackageArtifacts::ECurrentKeyResult FPackageArtifacts::TryCalculateCurrentKey(
 
 void FPackageArtifacts::Reset()
 {
-	BuildPackageDependencies.Reset();
-	ConfigDependencies.Reset();
-	RuntimePackageDependencies.Reset();
-	ScriptPackageDependencies.Reset();
-	RedirectorsAffectingPackageDependencies.Reset();
+	BuildDependencies.Reset();
+	RuntimeDependencies.Reset();
 	PackageName = FName();
 	StoredKey = FIoHash::Zero;
 	CurrentKey = FIoHash::Zero;
@@ -234,10 +248,8 @@ void FPackageArtifacts::Reset()
 void FPackageArtifacts::Empty()
 {
 	Reset();
-	BuildPackageDependencies.Empty();
-	ConfigDependencies.Empty();
-	RuntimePackageDependencies.Empty();
-	RedirectorsAffectingPackageDependencies.Empty();
+	BuildDependencies.Empty();
+	RuntimeDependencies.Empty();
 }
 
 enum class EPackageMountPoint
@@ -292,8 +304,8 @@ static void JoinMessagesIntoErrorReason(FStringBuilderBase& OutText,
 }
 
 FPackageArtifacts FPackageArtifacts::Collect(UPackage* Package, const ITargetPlatform* TargetPlatform,
-	FSavePackageResultStruct* SaveResult, const FGeneratedPackageResultStruct* GeneratedResult,
-	TArray<FName>&& RuntimeDependencies, TArray<TPair<ELogVerbosity::Type, FString>>* OutMessages)
+	FSavePackageResultStruct* SaveResult, FGenerationHelper* GenerationHelper, bool bGenerated,
+	TArray<FName>&& InRuntimeDependencies, TArray<TPair<ELogVerbosity::Type, FString>>* OutMessages)
 {
 	if (!Package)
 	{
@@ -314,27 +326,22 @@ FPackageArtifacts FPackageArtifacts::Collect(UPackage* Package, const ITargetPla
 	}
 
 	FPackageArtifacts Result;
+	// Skip the multiple reallocations for an array that grows from 0 to 128, but then reallocate according
+	// to normal TArray growth.
+	Result.BuildDependencies.Reserve(128);
 	Result.PackageName = Package->GetFName();
 	Result.bHasSaveResults = SaveResult != nullptr;
-	enum class EBuildDependency
-	{
-		PackageOnly,
-		Transitive,
-	};
-
-	TMap<FName, EBuildDependency> PackageBuildDependencies;
+	Result.BuildDependencies.Add(FCookDependency::Package(Result.PackageName));
+	Result.RuntimeDependencies = MoveTemp(InRuntimeDependencies);
 
 	TArray<FName> AssetDependencies;
-	if (!GeneratedResult)
+	// The equivalent of AssetRegistry dependencies for generated packages come from their generator's TryPopulate
+	// function and this function receives them via InRuntimeDependencies.
+	if (!bGenerated)
 	{
 		AssetRegistry->GetDependencies(Result.PackageName, AssetDependencies,
 			UE::AssetRegistry::EDependencyCategory::Package, UE::AssetRegistry::EDependencyQuery::Game);
-		RuntimeDependencies.Append(MoveTemp(AssetDependencies));
-	}
-	else
-	{
-		// GeneratedResult->PackageDependencies are already incorporated into the PackageHash of the generated
-		// package by FCookGenerationInfo::CreatePackageHash, so we do not need to add them into the BuildDependencies
+		Result.RuntimeDependencies.Append(MoveTemp(AssetDependencies));
 	}
 
 #if UE_WITH_PACKAGE_ACCESS_TRACKING
@@ -345,7 +352,6 @@ FPackageArtifacts FPackageArtifacts::Collect(UPackage* Package, const ITargetPla
 		TArray<TPair<FBuildDependencyAccessData, FResultProjectionList>> AccessDatas
 			= Tracker.GetAccessDatas(Result.PackageName);
 
-		PackageBuildDependencies.Reserve(AccessDatas.Num());
 		for (TPair<FBuildDependencyAccessData, FResultProjectionList>& Pair : AccessDatas)
 		{
 			FBuildDependencyAccessData& AccessData = Pair.Key;
@@ -358,27 +364,27 @@ FPackageArtifacts FPackageArtifacts::Collect(UPackage* Package, const ITargetPla
 					// We have not yet enabled marking the auto-added dependencies from TObjectPtr resolve
 					// as transitive, because it causes a performance regression and we are still working on fixing the
 					// regression.
-					PackageBuildDependencies.FindOrAdd(AccessData.ReferencedPackage, EBuildDependency::PackageOnly);
+					Result.BuildDependencies.Add(FCookDependency::Package(AccessData.ReferencedPackage));
 				}
 				else if (ProjectionList.bHasAll)
 				{
 					// TObjectPtr UE::BuildProjection::All dependencies are added as transitive build dependencies.
 					// We have to do this to be conservative, since we do not know which bytes from the target are
 					// dependended upon and which of the target's build dependencies influence those bytes.
-					PackageBuildDependencies.FindOrAdd(AccessData.ReferencedPackage) = EBuildDependency::Transitive;
+					Result.BuildDependencies.Add(FCookDependency::TransitiveBuild(AccessData.ReferencedPackage));
 				}
 				else
 				{
 					for (FTopLevelAssetPath ClassPath : ProjectionList.Classes)
 					{
-						Result.ClassDependencies.Add(ClassPath.ToString());
+						Result.BuildDependencies.Add(FCookDependency::NativeClass(WriteToString<256>(ClassPath)));
 					}
 
 					for (FName ResultProjection : ProjectionList.ResultProjections)
 					{
 						if (ResultProjection == UE::Cook::ResultProjection::PackageAndClass)
 						{
-							PackageBuildDependencies.FindOrAdd(AccessData.ReferencedPackage, EBuildDependency::PackageOnly);
+							Result.BuildDependencies.Add(FCookDependency::Package(AccessData.ReferencedPackage));
 						}
 						else
 						{
@@ -399,7 +405,7 @@ FPackageArtifacts FPackageArtifacts::Collect(UPackage* Package, const ITargetPla
 		// dependencies.
 		for (FName AssetDependency : AssetDependencies)
 		{
-			PackageBuildDependencies.FindOrAdd(AssetDependency) = EBuildDependency::Transitive;
+			Result.BuildDependencies.Add(FCookDependency::TransitiveBuild(AssetDependency));
 		}
 	}
 
@@ -410,183 +416,165 @@ FPackageArtifacts FPackageArtifacts::Collect(UPackage* Package, const ITargetPla
 		if (ConfigTracker.IsEnabled())
 		{
 			TArray<FConfigAccessData> ConfigKeys = ConfigTracker.GetPackageRecords(Result.PackageName, TargetPlatform);
-			Result.ConfigDependencies.Reserve(ConfigKeys.Num());
 			for (const FConfigAccessData& ConfigKey : ConfigKeys)
 			{
-				Result.ConfigDependencies.Add(ConfigKey.FullPathToString());
+				Result.BuildDependencies.Add(FCookDependency::Config(ConfigKey));
 			}
 		}
 	}
 #endif
 	if (SaveResult)
 	{
-		// Convert some FCookDependency types into our specially-handled dependencies, and put the more
-		// generic ones into this->CookDependencies
-		Result.CookDependencies.Reserve(SaveResult->CookDependencies.Num());
-
-		// Additional dependencies can be added from SettingsObject dependencies. Add these included dependencies
-		// to a followup list that we process in the second round, if they exist. Handle detecting cycles in the
-		// included dependencies; cycle detection is currently simple: nested includes are not yet required
-		// so give an error if we ever make it to a third round.
-		TArray<UE::Cook::FCookDependency> NextIncludedDependenciesBuffer[2];
-		TArray<UE::Cook::FCookDependency>* RemainingDependencies = &SaveResult->CookDependencies;
-		int32 IncludedDependenciesRound = 0;
-		constexpr int32 MaxNumberOfIncludedDependenciesRounds = 2;
-		while (!RemainingDependencies->IsEmpty())
+		Result.BuildDependencies.Append(SaveResult->CookDependencies);
+		for (FName SoftPackageReference : SaveResult->UntrackedSoftPackageReferences)
 		{
-			TArray<UE::Cook::FCookDependency>* NextIncludedDependencies =
-				&NextIncludedDependenciesBuffer[IncludedDependenciesRound % UE_ARRAY_COUNT(NextIncludedDependenciesBuffer)];
-			if (IncludedDependenciesRound++ > 0)
-			{
-				NextIncludedDependencies->Reset();
-				if (IncludedDependenciesRound > MaxNumberOfIncludedDependenciesRounds)
-				{
-					if (OutMessages)
-					{
-						OutMessages->Emplace(ELogVerbosity::Error,
-							TEXT("More rounds than expected when handling included dependencies."));
-					}
-					return FPackageArtifacts();
-				}
-			}
-
-			for (UE::Cook::FCookDependency& CookDependency : *RemainingDependencies)
-			{
-				switch (CookDependency.GetType())
-				{
-				case UE::Cook::ECookDependency::TransitiveBuild:
-				{
-					// Build dependencies from a package to itself have a performance cost and serve no
-					// purpose, so remove them. They can occur in some systems that naively add a build dependency
-					// from one object to another without checking whether the second object is in the same package.
-					FName DependencyPackageName = CookDependency.GetPackageName();
-					if (DependencyPackageName == Result.PackageName)
-					{
-						continue;
-					}
-					if (CookDependency.IsAlsoAddRuntimeDependency())
-					{
-						RuntimeDependencies.Add(DependencyPackageName);
-					}
-					PackageBuildDependencies.FindOrAdd(DependencyPackageName) = EBuildDependency::Transitive;
-					break;
-				}
-				case UE::Cook::ECookDependency::Package:
-					// Add it as a PackageOnly build dependency (lowest level) unless a higher level already exists.
-					PackageBuildDependencies.FindOrAdd(CookDependency.GetPackageName(), EBuildDependency::PackageOnly);
-					break;
-				case UE::Cook::ECookDependency::Config:
-					Result.ConfigDependencies.Add(FString(CookDependency.GetConfigPath()));
-					break;
-				case UE::Cook::ECookDependency::SettingsObject:
-				{
-					const UObject* SettingsObject = CookDependency.GetSettingsObject();
-					if (SettingsObject)
-					{
-						// We rely on the object to be rooted because we use its pointer as a key for the lifetime of
-						// the cook process, so it being garbage collected and something else allocated on the same
-						// pointer would break our key. IsRooted should have been validated by
-						// FCookDependency::SettingsObject.
-						check(SettingsObject->IsRooted());
-						FCookDependencyGroups::FRecordedDependencies& IncludeDependencies =
-							FCookDependencyGroups::Get().FindOrCreate(reinterpret_cast<UPTRINT>(SettingsObject));
-						if (!IncludeDependencies.bInitialized)
-						{
-							IncludeDependencies.Artifacts = CollectSettingsObject(SettingsObject,
-								&IncludeDependencies.Messages);
-							IncludeDependencies.bInitialized = true;
-						}
-						if (!IncludeDependencies.Artifacts.IsValid())
-						{
-							if (OutMessages)
-							{
-								TStringBuilder<256> ErrorText;
-								ErrorText << TEXT("Dependencies for SettingsObject ") << SettingsObject->GetPathName()
-									<< TEXT(" are unavailable");
-								JoinMessagesIntoErrorReason(ErrorText, IncludeDependencies.Messages);
-								OutMessages->Emplace(ELogVerbosity::Error, FString(*ErrorText));
-							}
-							return FPackageArtifacts();
-						}
-
-						for (const UE::Cook::FCookDependency& IncludeDependency :
-							IncludeDependencies.Artifacts.GetCookDependencies())
-						{
-							NextIncludedDependencies->Add(IncludeDependency);
-						}
-					}
-					break;
-				}
-				case UE::Cook::ECookDependency::NativeClass:
-					Result.ClassDependencies.Add(FString(CookDependency.GetClassPath()));
-					break;
-				default:
-					Result.CookDependencies.Add(MoveTemp(CookDependency));
-					break;
-				}
-			}
-			RemainingDependencies = NextIncludedDependencies;
+			Result.BuildDependencies.Add(FCookDependency::RedirectionTarget(SoftPackageReference));
 		}
-
-		Result.RedirectorsAffectingPackageDependencies = SaveResult->UntrackedSoftPackageReferences;
-		Result.RedirectorsAffectingPackageDependencies.Sort(FNameLexicalLess());
-		Result.RedirectorsAffectingPackageDependencies.SetNum(Algo::Unique(Result.RedirectorsAffectingPackageDependencies));
 	}
 
-	Algo::Sort(Result.CookDependencies);
-	Result.CookDependencies.SetNum(Algo::Unique(Result.CookDependencies));
-	for (TPair<FName, EBuildDependency>& Pair : PackageBuildDependencies)
-	{
-		if (Pair.Key == Result.PackageName)
+	// All Dependencies have been gathered. Format the lists for TryCalculateCurrentKey and for storage.
+
+	// Settings dependencies - Expand transitive dependencies on SettingsObjects into the list of dependencies
+	// recorded for that settings object.
+	TSet<const UObject*> SettingsDependencies;
+	Result.BuildDependencies.RemoveAllSwap([&SettingsDependencies](const FCookDependency& Dependency)
 		{
-			// Skip dependencies to self, for performance
-			continue;
-		}
-		if (GetPackageMountPoint(Pair.Key) != EPackageMountPoint::Content)
+			if (Dependency.GetType() == ECookDependency::SettingsObject)
+			{
+				SettingsDependencies.Add(Dependency.GetSettingsObject());
+				return true;
+			}
+			return false;
+		},
+		EAllowShrinking::No);
+	for (const UObject* SettingsObject : SettingsDependencies)
+	{
+		// We rely on the object to be rooted because we use its pointer as a key for the lifetime of
+		// the cook process, so it being garbage collected and something else allocated on the same
+		// pointer would break our key. IsRooted should have been validated by
+		// FCookDependency::SettingsObject.
+		check(SettingsObject->IsRooted());
+		FCookDependencyGroups::FRecordedDependencies& IncludeDependencies =
+			FCookDependencyGroups::Get().FindOrCreate(reinterpret_cast<UPTRINT>(SettingsObject));
+		if (!IncludeDependencies.bInitialized)
 		{
-			// Skip script and temp packages
-			continue;
+			IncludeDependencies.Artifacts = CollectSettingsObject(SettingsObject,
+				&IncludeDependencies.Messages);
+			IncludeDependencies.bInitialized = true;
 		}
-
-		TStringBuilder<256> PackagePath(EInPlace::InPlace, Pair.Key);
-		UPackage* PackageObject = FindPackage(nullptr, *PackagePath);
-		if (!PackageObject)
+		if (!IncludeDependencies.Artifacts.IsValid())
 		{
 			if (OutMessages)
 			{
 				TStringBuilder<256> ErrorText;
-				ErrorText << TEXT("Package ") << PackagePath << TEXT(" does not exist.");
+				ErrorText << TEXT("Dependencies for SettingsObject ") << SettingsObject->GetPathName()
+					<< TEXT(" are unavailable");
+				JoinMessagesIntoErrorReason(ErrorText, IncludeDependencies.Messages);
 				OutMessages->Emplace(ELogVerbosity::Error, FString(*ErrorText));
 			}
-
 			return FPackageArtifacts();
 		}
 
-		if(PackageObject->HasAnyPackageFlags(PKG_NewlyCreated))
+		for (const FCookDependency& IncludeDependency : IncludeDependencies.Artifacts.GetBuildDependencies())
 		{
-			// If the package is a newly created package (in-memory package) then ignore it.
-			// In-memory packages are ignored because we can't compute their digest. Only packages on disk have a digest.
-			continue;
+			// Recursive SettingsDependencies are not allowed. We haven't needed them yet, and not supporting them allows
+			// prevents the need for cycle detection.
+			if (IncludeDependency.GetType() == ECookDependency::SettingsObject)
+			{
+				if (OutMessages)
+				{
+					OutMessages->Emplace(ELogVerbosity::Error, FString::Printf(
+						TEXT("Settings dependency on object %s, but that object has a recursive Settings dependency on %s, and recursive Settings dependencies are not supported."),
+						*SettingsObject->GetPathName(), *IncludeDependency.GetSettingsObject()->GetPathName()));
+				}
+				return FPackageArtifacts();
+			}
+			Result.BuildDependencies.Add(IncludeDependency);
 		}
+	}
 
-		if (Pair.Value == EBuildDependency::Transitive)
-		{
-			Result.TransitiveBuildDependencies.Add(UE::Cook::FCookDependency::Package(Pair.Key));
-		}
-		else
+	// Process some rules for Package dependencies
+	TSet<FName> RedirectionTargets;
+	for (TArray<FCookDependency>::TIterator Iter(Result.BuildDependencies); Iter; ++Iter)
+	{
+		FCookDependency& Dependency = *Iter;
+		if (Dependency.GetType() == ECookDependency::TransitiveBuild ||
+			Dependency.GetType() == ECookDependency::Package)
 		{
-			Result.BuildPackageDependencies.Add(Pair.Key);
+			FName DependencyName = Dependency.GetPackageName();
+			// Remove transitive dependencies to self, for performance. But keep the package dependency to
+			// self; every cooked package has its EditorDomain package as a dependency.
+			if (Dependency.GetType() == ECookDependency::TransitiveBuild && DependencyName == Result.PackageName)
+			{
+				Iter.RemoveCurrentSwap();
+				continue;
+			}
+
+			// We do not hash dependencies to non-content packages (e.g. temp, memory, script),
+			// so remove package or transitive package dependencies to them.
+			EPackageMountPoint MountPoint = GetPackageMountPoint(DependencyName);
+			if (MountPoint != EPackageMountPoint::Content && MountPoint != EPackageMountPoint::GeneratedContent)
+			{
+				Iter.RemoveCurrentSwap();
+				continue;
+			}
+
+			// Remove dependencies to generated packages, except for a generated package's dependency to itself.
+			// We do not yet support the availability of the digest of other generated packages when requested from
+			// the savepackage and dependency collection of a generated package or a generator; 
+			// the digests only become available when the target generated package is saved, which can happen after
+			// the save of the packages that refer to it.
+			if (MountPoint == EPackageMountPoint::GeneratedContent && DependencyName != Result.PackageName)
+			{
+				Iter.RemoveCurrentSwap();
+				continue;
+			}
+
+			TStringBuilder<256> DependencyNameStr(InPlace, DependencyName);
+			UPackage* DependencyPackage = FindPackage(nullptr, *DependencyNameStr);
+			if (!DependencyPackage)
+			{
+				if (OutMessages)
+				{
+					OutMessages->Emplace(ELogVerbosity::Error,
+						FString::Printf(TEXT("Package %s does not exist."), *DependencyNameStr));
+				}
+
+				return FPackageArtifacts();
+			}
+
+			if (DependencyPackage->HasAnyPackageFlags(PKG_NewlyCreated))
+			{
+				// If the package is a newly created package (in-memory package) then ignore it.
+				// In-memory packages are ignored because we can't compute their digest. Only packages on disk have a digest.
+				Iter.RemoveCurrentSwap();
+				continue;
+			}
+
+			// Package dependencies of all kinds (Runtime, Build, TransitiveBuild) also cause RedirectionTarget
+			// dependencies.
+			RedirectionTargets.Add(DependencyName);
+
+			// Deprecated TransitiveBuildAndRuntime dependencies can also cause runtime dependencies;
+			// convert them to separate dependencies now.
+			if (Dependency.GetType() == ECookDependency::TransitiveBuild)
+			{
+				PRAGMA_DISABLE_DEPRECATION_WARNINGS;
+				if (Dependency.IsAlsoAddRuntimeDependency())
+				{
+					Result.RuntimeDependencies.Add(DependencyName);
+					// Remove the IsAlsoRuntimeDependency flag
+					Dependency = FCookDependency::TransitiveBuild(DependencyName);
+				}
+				PRAGMA_ENABLE_DEPRECATION_WARNINGS;
+			}
 		}
 	}
-	Algo::Sort(Result.TransitiveBuildDependencies);
-	Result.BuildPackageDependencies.Sort(FNameLexicalLess());
-	Algo::Sort(Result.ConfigDependencies);
-	Result.ConfigDependencies.SetNum(Algo::Unique(Result.ConfigDependencies));
-	Algo::Sort(Result.ClassDependencies);
-	Result.ClassDependencies.SetNum(Algo::Unique(Result.ClassDependencies));
-
 
-	for (TArray<FName>::TIterator Iter(RuntimeDependencies); Iter; ++Iter)
+	// Pull transient packages out of the runtime dependencies for performance; we don't need them for
+	// deciding what gets cooked. 
+	// Runtime and script dependencies also cause RedirectionTarget dependencies, so record those.
+	for (TArray<FName>::TIterator Iter(Result.RuntimeDependencies); Iter; ++Iter)
 	{
 		FName PackageName = *Iter;
 		EPackageMountPoint MountPoint = GetPackageMountPoint(PackageName);
@@ -594,22 +582,29 @@ FPackageArtifacts FPackageArtifacts::Collect(UPackage* Package, const ITargetPla
 		{
 		case EPackageMountPoint::GeneratedContent:
 		case EPackageMountPoint::Content:
-			Result.RuntimePackageDependencies.Add(PackageName);
-			break;
 		case EPackageMountPoint::Script:
-			Result.ScriptPackageDependencies.Add(PackageName);
+			// Keep it
+			RedirectionTargets.Add(PackageName);
 			break;
 		default:
+			Iter.RemoveCurrentSwap();
 			break;
 		}
 	}
-	Result.RuntimePackageDependencies.Sort(FNameLexicalLess());
-	Result.RuntimePackageDependencies.SetNum(Algo::Unique(Result.RuntimePackageDependencies), EAllowShrinking::Yes);
-	Result.ScriptPackageDependencies.Sort(FNameLexicalLess());
-	Result.ScriptPackageDependencies.SetNum(Algo::Unique(Result.ScriptPackageDependencies), EAllowShrinking::Yes);
 
-	const FAssetPackageData* AssetPackageData = GeneratedResult ? &GeneratedResult->AssetPackageData : nullptr;
-	ECurrentKeyResult CurrentKeyResult = Result.TryCalculateCurrentKey(AssetPackageData, OutMessages);
+	// Put all the extra redirection dependencies into BuildDependencies
+	for (FName RedirectionTarget : RedirectionTargets)
+	{
+		Result.BuildDependencies.Add(FCookDependency::RedirectionTarget(RedirectionTarget));
+	}
+
+	// Sort and remove duplicates
+	Algo::Sort(Result.BuildDependencies);
+	Result.BuildDependencies.SetNum(Algo::Unique(Result.BuildDependencies), EAllowShrinking::Yes);
+	Algo::Sort(Result.RuntimeDependencies, FNameLexicalLess());
+	Result.RuntimeDependencies.SetNum(Algo::Unique(Result.RuntimeDependencies), EAllowShrinking::Yes);
+
+	ECurrentKeyResult CurrentKeyResult = Result.TryCalculateCurrentKey(GenerationHelper, OutMessages);
 	if (CurrentKeyResult == ECurrentKeyResult::Error)
 	{
 		return FPackageArtifacts();
@@ -656,10 +651,10 @@ FPackageArtifacts FPackageArtifacts::CollectSettingsObject(const UObject* Object
 	TArray<UE::ConfigAccessTracking::FConfigAccessData> ConfigDatas;
 	const_cast<UObject*>(Object)->LoadConfig(nullptr /* ConfigClass */, nullptr /* Filename */,
 		UE::LCPF_None /* PropagationFlags */, nullptr /* PropertyToLoad */, &ConfigDatas);
-	Result.CookDependencies.Reserve(ConfigDatas.Num() + 1);
+	Result.BuildDependencies.Reserve(ConfigDatas.Num() + 1);
 	for (const UE::ConfigAccessTracking::FConfigAccessData& ConfigData : ConfigDatas)
 	{
-		Result.CookDependencies.Add(UE::Cook::FCookDependency::Config(ConfigData));
+		Result.BuildDependencies.Add(UE::Cook::FCookDependency::Config(ConfigData));
 	}
 
 	// In addition to adding the config dependencies, add a dependency on the class schema. If the current class has
@@ -672,11 +667,11 @@ FPackageArtifacts FPackageArtifacts::CollectSettingsObject(const UObject* Object
 	}
 	if (NativeClass)
 	{
-		Result.CookDependencies.Add(UE::Cook::FCookDependency::NativeClass(NativeClass));
+		Result.BuildDependencies.Add(UE::Cook::FCookDependency::NativeClass(NativeClass));
 	}
 
-	Algo::Sort(Result.CookDependencies);
-	Result.CookDependencies.SetNum(Algo::Unique(Result.CookDependencies));
+	Algo::Sort(Result.BuildDependencies);
+	Result.BuildDependencies.SetNum(Algo::Unique(Result.BuildDependencies));
 	Result.bValid = true;
 	return Result;
 }
@@ -715,58 +710,16 @@ bool LoadFromCompactBinary(FCbObjectView ObjectView, UE::Cook::FPackageArtifacts
 				return false;
 			}
 		}
-		if (FieldView.GetName().Equals(UTF8TEXTVIEW("BuildPackageDependencies")))
-		{
-			if (!LoadFromCompactBinary(FieldView++, Artifacts.BuildPackageDependencies))
-			{
-				return false;
-			}
-		}
-		if (FieldView.GetName().Equals(UTF8TEXTVIEW("ConfigDependencies")))
-		{
-			if (!LoadFromCompactBinary(FieldView++, Artifacts.ConfigDependencies))
-			{
-				return false;
-			}
-		}
-		if (FieldView.GetName().Equals(UTF8TEXTVIEW("RuntimePackageDependencies")))
-		{
-			if (!LoadFromCompactBinary(FieldView++, Artifacts.RuntimePackageDependencies))
-			{
-				return false;
-			}
-		}
-		if (FieldView->GetName().Equals(UTF8TEXTVIEW("RedirectorsAffectingPackageDependencies")))
+		if (FieldView.GetName().Equals(UTF8TEXTVIEW("BuildDependencies")))
 		{
-			if (!LoadFromCompactBinary(FieldView++, Artifacts.RedirectorsAffectingPackageDependencies))
+			if (!LoadFromCompactBinary(FieldView++, Artifacts.BuildDependencies))
 			{
 				return false;
 			}
 		}
-		if (FieldView.GetName().Equals(UTF8TEXTVIEW("ScriptPackageDependencies")))
+		if (FieldView.GetName().Equals(UTF8TEXTVIEW("RuntimeDependencies")))
 		{
-			if (!LoadFromCompactBinary(FieldView++, Artifacts.ScriptPackageDependencies))
-			{
-				return false;
-			}
-		}
-		if (FieldView.GetName().Equals(UTF8TEXTVIEW("CookDependencies")))
-		{
-			if (!LoadFromCompactBinary(FieldView++, Artifacts.CookDependencies))
-			{
-				return false;
-			}
-		}
-		if (FieldView.GetName().Equals(UTF8TEXTVIEW("TransitiveBuildDependencies")))
-		{
-			if (!LoadFromCompactBinary(FieldView++, Artifacts.TransitiveBuildDependencies))
-			{
-				return false;
-			}
-		}
-		if (FieldView.GetName().Equals(UTF8TEXTVIEW("ClassDependencies")))
-		{
-			if (!LoadFromCompactBinary(FieldView++, Artifacts.ClassDependencies))
+			if (!LoadFromCompactBinary(FieldView++, Artifacts.RuntimeDependencies))
 			{
 				return false;
 			}
@@ -792,37 +745,13 @@ FCbWriter& operator<<(FCbWriter& Writer, const UE::Cook::FPackageArtifacts& Arti
 	Writer << "Version" << PackageArtifactsVersion;
 	Writer << "StoredKey" << Artifacts.StoredKey;
 	Writer << "HasSaveResults" << Artifacts.bHasSaveResults;
-	if (!Artifacts.BuildPackageDependencies.IsEmpty())
-	{
-		Writer << "BuildPackageDependencies" << Artifacts.BuildPackageDependencies;
-	}
-	if (!Artifacts.ConfigDependencies.IsEmpty())
-	{
-		Writer << "ConfigDependencies" << Artifacts.ConfigDependencies;
-	}
-	if (!Artifacts.RuntimePackageDependencies.IsEmpty())
-	{
-		Writer << "RuntimePackageDependencies" << Artifacts.RuntimePackageDependencies;
-	}
-	if (!Artifacts.RedirectorsAffectingPackageDependencies.IsEmpty())
-	{
-		Writer << "RedirectorsAffectingPackageDependencies" << Artifacts.RedirectorsAffectingPackageDependencies;
-	}
-	if (!Artifacts.ScriptPackageDependencies.IsEmpty())
-	{
-		Writer << "ScriptPackageDependencies" << Artifacts.ScriptPackageDependencies;
-	}
-	if (!Artifacts.CookDependencies.IsEmpty())
-	{
-		Writer << "CookDependencies" << Artifacts.CookDependencies;
-	}
-	if (!Artifacts.TransitiveBuildDependencies.IsEmpty())
+	if (!Artifacts.BuildDependencies.IsEmpty())
 	{
-		Writer << "TransitiveBuildDependencies" << Artifacts.TransitiveBuildDependencies;
+		Writer << "BuildDependencies" << Artifacts.BuildDependencies;
 	}
-	if (!Artifacts.ClassDependencies.IsEmpty())
+	if (!Artifacts.RuntimeDependencies.IsEmpty())
 	{
-		Writer << "ClassDependencies" << Artifacts.ClassDependencies;
+		Writer << "RuntimeDependencies" << Artifacts.RuntimeDependencies;
 	}
 
 	Writer.EndObject();
@@ -957,7 +886,7 @@ void FIncrementalCookAttachments::AppendCommitAttachments(TArray<IPackageWriter:
 
 FIncrementalCookAttachments FIncrementalCookAttachments::Collect(UPackage* Package,
 	const ITargetPlatform* TargetPlatform, FSavePackageResultStruct* SaveResult,
-	const FGeneratedPackageResultStruct* GeneratedResult, TArray<FName>&& RuntimeDependencies)
+	FGenerationHelper* GenerationHelper, bool bGenerated, TArray<FName>&& RuntimeDependencies)
 {
 	FIncrementalCookAttachments Result;
 
@@ -965,7 +894,7 @@ FIncrementalCookAttachments FIncrementalCookAttachments::Collect(UPackage* Packa
 
 	TArray<TPair<ELogVerbosity::Type, FString>> Messages;
 	Result.Artifacts = FPackageArtifacts::Collect(Package, TargetPlatform, SaveResult,
-		GeneratedResult, MoveTemp(RuntimeDependencies), &Messages);
+		GenerationHelper, bGenerated, MoveTemp(RuntimeDependencies), &Messages);
 	if (!Result.Artifacts.IsValid())
 	{
 		TStringBuilder<256> LogText;
diff --git a/Engine/Source/Editor/UnrealEd/Private/Cooker/CookPackageArtifacts.h b/Engine/Source/Editor/UnrealEd/Private/Cooker/CookPackageArtifacts.h
--- a/Engine/Source/Editor/UnrealEd/Private/Cooker/CookPackageArtifacts.h
+++ b/Engine/Source/Editor/UnrealEd/Private/Cooker/CookPackageArtifacts.h
@@ -4,6 +4,7 @@
 
 #include "AssetRegistry/AssetData.h"
 #include "AssetRegistry/IAssetRegistry.h"
+#include "Cooker/CookDependency.h"
 #include "Containers/Array.h"
 #include "Containers/ArrayView.h"
 #include "Containers/Map.h"
@@ -24,9 +25,9 @@ class UPackage;
 struct FSavePackageResultStruct;
 template <typename FuncType> class TUniqueFunction;
 
-namespace UE::Cook { class FCookDependency; }
 namespace UE::Cook { class FPackageArtifacts; }
 namespace UE::Cook { struct FBuildDefinitionList; }
+namespace UE::Cook { struct FGenerationHelper; }
 namespace UE::Cook { struct FIncrementalCookAttachments; }
 
 bool LoadFromCompactBinary(FCbObjectView ObjectView, UE::Cook::FPackageArtifacts& CookAttachments);
@@ -37,16 +38,6 @@ FCbWriter& operator<<(FCbWriter& Writer, const UE::Cook::FBuildDefinitionList& D
 namespace UE::Cook
 {
 
-/**
- * Information collected from a CookPackageSplitter Generated package after it is saved; this information is needed to collect
- * the CookDependencies for the generated package.
- */
-struct FGeneratedPackageResultStruct
-{
-	FAssetPackageData AssetPackageData;
-	TArray<FAssetDependency> PackageDependencies;
-};
-
 /**
  * Non-runtime data recorded about each package and stored in the cook Oplog as attachments to the package.
  * Includes BuildResults built from the package that can be used for future incremental cooks, and the dependencies
@@ -69,30 +60,29 @@ public:
 	FPackageArtifacts(FPackageArtifacts&&);
 	FPackageArtifacts& operator=(const FPackageArtifacts&);
 	FPackageArtifacts& operator=(FPackageArtifacts&&);
-	// Build Dependencies
-	const TArray<FName>& GetBuildPackageDependencies() const;
-	const TArray<FString>& GetConfigDependencies() const;
-
-	// Runtime Dependencies
-	const TArray<FName>& GetRuntimePackageDependencies() const;
-	// Script dependencies, needed for TryCalculateCurrentKey
-	const TArray<FName>& GetScriptPackageDependencies() const;
 
-	// Cook Dependencies
-	const TArray<UE::Cook::FCookDependency>& GetCookDependencies() const;
-	const TArray<UE::Cook::FCookDependency>& GetTransitiveBuildDependencies() const;
+	
+	const TArray<FCookDependency>& GetBuildDependencies() const;
+	/** Get all of the runtime dependencies reported by the package, both script and content. */
+	const TArray<FName>& GetRuntimeDependencies() const;
+	/** Return runtime dependencies reported by the package that are content packages; script packages are removed. */
+	template <typename AllocatorType>
+	void GetRuntimeContentDependencies(TArray<FName, AllocatorType>& OutDependencies) const;
+	/** Helper function to filter GetBuildDependencies for the transitive build dependencies. */
+	template <typename AllocatorType>
+	void GetTransitiveBuildDependencies(TArray<FName, AllocatorType>& OutDependencies) const;
 
 	// Data about the Key and Package that are not dependencies
 	/**
-		* True if the structure has been calculated or fetched and accurately reports dependencies and
-		* key for the package. False if the stucture is default, has been reset, or was marked invalid.
-		*/
+	 * True if the structure has been calculated or fetched and accurately reports dependencies and
+	 * key for the package. False if the stucture is default, has been reset, or was marked invalid.
+	 */
 	bool IsValid() const;
 	bool HasSaveResults() const;
 	FName GetPackageName() const;
 	const FIoHash& GetStoredKey() const;
 	const FIoHash& GetCurrentKey() const;
-	bool HasKeyMatch(const FAssetPackageData* OverrideAssetPackageData);
+	bool HasKeyMatch(FGenerationHelper* GenerationHelper);
 
 	// Modifying the structure data
 	enum class ECurrentKeyResult
@@ -101,18 +91,28 @@ public:
 		Invalidated,
 		Error,
 	};
-	ECurrentKeyResult TryCalculateCurrentKey(const FAssetPackageData* OverrideAssetPackageData,
+	/**
+	 * Calculate the current key(s) from the BuildDependencies stored on this PackageArtifacts, and store it in
+	 * GetCurrentKey().
+	 * 
+	 * @param GenerationHelper - If non-null, provides the lookup for the AssetPackageData of generated packages.
+	 *                           Must be provided if any generated packages are in the dependencies.
+	 */
+	ECurrentKeyResult TryCalculateCurrentKey(FGenerationHelper* GenerationHelper,
 		TArray<TPair<ELogVerbosity::Type, FString>>* OutMessages = nullptr);
 	void Reset();
 	void Empty();
 
 	// Construction functions
 	/**
-		* Read dependencies for the given targetplatform of the given package out of global dependency trackers
-		* that have recorded its data during the package's load/other/save operations in the current cook session.
-		*/
+	 * Read dependencies for the given targetplatform of the given package out of global dependency trackers
+	 * that have recorded its data during the package's load/other/save operations in the current cook session.
+	 * 
+	 * @param GenerationHelper - If non-null, provides the lookup for the AssetPackageData of generated packages.
+	 *                           Must be provided if any generated packages are in the dependencies.
+	 */
 	static FPackageArtifacts Collect(UPackage* Package, const ITargetPlatform* TargetPlatform,
-		FSavePackageResultStruct* SaveResult, const FGeneratedPackageResultStruct* GeneratedResult,
+		FSavePackageResultStruct* SaveResult, FGenerationHelper* GenerationHelper, bool bGenerated,
 		TArray<FName>&& RuntimeDependencies, TArray<TPair<ELogVerbosity::Type, FString>>* OutMessages = nullptr);
 	static FPackageArtifacts CollectSettingsObject(const UObject* Object,
 		TArray<TPair<ELogVerbosity::Type, FString>>* OutMessages);
@@ -121,14 +121,8 @@ public:
 	// this structure. Use FIncrementalCookAttachments instead.
 
 private:
-	TArray<FName> BuildPackageDependencies;
-	TArray<FString> ConfigDependencies;
-	TArray<FName> RuntimePackageDependencies;
-	TArray<FName> ScriptPackageDependencies;
-	TArray<FName> RedirectorsAffectingPackageDependencies;
-	TArray<UE::Cook::FCookDependency> CookDependencies;
-	TArray<UE::Cook::FCookDependency> TransitiveBuildDependencies;
-	TArray<FString> ClassDependencies;
+	TArray<UE::Cook::FCookDependency> BuildDependencies;
+	TArray<FName> RuntimeDependencies;
 	FName PackageName;
 	FIoHash StoredKey;
 	FIoHash CurrentKey;
@@ -199,7 +193,7 @@ struct FIncrementalCookAttachments
 		TUniqueFunction<void(FName PackageName, FIncrementalCookAttachments&& Result)>&& Callback);
 	static FIncrementalCookAttachments Collect(UPackage* Package,
 		const ITargetPlatform* TargetPlatform, FSavePackageResultStruct* SaveResult,
-		const FGeneratedPackageResultStruct* GeneratedResult, TArray<FName>&& RuntimeDependencies);
+		FGenerationHelper* GenerationHelper, bool bGenerated, TArray<FName>&& RuntimeDependencies);
 };
 
 } // namespace UE::Cook
@@ -212,34 +206,39 @@ struct FIncrementalCookAttachments
 namespace UE::Cook
 {
 
-inline const TArray<FName>& FPackageArtifacts::GetBuildPackageDependencies() const
+inline const TArray<FName>& FPackageArtifacts::GetRuntimeDependencies() const
 {
-	return BuildPackageDependencies;
+	return RuntimeDependencies;
 }
 
-inline const TArray<FString>& FPackageArtifacts::GetConfigDependencies() const
+template <typename AllocatorType>
+inline void FPackageArtifacts::GetRuntimeContentDependencies(TArray<FName, AllocatorType>& OutDependencies) const
 {
-	return ConfigDependencies;
-}
-
-inline const TArray<FName>& FPackageArtifacts::GetRuntimePackageDependencies() const
-{
-	return RuntimePackageDependencies;
-}
-
-inline const TArray<FName>& FPackageArtifacts::GetScriptPackageDependencies() const
-{
-	return ScriptPackageDependencies;
+	OutDependencies.Reserve(OutDependencies.Num() + RuntimeDependencies.Num());
+	for (FName RuntimeDependency : RuntimeDependencies)
+	{
+		if (!FPackageName::IsScriptPackage(WriteToString<256>(RuntimeDependency)))
+		{
+			OutDependencies.Add(RuntimeDependency);
+		}
+	}
 }
 
-inline const TArray<UE::Cook::FCookDependency>& FPackageArtifacts::GetCookDependencies() const
+inline const TArray<FCookDependency>& FPackageArtifacts::GetBuildDependencies() const
 {
-	return CookDependencies;
+	return BuildDependencies;
 }
 
-inline const TArray<UE::Cook::FCookDependency>& FPackageArtifacts::GetTransitiveBuildDependencies() const
+template <typename AllocatorType>
+inline void FPackageArtifacts::GetTransitiveBuildDependencies(TArray<FName, AllocatorType>& OutDependencies) const
 {
-	return TransitiveBuildDependencies;
+	for (const FCookDependency& BuildDependency : BuildDependencies)
+	{
+		if (BuildDependency.GetType() == ECookDependency::TransitiveBuild)
+		{
+			OutDependencies.Add(BuildDependency.GetPackageName());
+		}
+	}
 }
 
 inline FName FPackageArtifacts::GetPackageName() const
diff --git a/Engine/Source/Editor/UnrealEd/Private/Cooker/CookRequestCluster.cpp b/Engine/Source/Editor/UnrealEd/Private/Cooker/CookRequestCluster.cpp
--- a/Engine/Source/Editor/UnrealEd/Private/Cooker/CookRequestCluster.cpp
+++ b/Engine/Source/Editor/UnrealEd/Private/Cooker/CookRequestCluster.cpp
@@ -1505,22 +1505,6 @@ void FRequestCluster::FGraphSearch::FExploreEdgesContext::CalculatePlatformsToPr
 	}
 }
 
-const FAssetPackageData* FRequestCluster::FVertexData::GetGeneratedAssetPackageData()
-{
-	check(PackageData); // Caller should not call if no PackageData
-	FPackageDatas& LocalPackageDatas = PackageData->GetPackageDatas();
-	FPackageData* ParentPackageData = LocalPackageDatas.FindPackageDataByPackageName(PackageData->GetParentGenerator());
-	if (ParentPackageData)
-	{
-		TRefCountPtr<FGenerationHelper> ParentGenerationHelper = ParentPackageData->GetGenerationHelper();
-		if (ParentGenerationHelper)
-		{
-			return ParentGenerationHelper->GetIncrementalCookAssetPackageData(*PackageData);
-		}
-	}
-	return nullptr;
-}
-
 bool FRequestCluster::FGraphSearch::FExploreEdgesContext::TryCalculateIterativelyUnmodified()
 {
 	using namespace UE::TargetDomain;
@@ -1553,36 +1537,36 @@ bool FRequestCluster::FGraphSearch::FExploreEdgesContext::TryCalculateIterativel
 		UE::Cook::FPackageArtifacts& Artifacts = QueryPlatformData.CookAttachments.Artifacts;
 		const FAssetPackageData* OverrideAssetPackageData = nullptr;
 		FPackageData* ParentPackageData = nullptr;
-		if (PackageData->IsGenerated())
+		TRefCountPtr<FGenerationHelper> GenerationHelper;
+		if (!PackageData->IsGenerated())
+		{
+			GenerationHelper = PackageData->GetGenerationHelper();
+		}
+		else
 		{
+			GenerationHelper = PackageData->GetOrFindParentGenerationHelper();
+			if (!GenerationHelper)
+			{
+				// Generated packages for which we do not have the GenerationHelper available are not iteratively
+				// skippable.
+				SetIsIterativelyUnmodified(PlatformIndex, false, PackagePlatformData);
+				continue;
+			}
+			ParentPackageData = &GenerationHelper->GetOwner();
 			// If a generator is marked iteratively unmodified, then by contract we are not required to test its
 			// generated packages; they are all marked iteratively unmodified as well
-			ParentPackageData = Cluster.PackageDatas.FindPackageDataByPackageName(
-				PackageData->GetParentGenerator());
-			if (ParentPackageData)
+			const FPackagePlatformData* ParentPlatformData =
+				ParentPackageData->GetPlatformDatas().Find(TargetPlatform);
+			if (ParentPlatformData)
 			{
-				const FPackagePlatformData* ParentPlatformData =
-					ParentPackageData->GetPlatformDatas().Find(TargetPlatform);
-				if (ParentPlatformData)
+				if (ParentPlatformData->IsIterativelyUnmodified())
 				{
-					if (ParentPlatformData->IsIterativelyUnmodified())
-					{
-						SetIsIterativelyUnmodified(PlatformIndex, true, PackagePlatformData);
-						continue;
-					}
+					SetIsIterativelyUnmodified(PlatformIndex, true, PackagePlatformData);
+					continue;
 				}
 			}
-
-			// If the generator was not marked iteratively unmodified, then we use the data provided by the generator
-			// to decide whether the generated package is iteratively unmodified.
-			OverrideAssetPackageData = Vertex->GetGeneratedAssetPackageData();
-			if (!OverrideAssetPackageData)
-			{
-				SetIsIterativelyUnmodified(PlatformIndex, false, PackagePlatformData);
-				continue;
-			}
 		}
-		if (!Artifacts.HasKeyMatch(OverrideAssetPackageData))
+		if (!Artifacts.HasKeyMatch(GenerationHelper.GetReference()))
 		{
 			SetIsIterativelyUnmodified(PlatformIndex, false, PackagePlatformData);
 			continue;
@@ -1607,10 +1591,10 @@ bool FRequestCluster::FGraphSearch::FExploreEdgesContext::TryCalculateIterativel
 		{
 			bool bAnyTransitiveBuildDependencyIsModified = false;
 			UnreadyTransitiveBuildVertices.Reset();
-			for (const UE::Cook::FCookDependency& TransitiveBuildDependency :
-				Artifacts.GetTransitiveBuildDependencies())
+			TArray<FName, TInlineAllocator<10>> TransitiveBuildDependencies;
+			Artifacts.GetTransitiveBuildDependencies(TransitiveBuildDependencies);
+			for (FName TransitiveBuildPackageName : TransitiveBuildDependencies)
 			{
-				FName TransitiveBuildPackageName = TransitiveBuildDependency.GetPackageName();
 				FVertexData& TransitiveBuildVertex = GraphSearch.Cluster.FindOrAddVertex(TransitiveBuildPackageName);
 				if (!TransitiveBuildVertex.GetPackageData())
 				{
@@ -2046,8 +2030,9 @@ void FRequestCluster::FGraphSearch::FExploreEdgesContext::ProcessPlatformAttachm
 			UE::Cook::FPackageArtifacts& Artifacts = PlatformAttachments.Artifacts;
 			if (bExploreRuntimeDependencies && PackagePlatformData.IsCookable() && Cluster.bAllowSoftDependencies)
 			{
-				AddPlatformDependencyRange(Artifacts.GetRuntimePackageDependencies(), PlatformIndex,
-					EInstigator::SoftDependency);
+				TArray<FName, TInlineAllocator<16>> RuntimeContentDependencies;
+				Artifacts.GetRuntimeContentDependencies(RuntimeContentDependencies);
+				AddPlatformDependencyRange(RuntimeContentDependencies, PlatformIndex, EInstigator::SoftDependency);
 			}
 
 			if (Cluster.bPreQueueBuildDefinitions && PackagePlatformData.IsCookable())
@@ -2057,11 +2042,11 @@ void FRequestCluster::FGraphSearch::FExploreEdgesContext::ProcessPlatformAttachm
 					PlatformAttachments.BuildDefinitions.Definitions);
 			}
 
-			for (const UE::Cook::FCookDependency& TransitiveBuildDependency : 
-				Artifacts.GetTransitiveBuildDependencies())
+			TArray<FName, TInlineAllocator<10>> TransitiveBuildDependencies;
+			Artifacts.GetTransitiveBuildDependencies(TransitiveBuildDependencies);
+			for (FName TransitivePackageName : TransitiveBuildDependencies)
 			{
-				AddPlatformDependency(TransitiveBuildDependency.GetPackageName(), PlatformIndex,
-					EInstigator::BuildDependency);
+				AddPlatformDependency(TransitivePackageName, PlatformIndex, EInstigator::BuildDependency);
 			}
 		}
 		bool bShouldIterativelySkip = bIterativelyUnmodified;
@@ -2148,18 +2133,17 @@ void FRequestCluster::FGraphSearch::FExploreEdgesContext::ProcessPlatformAttachm
 
 		if (PlatformAgnosticQueryData.bSchedulerThreadFetchCompleted)
 		{
-			bool bCanCheckHasKeyMatch = true;
 			const FAssetPackageData* OverrideAssetPackageData = nullptr;
-			if (PackageData->IsGenerated())
+			TRefCountPtr<FGenerationHelper> GenerationHelper;
+			if (!PackageData->IsGenerated())
 			{
-				OverrideAssetPackageData = Vertex->GetGeneratedAssetPackageData();
-				if (!OverrideAssetPackageData)
-				{
-					bCanCheckHasKeyMatch = false;
-				}
+				GenerationHelper = PackageData->GetGenerationHelper();
+			}
+			else
+			{
+				GenerationHelper = PackageData->GetOrFindParentGenerationHelper();
 			}
-			if (bCanCheckHasKeyMatch &&
-				PlatformAgnosticQueryData.CookAttachments.Artifacts.HasKeyMatch(OverrideAssetPackageData))
+			if (PlatformAgnosticQueryData.CookAttachments.Artifacts.HasKeyMatch(GenerationHelper.GetReference()))
 			{
 				Cluster.BuildDefinitions.AddBuildDefinitionList(PackageName, TargetPlatform,
 					PlatformAgnosticQueryData.CookAttachments.BuildDefinitions.Definitions);
@@ -2212,7 +2196,7 @@ FRequestCluster::FindOrAddVertex(FName PackageName, FGenerationHelper* ParentGen
 		if (!PackageData && ParentGenerationHelper && ICookPackageSplitter::IsUnderGeneratedPackageSubPath(NameBuffer))
 		{
 			const FAssetPackageData* PreviousPackageData =
-				ParentGenerationHelper->GetIncrementalCookAssetPackageData(PackageName);
+				ParentGenerationHelper->GetAssetPackageData(PackageName);
 			if (PreviousPackageData)
 			{
 				bool bIsMap = PreviousPackageData->Extension == EPackageExtension::Map;
diff --git a/Engine/Source/Editor/UnrealEd/Private/Cooker/CookRequestCluster.h b/Engine/Source/Editor/UnrealEd/Private/Cooker/CookRequestCluster.h
--- a/Engine/Source/Editor/UnrealEd/Private/Cooker/CookRequestCluster.h
+++ b/Engine/Source/Editor/UnrealEd/Private/Cooker/CookRequestCluster.h
@@ -174,7 +174,6 @@ private:
 		FPackageData* GetPackageData() const;
 		TArray<FVertexData*>& GetIterativelyModifiedListeners();
 		TSet<FVertexData*> GetUnreadyDependencies();
-		const FAssetPackageData* GetGeneratedAssetPackageData();
 
 		/**
 		 * Whether the package is owned by this cluster and the cluster should decide its next state.
diff --git a/Engine/Source/Editor/UnrealEd/Private/Cooker/CookSavePackage.cpp b/Engine/Source/Editor/UnrealEd/Private/Cooker/CookSavePackage.cpp
--- a/Engine/Source/Editor/UnrealEd/Private/Cooker/CookSavePackage.cpp
+++ b/Engine/Source/Editor/UnrealEd/Private/Cooker/CookSavePackage.cpp
@@ -415,28 +415,27 @@ void FSaveCookedPackageContext::FinishPlatform()
 
 	// Calculate up-to-date assetregistry data for Generator and Generated packages
 	TOptional<TArray<FAssetDependency>> OverridePackageDependencies;
-	TOptional<FGeneratedPackageResultStruct> GeneratedResult;
 	TOptional<FAssetPackageData> AssetPackageDataBuffer;
 	TOptional<FAssetPackageData> OverrideAssetPackageData;
 	const FAssetPackageData* AssetPackageData = nullptr;
+	TRefCountPtr<FGenerationHelper> GenerationHelper;
+	bool bGenerated = false;
 	if (CommitType == EReachability::Runtime)
 	{
 		CookResult = bSuccessful ? ECookResult::Succeeded : ECookResult::Failed;
 
-		if (TRefCountPtr<FGenerationHelper> GenerationHelper = PackageData.GetGenerationHelper();
-			GenerationHelper)
+		if (GenerationHelper = PackageData.GetGenerationHelper(); GenerationHelper)
 		{
 			OverridePackageDependencies.Emplace();
 			GenerationHelper->FinishGeneratorPlatformSave(PackageData, PlatformIndex == 0, *OverridePackageDependencies);
 		}
-		else if (TRefCountPtr<FGenerationHelper> ParentGenerationHelper = PackageData.GetParentGenerationHelper();
-			ParentGenerationHelper)
-		{
-			GeneratedResult.Emplace();
-			ParentGenerationHelper->FinishGeneratedPlatformSave(PackageData, *GeneratedResult);
-			OverridePackageDependencies.Emplace(GeneratedResult->PackageDependencies);
-			OverrideAssetPackageData.Emplace(GeneratedResult->AssetPackageData);
-			AssetPackageData = &GeneratedResult->AssetPackageData;
+		else if (GenerationHelper = PackageData.GetParentGenerationHelper(); GenerationHelper)
+		{
+			bGenerated = true;
+			OverrideAssetPackageData.Emplace();
+			OverridePackageDependencies.Emplace();
+			GenerationHelper->FinishGeneratedPlatformSave(PackageData, *OverrideAssetPackageData, *OverridePackageDependencies);
+			AssetPackageData = OverrideAssetPackageData.GetPtrOrNull();
 		}
 		if (!AssetPackageData)
 		{
@@ -455,7 +454,7 @@ void FSaveCookedPackageContext::FinishPlatform()
 		CalculatePlatformAgnosticRuntimeDependencies();
 		CalculatePlatformRuntimeDependencies();
 		// Note CalculateCookDependencies mutates this->SaveResult; it moves CookDependencies out of it
-		CalculateCookDependencies(GeneratedResult.GetPtrOrNull(), OverridePackageDependencies.GetPtrOrNull());
+		CalculateCookDependencies(GenerationHelper.GetReference(), bGenerated, OverridePackageDependencies.GetPtrOrNull());
 		RecordPlatformBuildDependencies();
 
 		ICookedPackageWriter::FCommitPackageInfo Info;
@@ -736,7 +735,7 @@ TArray<FName> FSaveCookedPackageContext::GetPlatformRuntimeDependencies() const
 	return PlatformDependencyNames;
 }
 
-void FSaveCookedPackageContext::CalculateCookDependencies(FGeneratedPackageResultStruct* GeneratedResult,
+void FSaveCookedPackageContext::CalculateCookDependencies(FGenerationHelper* GenerationHelper, bool bGenerated,
 	TArray<FAssetDependency>* OverrideDependencies)
 {
 	if (COTFS.bHybridIterativeEnabled)
@@ -752,7 +751,7 @@ void FSaveCookedPackageContext::CalculateCookDependencies(FGeneratedPackageResul
 			}
 		}
 		PlatformCookAttachments = FIncrementalCookAttachments::Collect(Package, TargetPlatform,
-			HasSavePackageResult() ? &SavePackageResult : nullptr, GeneratedResult,
+			HasSavePackageResult() ? &SavePackageResult : nullptr, GenerationHelper, bGenerated,
 			MoveTemp(PlatformRuntimeDependencies));
 	}
 }
@@ -762,11 +761,11 @@ void FSaveCookedPackageContext::RecordPlatformBuildDependencies()
 	if (COTFS.bHybridIterativeEnabled)
 	{
 		TSet<FPackageData*>& CurrentDependencies = PlatformDependencies[PlatformIndex].BuildDependencies;
-		for (const FCookDependency& BuildDependency :
-			PlatformCookAttachments.Artifacts.GetTransitiveBuildDependencies())
+		TArray<FName, TInlineAllocator<10>> TransitiveBuildDependencies;
+		PlatformCookAttachments.Artifacts.GetTransitiveBuildDependencies(TransitiveBuildDependencies);
+		for (FName TransitivePackageName : TransitiveBuildDependencies)
 		{
-			FPackageData* BuildPackageData = COTFS.PackageDatas->TryAddPackageDataByPackageName(
-				BuildDependency.GetPackageName());
+			FPackageData* BuildPackageData = COTFS.PackageDatas->TryAddPackageDataByPackageName(TransitivePackageName);
 			if (BuildPackageData)
 			{
 				CurrentDependencies.Add(BuildPackageData);
diff --git a/Engine/Source/Editor/UnrealEd/Private/Cooker/CookSavePackage.h b/Engine/Source/Editor/UnrealEd/Private/Cooker/CookSavePackage.h
--- a/Engine/Source/Editor/UnrealEd/Private/Cooker/CookSavePackage.h
+++ b/Engine/Source/Editor/UnrealEd/Private/Cooker/CookSavePackage.h
@@ -50,7 +50,7 @@ private: // Used only by UCookOnTheFlyServer, which has private access
 	void CalculatePlatformAgnosticRuntimeDependencies();
 	void CalculatePlatformRuntimeDependencies();
 	TArray<FName> GetPlatformRuntimeDependencies() const;
-	void CalculateCookDependencies(FGeneratedPackageResultStruct* GeneratedResult,
+	void CalculateCookDependencies(FGenerationHelper* GenerationHelper, bool bGenerated,
 		TArray<FAssetDependency>* OverrideDependencies);
 	void RecordPlatformBuildDependencies();
 	bool HasSavePackageResult() const;
diff --git a/Engine/Source/Editor/UnrealEd/Private/EditorDomain/EditorDomainUtils.cpp b/Engine/Source/Editor/UnrealEd/Private/EditorDomain/EditorDomainUtils.cpp
--- a/Engine/Source/Editor/UnrealEd/Private/EditorDomain/EditorDomainUtils.cpp
+++ b/Engine/Source/Editor/UnrealEd/Private/EditorDomain/EditorDomainUtils.cpp
@@ -1982,7 +1982,8 @@ bool TrySavePackage(UPackage* Package)
 	{
 		UE::Cook::FIncrementalCookAttachments CookAttachments =
 			UE::Cook::FIncrementalCookAttachments::Collect(
-			Package, nullptr, nullptr, nullptr, TArray<FName>());
+				Package, nullptr /* TargetPlatform */, nullptr /* SaveResult */,
+				nullptr /* GenerationHelper */, false /* bGenerated */, TArray<FName>());
 		TArray<IPackageWriter::FCommitAttachmentInfo> Attachments;
 		CookAttachments.AppendCommitAttachments(Attachments);
 		if (!Attachments.IsEmpty())
diff --git a/Engine/Source/Runtime/CoreUObject/Private/Cooker/CookDependency.cpp b/Engine/Source/Runtime/CoreUObject/Private/Cooker/CookDependency.cpp
--- a/Engine/Source/Runtime/CoreUObject/Private/Cooker/CookDependency.cpp
+++ b/Engine/Source/Runtime/CoreUObject/Private/Cooker/CookDependency.cpp
@@ -67,11 +67,20 @@ FCookDependency FCookDependency::Function(FName InFunctionName, FCbFieldIterator
 	return Result;
 }
 
+FCookDependency FCookDependency::TransitiveBuild(FName PackageName)
+{
+	FCookDependency Result(ECookDependency::TransitiveBuild);
+	Result.TransitiveBuildData.PackageName = PackageName;
+	return Result;
+}
+
 FCookDependency FCookDependency::TransitiveBuildAndRuntime(FName PackageName)
 {
 	FCookDependency Result(ECookDependency::TransitiveBuild);
 	Result.TransitiveBuildData.PackageName = PackageName;
+	PRAGMA_DISABLE_DEPRECATION_WARNINGS;
 	Result.TransitiveBuildData.bAlsoAddRuntimeDependency = true;
+	PRAGMA_ENABLE_DEPRECATION_WARNINGS;
 	return Result;
 }
 
@@ -189,6 +198,13 @@ FCookDependency FCookDependency::AssetRegistryQuery(FARFilter Filter)
 	return Result;
 }
 
+FCookDependency FCookDependency::RedirectionTarget(FName PackageName)
+{
+	FCookDependency Result(ECookDependency::RedirectionTarget);
+	Result.NameData = PackageName;
+	return Result;
+}
+
 FCookDependency::FCookDependency()
 	: Type(ECookDependency::None)
 {
@@ -251,6 +267,9 @@ FCookDependency& FCookDependency::operator=(const FCookDependency& Other)
 	case ECookDependency::AssetRegistryQuery:
 		ARFilter.Reset(Other.ARFilter.IsValid() ? new FARFilter(*Other.ARFilter) : nullptr);
 		break;
+	case ECookDependency::RedirectionTarget:
+		NameData = Other.NameData;
+		break;
 	default:
 		checkNoEntry();
 		break;
@@ -290,6 +309,9 @@ FCookDependency& FCookDependency::operator=(FCookDependency&& Other)
 	case ECookDependency::AssetRegistryQuery:
 		ARFilter = MoveTemp(Other.ARFilter);
 		break;
+	case ECookDependency::RedirectionTarget:
+		NameData = MoveTemp(Other.NameData);
+		break;
 	default:
 		checkNoEntry();
 		break;
@@ -297,6 +319,12 @@ FCookDependency& FCookDependency::operator=(FCookDependency&& Other)
 	return *this;
 }
 
+UE::ConfigAccessTracking::FConfigAccessData FCookDependency::GetConfigAccessData() const
+{
+	return Type == ECookDependency::Config && ConfigAccessData.IsValid() ?
+		*ConfigAccessData : UE::ConfigAccessTracking::FConfigAccessData();
+}
+
 FString FCookDependency::GetConfigPath() const
 {
 	return Type == ECookDependency::Config && ConfigAccessData.IsValid() ?
@@ -419,6 +447,11 @@ void FCookDependency::UpdateHash(FCookDependencyContext& Context) const
 			}
 		}
 		return;
+	case ECookDependency::RedirectionTarget:
+		Context.LogError(FString::Printf(
+			TEXT("FCookDependency::RedirectionTarget('%s') failed to UpdateHash: RedirectionTarget dependencies do not implement UpdateHash and it should not be called on them."),
+			*NameData.ToString()));
+		return;
 	default:
 		checkNoEntry();
 		return;
@@ -456,6 +489,9 @@ void FCookDependency::Construct()
 		// Set the union's bytes equal to a TUniquePtr with a null internal pointer
 		new(&ARFilter) TUniquePtr<FARFilter>(nullptr);
 		break;
+	case ECookDependency::RedirectionTarget:
+		new(&NameData) FName();
+		break;
 	default:
 		checkNoEntry();
 		break;
@@ -493,6 +529,9 @@ void FCookDependency::Destruct()
 		// TUniquePtr's destructor will also destruct its internal pointer
 		ARFilter.~TUniquePtr<FARFilter>();
 		break;
+	case ECookDependency::RedirectionTarget:
+		NameData.~FName();
+		break;
 	default:
 		checkNoEntry();
 		break;
@@ -521,7 +560,6 @@ void FCookDependency::Save(FCbWriter& Writer) const
 		break;
 	case ECookDependency::TransitiveBuild:
 		Writer << TransitiveBuildData.PackageName;
-		Writer << TransitiveBuildData.bAlsoAddRuntimeDependency;
 		break;
 	case ECookDependency::Package:
 		Writer << NameData;
@@ -542,6 +580,9 @@ void FCookDependency::Save(FCbWriter& Writer) const
 		}
 		break;
 	}
+	case ECookDependency::RedirectionTarget:
+		Writer << NameData;
+		break;
 	default:
 		checkNoEntry();
 		break;
@@ -597,13 +638,7 @@ bool FCookDependency::Load(FCbFieldView Value)
 		{
 			return false;
 		}
-		bool bLocalAlsoAddRuntimeDependency = Field.AsBool();
-		if ((Field++).HasError())
-		{
-			return false;
-		}
-		*this = FCookDependency::TransitiveBuildAndRuntime(LocalPackageName);
-		TransitiveBuildData.bAlsoAddRuntimeDependency = bLocalAlsoAddRuntimeDependency;
+		*this = FCookDependency::TransitiveBuild(LocalPackageName);
 		return true;
 	}
 	case ECookDependency::Package:
@@ -681,6 +716,16 @@ bool FCookDependency::Load(FCbFieldView Value)
 		}
 		return true;
 	}
+	case ECookDependency::RedirectionTarget:
+	{
+		FName LocalPackageName;
+		if (!LoadFromCompactBinary(Field++, LocalPackageName))
+		{
+			return false;
+		}
+		*this = FCookDependency::RedirectionTarget(LocalPackageName);
+		return true;
+	}
 	default:
 		break;
 	}
diff --git a/Engine/Source/Runtime/CoreUObject/Public/Cooker/CookDependency.h b/Engine/Source/Runtime/CoreUObject/Public/Cooker/CookDependency.h
--- a/Engine/Source/Runtime/CoreUObject/Public/Cooker/CookDependency.h
+++ b/Engine/Source/Runtime/CoreUObject/Public/Cooker/CookDependency.h
@@ -133,6 +133,7 @@ enum class ECookDependency : uint8
 	SettingsObject			= 7,
 	NativeClass				= 8,
 	AssetRegistryQuery		= 9,
+	RedirectionTarget		= 10,
 
 	Count,
 };
@@ -196,11 +197,9 @@ public:
 	 * Create a transitive build dependency on another package. In an incremental cook if the other package was not
 	 * cooked in a previous cook session, or its previous cook result was invalidated, the current package will also
 	 * have its cook result invalidated.
-	 *
-	 * This version of the function also adds a runtime dependency - the requested package will be staged for the
-	 * current platform. Adding a transitive build dependency without adding a runtime dependency is not yet supported
-	 * due to limitations in the cooker.
 	 */
+	COREUOBJECT_API static FCookDependency TransitiveBuild(FName PackageName);
+	UE_DEPRECATED(5.6, "Add a TransitiveBuild and a Runtime dependency separately.")
 	COREUOBJECT_API static FCookDependency TransitiveBuildAndRuntime(FName PackageName);
 
 	/**
@@ -236,6 +235,9 @@ public:
 	COREUOBJECT_API static FCookDependency NativeClass(const UClass* InClass);
 	COREUOBJECT_API static FCookDependency NativeClass(FStringView ClassPath);
 
+	/** Adds a dependency on the ObjectRedirectors and CoreRedirects that affect the given PackageName. */
+	COREUOBJECT_API static FCookDependency RedirectionTarget(FName PackageName);
+
 	/** Adds a dependency on the results reported by an AssetRegistry query. */
 	COREUOBJECT_API static FCookDependency AssetRegistryQuery(FARFilter Filter);
 
@@ -261,9 +263,11 @@ public:
 
 	/** PackageName if GetType() == TransitiveBuild or GetType() == Package, else NAME_None. */
 	FName GetPackageName() const;
-	/** If GetType() == TransitiveBuild, whether AlsoAddRuntimeDependency was selected, otherwise false. */
+	UE_DEPRECATED(5.6, "This function is only needed for TransitiveBuildAndRuntime dependencies, which are deprecated.")
 	bool IsAlsoAddRuntimeDependency() const;
 
+	/** Returns the ConfigAccess in its struct form (File,Section,Key) if GetType() == Config, else empty. */
+	COREUOBJECT_API UE::ConfigAccessTracking::FConfigAccessData GetConfigAccessData() const;
 	/**
 	 * Returns the full path of the config access (e.g. Platform.Filename.Section.ValueName)
 	 * if GetType() == Config, else empty.
@@ -332,7 +336,7 @@ private:
 	struct FTransitiveBuildData
 	{
 		FName PackageName;
-		bool bAlsoAddRuntimeDependency = true;
+		bool bAlsoAddRuntimeDependency = false;
 	};
 	union
 	{
@@ -462,6 +466,8 @@ inline FName FCookDependency::GetPackageName() const
 		return TransitiveBuildData.PackageName;
 	case ECookDependency::Package:
 		return NameData;
+	case ECookDependency::RedirectionTarget:
+		return NameData;
 	default:
 		return NAME_None;
 	};
@@ -470,7 +476,9 @@ inline FName FCookDependency::GetPackageName() const
 
 inline bool FCookDependency::IsAlsoAddRuntimeDependency() const
 {
+	PRAGMA_DISABLE_DEPRECATION_WARNINGS;
 	return Type == ECookDependency::TransitiveBuild ? TransitiveBuildData.bAlsoAddRuntimeDependency : false;
+	PRAGMA_ENABLE_DEPRECATION_WARNINGS;
 }
 
 inline const UObject* FCookDependency::GetSettingsObject() const
@@ -528,10 +536,6 @@ inline bool FCookDependency::operator<(const FCookDependency& Other) const
 		{
 			return Compare < 0;
 		}
-		if (TransitiveBuildData.bAlsoAddRuntimeDependency != Other.TransitiveBuildData.bAlsoAddRuntimeDependency)
-		{
-			return TransitiveBuildData.bAlsoAddRuntimeDependency == false;
-		}
 		return false;
 	}
 	case ECookDependency::Package:
@@ -559,6 +563,8 @@ inline bool FCookDependency::operator<(const FCookDependency& Other) const
 			return false; // equal
 		}
 		return ARFilterLessThan(*ARFilter, *Other.ARFilter);
+	case ECookDependency::RedirectionTarget:
+		return NameData.Compare(Other.NameData) < 0;
 	case ECookDependency::Count:
 	default:
 		checkNoEntry();
@@ -605,7 +611,7 @@ inline bool FCookDependency::operator==(const FCookDependency& Other) const
 		{
 			return false;
 		}
-		return TransitiveBuildData.bAlsoAddRuntimeDependency == Other.TransitiveBuildData.bAlsoAddRuntimeDependency;
+		return true;
 	}
 	case ECookDependency::Package:
 		return NameData.Compare(Other.NameData) == 0;
@@ -635,6 +641,8 @@ inline bool FCookDependency::operator==(const FCookDependency& Other) const
 
 		return ARFilterEqual(*ARFilter, *Other.ARFilter);
 	}
+	case ECookDependency::RedirectionTarget:
+		return NameData.Compare(Other.NameData) == 0;
 	case ECookDependency::Count:
 	default:
 		checkNoEntry();
diff --git a/Engine/Source/Runtime/Engine/Private/Materials/MaterialInstance.cpp b/Engine/Source/Runtime/Engine/Private/Materials/MaterialInstance.cpp
--- a/Engine/Source/Runtime/Engine/Private/Materials/MaterialInstance.cpp
+++ b/Engine/Source/Runtime/Engine/Private/Materials/MaterialInstance.cpp
@@ -5070,7 +5070,7 @@ void UMaterialInstance::PreSave(FObjectPreSaveContext ObjectSaveContext)
 		UMaterialInterface* EffectiveParent = Parent ? Parent : UMaterial::GetDefaultMaterial(MD_Surface);
 		if (EffectiveParent && EffectiveParent->GetPackage() != GetPackage())
 		{
-			ObjectSaveContext.AddCookBuildDependency(UE::Cook::FCookDependency::TransitiveBuildAndRuntime(
+			ObjectSaveContext.AddCookBuildDependency(UE::Cook::FCookDependency::TransitiveBuild(
 				EffectiveParent->GetPackage()->GetFName()));
 		}
 	}
