PR #12423: Allow custom ISM editor proxy objects

diff --git a/Engine/Source/Editor/UnrealEd/Private/Elements/SMInstance/SMInstanceElementDetailsInterface.cpp b/Engine/Source/Editor/UnrealEd/Private/Elements/SMInstance/SMInstanceElementDetailsInterface.cpp
--- a/Engine/Source/Editor/UnrealEd/Private/Elements/SMInstance/SMInstanceElementDetailsInterface.cpp
+++ b/Engine/Source/Editor/UnrealEd/Private/Elements/SMInstance/SMInstanceElementDetailsInterface.cpp
@@ -3,14 +3,23 @@
 #include "Elements/SMInstance/SMInstanceElementDetailsInterface.h"
 #include "Elements/SMInstance/SMInstanceElementDetailsProxyObject.h"
 #include "Elements/SMInstance/SMInstanceElementData.h"
+#include "Elements/SMInstance/SMInstanceManager.h"
 #include "Components/InstancedStaticMeshComponent.h"
 
 class FSMInstanceTypedElementDetailsObject : public ITypedElementDetailsObject
 {
 public:
-	explicit FSMInstanceTypedElementDetailsObject(const FSMInstanceElementId& InSMInstanceElementId)
+	explicit FSMInstanceTypedElementDetailsObject(const FSMInstanceElementId& InSMInstanceElementId, const ISMInstanceManager* InInstanceManager)
 	{
-		USMInstanceElementDetailsProxyObject* InstanceProxyObjectPtr = NewObject<USMInstanceElementDetailsProxyObject>();
+		USMInstanceProxyEditingObject* InstanceProxyObjectPtr = nullptr;
+		if (UClass* ProxyObjectClass = InInstanceManager->GetSMInstanceEditingProxyClass())
+		{
+			InstanceProxyObjectPtr = NewObject<USMInstanceProxyEditingObject>((UObject*)GetTransientPackage(), ProxyObjectClass);
+		}
+		else
+		{
+			InstanceProxyObjectPtr = NewObject<USMInstanceElementDetailsProxyObject>();
+		}
 		InstanceProxyObjectPtr->Initialize(InSMInstanceElementId);
 
 		InstanceProxyObject = InstanceProxyObjectPtr;
@@ -18,7 +27,7 @@ public:
 
 	~FSMInstanceTypedElementDetailsObject()
 	{
-		if (USMInstanceElementDetailsProxyObject* InstanceProxyObjectPtr = InstanceProxyObject.Get())
+		if (USMInstanceProxyEditingObject* InstanceProxyObjectPtr = InstanceProxyObject.Get())
 		{
 			InstanceProxyObjectPtr->Shutdown();
 		}
@@ -31,7 +40,7 @@ public:
 
 	virtual void AddReferencedObjects(FReferenceCollector& Collector) override
 	{
-		if (USMInstanceElementDetailsProxyObject* InstanceProxyObjectPtr = InstanceProxyObject.Get())
+		if (USMInstanceProxyEditingObject* InstanceProxyObjectPtr = InstanceProxyObject.Get())
 		{
 			Collector.AddReferencedObject(InstanceProxyObject);
 			InstanceProxyObject = InstanceProxyObjectPtr;
@@ -39,14 +48,15 @@ public:
 	}
 
 private:
-	TWeakObjectPtr<USMInstanceElementDetailsProxyObject> InstanceProxyObject;
+	TWeakObjectPtr<USMInstanceProxyEditingObject> InstanceProxyObject;
 };
 
 TUniquePtr<ITypedElementDetailsObject> USMInstanceElementDetailsInterface::GetDetailsObject(const FTypedElementHandle& InElementHandle)
 {
 	if (const FSMInstanceElementData* SMInstanceElement = InElementHandle.GetData<FSMInstanceElementData>())
 	{
-		return MakeUnique<FSMInstanceTypedElementDetailsObject>(SMInstanceElement->InstanceElementId);
+		FSMInstanceManager InstanceManager = SMInstanceElementDataUtil::GetSMInstanceFromHandleChecked(InElementHandle);
+		return MakeUnique<FSMInstanceTypedElementDetailsObject>(SMInstanceElement->InstanceElementId, InstanceManager.GetInstanceManager());
 	}
 	return nullptr;
 }
diff --git a/Engine/Source/Editor/UnrealEd/Private/Elements/SMInstance/SMInstanceElementDetailsProxyObject.h b/Engine/Source/Editor/UnrealEd/Private/Elements/SMInstance/SMInstanceElementDetailsProxyObject.h
--- a/Engine/Source/Editor/UnrealEd/Private/Elements/SMInstance/SMInstanceElementDetailsProxyObject.h
+++ b/Engine/Source/Editor/UnrealEd/Private/Elements/SMInstance/SMInstanceElementDetailsProxyObject.h
@@ -12,13 +12,13 @@
 #include "SMInstanceElementDetailsProxyObject.generated.h"
 
 UCLASS(Transient)
-class USMInstanceElementDetailsProxyObject : public UObject
+class USMInstanceElementDetailsProxyObject : public USMInstanceProxyEditingObject
 {
 	GENERATED_BODY()
 
 public:
-	void Initialize(const FSMInstanceElementId& InSMInstanceElementId);
-	void Shutdown();
+	virtual void Initialize(const FSMInstanceElementId& InSMInstanceElementId) override;
+	virtual void Shutdown() override;
 
 	//~ UObject interface
 	virtual void BeginDestroy() override;
diff --git a/Engine/Source/Runtime/Engine/Public/Elements/SMInstance/SMInstanceManager.h b/Engine/Source/Runtime/Engine/Public/Elements/SMInstance/SMInstanceManager.h
--- a/Engine/Source/Runtime/Engine/Public/Elements/SMInstance/SMInstanceManager.h
+++ b/Engine/Source/Runtime/Engine/Public/Elements/SMInstance/SMInstanceManager.h
@@ -6,9 +6,11 @@
 #include "UObject/Interface.h"
 #include "Containers/ArrayView.h"
 #include "Elements/SMInstance/SMInstanceElementId.h"
+#include "Templates/SubclassOf.h"
 #include "SMInstanceManager.generated.h"
 
 enum class ETypedElementWorldType : uint8;
+class USMInstanceProxyEditingObject;
 
 /**
  * An interface for actors that manage static mesh instances.
@@ -138,6 +140,27 @@ public:
 	 * @return True if any instances were duplicated, false otherwise.
 	 */
 	virtual bool DuplicateSMInstances(TArrayView<const FSMInstanceId> InstanceIds, TArray<FSMInstanceId>& OutNewInstanceIds) = 0;
+
+	/**
+	* Optionally specify a different class for a proxy object than the default one.
+	* This can be useful if a full UObject is needed to represent the SM Instance
+	*/
+	virtual TSubclassOf<USMInstanceProxyEditingObject> GetSMInstanceEditingProxyClass() const
+	{
+		return nullptr;
+	}
+};
+
+UCLASS(MinimalAPI, Abstract)
+class USMInstanceProxyEditingObject : public UObject
+{
+	GENERATED_BODY()
+
+public:
+	virtual void Initialize(const FSMInstanceElementId& InSMInstanceElementId) {}
+	virtual void Shutdown() {}
+private:
+	virtual bool SyncProxyStateFromInstance() { return false; }
 };
 
 /**
