Apply world size clamping to camera location in FViewportCameraTransform::SetLocation, only if AWorldSettings::bEnableWorldBoundsChecks is enabled.

diff --git a/Engine/Source/Editor/UnrealEd/Private/EditorViewportClient.cpp b/Engine/Source/Editor/UnrealEd/Private/EditorViewportClient.cpp
--- a/Engine/Source/Editor/UnrealEd/Private/EditorViewportClient.cpp
+++ b/Engine/Source/Editor/UnrealEd/Private/EditorViewportClient.cpp
@@ -265,11 +265,12 @@ FViewportCameraTransform::FViewportCameraTransform()
 	, LookAt( FVector::ZeroVector )
 	, StartLocation( FVector::ZeroVector )
 	, OrthoZoom( DEFAULT_ORTHOZOOM )
+	, MaxLocation( UE_DOUBLE_HUGE_DISTANCE )
 {}
 
 void FViewportCameraTransform::SetLocation( const FVector& Position )
 {
-	FVector ClampedPosition = Position.BoundToCube(WORLD_MAX);
+	FVector ClampedPosition = Position.BoundToCube(MaxLocation);
 	ViewLocation = ClampedPosition;
 	DesiredLocation = ViewLocation;
 }
@@ -549,6 +550,19 @@ FEditorViewportClient::FEditorViewportClient(FEditorModeTools* InModeTools, FPre
 	FSlateApplication::Get().OnWindowDPIScaleChanged().AddRaw(this, &FEditorViewportClient::HandleWindowDPIScaleChanged);
 
 	RegisterDragTools();
+
+	double MaxViewLocation = HALF_WORLD_MAX;
+	if(GetScene() != nullptr && GetScene()->GetWorld() != nullptr)
+	{
+		AWorldSettings* WorldSettings = GetScene()->GetWorld()->GetWorldSettings(false);
+		bool bEnableWorldBoundsChecks = WorldSettings ? WorldSettings->bEnableWorldBoundsChecks : true;
+		if (!bEnableWorldBoundsChecks)
+		{
+			MaxViewLocation = UE_DOUBLE_HUGE_DISTANCE;
+		}
+	}
+	ViewTransformPerspective.SetMaxLocation(MaxViewLocation);
+	ViewTransformOrthographic.SetMaxLocation(MaxViewLocation);
 }
 
 FEditorViewportClient::~FEditorViewportClient()
diff --git a/Engine/Source/Editor/UnrealEd/Public/EditorViewportClient.h b/Engine/Source/Editor/UnrealEd/Public/EditorViewportClient.h
--- a/Engine/Source/Editor/UnrealEd/Public/EditorViewportClient.h
+++ b/Engine/Source/Editor/UnrealEd/Public/EditorViewportClient.h
@@ -244,6 +244,11 @@ public:
 		OrthoZoom = InOrthoZoom;
 	}
 
+	void SetMaxLocation( double InMaxLocation )
+	{
+		MaxLocation = InMaxLocation;
+	}
+
 	/** Check if transition curve is playing. */
 	UNREALED_API bool IsPlaying();
 
@@ -297,6 +302,8 @@ private:
 	FVector StartLocation;
 	/** Ortho zoom amount */
 	float OrthoZoom;
+	/** Location is clamped to a box around the origin with this radius */
+	double MaxLocation;
 };
 
 /** Parameter struct for editor viewport view modifiers */
