More long term fix for calling GetCppStructOps on a dead Verse type in the transient package.
We need to keep the GetCppStructOps check because cycling the struct ops for some types is a problem.

#rnx

diff --git a/Engine/Source/Editor/UnrealEd/Private/Kismet2/ReloadUtilities.cpp b/Engine/Source/Editor/UnrealEd/Private/Kismet2/ReloadUtilities.cpp
--- a/Engine/Source/Editor/UnrealEd/Private/Kismet2/ReloadUtilities.cpp
+++ b/Engine/Source/Editor/UnrealEd/Private/Kismet2/ReloadUtilities.cpp
@@ -1000,19 +1000,19 @@ void FReload::NotifyChange(UPackage* New, UPackage* Old)
 		++PackageStats.New;
 	}
 
-	Packages.AddUnique(New);
+	Packages.Add(New);
 }
 
 namespace
 {
 	template<typename T>
-	void CollectPackages(TArray<UPackage*>& Packages, const TMap<T*, T*>& Reinstances)
+	void CollectPackages(TSet<UPackage*>& Packages, const TMap<T*, T*>& Reinstances)
 	{
 		for (const TPair<T*, T*>& Pair : Reinstances)
 		{
 			T* Old = Pair.Key;
 			T* New = Pair.Value;
-			Packages.AddUnique(New ? New->GetPackage() : Old->GetPackage());
+			Packages.Add(New ? New->GetPackage() : Old->GetPackage());
 		}
 	}
 }
@@ -1035,6 +1035,7 @@ void FReload::Reinstance()
 	}
 
 	// Remap all native functions (and gather scriptstructs)
+	UPackage* TransientPackage = GetTransientPackage();
 	TArray<UScriptStruct*> ScriptStructs;
 	for (FRawObjectIterator It; It; ++It)
 	{
@@ -1048,8 +1049,13 @@ void FReload::Reinstance()
 		} 
 		else if (UScriptStruct* ScriptStruct = Cast<UScriptStruct>(static_cast<UObject*>(It->Object)))
 		{
-			if (!ScriptStruct->HasAnyFlags(RF_ClassDefaultObject) && 
-				Packages.ContainsByPredicate([ScriptStruct](UPackage* Package) { return ScriptStruct->IsIn(Package); }))
+			// IMPORTANT: The order here is very important.  There is a chance that GetCppStructOps will fail for 
+			// types that are in a package we don't care about.  So we have to filter on the package first.
+			UPackage* ScriptStructPackage = ScriptStruct->GetPackage();
+			if (!ScriptStruct->HasAnyFlags(RF_ClassDefaultObject) &&
+				ScriptStructPackage != TransientPackage &&
+				Packages.Contains(ScriptStruct->GetPackage()) &&
+				ScriptStruct->GetCppStructOps())
 			{
 				ScriptStructs.Add(ScriptStruct);
 			}
diff --git a/Engine/Source/Editor/UnrealEd/Public/Kismet2/ReloadUtilities.h b/Engine/Source/Editor/UnrealEd/Public/Kismet2/ReloadUtilities.h
--- a/Engine/Source/Editor/UnrealEd/Public/Kismet2/ReloadUtilities.h
+++ b/Engine/Source/Editor/UnrealEd/Public/Kismet2/ReloadUtilities.h
@@ -142,7 +142,7 @@ private:
 	const TCHAR* Prefix = nullptr;
 
 	/** List of packages affected by the reload */
-	TArray<UPackage*> Packages;
+	TSet<UPackage*> Packages;
 	
 	/** Output device for any logging */
 	FOutputDevice& Ar;
diff --git a/Engine/Source/Runtime/CoreUObject/Public/UObject/Class.h b/Engine/Source/Runtime/CoreUObject/Public/UObject/Class.h
--- a/Engine/Source/Runtime/CoreUObject/Public/UObject/Class.h
+++ b/Engine/Source/Runtime/CoreUObject/Public/UObject/Class.h
@@ -1793,7 +1793,7 @@ public:
 	/** Returns the CppStructOps that can be used to do custom operations */
 	FORCEINLINE ICppStructOps* GetCppStructOps() const
 	{
-		checkf(bPrepareCppStructOpsCompleted, TEXT("GetCppStructOps: PrepareCppStructOps() has not been called for class %s"), *GetName());
+		checkf(bPrepareCppStructOpsCompleted, TEXT("GetCppStructOps: PrepareCppStructOps() has not been called for class %s"), *GetPathName());
 		return CppStructOps;
 	}
 
