In  FPackageArtifacts::Collect, error out if a package can't be found in the asset registry and in memory.

#rnx
#incrementalcook

diff --git a/Engine/Source/Editor/UnrealEd/Private/Cooker/CookPackageArtifacts.cpp b/Engine/Source/Editor/UnrealEd/Private/Cooker/CookPackageArtifacts.cpp
--- a/Engine/Source/Editor/UnrealEd/Private/Cooker/CookPackageArtifacts.cpp
+++ b/Engine/Source/Editor/UnrealEd/Private/Cooker/CookPackageArtifacts.cpp
@@ -530,25 +530,32 @@ FPackageArtifacts FPackageArtifacts::Collect(UPackage* Package, const ITargetPla
 				continue;
 			}
 
-			TStringBuilder<256> DependencyNameStr(InPlace, DependencyName);
-			UPackage* DependencyPackage = FindPackage(nullptr, *DependencyNameStr);
-			if (!DependencyPackage)
+			bool bPackageExistOnDisk = AssetRegistry->DoesPackageExistOnDisk(DependencyName);
+			if (!bPackageExistOnDisk)
 			{
-				if (OutMessages)
+				TStringBuilder<256> PackagePath(EInPlace::InPlace, DependencyName);
+				UPackage* PackageObject = FindPackage(nullptr, *PackagePath); // The package doesn't exist on disk, maybe it's an in-memory package.
+				if (!PackageObject)
 				{
-					OutMessages->Emplace(ELogVerbosity::Error,
-						FString::Printf(TEXT("Package %s does not exist."), *DependencyNameStr));
-				}
+					// The package does not exist anywhere so error out.
+					if (OutMessages)
+					{
+						TStringBuilder<256> ErrorText;
+						ErrorText << TEXT("Package ") << PackagePath << TEXT(" does not exist.");
+						OutMessages->Emplace(ELogVerbosity::Error, FString(*ErrorText));
+					}
 
-				return FPackageArtifacts();
-			}
+					return FPackageArtifacts();
+				}
+				else if (PackageObject->HasAnyPackageFlags(PKG_NewlyCreated))
+				{
+					// If the package is a newly created package (in-memory package) then ignore it.
+					// In-memory packages are ignored because we can't compute their digest. Only packages on disk have a digest.
+					Iter.RemoveCurrentSwap();
+					continue;
+				}
 
-			if (DependencyPackage->HasAnyPackageFlags(PKG_NewlyCreated))
-			{
-				// If the package is a newly created package (in-memory package) then ignore it.
-				// In-memory packages are ignored because we can't compute their digest. Only packages on disk have a digest.
-				Iter.RemoveCurrentSwap();
-				continue;
+				// Here the package is not on disk, in-memory and not newly created. It's a strange edge case but let's register it to the dependencies.
 			}
 
 			// Package dependencies of all kinds (Runtime, Build, TransitiveBuild) also cause RedirectionTarget
