IncrementalCook: Refactor CookPackageArtifacts to store dependencies and hashes in an FBuildDependencySet, and add the placeholder FBuildDependencySet for Load.
#rnx
#incrementalcook

diff --git a/Engine/Source/Editor/UnrealEd/Private/Cooker/CookPackageArtifacts.cpp b/Engine/Source/Editor/UnrealEd/Private/Cooker/CookPackageArtifacts.cpp
--- a/Engine/Source/Editor/UnrealEd/Private/Cooker/CookPackageArtifacts.cpp
+++ b/Engine/Source/Editor/UnrealEd/Private/Cooker/CookPackageArtifacts.cpp
@@ -34,33 +34,12 @@ FGuid TargetDomainVersion(0x10E1BAE8, 0x6CA04033, 0xA7E57A37, 0xD967F493);
 
 // Bump PackageArtifactsVersion version when the serialization of PackageArtifacts has changed and we want to add
 // backwards compatibility rather than invalidating everything.
-constexpr uint32 PackageArtifactsVersion = 0x00000003;
+constexpr uint32 PackageArtifactsVersion = 0x00000004;
 
 static const FUtf8StringView PackageArtifactsAttachmentKey = UTF8TEXTVIEW("CookPackageArtifacts");
 static const FUtf8StringView BuildDefinitionsAttachmentKey = UTF8TEXTVIEW("BuildDefinitionsAttachmentKey");
 
-// Constructor/Destructor defined here in cpp rather than header so we can 
-// avoid needing the definition of FCookDependency in the header; it is needed
-// for construct/destruct of TArray<FCookDependency>.
-FPackageArtifacts::FPackageArtifacts() = default;
-FPackageArtifacts::~FPackageArtifacts() = default;
-FPackageArtifacts::FPackageArtifacts(const FPackageArtifacts&) = default;
-FPackageArtifacts::FPackageArtifacts(FPackageArtifacts&&) = default;
-FPackageArtifacts& FPackageArtifacts::operator=(const FPackageArtifacts&) = default;
-FPackageArtifacts& FPackageArtifacts::operator=(FPackageArtifacts&&) = default;
-
-
-bool FPackageArtifacts::IsValid() const
-{
-	return bValid;
-}
-
-bool FPackageArtifacts::HasSaveResults() const
-{
-	return bHasSaveResults;
-}
-
-bool FPackageArtifacts::HasKeyMatch(FGenerationHelper* GenerationHelper)
+bool FBuildDependencySet::HasKeyMatch(FName PackageName, FGenerationHelper* GenerationHelper)
 {
 	if (!bValid)
 	{
@@ -72,7 +51,7 @@ bool FPackageArtifacts::HasKeyMatch(FGenerationHelper* GenerationHelper)
 	}
 	if (CurrentKey.IsZero())
 	{
-		ECurrentKeyResult Result = TryCalculateCurrentKey(GenerationHelper);
+		ECurrentKeyResult Result = TryCalculateCurrentKey(PackageName, GenerationHelper);
 		if (Result != ECurrentKeyResult::Success)
 		{
 			CurrentKey.Reset();
@@ -82,7 +61,7 @@ bool FPackageArtifacts::HasKeyMatch(FGenerationHelper* GenerationHelper)
 	return CurrentKey == StoredKey;
 }
 
-FPackageArtifacts::ECurrentKeyResult FPackageArtifacts::TryCalculateCurrentKey(
+FBuildDependencySet::ECurrentKeyResult FBuildDependencySet::TryCalculateCurrentKey(FName PackageName,
 	FGenerationHelper* GenerationHelper, TArray<TPair<ELogVerbosity::Type, FString>>* OutMessages)
 {
 	IAssetRegistry* AssetRegistry = IAssetRegistry::Get();
@@ -106,7 +85,7 @@ FPackageArtifacts::ECurrentKeyResult FPackageArtifacts::TryCalculateCurrentKey(
 	KeyBuilder.Update(&TargetDomainVersion, sizeof(TargetDomainVersion));
 
 	ECurrentKeyResult Result = ECurrentKeyResult::Success;
-	UE::Cook::FCookDependencyContext Context(&KeyBuilder,
+	FCookDependencyContext Context(&KeyBuilder,
 		[&Result, OutMessages](ELogVerbosity::Type Verbosity, FString&& Message)
 		{
 			if (OutMessages)
@@ -126,16 +105,16 @@ FPackageArtifacts::ECurrentKeyResult FPackageArtifacts::TryCalculateCurrentKey(
 	// function sorts CookDependencies of the same type together. This allows us
 	// to create batches for CookDependencies with Update functions that benefit
 	// from being updated in batches.
-	int32 NumDependencies = BuildDependencies.Num();
+	int32 NumDependencies = Dependencies.Num();
 	for (int32 BatchStart = 0; BatchStart < NumDependencies; )
 	{
-		ECookDependency BatchCategory = BuildDependencies[BatchStart].GetType();
+		ECookDependency BatchCategory = Dependencies[BatchStart].GetType();
 		int32 BatchEnd = BatchStart + 1;
-		while (BatchEnd < NumDependencies && BuildDependencies[BatchEnd].GetType() == BatchCategory)
+		while (BatchEnd < NumDependencies && Dependencies[BatchEnd].GetType() == BatchCategory)
 		{
 			++BatchEnd;
 		}
-		TConstArrayView<FCookDependency> Batch(BuildDependencies.GetData() + BatchStart, BatchEnd - BatchStart);
+		TConstArrayView<FCookDependency> Batch(Dependencies.GetData() + BatchStart, BatchEnd - BatchStart);
 		BatchStart += Batch.Num();
 
 		// Some CookDependency types can not handle Update being called, because their Update relies on functions
@@ -235,21 +214,130 @@ FPackageArtifacts::ECurrentKeyResult FPackageArtifacts::TryCalculateCurrentKey(
 	return Result;
 }
 
-void FPackageArtifacts::Reset()
+void FBuildDependencySet::Empty()
 {
-	BuildDependencies.Reset();
-	RuntimeDependencies.Reset();
-	PackageName = FName();
+	Dependencies.Empty();
 	StoredKey = FIoHash::Zero;
 	CurrentKey = FIoHash::Zero;
 	bValid = false;
 }
 
+bool FBuildDependencySet::TryLoad(FCbFieldView InFieldView)
+{
+	Empty();
+
+	for (FCbFieldViewIterator FieldView(InFieldView.CreateViewIterator()); FieldView; )
+	{
+		const FCbFieldViewIterator Last = FieldView;
+		if (FieldView.GetName().Equals(UTF8TEXTVIEW("Name")))
+		{
+			if (!LoadFromCompactBinary(FieldView++, Name))
+			{
+				return false;
+			}
+		}
+		if (FieldView.GetName().Equals(UTF8TEXTVIEW("StoredKey")))
+		{
+			if (!LoadFromCompactBinary(FieldView++, StoredKey))
+			{
+				return false;
+			}
+		}
+		if (FieldView.GetName().Equals(UTF8TEXTVIEW("Dependencies")))
+		{
+			if (!LoadFromCompactBinary(FieldView++, Dependencies))
+			{
+				return false;
+			}
+		}
+		if (FieldView == Last)
+		{
+			++FieldView;
+		}
+	}
+	bValid = true;
+	return true;
+}
+
+void FBuildDependencySet::Save(FCbWriter& Writer) const
+{
+	Writer.BeginObject();
+	Writer << "Name" << Name;
+	Writer << "StoredKey" << StoredKey;
+	if (!Dependencies.IsEmpty())
+	{
+		Writer << "Dependencies" << Dependencies;
+	}
+	Writer.EndObject();
+}
+
+FName FPackageArtifacts::NAME_Save(TEXT("Save"));
+FName FPackageArtifacts::NAME_Load(TEXT("Load"));
+
+FPackageArtifacts::FPackageArtifacts()
+{
+	LoadBuildDependencies.SetName(NAME_Load);
+	SaveBuildDependencies.SetName(NAME_Save);
+}
+
+FBuildDependencySet& FPackageArtifacts::FindOrAddBuildDependencySet(FName ResultName)
+{
+	if (ResultName == NAME_Save)
+	{
+		return SaveBuildDependencies;
+	}
+	else if (ResultName == NAME_Load)
+	{
+		return LoadBuildDependencies;
+	}
+	else
+	{
+		// Not yet implemented
+		check(false);
+		return SaveBuildDependencies;
+	}
+}
+
+FBuildDependencySet* FPackageArtifacts::FindBuildDependencySet(FName ResultName)
+{
+	if (ResultName == NAME_Save)
+	{
+		return &SaveBuildDependencies;
+	}
+	else if (ResultName == NAME_Load)
+	{
+		return &LoadBuildDependencies;
+	}
+	else
+	{
+		return nullptr;
+	}
+}
+
+bool FPackageArtifacts::HasSaveResults() const
+{
+	return bHasSaveResults;
+}
+
+bool FPackageArtifacts::HasKeyMatch(FGenerationHelper* GenerationHelper)
+{
+	return bValid && SaveBuildDependencies.HasKeyMatch(PackageName, GenerationHelper);
+}
+
+FBuildDependencySet::ECurrentKeyResult FPackageArtifacts::TryCalculateCurrentKey(
+	FGenerationHelper* GenerationHelper, TArray<TPair<ELogVerbosity::Type, FString>>* OutMessages)
+{
+	return SaveBuildDependencies.TryCalculateCurrentKey(PackageName, GenerationHelper, OutMessages);
+}
+
 void FPackageArtifacts::Empty()
 {
-	Reset();
-	BuildDependencies.Empty();
+	SaveBuildDependencies.Empty();
+	LoadBuildDependencies.Empty();
 	RuntimeDependencies.Empty();
+	PackageName = FName();
+	bHasSaveResults = false;
+	bValid = false;
 }
 
 enum class EPackageMountPoint
@@ -328,10 +416,11 @@ FPackageArtifacts FPackageArtifacts::Collect(UPackage* Package, const ITargetPla
 	FPackageArtifacts Result;
 	// Skip the multiple reallocations for an array that grows from 0 to 128, but then reallocate according
 	// to normal TArray growth.
-	Result.BuildDependencies.Reserve(128);
+	TArray<FCookDependency> BuildDependencies;
+	BuildDependencies.Reserve(128);
 	Result.PackageName = Package->GetFName();
 	Result.bHasSaveResults = SaveResult != nullptr;
-	Result.BuildDependencies.Add(FCookDependency::Package(Result.PackageName));
+	BuildDependencies.Add(FCookDependency::Package(Result.PackageName));
 	Result.RuntimeDependencies = MoveTemp(InRuntimeDependencies);
 
 	TArray<FName> AssetDependencies;
@@ -364,27 +453,27 @@ FPackageArtifacts FPackageArtifacts::Collect(UPackage* Package, const ITargetPla
 					// We have not yet enabled marking the auto-added dependencies from TObjectPtr resolve
 					// as transitive, because it causes a performance regression and we are still working on fixing the
 					// regression.
-					Result.BuildDependencies.Add(FCookDependency::Package(AccessData.ReferencedPackage));
+					BuildDependencies.Add(FCookDependency::Package(AccessData.ReferencedPackage));
 				}
 				else if (ProjectionList.bHasAll)
 				{
 					// TObjectPtr UE::BuildProjection::All dependencies are added as transitive build dependencies.
 					// We have to do this to be conservative, since we do not know which bytes from the target are
 					// dependended upon and which of the target's build dependencies influence those bytes.
-					Result.BuildDependencies.Add(FCookDependency::TransitiveBuild(AccessData.ReferencedPackage));
+					BuildDependencies.Add(FCookDependency::TransitiveBuild(AccessData.ReferencedPackage));
 				}
 				else
 				{
 					for (FTopLevelAssetPath ClassPath : ProjectionList.Classes)
 					{
-						Result.BuildDependencies.Add(FCookDependency::NativeClass(WriteToString<256>(ClassPath)));
+						BuildDependencies.Add(FCookDependency::NativeClass(WriteToString<256>(ClassPath)));
 					}
 
 					for (FName ResultProjection : ProjectionList.ResultProjections)
 					{
 						if (ResultProjection == UE::Cook::ResultProjection::PackageAndClass)
 						{
-							Result.BuildDependencies.Add(FCookDependency::Package(AccessData.ReferencedPackage));
+							BuildDependencies.Add(FCookDependency::Package(AccessData.ReferencedPackage));
 						}
 						else
 						{
@@ -405,7 +494,7 @@ FPackageArtifacts FPackageArtifacts::Collect(UPackage* Package, const ITargetPla
 		// dependencies.
 		for (FName AssetDependency : AssetDependencies)
 		{
-			Result.BuildDependencies.Add(FCookDependency::TransitiveBuild(AssetDependency));
+			BuildDependencies.Add(FCookDependency::TransitiveBuild(AssetDependency));
 		}
 	}
 
@@ -418,17 +507,17 @@ FPackageArtifacts FPackageArtifacts::Collect(UPackage* Package, const ITargetPla
 			TArray<FConfigAccessData> ConfigKeys = ConfigTracker.GetPackageRecords(Result.PackageName, TargetPlatform);
 			for (const FConfigAccessData& ConfigKey : ConfigKeys)
 			{
-				Result.BuildDependencies.Add(FCookDependency::Config(ConfigKey));
+				BuildDependencies.Add(FCookDependency::Config(ConfigKey));
 			}
 		}
 	}
 #endif
 	if (SaveResult)
 	{
-		Result.BuildDependencies.Append(SaveResult->CookDependencies);
+		BuildDependencies.Append(SaveResult->CookDependencies);
 		for (FName SoftPackageReference : SaveResult->UntrackedSoftPackageReferences)
 		{
-			Result.BuildDependencies.Add(FCookDependency::RedirectionTarget(SoftPackageReference));
+			BuildDependencies.Add(FCookDependency::RedirectionTarget(SoftPackageReference));
 		}
 	}
 
@@ -437,7 +526,7 @@ FPackageArtifacts FPackageArtifacts::Collect(UPackage* Package, const ITargetPla
 	// Settings dependencies - Expand transitive dependencies on SettingsObjects into the list of dependencies
 	// recorded for that settings object.
 	TSet<const UObject*> SettingsDependencies;
-	Result.BuildDependencies.RemoveAllSwap([&SettingsDependencies](const FCookDependency& Dependency)
+	BuildDependencies.RemoveAllSwap([&SettingsDependencies](const FCookDependency& Dependency)
 		{
 			if (Dependency.GetType() == ECookDependency::SettingsObject)
 			{
@@ -458,11 +547,11 @@ FPackageArtifacts FPackageArtifacts::Collect(UPackage* Package, const ITargetPla
 			FCookDependencyGroups::Get().FindOrCreate(reinterpret_cast<UPTRINT>(SettingsObject));
 		if (!IncludeDependencies.bInitialized)
 		{
-			IncludeDependencies.Artifacts = CollectSettingsObject(SettingsObject,
+			IncludeDependencies.Dependencies = CollectSettingsObject(SettingsObject,
 				&IncludeDependencies.Messages);
 			IncludeDependencies.bInitialized = true;
 		}
-		if (!IncludeDependencies.Artifacts.IsValid())
+		if (!IncludeDependencies.Dependencies.IsValid())
 		{
 			if (OutMessages)
 			{
@@ -475,7 +564,7 @@ FPackageArtifacts FPackageArtifacts::Collect(UPackage* Package, const ITargetPla
 			return FPackageArtifacts();
 		}
 
-		for (const FCookDependency& IncludeDependency : IncludeDependencies.Artifacts.GetBuildDependencies())
+		for (const FCookDependency& IncludeDependency : IncludeDependencies.Dependencies.GetDependencies())
 		{
 			// Recursive SettingsDependencies are not allowed. We haven't needed them yet, and not supporting them allows
 			// prevents the need for cycle detection.
@@ -489,13 +578,13 @@ FPackageArtifacts FPackageArtifacts::Collect(UPackage* Package, const ITargetPla
 				}
 				return FPackageArtifacts();
 			}
-			Result.BuildDependencies.Add(IncludeDependency);
+			BuildDependencies.Add(IncludeDependency);
 		}
 	}
 
 	// Process some rules for Package dependencies
 	TSet<FName> RedirectionTargets;
-	for (TArray<FCookDependency>::TIterator Iter(Result.BuildDependencies); Iter; ++Iter)
+	for (TArray<FCookDependency>::TIterator Iter(BuildDependencies); Iter; ++Iter)
 	{
 		FCookDependency& Dependency = *Iter;
 		if (Dependency.GetType() == ECookDependency::TransitiveBuild ||
@@ -602,27 +691,30 @@ FPackageArtifacts FPackageArtifacts::Collect(UPackage* Package, const ITargetPla
 	// Put all the extra redirection dependencies into BuildDependencies
 	for (FName RedirectionTarget : RedirectionTargets)
 	{
-		Result.BuildDependencies.Add(FCookDependency::RedirectionTarget(RedirectionTarget));
+		BuildDependencies.Add(FCookDependency::RedirectionTarget(RedirectionTarget));
 	}
 
 	// Sort and remove duplicates
-	Algo::Sort(Result.BuildDependencies);
-	Result.BuildDependencies.SetNum(Algo::Unique(Result.BuildDependencies), EAllowShrinking::Yes);
+	Algo::Sort(BuildDependencies);
+	BuildDependencies.SetNum(Algo::Unique(BuildDependencies), EAllowShrinking::Yes);
 	Algo::Sort(Result.RuntimeDependencies, FNameLexicalLess());
 	Result.RuntimeDependencies.SetNum(Algo::Unique(Result.RuntimeDependencies), EAllowShrinking::Yes);
 
-	ECurrentKeyResult CurrentKeyResult = Result.TryCalculateCurrentKey(GenerationHelper, OutMessages);
-	if (CurrentKeyResult == ECurrentKeyResult::Error)
+	Result.SaveBuildDependencies.SetNormalizedDependencies(MoveTemp(BuildDependencies));
+	FBuildDependencySet::ECurrentKeyResult
+		CurrentKeyResult = Result.TryCalculateCurrentKey(GenerationHelper, OutMessages);
+	if (CurrentKeyResult == FBuildDependencySet::ECurrentKeyResult::Error)
 	{
 		return FPackageArtifacts();
 	}
-	Result.StoredKey = Result.CurrentKey;
+	Result.SaveBuildDependencies.StoreCurrentKey();
+	Result.SaveBuildDependencies.SetValid(true);
 	Result.bValid = true;
 
 	return Result;
 }
 
-FPackageArtifacts FPackageArtifacts::CollectSettingsObject(const UObject* Object,
+FBuildDependencySet FPackageArtifacts::CollectSettingsObject(const UObject* Object,
 	TArray<TPair<ELogVerbosity::Type, FString>>* OutMessages)
 {
 	if (!Object)
@@ -631,7 +723,7 @@ FPackageArtifacts FPackageArtifacts::CollectSettingsObject(const UObject* Object
 		{
 			OutMessages->Emplace(ELogVerbosity::Error, TEXT("Invalid null Object."));
 		}
-		return FPackageArtifacts();
+		return FBuildDependencySet();
 	}
 
 	UClass* Class = Object->GetClass();
@@ -642,7 +734,7 @@ FPackageArtifacts FPackageArtifacts::CollectSettingsObject(const UObject* Object
 			OutMessages->Emplace(ELogVerbosity::Error, FString::Printf(
 				TEXT("Class %s is not a config class."), *Class->GetPathName()));
 		}
-		return FPackageArtifacts();
+		return FBuildDependencySet();
 	}
 	if (!Class->HasAnyClassFlags(CLASS_PerObjectConfig) && Object != Class->GetDefaultObject())
 	{
@@ -651,17 +743,17 @@ FPackageArtifacts FPackageArtifacts::CollectSettingsObject(const UObject* Object
 			OutMessages->Emplace(ELogVerbosity::Error, FString::Printf(
 				TEXT("Class %s is not a per-object-config class."), *Class->GetPathName()));
 		}
-		return FPackageArtifacts();
+		return FBuildDependencySet();
 	}
 
-	FPackageArtifacts Result;
+	TArray<FCookDependency> BuildDependencies;
 	TArray<UE::ConfigAccessTracking::FConfigAccessData> ConfigDatas;
 	const_cast<UObject*>(Object)->LoadConfig(nullptr /* ConfigClass */, nullptr /* Filename */,
 		UE::LCPF_None /* PropagationFlags */, nullptr /* PropertyToLoad */, &ConfigDatas);
-	Result.BuildDependencies.Reserve(ConfigDatas.Num() + 1);
+	BuildDependencies.Reserve(ConfigDatas.Num() + 1);
 	for (const UE::ConfigAccessTracking::FConfigAccessData& ConfigData : ConfigDatas)
 	{
-		Result.BuildDependencies.Add(UE::Cook::FCookDependency::Config(ConfigData));
+		BuildDependencies.Add(UE::Cook::FCookDependency::Config(ConfigData));
 	}
 
 	// In addition to adding the config dependencies, add a dependency on the class schema. If the current class has
@@ -674,12 +766,14 @@ FPackageArtifacts FPackageArtifacts::CollectSettingsObject(const UObject* Object
 	}
 	if (NativeClass)
 	{
-		Result.BuildDependencies.Add(UE::Cook::FCookDependency::NativeClass(NativeClass));
+		BuildDependencies.Add(UE::Cook::FCookDependency::NativeClass(NativeClass));
 	}
 
-	Algo::Sort(Result.BuildDependencies);
-	Result.BuildDependencies.SetNum(Algo::Unique(Result.BuildDependencies));
-	Result.bValid = true;
+	Algo::Sort(BuildDependencies);
+	BuildDependencies.SetNum(Algo::Unique(BuildDependencies));
+	FBuildDependencySet Result;
+	Result.SetNormalizedDependencies(MoveTemp(BuildDependencies));
+	Result.SetValid(true);
 	return Result;
 }
 
@@ -689,7 +783,7 @@ bool LoadFromCompactBinary(FCbObjectView ObjectView, UE::Cook::FPackageArtifacts
 {
 	using namespace UE::Cook;
 
-	Artifacts.Reset();
+	Artifacts.Empty();
 	int32 Version = -1;
 
 	for (FCbFieldViewIterator FieldView(ObjectView.CreateViewIterator()); FieldView; )
@@ -703,23 +797,23 @@ bool LoadFromCompactBinary(FCbObjectView ObjectView, UE::Cook::FPackageArtifacts
 				return false;
 			}
 		}
-		if (FieldView.GetName().Equals(UTF8TEXTVIEW("StoredKey")))
+		if (FieldView.GetName().Equals(UTF8TEXTVIEW("HasSaveResults")))
 		{
-			if (!LoadFromCompactBinary(FieldView++, Artifacts.StoredKey))
+			if (!LoadFromCompactBinary(FieldView++, Artifacts.bHasSaveResults))
 			{
 				return false;
 			}
 		}
-		if (FieldView.GetName().Equals(UTF8TEXTVIEW("HasSaveResults")))
+		if (FieldView.GetName().Equals(UTF8TEXTVIEW("SaveBuildDependencies")))
 		{
-			if (!LoadFromCompactBinary(FieldView++, Artifacts.bHasSaveResults))
+			if (!LoadFromCompactBinary(FieldView++, Artifacts.SaveBuildDependencies))
 			{
 				return false;
 			}
 		}
-		if (FieldView.GetName().Equals(UTF8TEXTVIEW("BuildDependencies")))
+		if (FieldView.GetName().Equals(UTF8TEXTVIEW("LoadBuildDependencies")))
 		{
-			if (!LoadFromCompactBinary(FieldView++, Artifacts.BuildDependencies))
+			if (!LoadFromCompactBinary(FieldView++, Artifacts.LoadBuildDependencies))
 			{
 				return false;
 			}
@@ -750,11 +844,14 @@ FCbWriter& operator<<(FCbWriter& Writer, const UE::Cook::FPackageArtifacts& Arti
 
 	Writer.BeginObject();
 	Writer << "Version" << PackageArtifactsVersion;
-	Writer << "StoredKey" << Artifacts.StoredKey;
 	Writer << "HasSaveResults" << Artifacts.bHasSaveResults;
-	if (!Artifacts.BuildDependencies.IsEmpty())
+	if (Artifacts.SaveBuildDependencies.IsValid())
+	{
+		Writer << "SaveBuildDependencies" << Artifacts.SaveBuildDependencies;
+	}
+	if (Artifacts.LoadBuildDependencies.IsValid())
 	{
-		Writer << "BuildDependencies" << Artifacts.BuildDependencies;
+		Writer << "LoadBuildDependencies" << Artifacts.LoadBuildDependencies;
 	}
 	if (!Artifacts.RuntimeDependencies.IsEmpty())
 	{
@@ -802,11 +899,6 @@ FBuildDefinitionList FBuildDefinitionList::Collect(UPackage* Package, const ITar
 	return Result;
 }
 
-void FBuildDefinitionList::Reset()
-{
-	Definitions.Reset();
-}
-
 void FBuildDefinitionList::Empty()
 {
 	Definitions.Empty();
@@ -858,12 +950,6 @@ FCbWriter& operator<<(FCbWriter& Writer, const UE::Cook::FBuildDefinitionList& D
 namespace UE::Cook
 {
 
-void FIncrementalCookAttachments::Reset()
-{
-	Artifacts.Reset();
-	BuildDefinitions.Reset();
-}
-
 void FIncrementalCookAttachments::Empty()
 {
 	Artifacts.Empty();
diff --git a/Engine/Source/Editor/UnrealEd/Private/Cooker/CookPackageArtifacts.h b/Engine/Source/Editor/UnrealEd/Private/Cooker/CookPackageArtifacts.h
--- a/Engine/Source/Editor/UnrealEd/Private/Cooker/CookPackageArtifacts.h
+++ b/Engine/Source/Editor/UnrealEd/Private/Cooker/CookPackageArtifacts.h
@@ -39,5 +39,115 @@
 {
 
 /**
+ * A list of dependencies that affect the build of a BuildResult. BuildResults can be a package load, a package save,
+ * or a system-specific set of data that is produced alongside the package load or save.
+ *
+ * In the build operation for the cook of a package, the load of the package is recorded as a BuildResult with no
+ * payload (it has an implicit payload which is the loaded package as it exists in memory, but we do not store that
+ * payload in the oplog). The dependencies for that BuildResult are the most commonly used source for transitive build
+ * dependencies.
+ * 
+ * The second BuildResult in a package's cook is the bytes of the saved package. That BuildResult stores some of its
+ * payload - the package bytes - as a special payload which is not stored in the BuildResult itself, but rather is
+ * stored externally in the oplog. It stores the rest of its payload - the runtime dependencies - in the
+ * FPackageArtifacts. The BuildDependencies of the save BuildResult are used to decide whether the package can be
+ * incrementally skipped.
+ *
+ * System-specific BuildResults are not saved during the cook and each one must be recalculated on demand during the
+ * cook by each dependent package that incorporates their data, and stored in the package data of the owning package
+ * until it gets garbage collected and the data has to be recreated. The dependencies of that operation are stored in
+ * the FPackageArtifacts for the owning package and can be used as a build dependency for the owning package and the
+ * dependent packages.
+ *
+ * TODO: Add a system to preserve build results along with the dependencies.
+ */
+class FBuildDependencySet
+{
+public:
+	/**
+	 * True if the structure has been calculated or Set since last reset. False if the stucture is default or has
+	 * been reset.
+	 */
+	bool IsValid() const;
+
+	/** Name of the build result owning the DependencySet, used for lookup by transitive build dependencies. */
+	FName GetName() const;
+	void SetName(FName InName);
+
+	/** The list of dependencies in the BuildDependencySet. */
+	const TArray<FCookDependency>& GetDependencies() const;
+	/**
+	 * When constructing the set, the dependencies must be normalized - rules processed, sorted, made unique
+	 * by the caller before being set into the BuildDependencySet.
+	 */
+	void SetNormalizedDependencies(TArray<FCookDependency> InDependencies);
+	/** Copy CurrentKey into StoredKey, called after SetNormalizedDependencies before Save. */
+	void StoreCurrentKey();
+	/** Sets IsValid to the given argument, called after all values have been written by the Caller before Save */
+	void SetValid(bool bInValid);
+
+	/** Helper function to filter GetDependencies for the transitive build dependencies. */
+	template <typename AllocatorType>
+	void GetTransitiveDependencies(TArray<FName, AllocatorType>& OutDependencies) const;
+
+	/** Return the Key that was hashed from the BuildDependencySet in the cookprocess that created it. */
+	const FIoHash& GetStoredKey() const;
+	/**
+	 * Return the Key that was hashed from the BuildDependencySet in the current cookprocess.
+	 * Will be the zero-hash if not yet calculated.
+	 */
+	const FIoHash& GetCurrentKey() const;
+	/** Call TryCalculateCurrentKey if not yet called, and return whether StoredKey == CurrentKey. */
+	bool HasKeyMatch(FName PackageName, FGenerationHelper* GenerationHelper);
+
+	enum class ECurrentKeyResult
+	{
+		Success,
+		Invalidated,
+		Error,
+	};
+	/**
+	 * Calculate the current key(s) from the Dependencies and store it in GetCurrentKey().
+	 *
+	 * @param GenerationHelper - If non-null, provides the lookup for the AssetPackageData of generated packages.
+	 *                           Must be provided if any generated packages are in the dependencies.
+	 */
+	ECurrentKeyResult TryCalculateCurrentKey(FName PackageName, FGenerationHelper* GenerationHelper,
+		TArray<TPair<ELogVerbosity::Type, FString>>* OutMessages = nullptr);
+
+	/** Clear data (except Name) and free memory. */
+	void Empty();
+
+	bool TryLoad(FCbFieldView FieldView);
+	void Save(FCbWriter& Writer) const;
+
+	// Hidden friend operators
+private:
+	friend bool LoadFromCompactBinary(FCbFieldView FieldView, UE::Cook::FBuildDependencySet& DependencySet)
+	{
+		return DependencySet.TryLoad(FieldView);
+	}
+	friend FCbWriter& operator<<(FCbWriter& Writer, const UE::Cook::FBuildDependencySet& DependencySet)
+	{
+		DependencySet.Save(Writer);
+		return Writer;
+	}
+
+private:
+	/* Name used to look up the BuildResult for transitive dependencies and data derived from it. */
+	FName Name;
+	/**
+	 * The dependencies that impact the creation of the BuildResult that owns this set.
+	 * These dependencies are normalized and sorted before storage.
+	 */
+	TArray<UE::Cook::FCookDependency> Dependencies;
+	/** The hash of the dependencies that was calculated in the cooksession that created the BuildResult. */
+	FIoHash StoredKey;
+	/** The hash of the dependencies that was calculated during the current cook session. */
+	FIoHash CurrentKey;
+	bool bValid = false;
+};
+
+/**
  * Non-runtime data recorded about each package and stored in the cook Oplog as attachments to the package.
  * Includes BuildResults built from the package that can be used for future incremental cooks, and the dependencies
@@ -55,42 +165,29 @@
 {
 public:
 	FPackageArtifacts();
-	~FPackageArtifacts();
-	FPackageArtifacts(const FPackageArtifacts&);
-	FPackageArtifacts(FPackageArtifacts&&);
-	FPackageArtifacts& operator=(const FPackageArtifacts&);
-	FPackageArtifacts& operator=(FPackageArtifacts&&);
 
-	
-	const TArray<FCookDependency>& GetBuildDependencies() const;
-	/** Get all of the runtime dependencies reported by the package, both script and content. */
-	const TArray<FName>& GetRuntimeDependencies() const;
-	/** Return runtime dependencies reported by the package that are content packages; script packages are removed. */
-	template <typename AllocatorType>
-	void GetRuntimeContentDependencies(TArray<FName, AllocatorType>& OutDependencies) const;
-	/** Helper function to filter GetBuildDependencies for the transitive build dependencies. */
-	template <typename AllocatorType>
-	void GetTransitiveBuildDependencies(TArray<FName, AllocatorType>& OutDependencies) const;
-
-	// Data about the Key and Package that are not dependencies
 	/**
 	 * True if the structure has been calculated or fetched and accurately reports dependencies and
 	 * key for the package. False if the stucture is default, has been reset, or was marked invalid.
 	 */
 	bool IsValid() const;
+
+	FBuildDependencySet& FindOrAddBuildDependencySet(FName ResultName);
+	FBuildDependencySet* FindBuildDependencySet(FName ResultName);
+	/** The fixed name used for the SavePackage BuildResult. */
+	static FName NAME_Save;
+	/** The fixed name used for the LoadPackage BuildResult. */
+	static FName NAME_Load;
+
+	/** Get all of the runtime dependencies reported by the package, both script and content. */
+	const TArray<FName>& GetRuntimeDependencies() const;
+	/** Return runtime dependencies reported by the package that are content packages; script packages are removed. */
+	template <typename AllocatorType>
+	void GetRuntimeContentDependencies(TArray<FName, AllocatorType>& OutDependencies) const;
+
 	bool HasSaveResults() const;
 	FName GetPackageName() const;
-	const FIoHash& GetStoredKey() const;
-	const FIoHash& GetCurrentKey() const;
-	bool HasKeyMatch(FGenerationHelper* GenerationHelper);
 
-	// Modifying the structure data
-	enum class ECurrentKeyResult
-	{
-		Success,
-		Invalidated,
-		Error,
-	};
 	/**
 	 * Calculate the current key(s) from the BuildDependencies stored on this PackageArtifacts, and store it in
 	 * GetCurrentKey().
@@ -98,9 +195,8 @@
 	 * @param GenerationHelper - If non-null, provides the lookup for the AssetPackageData of generated packages.
 	 *                           Must be provided if any generated packages are in the dependencies.
 	 */
-	ECurrentKeyResult TryCalculateCurrentKey(FGenerationHelper* GenerationHelper,
+	FBuildDependencySet::ECurrentKeyResult TryCalculateCurrentKey(FGenerationHelper* GenerationHelper,
 		TArray<TPair<ELogVerbosity::Type, FString>>* OutMessages = nullptr);
-	void Reset();
 	void Empty();
 
 	// Construction functions
@@ -114,18 +210,25 @@
 	static FPackageArtifacts Collect(UPackage* Package, const ITargetPlatform* TargetPlatform,
 		FSavePackageResultStruct* SaveResult, FGenerationHelper* GenerationHelper, bool bGenerated,
 		TArray<FName>&& RuntimeDependencies, TArray<TPair<ELogVerbosity::Type, FString>>* OutMessages = nullptr);
-	static FPackageArtifacts CollectSettingsObject(const UObject* Object,
+	static FBuildDependencySet CollectSettingsObject(const UObject* Object,
 		TArray<TPair<ELogVerbosity::Type, FString>>* OutMessages);
 
 	// Fetch function to load the dependencies from a PackageStore is not yet implemented independently for
 	// this structure. Use FIncrementalCookAttachments instead.
 
+	// Legacy API before FBuildDependencySet. TODO: Change callers to use FindBuildResult. */
+	const TArray<FCookDependency>& GetBuildDependencies() const;
+	template <typename AllocatorType>
+	void GetTransitiveBuildDependencies(TArray<FName, AllocatorType>& OutDependencies) const;
+	const FIoHash& GetStoredKey() const;
+	const FIoHash& GetCurrentKey() const;
+	bool HasKeyMatch(FGenerationHelper* GenerationHelper);
+
 private:
-	TArray<UE::Cook::FCookDependency> BuildDependencies;
+	FBuildDependencySet LoadBuildDependencies;
+	FBuildDependencySet SaveBuildDependencies;
 	TArray<FName> RuntimeDependencies;
 	FName PackageName;
-	FIoHash StoredKey;
-	FIoHash CurrentKey;
 	bool bHasSaveResults = false;
 	bool bValid = false;
 
@@ -147,7 +250,7 @@
 public:
 	struct FRecordedDependencies
 	{
-		FPackageArtifacts Artifacts;
+		FBuildDependencySet Dependencies;
 		TArray<TPair<ELogVerbosity::Type, FString>> Messages;
 		bool bInitialized = false;
 	};
@@ -165,7 +268,6 @@
 public:
 	TArray<UE::DerivedData::FBuildDefinition> Definitions;
 
-	void Reset();
 	void Empty();
 
 	/** Collect DDC BuildDefinitions that were issued from the load/save of the given package and platform. */
@@ -184,7 +286,6 @@
 	FBuildDefinitionList BuildDefinitions;
 	IPackageWriter::ECommitStatus CommitStatus = IPackageWriter::ECommitStatus::NotCommitted;
 
-	void Reset();
 	void Empty();
 	void AppendCommitAttachments(TArray<IPackageWriter::FCommitAttachmentInfo>& OutAttachments);
 
@@ -206,6 +307,68 @@
 namespace UE::Cook
 {
 
+inline bool FBuildDependencySet::IsValid() const
+{
+	return bValid;
+}
+
+inline FName FBuildDependencySet::GetName() const
+{
+	return Name;
+}
+
+inline void FBuildDependencySet::SetName(FName InName)
+{
+	Name = InName;
+}
+
+inline const TArray<FCookDependency>& FBuildDependencySet::GetDependencies() const
+{
+	return Dependencies;
+}
+
+inline void FBuildDependencySet::SetNormalizedDependencies(TArray<FCookDependency> InDependencies)
+{
+	Dependencies = MoveTemp(InDependencies);
+}
+
+inline void FBuildDependencySet::StoreCurrentKey()
+{
+	StoredKey = CurrentKey;
+}
+
+inline void FBuildDependencySet::SetValid(bool bInValid)
+{
+	bValid = bInValid;
+}
+
+template <typename AllocatorType>
+inline void FBuildDependencySet::GetTransitiveDependencies(TArray<FName, AllocatorType>& OutDependencies) const
+{
+	for (const FCookDependency& BuildDependency : Dependencies)
+	{
+		if (BuildDependency.GetType() == ECookDependency::TransitiveBuild)
+		{
+			OutDependencies.Add(BuildDependency.GetPackageName());
+		}
+	}
+}
+
+inline const FIoHash& FBuildDependencySet::GetStoredKey() const
+{
+	return StoredKey;
+}
+
+inline const FIoHash& FBuildDependencySet::GetCurrentKey() const
+{
+	return CurrentKey;
+}
+
+inline bool FPackageArtifacts::IsValid() const
+{
+	return bValid;
+}
+
 inline const TArray<FName>& FPackageArtifacts::GetRuntimeDependencies() const
 {
 	return RuntimeDependencies;
@@ -223,20 +386,3 @@
 		}
 	}
 }
-
-inline const TArray<FCookDependency>& FPackageArtifacts::GetBuildDependencies() const
-{
-	return BuildDependencies;
-}
-
-template <typename AllocatorType>
-inline void FPackageArtifacts::GetTransitiveBuildDependencies(TArray<FName, AllocatorType>& OutDependencies) const
-{
-	for (const FCookDependency& BuildDependency : BuildDependencies)
-	{
-		if (BuildDependency.GetType() == ECookDependency::TransitiveBuild)
-		{
-			OutDependencies.Add(BuildDependency.GetPackageName());
-		}
-	}
-}
@@ -246,14 +392,25 @@
 	return PackageName;
 }
 
+inline const TArray<FCookDependency>& FPackageArtifacts::GetBuildDependencies() const
+{
+	return SaveBuildDependencies.GetDependencies();
+}
+
+template <typename AllocatorType>
+inline void FPackageArtifacts::GetTransitiveBuildDependencies(TArray<FName, AllocatorType>& OutDependencies) const
+{
+	return SaveBuildDependencies.GetTransitiveDependencies(OutDependencies);
+}
+
 inline const FIoHash& FPackageArtifacts::GetStoredKey() const
 {
-	return StoredKey;
+	return SaveBuildDependencies.GetStoredKey();
 }
 
 inline const FIoHash& FPackageArtifacts::GetCurrentKey() const
 {
-	return CurrentKey;
+	return SaveBuildDependencies.GetCurrentKey();
 }
 
 } // namespace UE::Cook
