Add an AxisDisplayInfo abstraction to the engine for getting axes display names and colors.

diff --git a/Engine/Config/BaseEditorPerProjectUserSettings.ini b/Engine/Config/BaseEditorPerProjectUserSettings.ini
--- a/Engine/Config/BaseEditorPerProjectUserSettings.ini
+++ b/Engine/Config/BaseEditorPerProjectUserSettings.ini
@@ -83,6 +83,11 @@ bEnableUserEditorLayoutManagement=True
 ColorVisionDeficiencyPreviewType=NormalVision
 ColorVisionDeficiencySeverity=3
 
+; Default axis colors
+XAxisColor=(R=0.594,G=0.0197,B=0.0,A=1.0)
+YAxisColor=(R=0.1349,G=0.3959,B=0.0,A=1.0)
+ZAxisColor=(R=0.0251,G=0.207,B=0.85,A=1.0)
+
 ; If true, all toolbars will use small icons without labels
 bUseSmallToolBarIcons=False
 ;If true, Variables names in the script editor are displayed in a sanitized format
diff --git a/Engine/Source/Editor/DetailCustomizations/Private/MathStructProxyCustomizations.cpp b/Engine/Source/Editor/DetailCustomizations/Private/MathStructProxyCustomizations.cpp
--- a/Engine/Source/Editor/DetailCustomizations/Private/MathStructProxyCustomizations.cpp
+++ b/Engine/Source/Editor/DetailCustomizations/Private/MathStructProxyCustomizations.cpp
@@ -24,6 +24,7 @@
 #include "Math/UnrealMathSSE.h"
 #include "Math/VectorRegister.h"
 #include "Misc/AssertionMacros.h"
+#include "Misc/AxisDisplayInfo.h"
 #include "PropertyHandle.h"
 #include "ScopedTransaction.h"
 #include "SlotBase.h"
@@ -184,17 +185,17 @@ void FMatrixStructCustomization<T>::CustomizeLocation(TSharedRef<class IProperty
 		+ SHorizontalBox::Slot()
 		.Padding(FMargin(0.0f, 2.0f, 3.0f, 2.0f))
 		[
-			MakeNumericProxyWidget<UE::Math::TVector<T>, T>(StructPropertyHandle, CachedTranslationX, LOCTEXT("TranslationX", "X"), false, SNumericEntryBox<T>::RedLabelBackgroundColor)
+			MakeNumericProxyWidget<UE::Math::TVector<T>, T>(StructPropertyHandle, CachedTranslationX, LOCTEXT("TranslationX", "X"), false, AxisDisplayInfo::GetAxisColor(EAxisList::X))
 		]
 		+ SHorizontalBox::Slot()
 		.Padding(FMargin(0.0f, 2.0f, 3.0f, 2.0f))
 		[
-			MakeNumericProxyWidget<UE::Math::TVector<T>, T>(StructPropertyHandle, CachedTranslationY, LOCTEXT("TranslationY", "Y"), false, SNumericEntryBox<T>::GreenLabelBackgroundColor)
+			MakeNumericProxyWidget<UE::Math::TVector<T>, T>(StructPropertyHandle, CachedTranslationY, LOCTEXT("TranslationY", "Y"), false, AxisDisplayInfo::GetAxisColor(EAxisList::Y))
 		]
 		+ SHorizontalBox::Slot()
 		.Padding(FMargin(0.0f, 2.0f, 0.0f, 2.0f))
 		[
-			MakeNumericProxyWidget<UE::Math::TVector<T>, T>(StructPropertyHandle, CachedTranslationZ, LOCTEXT("TranslationZ", "Z"), false, SNumericEntryBox<T>::BlueLabelBackgroundColor)
+			MakeNumericProxyWidget<UE::Math::TVector<T>, T>(StructPropertyHandle, CachedTranslationZ, LOCTEXT("TranslationZ", "Z"), false, AxisDisplayInfo::GetAxisColor(EAxisList::Z))
 		]
 	];
 }
@@ -224,17 +225,17 @@ void FMatrixStructCustomization<T>::CustomizeRotation(TSharedRef<class IProperty
 		+ SHorizontalBox::Slot()
 		.Padding(FMargin(0.0f, 2.0f, 3.0f, 2.0f))
 		[
-			MakeNumericProxyWidget<UE::Math::TRotator<T>, T>(StructPropertyHandle, CachedRotationRoll, LOCTEXT("RotationRoll", "X"), true, SNumericEntryBox<T>::RedLabelBackgroundColor)
+			MakeNumericProxyWidget<UE::Math::TRotator<T>, T>(StructPropertyHandle, CachedRotationRoll, LOCTEXT("RotationRoll", "X"), true, AxisDisplayInfo::GetAxisColor(EAxisList::X))
 		]
 		+ SHorizontalBox::Slot()
 		.Padding(FMargin(0.0f, 2.0f, 3.0f, 2.0f))
 		[
-			MakeNumericProxyWidget<UE::Math::TRotator<T>, T>(StructPropertyHandle, CachedRotationPitch, LOCTEXT("RotationPitch", "Y"), true, SNumericEntryBox<T>::GreenLabelBackgroundColor)
+			MakeNumericProxyWidget<UE::Math::TRotator<T>, T>(StructPropertyHandle, CachedRotationPitch, LOCTEXT("RotationPitch", "Y"), true, AxisDisplayInfo::GetAxisColor(EAxisList::Y))
 		]
 		+ SHorizontalBox::Slot()
 		.Padding(FMargin(0.0f, 2.0f, 0.0f, 2.0f))
 		[
-			MakeNumericProxyWidget<UE::Math::TRotator<T>, T>(StructPropertyHandle, CachedRotationYaw, LOCTEXT("RotationYaw", "Z"), true, SNumericEntryBox<T>::BlueLabelBackgroundColor)
+			MakeNumericProxyWidget<UE::Math::TRotator<T>, T>(StructPropertyHandle, CachedRotationYaw, LOCTEXT("RotationYaw", "Z"), true, AxisDisplayInfo::GetAxisColor(EAxisList::Z))
 		]
 	];
 }
@@ -264,17 +265,17 @@ void FMatrixStructCustomization<T>::CustomizeScale(TSharedRef<class IPropertyHan
 		+ SHorizontalBox::Slot()
 		.Padding(FMargin(0.0f, 2.0f, 3.0f, 2.0f))
 		[
-			MakeNumericProxyWidget<UE::Math::TVector<T>, T>(StructPropertyHandle, CachedScaleX, LOCTEXT("ScaleX", "X"), false, SNumericEntryBox<T>::RedLabelBackgroundColor)
+			MakeNumericProxyWidget<UE::Math::TVector<T>, T>(StructPropertyHandle, CachedScaleX, LOCTEXT("ScaleX", "X"), false, AxisDisplayInfo::GetAxisColor(EAxisList::X))
 		]
 		+ SHorizontalBox::Slot()
 		.Padding(FMargin(0.0f, 2.0f, 3.0f, 2.0f))
 		[
-			MakeNumericProxyWidget<UE::Math::TVector<T>, T>(StructPropertyHandle, CachedScaleY, LOCTEXT("ScaleY", "Y"), false, SNumericEntryBox<T>::GreenLabelBackgroundColor)
+			MakeNumericProxyWidget<UE::Math::TVector<T>, T>(StructPropertyHandle, CachedScaleY, LOCTEXT("ScaleY", "Y"), false, AxisDisplayInfo::GetAxisColor(EAxisList::Y))
 		]
 		+ SHorizontalBox::Slot()
 		.Padding(FMargin(0.0f, 2.0f, 0.0f, 2.0f))
 		[
-			MakeNumericProxyWidget<UE::Math::TVector<T>, T>(StructPropertyHandle, CachedScaleZ, LOCTEXT("ScaleZ", "Z"), false, SNumericEntryBox<T>::BlueLabelBackgroundColor)
+			MakeNumericProxyWidget<UE::Math::TVector<T>, T>(StructPropertyHandle, CachedScaleZ, LOCTEXT("ScaleZ", "Z"), false, AxisDisplayInfo::GetAxisColor(EAxisList::Z))
 		]
 	];
 }
diff --git a/Engine/Source/Editor/DetailCustomizations/Private/RotatorStructCustomization.cpp b/Engine/Source/Editor/DetailCustomizations/Private/RotatorStructCustomization.cpp
--- a/Engine/Source/Editor/DetailCustomizations/Private/RotatorStructCustomization.cpp
+++ b/Engine/Source/Editor/DetailCustomizations/Private/RotatorStructCustomization.cpp
@@ -2,10 +2,10 @@
 
 #include "RotatorStructCustomization.h"
 
-#include "Customizations/VectorStructCustomization.h"
 #include "HAL/Platform.h"
 #include "HAL/PlatformCrt.h"
 #include "Misc/AssertionMacros.h"
+#include "Misc/AxisDisplayInfo.h"
 #include "PropertyHandle.h"
 #include "UObject/NameTypes.h"
 #include "UObject/UnrealType.h"
@@ -42,7 +42,7 @@ void FRotatorStructCustomization::GetSortedChildren(TSharedRef<IPropertyHandle>
 	static const FName Pitch("Pitch");
 	static const FName Yaw("Yaw");
 
-	const bool bUseForwardRightUpDisplayNames = FVectorStructCustomization::UseForwardRightUpDisplayNames();
+	const bool bUseForwardRightUpDisplayNames = AxisDisplayInfo::UseForwardRightUpDisplayNames();
 
 	TSharedPtr< IPropertyHandle > RotatorChildren[3];
 
diff --git a/Engine/Source/Editor/DetailCustomizations/Private/VectorStructCustomization.cpp b/Engine/Source/Editor/DetailCustomizations/Private/VectorStructCustomization.cpp
--- a/Engine/Source/Editor/DetailCustomizations/Private/VectorStructCustomization.cpp
+++ b/Engine/Source/Editor/DetailCustomizations/Private/VectorStructCustomization.cpp
@@ -5,14 +5,13 @@
 #include "HAL/Platform.h"
 #include "HAL/PlatformCrt.h"
 #include "Misc/AssertionMacros.h"
+#include "Misc/AxisDisplayInfo.h"
 #include "PropertyHandle.h"
 #include "UObject/NameTypes.h"
 #include "UObject/UnrealType.h"
 
 class IPropertyTypeCustomization;
 
-#define LOCTEXT_NAMESPACE "FVectorStructCustomization"
-
 namespace UE::Editor::DetailCustomizations::Private
 {
 	static TAutoConsoleVariable<bool> ShowVector3Children(
@@ -20,8 +19,6 @@ namespace UE::Editor::DetailCustomizations::Private
 		true,
 		TEXT("When true, the detail customizations for Vector3 variants expand to show children")
 	);
-
-	static TOptional<bool> UseForwardRightUpDisplayNames;
 }
 
 TSharedRef<IPropertyTypeCustomization> FVectorStructCustomization::MakeInstance() 
@@ -29,18 +26,6 @@ TSharedRef<IPropertyTypeCustomization> FVectorStructCustomization::MakeInstance(
 	return MakeShareable(new FVectorStructCustomization);
 }
 
-bool FVectorStructCustomization::UseForwardRightUpDisplayNames()
-{
-	if (!UE::Editor::DetailCustomizations::Private::UseForwardRightUpDisplayNames.IsSet())
-	{
-		if (IConsoleVariable* UseForwardRightUpDisplayNamesVariable = IConsoleManager::Get().FindConsoleVariable(TEXT("Editor.UseForwardRightUpDisplayNames")))
-		{
-			UE::Editor::DetailCustomizations::Private::UseForwardRightUpDisplayNames = UseForwardRightUpDisplayNamesVariable->GetBool();
-		}
-	}
-	return UE::Editor::DetailCustomizations::Private::UseForwardRightUpDisplayNames.Get(false/*DefaultValue*/);
-}
-
 void FVectorStructCustomization::CustomizeChildren(TSharedRef<IPropertyHandle> StructPropertyHandle, IDetailChildrenBuilder& StructBuilder, IPropertyTypeCustomizationUtils& StructCustomizationUtils)
 {
 	if (UE::Editor::DetailCustomizations::Private::ShowVector3Children.GetValueOnGameThread())
@@ -55,8 +40,6 @@ void FVectorStructCustomization::GetSortedChildren(TSharedRef<IPropertyHandle> S
 	static const FName Y("Y");
 	static const FName Z("Z");
 
-	const bool bUseForwardRightUpDisplayNames = UseForwardRightUpDisplayNames();
-	
 	TSharedPtr<IPropertyHandle> VectorChildren[3];
 
 	uint32 NumChildren;
@@ -69,21 +52,15 @@ void FVectorStructCustomization::GetSortedChildren(TSharedRef<IPropertyHandle> S
 
 		if (PropertyName == X)
 		{
-			if (bUseForwardRightUpDisplayNames)
-			{
-				ChildHandle->SetPropertyDisplayName(LOCTEXT("ForwardDisplayName", "Forward (was X)"));
-				ChildHandle->SetToolTipText(LOCTEXT("ForwardToolTip", "Forward (was X)"));
-			}
+			ChildHandle->SetPropertyDisplayName(AxisDisplayInfo::GetAxisDisplayName(EAxisList::X));
+			ChildHandle->SetToolTipText(AxisDisplayInfo::GetAxisDisplayName(EAxisList::X));
 
 			VectorChildren[0] = ChildHandle;
 		}
 		else if (PropertyName == Y)
 		{
-			if (bUseForwardRightUpDisplayNames)
-			{
-				ChildHandle->SetPropertyDisplayName(LOCTEXT("RightDisplayName", "Right (was Y)"));
-				ChildHandle->SetToolTipText(LOCTEXT("RightToolTip", "Right (was Y)"));
-			}
+			ChildHandle->SetPropertyDisplayName(AxisDisplayInfo::GetAxisDisplayName(EAxisList::Y));
+			ChildHandle->SetToolTipText(AxisDisplayInfo::GetAxisDisplayName(EAxisList::Y));
 
 			VectorChildren[1] = ChildHandle;
 		}
@@ -91,11 +68,8 @@ void FVectorStructCustomization::GetSortedChildren(TSharedRef<IPropertyHandle> S
 		{
 			check(PropertyName == Z);
 
-			if (bUseForwardRightUpDisplayNames)
-			{
-				ChildHandle->SetPropertyDisplayName(LOCTEXT("UpDisplayName", "Up (was Z)"));
-				ChildHandle->SetToolTipText(LOCTEXT("UpToolTip", "Up (was Z)"));
-			}
+			ChildHandle->SetPropertyDisplayName(AxisDisplayInfo::GetAxisDisplayName(EAxisList::Z));
+			ChildHandle->SetToolTipText(AxisDisplayInfo::GetAxisDisplayName(EAxisList::Z));
 
 			VectorChildren[2] = ChildHandle;
 		}
@@ -105,5 +79,3 @@ void FVectorStructCustomization::GetSortedChildren(TSharedRef<IPropertyHandle> S
 	OutChildren.Add(VectorChildren[1].ToSharedRef());
 	OutChildren.Add(VectorChildren[2].ToSharedRef());
 }
-
-#undef LOCTEXT_NAMESPACE
diff --git a/Engine/Source/Editor/DetailCustomizations/Public/Customizations/VectorStructCustomization.h b/Engine/Source/Editor/DetailCustomizations/Public/Customizations/VectorStructCustomization.h
--- a/Engine/Source/Editor/DetailCustomizations/Public/Customizations/VectorStructCustomization.h
+++ b/Engine/Source/Editor/DetailCustomizations/Public/Customizations/VectorStructCustomization.h
@@ -20,8 +20,6 @@ public:
 	/** @return A new instance of this class */
 	static TSharedRef<IPropertyTypeCustomization> MakeInstance();
 
-	static bool UseForwardRightUpDisplayNames();
-
 private:
 
 	/** IPropertyTypeCustomization interface */
diff --git a/Engine/Source/Editor/EditorFramework/Private/EditorFrameworkModule.cpp b/Engine/Source/Editor/EditorFramework/Private/EditorFrameworkModule.cpp
--- a/Engine/Source/Editor/EditorFramework/Private/EditorFrameworkModule.cpp
+++ b/Engine/Source/Editor/EditorFramework/Private/EditorFrameworkModule.cpp
@@ -2,19 +2,8 @@
 
 #include "EditorFrameworkModule.h"
 
-#include "HAL/IConsoleManager.h"
 #include "Modules/ModuleManager.h"
 
-namespace UE::Editor::Private
-{
-	static TAutoConsoleVariable<bool> EditorUseForwardRightUpDisplayNames(
-		TEXT("Editor.UseForwardRightUpDisplayNames"),
-		false,
-		TEXT("When true, the editor uses (Forward, Right, Up) nomenclature for coordinate axes"),
-		ECVF_ReadOnly
-	);
-}
-
 void FEditorFrameworkModule::StartupModule()
 {
 }
diff --git a/Engine/Source/Editor/Experimental/EditorInteractiveToolsFramework/Private/EditorGizmos/TransformGizmo.cpp b/Engine/Source/Editor/Experimental/EditorInteractiveToolsFramework/Private/EditorGizmos/TransformGizmo.cpp
--- a/Engine/Source/Editor/Experimental/EditorInteractiveToolsFramework/Private/EditorGizmos/TransformGizmo.cpp
+++ b/Engine/Source/Editor/Experimental/EditorInteractiveToolsFramework/Private/EditorGizmos/TransformGizmo.cpp
@@ -16,6 +16,7 @@
 #include "Engine/CollisionProfile.h"
 #include "Engine/World.h"
 #include "Materials/Material.h"
+#include "Misc/AxisDisplayInfo.h"
 #include "UnrealEngine.h"
 #include "Behaviors/MultiButtonClickDragBehavior.h"
 #include "Intersection/IntersectionUtil.h"
@@ -340,13 +341,13 @@ void UTransformGizmo::SetupMaterials()
 	UMaterial* AxisMaterialBase = GetBaseMaterial(); 
 
 	AxisMaterialX = UMaterialInstanceDynamic::Create(AxisMaterialBase, NULL);
-	AxisMaterialX->SetVectorParameterValue("GizmoColor", AxisColorX);
+	AxisMaterialX->SetVectorParameterValue("GizmoColor", AxisDisplayInfo::GetAxisColor(EAxisList::X));
 
 	AxisMaterialY = UMaterialInstanceDynamic::Create(AxisMaterialBase, NULL);
-	AxisMaterialY->SetVectorParameterValue("GizmoColor", AxisColorY);
+	AxisMaterialY->SetVectorParameterValue("GizmoColor", AxisDisplayInfo::GetAxisColor(EAxisList::Y));
 
 	AxisMaterialZ = UMaterialInstanceDynamic::Create(AxisMaterialBase, NULL);
-	AxisMaterialZ->SetVectorParameterValue("GizmoColor", AxisColorZ);
+	AxisMaterialZ->SetVectorParameterValue("GizmoColor", AxisDisplayInfo::GetAxisColor(EAxisList::Z));
 
 	GreyMaterial = UMaterialInstanceDynamic::Create(AxisMaterialBase, NULL);
 	GreyMaterial->SetVectorParameterValue("GizmoColor", GreyColor);
@@ -805,15 +806,15 @@ void UTransformGizmo::EnablePlanarObjects(bool bTranslate, bool bEnableX, bool b
 
 	if (bTranslate)
 	{
-		EnablePlanarElement(TranslatePlanarXYElement, ETransformGizmoPartIdentifier::TranslateXYPlanar, XAxis, YAxis, ZAxis, AxisColorZ, bEnableXY);
-		EnablePlanarElement(TranslatePlanarYZElement, ETransformGizmoPartIdentifier::TranslateYZPlanar, YAxis, ZAxis, XAxis, AxisColorX, bEnableYZ);
-		EnablePlanarElement(TranslatePlanarXZElement, ETransformGizmoPartIdentifier::TranslateXZPlanar, ZAxis, XAxis, YAxis, AxisColorY, bEnableXZ);
+		EnablePlanarElement(TranslatePlanarXYElement, ETransformGizmoPartIdentifier::TranslateXYPlanar, XAxis, YAxis, ZAxis, AxisDisplayInfo::GetAxisColor(EAxisList::Z), bEnableXY);
+		EnablePlanarElement(TranslatePlanarYZElement, ETransformGizmoPartIdentifier::TranslateYZPlanar, YAxis, ZAxis, XAxis, AxisDisplayInfo::GetAxisColor(EAxisList::X), bEnableYZ);
+		EnablePlanarElement(TranslatePlanarXZElement, ETransformGizmoPartIdentifier::TranslateXZPlanar, ZAxis, XAxis, YAxis, AxisDisplayInfo::GetAxisColor(EAxisList::Y), bEnableXZ);
 	}
 	else
 	{
-		EnablePlanarElement(ScalePlanarXYElement, ETransformGizmoPartIdentifier::ScaleXYPlanar, XAxis, YAxis, ZAxis, AxisColorZ, bEnableXY);
-		EnablePlanarElement(ScalePlanarYZElement, ETransformGizmoPartIdentifier::ScaleYZPlanar, YAxis, ZAxis, XAxis, AxisColorX, bEnableYZ);
-		EnablePlanarElement(ScalePlanarXZElement, ETransformGizmoPartIdentifier::ScaleXZPlanar, ZAxis, XAxis, YAxis, AxisColorY, bEnableXZ);
+		EnablePlanarElement(ScalePlanarXYElement, ETransformGizmoPartIdentifier::ScaleXYPlanar, XAxis, YAxis, ZAxis, AxisDisplayInfo::GetAxisColor(EAxisList::Z), bEnableXY);
+		EnablePlanarElement(ScalePlanarYZElement, ETransformGizmoPartIdentifier::ScaleYZPlanar, YAxis, ZAxis, XAxis, AxisDisplayInfo::GetAxisColor(EAxisList::X), bEnableYZ);
+		EnablePlanarElement(ScalePlanarXZElement, ETransformGizmoPartIdentifier::ScaleXZPlanar, ZAxis, XAxis, YAxis, AxisDisplayInfo::GetAxisColor(EAxisList::Y), bEnableXZ);
 	}
 }
 
diff --git a/Engine/Source/Editor/Experimental/EditorInteractiveToolsFramework/Public/EditorGizmos/TransformGizmo.h b/Engine/Source/Editor/Experimental/EditorInteractiveToolsFramework/Public/EditorGizmos/TransformGizmo.h
--- a/Engine/Source/Editor/Experimental/EditorInteractiveToolsFramework/Public/EditorGizmos/TransformGizmo.h
+++ b/Engine/Source/Editor/Experimental/EditorInteractiveToolsFramework/Public/EditorGizmos/TransformGizmo.h
@@ -162,8 +162,11 @@ public:
 	static constexpr float PlanarHandleSize = 15.0f;
 
 	static constexpr float AxisTransp = 0.8f;
+	UE_DEPRECATED(5.6, "Use AxisDisplayInfo::GetAxisColor(EAxisList::X) instead")
 	static constexpr FLinearColor AxisColorX = FLinearColor(0.594f, 0.0197f, 0.0f);
+	UE_DEPRECATED(5.6, "Use AxisDisplayInfo::GetAxisColor(EAxisList::Y) instead")
 	static constexpr FLinearColor AxisColorY = FLinearColor(0.1349f, 0.3959f, 0.0f);
+	UE_DEPRECATED(5.6, "Use AxisDisplayInfo::GetAxisColor(EAxisList::Z) instead")
 	static constexpr FLinearColor AxisColorZ = FLinearColor(0.0251f, 0.207f, 0.85f);
 	static constexpr FLinearColor ScreenAxisColor = FLinearColor(0.76f, 0.72f, 0.14f);
 	static constexpr FLinearColor PlaneColorXY = FLinearColor(1.0f, 1.0f, 0.0f);
diff --git a/Engine/Source/Editor/UnrealEd/Classes/Settings/EditorStyleSettings.h b/Engine/Source/Editor/UnrealEd/Classes/Settings/EditorStyleSettings.h
--- a/Engine/Source/Editor/UnrealEd/Classes/Settings/EditorStyleSettings.h
+++ b/Engine/Source/Editor/UnrealEd/Classes/Settings/EditorStyleSettings.h
@@ -76,6 +76,18 @@ public:
 	UPROPERTY(EditAnywhere, config, Category = "Accessibility")
 	bool bColorVisionDeficiencyCorrectionPreviewWithDeficiency;
 
+	/** The color used for the X axis */
+	UPROPERTY(EditAnywhere, config, Category = UserInterface, meta = (ConfigRestartRequired = true))
+	FLinearColor XAxisColor;
+
+	/** The color used for the Y axis */
+	UPROPERTY(EditAnywhere, config, Category = UserInterface, meta = (ConfigRestartRequired = true))
+	FLinearColor YAxisColor;
+
+	/** The color used for the Z axis */
+	UPROPERTY(EditAnywhere, config, Category = UserInterface, meta = (ConfigRestartRequired = true))
+	FLinearColor ZAxisColor;
+
 	/** The color used to represent selection */
 	UPROPERTY(EditAnywhere, config, Category=UserInterface, meta=(DisplayName="Viewport Selection Color"))
 	FLinearColor SelectionColor;
diff --git a/Engine/Source/Editor/UnrealEd/Private/EditorAxisDisplayInfo.cpp b/Engine/Source/Editor/UnrealEd/Private/EditorAxisDisplayInfo.cpp
new file mode 100644
--- /dev/null
+++ b/Engine/Source/Editor/UnrealEd/Private/EditorAxisDisplayInfo.cpp
@@ -0,0 +1,149 @@
+// Copyright Epic Games, Inc. All Rights Reserved.
+
+#include "EditorAxisDisplayInfo.h"
+
+#include "HAL/IConsoleManager.h"
+#include "Internationalization/Internationalization.h"
+#include "Settings/EditorStyleSettings.h"
+
+#define LOCTEXT_NAMESPACE "EditorAxisDisplayInfo"
+
+namespace UE::Editor::Private
+{
+	static TAutoConsoleVariable<bool> EditorUseForwardRightUpDisplayNames(
+		TEXT("Editor.UseForwardRightUpDisplayNames"),
+		false,
+		TEXT("When true, the editor uses (Forward, Right, Up) nomenclature for coordinate axes"),
+		ECVF_ReadOnly
+	);
+}
+
+FEditorAxisDisplayInfo::FEditorAxisDisplayInfo()
+{
+	if (FProperty* XAxisColorProperty = UEditorStyleSettings::StaticClass()->FindPropertyByName(GET_MEMBER_NAME_CHECKED(UEditorStyleSettings, XAxisColor)))
+	{
+		XAxisColorProperty->SetMetaData(TEXT("DisplayName"), *FText::Format(LOCTEXT("XAxisColorDisplayName", "{0} Axis Color"), GetAxisDisplayName(EAxisList::X)).ToString());
+	}
+	if (FProperty* YAxisColorProperty = UEditorStyleSettings::StaticClass()->FindPropertyByName(GET_MEMBER_NAME_CHECKED(UEditorStyleSettings, YAxisColor)))
+	{
+		YAxisColorProperty->SetMetaData(TEXT("DisplayName"), *FText::Format(LOCTEXT("YAxisColorDisplayName", "{0} Axis Color"), GetAxisDisplayName(EAxisList::Y)).ToString());
+	}
+	if (FProperty* ZAxisColorProperty = UEditorStyleSettings::StaticClass()->FindPropertyByName(GET_MEMBER_NAME_CHECKED(UEditorStyleSettings, ZAxisColor)))
+	{
+		ZAxisColorProperty->SetMetaData(TEXT("DisplayName"), *FText::Format(LOCTEXT("ZAxisColorDisplayName", "{0} Axis Color"), GetAxisDisplayName(EAxisList::Z)).ToString());
+	}
+}
+
+FText FEditorAxisDisplayInfo::GetAxisDisplayName(EAxisList::Type Axis)
+{
+	if (Axis == EAxisList::X)
+	{
+		if (UseForwardRightUpDisplayNames())
+		{
+			return LOCTEXT("ForwardDisplayName", "Forward (was X)");
+		}
+		else
+		{
+			return LOCTEXT("XDisplayName", "X");
+		}
+	}
+	else if (Axis == EAxisList::Y)
+	{
+		if (UseForwardRightUpDisplayNames())
+		{
+			return LOCTEXT("RightDisplayName", "Right (was Y)");
+		}
+		else
+		{
+			return LOCTEXT("YDisplayName", "Y");
+		}
+	}
+	else if (Axis == EAxisList::Z)
+	{
+		if (UseForwardRightUpDisplayNames())
+		{
+			return LOCTEXT("UpDisplayName", "Up (was Z)");
+		}
+		else
+		{
+			return LOCTEXT("ZDisplayName", "Z");
+		}
+	}
+
+	ensureMsgf(false, TEXT("Unsupported Axis: %d"), Axis);
+
+	return LOCTEXT("UnsupportedDisplayName", "Unsupported");
+}
+
+FText FEditorAxisDisplayInfo::GetAxisDisplayNameShort(EAxisList::Type Axis)
+{
+	if (Axis == EAxisList::X)
+	{
+		if (UseForwardRightUpDisplayNames())
+		{
+			return LOCTEXT("ForwardDisplayNameShort", "F");
+		}
+		else
+		{
+			return LOCTEXT("XDisplayNameShort", "X");
+		}
+	}
+	else if (Axis == EAxisList::Y)
+	{
+		if (UseForwardRightUpDisplayNames())
+		{
+			return LOCTEXT("RightDisplayNameShort", "R");
+		}
+		else
+		{
+			return LOCTEXT("YDisplayNameShort", "Y");
+		}
+	}
+	else if (Axis == EAxisList::Z)
+	{
+		if (UseForwardRightUpDisplayNames())
+		{
+			return LOCTEXT("UpDisplayNameShort", "U");
+		}
+		else
+		{
+			return LOCTEXT("ZDisplayNameShort", "Z");
+		}
+	}
+
+	ensureMsgf(false, TEXT("Unsupported Axis: %d"), Axis);
+
+	return LOCTEXT("UnsupportedDisplayNameShort", "?");
+}
+
+FLinearColor FEditorAxisDisplayInfo::GetAxisColor(EAxisList::Type Axis)
+{
+	if (Axis == EAxisList::X)
+	{
+		return GetDefault<UEditorStyleSettings>()->XAxisColor;
+	}
+	else if (Axis == EAxisList::Y)
+	{
+		return GetDefault<UEditorStyleSettings>()->YAxisColor;
+	}
+	else if (Axis == EAxisList::Z)
+	{
+		return GetDefault<UEditorStyleSettings>()->ZAxisColor;
+	}
+
+	ensureMsgf(false, TEXT("Unsupported Axis: %d"), Axis);
+
+	return FLinearColor::Black;
+}
+
+bool FEditorAxisDisplayInfo::UseForwardRightUpDisplayNames()
+{
+	if (!bUseForwardRightUpDisplayNames.IsSet())
+	{
+		bUseForwardRightUpDisplayNames = UE::Editor::Private::EditorUseForwardRightUpDisplayNames.GetValueOnAnyThread();
+	}
+
+	return bUseForwardRightUpDisplayNames.Get(false/*DefaultValue*/);
+}
+
+#undef LOCTEXT_NAMESPACE
diff --git a/Engine/Source/Editor/UnrealEd/Private/EditorViewportClient.cpp b/Engine/Source/Editor/UnrealEd/Private/EditorViewportClient.cpp
--- a/Engine/Source/Editor/UnrealEd/Private/EditorViewportClient.cpp
+++ b/Engine/Source/Editor/UnrealEd/Private/EditorViewportClient.cpp
@@ -8,6 +8,7 @@
 #include "Elements/Interfaces/TypedElementObjectInterface.h"
 #include "PreviewScene.h"
 #include "HAL/FileManager.h"
+#include "Misc/AxisDisplayInfo.h"
 #include "Misc/FileHelper.h"
 #include "Misc/Paths.h"
 #include "Misc/CoreDelegates.h"
@@ -3648,12 +3649,19 @@ void FEditorViewportClient::DrawAxes(FViewport* InViewport, FCanvas* Canvas, con
 	FVector2D OutsideCenterSize = FVector2D(4, 4);
 	FCanvasTileItem OutsideTileItem(FVector2D::ZeroVector, OutsideCenterSize, FLinearColor(0.05f, 0.05f, 0.05f));
 	
-	FLinearColor XBaseColor = FLinearColor(0.594f, 0.0197f, 0.0f);
-	FLinearColor XShadowColor = FLinearColor(0.086f, 0.012f, 0.004f);
-	FLinearColor YBaseColor = FLinearColor(0.1349f, 0.3959f, 0.0f);
-	FLinearColor YShadowColor = FLinearColor(0.039f, 0.075f, 0.008f);
-	FLinearColor ZBaseColor = FLinearColor(0.0251f, 0.207f, 0.85f);
-	FLinearColor ZShadowColor = FLinearColor(0.012f, 0.035f, 0.094f);
+	auto GetShadowColor = [](const FLinearColor& Color) -> FLinearColor
+		{
+			FLinearColor ColorHSV = Color.LinearRGBToHSV();
+			ColorHSV.B *= .2f;
+			return ColorHSV.HSVToLinearRGB();
+		};
+
+	FLinearColor XBaseColor = AxisDisplayInfo::GetAxisColor(EAxisList::X);
+	FLinearColor XShadowColor = GetShadowColor(XBaseColor);
+	FLinearColor YBaseColor = AxisDisplayInfo::GetAxisColor(EAxisList::Y);
+	FLinearColor YShadowColor = GetShadowColor(YBaseColor);
+	FLinearColor ZBaseColor = AxisDisplayInfo::GetAxisColor(EAxisList::Z);
+	FLinearColor ZShadowColor = GetShadowColor(ZBaseColor);
 
 	auto GetShadowOffset = [this](FIntPoint EndPoint, FIntPoint OriginPoint) -> FIntPoint
 		{
@@ -3662,38 +3670,6 @@ void FEditorViewportClient::DrawAxes(FViewport* InViewport, FCanvas* Canvas, con
 			return FIntPoint(X, Y);
 		};
 		
-	const TStaticArray<FText, 3> AxisNames = []()
-	{
-		static bool bUseForwardRightUpDisplayNames = []()
-		{
-			if (IConsoleVariable* UseForwardRightUpDisplayNamesVariable = IConsoleManager::Get().FindConsoleVariable(TEXT("Editor.UseForwardRightUpDisplayNames")))
-			{
-				return UseForwardRightUpDisplayNamesVariable->GetBool();
-			}
-			return false;
-		}();
-
-		if (bUseForwardRightUpDisplayNames)
-		{
-			return TStaticArray<FText, 3>({LOCTEXT("ForwardAxis","F"), LOCTEXT("RightAxis","R"), LOCTEXT("UpAxis","U")});
-		}
-		else
-		{
-			return TStaticArray<FText, 3>({LOCTEXT("XAxis","X"), LOCTEXT("YAxis","Y"), LOCTEXT("ZAxis","Z")});
-		}
-	}();
-	auto GetAxisName = [&AxisNames](EAxisList::Type InAxis) -> FText
-	{
-		switch (InAxis)
-		{
-		case EAxisList::X: return AxisNames[0];
-		case EAxisList::Y: return AxisNames[1];
-		case EAxisList::Z: return AxisNames[2];
-		default:
-			return FText();
-		}
-	};
-
 	if( ( InAxis & EAxisList::X ) ==  EAxisList::X)
 	{
 		AxisVec = AxisSize * ViewTM.InverseTransformVector( FVector(1,0,0) );
@@ -3704,7 +3680,7 @@ void FEditorViewportClient::DrawAxes(FViewport* InViewport, FCanvas* Canvas, con
 		LineItem.Draw(Canvas, AxisOrigin + GetShadowOffset(AxisEnd, AxisOrigin), AxisEnd + GetShadowOffset(AxisEnd, AxisOrigin));
 		LineItem.SetColor(XBaseColor);
 		LineItem.Draw( Canvas, AxisOrigin, AxisEnd );
-		TextItem.Text = GetAxisName(EAxisList::X);
+		TextItem.Text = AxisDisplayInfo::GetAxisDisplayNameShort(EAxisList::X);
 		TextItem.SetColor(XBaseColor);
 		TextItem.EnableShadow(XShadowColor);
 		TextItem.Draw( Canvas, FVector2D(TextCenter.X - 0.5*XL, TextCenter.Y - 0.5*YL) );
@@ -3720,7 +3696,7 @@ void FEditorViewportClient::DrawAxes(FViewport* InViewport, FCanvas* Canvas, con
 		LineItem.Draw(Canvas, AxisOrigin + GetShadowOffset(AxisEnd, AxisOrigin), AxisEnd + GetShadowOffset(AxisEnd, AxisOrigin));
 		LineItem.SetColor(YBaseColor);
 		LineItem.Draw( Canvas, AxisOrigin, AxisEnd );
-		TextItem.Text = GetAxisName(EAxisList::Y);
+		TextItem.Text = AxisDisplayInfo::GetAxisDisplayNameShort(EAxisList::Y);
 		TextItem.SetColor(YBaseColor);
 		TextItem.EnableShadow(YShadowColor);
 		TextItem.Draw( Canvas, FVector2D(TextCenter.X - 0.5*XL, TextCenter.Y - 0.5*YL) );
@@ -3737,7 +3713,7 @@ void FEditorViewportClient::DrawAxes(FViewport* InViewport, FCanvas* Canvas, con
 		LineItem.Draw(Canvas, AxisOrigin + GetShadowOffset(AxisEnd, AxisOrigin), AxisEnd + GetShadowOffset(AxisEnd, AxisOrigin));
 		LineItem.SetColor(ZBaseColor);
 		LineItem.Draw( Canvas, AxisOrigin, AxisEnd );
-		TextItem.Text = GetAxisName(EAxisList::Z);
+		TextItem.Text = AxisDisplayInfo::GetAxisDisplayNameShort(EAxisList::Z);
 		TextItem.SetColor(ZBaseColor);
 		TextItem.EnableShadow(ZShadowColor);
 		TextItem.Draw( Canvas, FVector2D(TextCenter.X - 0.5*XL, TextCenter.Y - 0.5*YL) );
diff --git a/Engine/Source/Editor/UnrealEd/Private/UnrealEdMisc.cpp b/Engine/Source/Editor/UnrealEd/Private/UnrealEdMisc.cpp
--- a/Engine/Source/Editor/UnrealEd/Private/UnrealEdMisc.cpp
+++ b/Engine/Source/Editor/UnrealEd/Private/UnrealEdMisc.cpp
@@ -96,6 +96,7 @@
 #include "DeveloperToolSettingsDelegates.h"
 #include "Cooker/CookConfigAccessTracker.h"
 #include "Cooker/PackageBuildDependencyTracker.h"
+#include "EditorAxisDisplayInfo.h"
 
 #define USE_UNIT_TESTS 0
 
@@ -255,6 +256,10 @@ void FUnrealEdMisc::OnInit()
 	CmdExec = MakeUnique<FConsoleCommandExecutor>();
 	IModularFeatures::Get().RegisterModularFeature(IConsoleCommandExecutor::ModularFeatureName(), CmdExec.Get());
 
+	// Register the EditorAxisDisplayInfo modular feature
+	EditorAxisDisplayInfo = MakeUnique<FEditorAxisDisplayInfo>();
+	IModularFeatures::Get().RegisterModularFeature(IAxisDisplayInfo::GetModularFeatureName(), EditorAxisDisplayInfo.Get());
+
 	// Register all callback notifications
 	FEditorDelegates::SelectedProps.AddRaw(this, &FUnrealEdMisc::CB_SelectedProps);
 	FEditorDelegates::DisplayLoadErrors.AddRaw(this, &FUnrealEdMisc::CB_DisplayLoadErrors);
@@ -1138,6 +1143,10 @@ void FUnrealEdMisc::OnExit()
 		}
 	}
 
+	// Unregister the EditorAxisDisplayInfo modular feature
+	IModularFeatures::Get().UnregisterModularFeature(IAxisDisplayInfo::GetModularFeatureName(), EditorAxisDisplayInfo.Get());
+	EditorAxisDisplayInfo.Reset();
+
 	// Unregister the command executor
 	IModularFeatures::Get().UnregisterModularFeature(IConsoleCommandExecutor::ModularFeatureName(), CmdExec.Get());
 	CmdExec.Reset();
diff --git a/Engine/Source/Editor/UnrealEd/Private/UnrealWidget.cpp b/Engine/Source/Editor/UnrealEd/Private/UnrealWidget.cpp
--- a/Engine/Source/Editor/UnrealEd/Private/UnrealWidget.cpp
+++ b/Engine/Source/Editor/UnrealEd/Private/UnrealWidget.cpp
@@ -11,6 +11,7 @@
 #include "GenericPlatform/ICursor.h"
 #include "Materials/Material.h"
 #include "Materials/MaterialInstanceDynamic.h"
+#include "Misc/AxisDisplayInfo.h"
 #include "SceneView.h"
 #include "Settings/LevelEditorViewportSettings.h"
 #include "SnappingUtils.h"
@@ -32,9 +33,9 @@ FWidget::FWidget()
 	TotalDeltaRotation   = 0;
 	CurrentDeltaRotation = 0;
 
-	AxisColorX       = FLinearColor(0.594f, 0.0197f, 0.0f);
-	AxisColorY       = FLinearColor(0.1349f, 0.3959f, 0.0f);
-	AxisColorZ       = FLinearColor(0.0251f, 0.207f, 0.85f);
+	AxisColorX       = AxisDisplayInfo::GetAxisColor(EAxisList::X);
+	AxisColorY       = AxisDisplayInfo::GetAxisColor(EAxisList::Y);
+	AxisColorZ       = AxisDisplayInfo::GetAxisColor(EAxisList::Z);
 	ScreenAxisColor  = FLinearColor(0.76, 0.72, 0.14f);
 	PlaneColorXY     = FColor::Yellow;
 	ArcBallColor     = FColor(128, 128, 128, 6);
diff --git a/Engine/Source/Editor/UnrealEd/Public/EditorAxisDisplayInfo.h b/Engine/Source/Editor/UnrealEd/Public/EditorAxisDisplayInfo.h
new file mode 100644
--- /dev/null
+++ b/Engine/Source/Editor/UnrealEd/Public/EditorAxisDisplayInfo.h
@@ -0,0 +1,26 @@
+// Copyright Epic Games, Inc. All Rights Reserved.
+
+#pragma once
+
+#include "Misc/AxisDisplayInfo.h"
+
+class FEditorAxisDisplayInfo final : public IAxisDisplayInfo
+{
+public:
+
+	FEditorAxisDisplayInfo();
+	virtual ~FEditorAxisDisplayInfo() = default;
+
+	virtual FText GetAxisDisplayName(EAxisList::Type Axis) override;
+
+	virtual FText GetAxisDisplayNameShort(EAxisList::Type Axis) override;
+
+	virtual FLinearColor GetAxisColor(EAxisList::Type Axis) override;
+
+	virtual bool UseForwardRightUpDisplayNames() override;
+
+private:
+
+	TOptional<bool> bUseForwardRightUpDisplayNames;
+
+};
diff --git a/Engine/Source/Editor/UnrealEd/Public/UnrealEdMisc.h b/Engine/Source/Editor/UnrealEd/Public/UnrealEdMisc.h
--- a/Engine/Source/Editor/UnrealEd/Public/UnrealEdMisc.h
+++ b/Engine/Source/Editor/UnrealEd/Public/UnrealEdMisc.h
@@ -11,6 +11,7 @@ class FPerformanceAnalyticsStats;
 class FTickableEditorObject;
 class FUICommandInfo;
 class FConsoleCommandExecutor;
+class FEditorAxisDisplayInfo;
 
 enum class EMapChangeType : uint8
 {
@@ -331,6 +332,9 @@ private:
 	/** Pointer to the classic "Cmd" executor */
 	TUniquePtr<FConsoleCommandExecutor> CmdExec;
 
+	/** Pointer to the EditorAxisDisplayInfo implementation. */
+	TUniquePtr<FEditorAxisDisplayInfo> EditorAxisDisplayInfo;
+
 	/** Handle to the registered OnUserDefinedChordChanged delegate. */
 	FDelegateHandle OnUserDefinedChordChangedDelegateHandle;
 
diff --git a/Engine/Source/Editor/ViewportInteraction/Private/ViewportWorldInteraction.cpp b/Engine/Source/Editor/ViewportInteraction/Private/ViewportWorldInteraction.cpp
--- a/Engine/Source/Editor/ViewportInteraction/Private/ViewportWorldInteraction.cpp
+++ b/Engine/Source/Editor/ViewportInteraction/Private/ViewportWorldInteraction.cpp
@@ -1,7 +1,9 @@
 // Copyright Epic Games, Inc. All Rights Reserved.
 
 #include "ViewportWorldInteraction.h"
+
 #include "Materials/MaterialInterface.h"
+#include "Misc/AxisDisplayInfo.h"
 #include "Components/PrimitiveComponent.h"
 #include "Components/StaticMeshComponent.h"
 #include "Settings/LevelEditorViewportSettings.h"
@@ -350,9 +352,9 @@ void UViewportWorldInteraction::InitColors()
 	Colors.SetNumZeroed((int32)EColors::TotalCount);
 	{
 		Colors[(int32)EColors::DefaultColor] = FLinearColor(0.7f, 0.7f, 0.7f, 1.0f);
-		Colors[(int32)EColors::Forward] = FLinearColor(0.594f, 0.0197f, 0.0f, 1.0f);
-		Colors[(int32)EColors::Right] = FLinearColor(0.1349f, 0.3959f, 0.0f, 1.0f);
-		Colors[(int32)EColors::Up] = FLinearColor(0.0251f, 0.207f, 0.85f, 1.0f);
+		Colors[(int32)EColors::Forward] = AxisDisplayInfo::GetAxisColor(EAxisList::X);
+		Colors[(int32)EColors::Right] = AxisDisplayInfo::GetAxisColor(EAxisList::Y);
+		Colors[(int32)EColors::Up] = AxisDisplayInfo::GetAxisColor(EAxisList::Z);
 		Colors[(int32)EColors::GizmoHover] = FLinearColor::Yellow;
 		Colors[(int32)EColors::GizmoDragging] = FLinearColor::Yellow;
 	}
diff --git a/Engine/Source/Runtime/Core/Private/Misc/AxisDisplayInfo.cpp b/Engine/Source/Runtime/Core/Private/Misc/AxisDisplayInfo.cpp
new file mode 100644
--- /dev/null
+++ b/Engine/Source/Runtime/Core/Private/Misc/AxisDisplayInfo.cpp
@@ -0,0 +1,99 @@
+// Copyright Epic Games, Inc. All Rights Reserved.
+
+#include "Misc/AxisDisplayInfo.h"
+
+#include "Features/IModularFeatures.h"
+#include "Internationalization/Internationalization.h"
+
+#define LOCTEXT_NAMESPACE "AxisDisplayInfo"
+
+FText AxisDisplayInfo::GetAxisDisplayName(EAxisList::Type Axis)
+{
+	if (IModularFeatures::Get().IsModularFeatureAvailable(IAxisDisplayInfo::GetModularFeatureName()))
+	{
+		IAxisDisplayInfo& DisplayInfo = IModularFeatures::Get().GetModularFeature<IAxisDisplayInfo>(IAxisDisplayInfo::GetModularFeatureName());
+		return DisplayInfo.GetAxisDisplayName(Axis);
+	}
+
+	if (Axis == EAxisList::X)
+	{
+		return LOCTEXT("XDisplayName", "X");
+	}
+	else if (Axis == EAxisList::Y)
+	{
+		return LOCTEXT("YDisplayName", "Y");
+	}
+	else if (Axis == EAxisList::Z)
+	{
+		return LOCTEXT("ZDisplayName", "Z");
+	}
+
+	ensureMsgf(false, TEXT("Unsupported Axis: %d"), Axis);
+
+	return LOCTEXT("UnsupportedDisplayName", "Unsupported");
+}
+
+FText AxisDisplayInfo::GetAxisDisplayNameShort(EAxisList::Type Axis)
+{
+	if (IModularFeatures::Get().IsModularFeatureAvailable(IAxisDisplayInfo::GetModularFeatureName()))
+	{
+		IAxisDisplayInfo& DisplayInfo = IModularFeatures::Get().GetModularFeature<IAxisDisplayInfo>(IAxisDisplayInfo::GetModularFeatureName());
+		return DisplayInfo.GetAxisDisplayNameShort(Axis);
+	}
+
+	if (Axis == EAxisList::X)
+	{
+		return LOCTEXT("XDisplayNameShort", "X");
+	}
+	else if (Axis == EAxisList::Y)
+	{
+		return LOCTEXT("YDisplayNameShort", "Y");
+	}
+	else if (Axis == EAxisList::Z)
+	{
+		return LOCTEXT("ZDisplayNameShort", "Z");
+	}
+
+	ensureMsgf(false, TEXT("Unsupported Axis: %d"), Axis);
+
+	return LOCTEXT("UnsupportedDisplayNameShort", "?");
+}
+
+FLinearColor AxisDisplayInfo::GetAxisColor(EAxisList::Type Axis)
+{
+	if (IModularFeatures::Get().IsModularFeatureAvailable(IAxisDisplayInfo::GetModularFeatureName()))
+	{
+		IAxisDisplayInfo& DisplayInfo = IModularFeatures::Get().GetModularFeature<IAxisDisplayInfo>(IAxisDisplayInfo::GetModularFeatureName());
+		return DisplayInfo.GetAxisColor(Axis);
+	}
+
+	if (Axis == EAxisList::X)
+	{
+		return FLinearColor(0.594f, 0.0197f, 0.0f);
+	}
+	else if (Axis == EAxisList::Y)
+	{
+		return FLinearColor(0.1349f, 0.3959f, 0.0f);
+	}
+	else if (Axis == EAxisList::Z)
+	{
+		return FLinearColor(0.0251f, 0.207f, 0.85f);
+	}
+
+	ensureMsgf(false, TEXT("Unsupported Axis: %d"), Axis);
+
+	return FLinearColor::Black;
+}
+
+bool AxisDisplayInfo::UseForwardRightUpDisplayNames()
+{
+	if (IModularFeatures::Get().IsModularFeatureAvailable(IAxisDisplayInfo::GetModularFeatureName()))
+	{
+		IAxisDisplayInfo& DisplayInfo = IModularFeatures::Get().GetModularFeature<IAxisDisplayInfo>(IAxisDisplayInfo::GetModularFeatureName());
+		return DisplayInfo.UseForwardRightUpDisplayNames();
+	}
+
+	return false;
+}
+
+#undef LOCTEXT_NAMESPACE
diff --git a/Engine/Source/Runtime/Core/Public/Misc/AxisDisplayInfo.h b/Engine/Source/Runtime/Core/Public/Misc/AxisDisplayInfo.h
new file mode 100644
--- /dev/null
+++ b/Engine/Source/Runtime/Core/Public/Misc/AxisDisplayInfo.h
@@ -0,0 +1,81 @@
+// Copyright Epic Games, Inc. All Rights Reserved.
+
+#pragma once
+
+#include "Features/IModularFeature.h"
+#include "Internationalization/Text.h"
+#include "Math/Axis.h"
+#include "Math/Color.h"
+
+namespace AxisDisplayInfo
+{
+	/**
+	 * Gets the display name to use for the given axis.
+	 *
+	 * @param Axis The axis whose display name you are requesting
+	 * @return The display name corresponding to the given axis
+	 */
+	CORE_API FText GetAxisDisplayName(EAxisList::Type Axis);
+
+	/**
+	 * Gets the short display name to use for the given axis.
+	 *
+	 * @param Axis The axis whose short display name you are requesting
+	 * @return The short display name corresponding to the given axis
+	 */
+	CORE_API FText GetAxisDisplayNameShort(EAxisList::Type Axis);
+
+	/**
+	 * Gets the color to use for the given axis.
+	 * 
+	 * @param Axis The axis whose color you are requesting
+	 * @return The color corresponding to the given axis
+	 */
+	CORE_API FLinearColor GetAxisColor(EAxisList::Type Axis);
+
+	/**
+	 * Gets whether or not the engine uses forward/right/up nomenclature for display names.
+	 */
+	CORE_API bool UseForwardRightUpDisplayNames();
+};
+
+class IAxisDisplayInfo : public IModularFeature
+{
+public:
+	virtual ~IAxisDisplayInfo() = default;
+
+	static FName GetModularFeatureName()
+	{
+		static const FName Name = TEXT("AxisDisplayInfo");
+		return Name;
+	}
+
+	/**
+	 * Gets the display name to use for the given axis.
+	 *
+	 * @param Axis The axis whose display name you are requesting
+	 * @return The display name corresponding to the given axis
+	 */
+	virtual FText GetAxisDisplayName(EAxisList::Type Axis) = 0;
+
+	/**
+	 * Gets the short display name to use for the given axis.
+	 *
+	 * @param Axis The axis whose short display name you are requesting
+	 * @return The short display name corresponding to the given axis
+	 */
+	virtual FText GetAxisDisplayNameShort(EAxisList::Type Axis) = 0;
+
+	/**
+	 * Gets the color to use for the given axis.
+	 *
+	 * @param Axis The axis whose color you are requesting
+	 * @return The color corresponding to the given axis
+	 */
+	virtual FLinearColor GetAxisColor(EAxisList::Type Axis) = 0;
+
+	/**
+	 * Gets whether or not the engine uses forward/right/up nomenclature for display names.
+	 */
+	virtual bool UseForwardRightUpDisplayNames() = 0;
+};
diff --git a/Engine/Source/Runtime/InteractiveToolsFramework/Private/BaseGizmos/GizmoRenderingUtil.cpp b/Engine/Source/Runtime/InteractiveToolsFramework/Private/BaseGizmos/GizmoRenderingUtil.cpp
--- a/Engine/Source/Runtime/InteractiveToolsFramework/Private/BaseGizmos/GizmoRenderingUtil.cpp
+++ b/Engine/Source/Runtime/InteractiveToolsFramework/Private/BaseGizmos/GizmoRenderingUtil.cpp
@@ -10,6 +10,7 @@
 #include "Materials/MaterialInstanceDynamic.h"
 #include "Math/UnrealMathUtility.h"
 #include "Math/NumericLimits.h"
+#include "Misc/AxisDisplayInfo.h"
 
 // yuck global value set by Editor
 static const FSceneView* GlobalCurrentSceneView = nullptr;
@@ -224,11 +225,11 @@ FLinearColor UE::GizmoRenderingUtil::GetDefaultAxisColor(EAxis::Type Axis)
 	switch (Axis)
 	{
 	case EAxis::X:
-		return FLinearColor(0.594f, 0.0197f, 0.0f);
+		return AxisDisplayInfo::GetAxisColor(EAxisList::X);
 	case EAxis::Y:
-		return FLinearColor(0.1349f, 0.3959f, 0.0f);
+		return AxisDisplayInfo::GetAxisColor(EAxisList::Y);
 	case EAxis::Z:
-		return FLinearColor(0.0251f, 0.207f, 0.85f);
+		return AxisDisplayInfo::GetAxisColor(EAxisList::Z);
 	default:
 		return FLinearColor::White;
 	}
diff --git a/Engine/Source/Runtime/Slate/Public/Widgets/Input/SRotatorInputBox.h b/Engine/Source/Runtime/Slate/Public/Widgets/Input/SRotatorInputBox.h
--- a/Engine/Source/Runtime/Slate/Public/Widgets/Input/SRotatorInputBox.h
+++ b/Engine/Source/Runtime/Slate/Public/Widgets/Input/SRotatorInputBox.h
@@ -4,6 +4,7 @@
 
 #include "CoreMinimal.h"
 #include "Layout/Margin.h"
+#include "Misc/AxisDisplayInfo.h"
 #include "Fonts/SlateFontInfo.h"
 #include "Widgets/DeclarativeSyntaxSupport.h"
 #include "Widgets/SWidget.h"
@@ -172,7 +173,7 @@ public:
 				.LabelLocation(SNumericEntryBox<NumericType>::ELabelLocation::Inside)
 				.Label()
 				[
-					InArgs._bColorAxisLabels ? SNumericEntryBox<NumericType>::BuildNarrowColorLabel(SNumericEntryBox<NumericType>::RedLabelBackgroundColor) : SNullWidget::NullWidget
+					InArgs._bColorAxisLabels ? SNumericEntryBox<NumericType>::BuildNarrowColorLabel(AxisDisplayInfo::GetAxisColor(EAxisList::X)) : SNullWidget::NullWidget
 				]
 				.Font( InArgs._Font )
 				.Value( InArgs._Roll )
@@ -209,7 +210,7 @@ public:
 				.LabelLocation(SNumericEntryBox<NumericType>::ELabelLocation::Inside)
 				.Label()
 				[
-					InArgs._bColorAxisLabels ? SNumericEntryBox<NumericType>::BuildNarrowColorLabel(SNumericEntryBox<NumericType>::GreenLabelBackgroundColor) : SNullWidget::NullWidget
+					InArgs._bColorAxisLabels ? SNumericEntryBox<NumericType>::BuildNarrowColorLabel(AxisDisplayInfo::GetAxisColor(EAxisList::Y)) : SNullWidget::NullWidget
 				]
 				.Font( InArgs._Font )
 				.Value( InArgs._Pitch )
@@ -246,7 +247,7 @@ public:
 				.LabelLocation(SNumericEntryBox<NumericType>::ELabelLocation::Inside)
 				.Label()
 				[
-					InArgs._bColorAxisLabels ? SNumericEntryBox<NumericType>::BuildNarrowColorLabel(SNumericEntryBox<NumericType>::BlueLabelBackgroundColor) : SNullWidget::NullWidget
+					InArgs._bColorAxisLabels ? SNumericEntryBox<NumericType>::BuildNarrowColorLabel(AxisDisplayInfo::GetAxisColor(EAxisList::Z)) : SNullWidget::NullWidget
 				]
 				.Font( InArgs._Font )
 				.Value( InArgs._Yaw )
diff --git a/Engine/Source/Runtime/Slate/Public/Widgets/Input/SVectorInputBox.h b/Engine/Source/Runtime/Slate/Public/Widgets/Input/SVectorInputBox.h
--- a/Engine/Source/Runtime/Slate/Public/Widgets/Input/SVectorInputBox.h
+++ b/Engine/Source/Runtime/Slate/Public/Widgets/Input/SVectorInputBox.h
@@ -4,6 +4,7 @@
 
 #include "CoreMinimal.h"
 #include "Layout/Margin.h"
+#include "Misc/AxisDisplayInfo.h"
 #include "Fonts/SlateFontInfo.h"
 #include "Widgets/DeclarativeSyntaxSupport.h"
 #include "Widgets/Input/NumericTypeInterface.h"
@@ -476,7 +477,7 @@ private:
 	{
 		ConstructComponent(0,
 			InArgs,
-			SNumericEntryBox<NumericType>::RedLabelBackgroundColor,
+			AxisDisplayInfo::GetAxisColor(EAxisList::X),
 			NSLOCTEXT("SVectorInputBox", "X_ToolTip", "X: {0}"),
 			HorizontalBox,
 			InArgs._X,
@@ -497,7 +498,7 @@ private:
 	{
 		ConstructComponent(1,
 			InArgs,
-			SNumericEntryBox<NumericType>::GreenLabelBackgroundColor,
+			AxisDisplayInfo::GetAxisColor(EAxisList::Y),
 			NSLOCTEXT("SVectorInputBox", "Y_ToolTip", "Y: {0}"),
 			HorizontalBox,
 			InArgs._Y,
@@ -518,7 +519,7 @@ private:
 	{
 		ConstructComponent(2,
 			InArgs,
-			SNumericEntryBox<NumericType>::BlueLabelBackgroundColor,
+			AxisDisplayInfo::GetAxisColor(EAxisList::Z),
 			NSLOCTEXT("SVectorInputBox", "Z_ToolTip", "Z: {0}"),
 			HorizontalBox,
 			InArgs._Z,
