Skein: Initial work on autosaves.

#gh 1525
#rnx

diff --git a/Engine/Source/Developer/SourceControl/Public/SourceControlOperations.h b/Engine/Source/Developer/SourceControl/Public/SourceControlOperations.h
--- a/Engine/Source/Developer/SourceControl/Public/SourceControlOperations.h
+++ b/Engine/Source/Developer/SourceControl/Public/SourceControlOperations.h
@@ -1544,4 +1544,22 @@ private:
 	TArray<FileInfo> Files;
 };
 
+/**
+ * This operation will save files into revision control without actually doing a check-in
+ */
+class FSave : public FSourceControlOperationBase
+{
+public:
+	// ISourceControlOperation interface
+	virtual FName GetName() const override
+	{
+		return "Save";
+	}
+
+	virtual FText GetInProgressString() const override
+	{
+		return LOCTEXT("SourceControl_Save", "Saving file(s) into Revision Control...");
+	}
+};
+
 #undef LOCTEXT_NAMESPACE
diff --git a/Engine/Source/Editor/UnrealEd/Private/FileHelpers.cpp b/Engine/Source/Editor/UnrealEd/Private/FileHelpers.cpp
--- a/Engine/Source/Editor/UnrealEd/Private/FileHelpers.cpp
+++ b/Engine/Source/Editor/UnrealEd/Private/FileHelpers.cpp
@@ -4216,7 +4216,13 @@ FEditorFileUtils::EPromptReturnCode InternalPromptForCheckoutAndSave(const TArra
 	PackagesToSave.Reserve(FinalSaveList.Num());
 
 	{
-		FScopedSlowTask SlowTask(static_cast<float>(FinalSaveList.Num() * 2), NSLOCTEXT("UnrealEd", "SavingPackages", "Saving packages..."));
+		auto SaveOperation = ISourceControlOperation::Create<FSave>();
+		bool bCanSaveToRevisionControl = 
+			ISourceControlModule::Get().IsEnabled() &&
+			ISourceControlModule::Get().GetProvider().IsAvailable() &&
+			ISourceControlModule::Get().GetProvider().CanExecuteOperation(SaveOperation);
+
+		FScopedSlowTask SlowTask(static_cast<float>(FinalSaveList.Num() * (bCanSaveToRevisionControl ? 3 : 2)), NSLOCTEXT("UnrealEd", "SavingPackages", "Saving packages..."));
 		SlowTask.MakeDialog();
 
 		UWorld* ActorsWorld = nullptr;
@@ -4306,6 +4312,14 @@ FEditorFileUtils::EPromptReturnCode InternalPromptForCheckoutAndSave(const TArra
 		{
 			FEditorDelegates::PostSaveExternalActors.Broadcast(ActorsWorld);
 		}
+
+		if (bCanSaveToRevisionControl)
+		{
+			FText SavingRevisionControlText = NSLOCTEXT("UnrealEd", "SavingRevisionControl", "Saving to Revision Control");
+			SlowTask.EnterProgressFrame(PackagesToSave.Num() + PackagesToClean.Num(), SavingRevisionControlText);
+
+			ISourceControlModule::Get().GetProvider().Execute(SaveOperation, USourceControlHelpers::PackageFilenames(PackagesToSave), EConcurrency::Synchronous);
+		}
 	}
 
 	SaveErrors.Flush();
