Asset guidelines will now detect project settings set in Base ini files.

Previously they would prompt the user to change the project setting unless it was set in a Default ini file, e.g. DefaultEngine.ini.

diff --git a/Engine/Source/Editor/UnrealEd/Classes/Editor/AssetGuideline.h b/Engine/Source/Editor/UnrealEd/Classes/Editor/AssetGuideline.h
--- a/Engine/Source/Editor/UnrealEd/Classes/Editor/AssetGuideline.h
+++ b/Engine/Source/Editor/UnrealEd/Classes/Editor/AssetGuideline.h
@@ -81,8 +81,10 @@ public:
 	static bool AreAssetGuidelinesEnabled() { return bAssetGuidelinesEnabled; }
 
 private:
-	UNREALED_API void EnableMissingGuidelines(TArray<FString> IncorrectPlugins, TArray<FIniStringValue> IncorrectProjectSettings);
-	UNREALED_API void DismissNotifications();
-	UNREALED_API void RemoveAssetGuideline();
+	[[nodiscard]] static bool TryDetectIniBranchFromFilename(const FString& Filename, FString& OutIniBranch);
+
+	void EnableMissingGuidelines(TArray<FString> IncorrectPlugins, TArray<FIniStringValue> IncorrectProjectSettings);
+	void DismissNotifications();
+	void RemoveAssetGuideline();
 #endif // WITH_EDITOR
 };
diff --git a/Engine/Source/Editor/UnrealEd/Private/AssetGuideline.cpp b/Engine/Source/Editor/UnrealEd/Private/AssetGuideline.cpp
--- a/Engine/Source/Editor/UnrealEd/Private/AssetGuideline.cpp
+++ b/Engine/Source/Editor/UnrealEd/Private/AssetGuideline.cpp
@@ -91,6 +91,26 @@ void UAssetGuideline::PostLoad()
 	{
 		if (IConsoleManager::Get().FindConsoleVariable(*ProjectSetting.Key))
 		{
+			FString Branch = ProjectSetting.Branch;
+			// If the branch was specified, or we can parse it from the filename, use that to read
+			// the ini value, as it will take Base inis into account.
+			if (!Branch.IsEmpty()
+				|| TryDetectIniBranchFromFilename(ProjectSetting.Filename, Branch))
+			{
+				FString CurrentIniValue;
+				if (GConfig->GetString(*ProjectSetting.Section, *ProjectSetting.Key, CurrentIniValue, Branch))
+				{
+					if (CurrentIniValue != ProjectSetting.Value)
+					{
+						NeededProjectSettings += FString::Printf(TEXT("[%s] %s = %s\n"), *ProjectSetting.Section, *ProjectSetting.Key, *ProjectSetting.Value);
+						IncorrectProjectSettings.Add(ProjectSetting);
+					}
+ 
+					continue;
+				}
+			}
+ 
+			// No branch could be detected. Read directly from the specified file.
 			FString CurrentIniValue;
 			FString FilenamePath = FConfigCacheIni::NormalizeConfigIniPath(FPaths::ProjectDir() + ProjectSetting.Filename);
 			if (GConfig->GetString(*ProjectSetting.Section, *ProjectSetting.Key, CurrentIniValue, FilenamePath))
@@ -249,6 +269,26 @@ void UAssetGuideline::BeginDestroy()
 	Super::BeginDestroy();
 }
 
+bool UAssetGuideline::TryDetectIniBranchFromFilename(const FString& Filename, FString& OutIniBranch)
+{
+	const FRegexPattern DefaultPathPattern(TEXT("^/Config/Default(\\w+)\\.ini$"), ERegexPatternFlags::CaseInsensitive);
+	FRegexMatcher Matcher(DefaultPathPattern, Filename);
+	if (Matcher.FindNext())
+	{
+		OutIniBranch = Matcher.GetCaptureGroup(1);
+		// The regex should ensure this capture group is non-empty
+		check(!OutIniBranch.IsEmpty());
+
+		UE_LOGFMT(LogAssetGuideline, Log, "Detected ini branch name {Branch} from filename {Filename}", OutIniBranch, Filename);
+		return true;
+	}
+	else
+	{
+		UE_LOGFMT(LogAssetGuideline, Error, "Failed to detect ini branch name from filename {Filename}. Please specify the ini branch explicitly in the asset guideline.", Filename);
+		return false;
+	}
+}
+
 void UAssetGuideline::EnableMissingGuidelines(TArray<FString> IncorrectPlugins, TArray<FIniStringValue> IncorrectProjectSettings)
 {
 	if (TSharedPtr<SNotificationItem> NotificationPin = NotificationPtr.Pin())
@@ -289,12 +329,8 @@ void UAssetGuideline::EnableMissingGuidelines(TArray<FString> IncorrectPlugins,
 				FString Branch = IncorrectProjectSetting.Branch;
 				if (Branch.Len() == 0)
 				{
-					// Try to detect the branch from the filename
-					const FRegexPattern DefaultPathPattern(TEXT("^/Config/Default(\\w+)\\.ini$"), ERegexPatternFlags::CaseInsensitive);
-					FRegexMatcher Matcher(DefaultPathPattern, IncorrectProjectSetting.Filename);
-					if (Matcher.FindNext())
+					if (TryDetectIniBranchFromFilename(IncorrectProjectSetting.Filename, Branch))
 					{
-						Branch = Matcher.GetCaptureGroup(1);
 						UE_LOGFMT(LogAssetGuideline, Log, "Detected ini branch name {Branch} from filename {Filename}", Branch, IncorrectProjectSetting.Filename);
 					}
 					else
@@ -395,4 +431,4 @@ void UAssetGuideline::RemoveAssetGuideline()
 
 #undef LOCTEXT_NAMESPACE 
 
-#endif // WITH_EDITOR
\ No newline at end of file
+#endif // WITH_EDITOR
