Cleaned up some logging in AssetGuideline

diff --git a/Engine/Source/Editor/UnrealEd/Private/AssetGuideline.cpp b/Engine/Source/Editor/UnrealEd/Private/AssetGuideline.cpp
--- a/Engine/Source/Editor/UnrealEd/Private/AssetGuideline.cpp
+++ b/Engine/Source/Editor/UnrealEd/Private/AssetGuideline.cpp
@@ -327,19 +327,12 @@ void UAssetGuideline::EnableMissingGuidelines(TArray<FString> IncorrectPlugins,
 			for (const FIniStringValue& IncorrectProjectSetting : IncorrectProjectSettings)
 			{
 				FString Branch = IncorrectProjectSetting.Branch;
-				if (Branch.Len() == 0)
+				if (Branch.Len() == 0
+					&& !TryDetectIniBranchFromFilename(IncorrectProjectSetting.Filename, Branch))
 				{
-					if (TryDetectIniBranchFromFilename(IncorrectProjectSetting.Filename, Branch))
-					{
-						UE_LOGFMT(LogAssetGuideline, Log, "Detected ini branch name {Branch} from filename {Filename}", Branch, IncorrectProjectSetting.Filename);
-					}
-					else
-					{
-						UE_LOGFMT(LogAssetGuideline, Error, "Failed to detect ini branch name from filename {Filename}. Please specify the ini branch explicitly in the asset guideline.", IncorrectProjectSetting.Filename);
-						
-						bSuccess = false;
-						break;
-					}
+					// Branch wasn't specified and couldn't be detected from the filename
+					bSuccess = false;
+					break;
 				}
 
 				const FString FilenamePath = FConfigCacheIni::NormalizeConfigIniPath(FPaths::ProjectDir() + IncorrectProjectSetting.Filename);
