Added additional platform and restricted paths to FSourceFileDatabase so that source navigation and querying works for modules under these paths

#rnx

diff --git a/Engine/Source/Editor/UnrealEd/Private/SourceCodeNavigation.cpp b/Engine/Source/Editor/UnrealEd/Private/SourceCodeNavigation.cpp
--- a/Engine/Source/Editor/UnrealEd/Private/SourceCodeNavigation.cpp
+++ b/Engine/Source/Editor/UnrealEd/Private/SourceCodeNavigation.cpp
@@ -1,35 +1,37 @@
 // Copyright Epic Games, Inc. All Rights Reserved.
 
-
 #include "SourceCodeNavigation.h"
-#include "HAL/PlatformStackWalk.h"
+
+#include "Async/AsyncWork.h"
+#include "Async/TaskGraphInterfaces.h"
+#include "Editor/UnrealEdEngine.h"
+#include "GameFramework/Actor.h"
+#include "GameFramework/Pawn.h"
 #include "GenericPlatform/GenericPlatformFile.h"
 #include "HAL/FileManager.h"
+#include "HAL/PlatformStackWalk.h"
+#include "HttpModule.h"
+#include "Interfaces/IHttpRequest.h"
+#include "Interfaces/IHttpResponse.h"
+#include "ISourceCodeAccessModule.h"
+#include "ISourceCodeAccessor.h"
+#include "Misc/ConfigCacheIni.h"
+#include "Misc/DataDrivenPlatformInfoRegistry.h"
+#include "Misc/EngineBuildSettings.h"
+#include "Misc/PackageName.h"
 #include "Misc/Paths.h"
 #include "Misc/ScopeLock.h"
-#include "Misc/ConfigCacheIni.h"
-#include "Async/AsyncWork.h"
 #include "Modules/ModuleManager.h"
-#include "UObject/Class.h"
-#include "UObject/UObjectHash.h"
-#include "UObject/UObjectIterator.h"
-#include "UObject/Package.h"
-#include "UObject/MetaData.h"
-#include "Misc/PackageName.h"
-#include "Async/TaskGraphInterfaces.h"
+#include "Preferences/UnrealEdOptions.h"
 #include "Styling/AppStyle.h"
-#include "GameFramework/Actor.h"
-#include "GameFramework/Pawn.h"
 #include "TickableEditorObject.h"
-#include "UnrealEdMisc.h"
-#include "ISourceCodeAccessor.h"
-#include "ISourceCodeAccessModule.h"
-#include "Interfaces/IHttpResponse.h"
-#include "Interfaces/IHttpRequest.h"
-#include "HttpModule.h"
-#include "Editor/UnrealEdEngine.h"
-#include "Preferences/UnrealEdOptions.h"
 #include "UnrealEdGlobals.h"
+#include "UnrealEdMisc.h"
+#include "UObject/Class.h"
+#include "UObject/MetaData.h"
+#include "UObject/Package.h"
+#include "UObject/UObjectHash.h"
+#include "UObject/UObjectIterator.h"
 
 #if PLATFORM_WINDOWS
 #include "Windows/AllowWindowsPlatformTypes.h"
@@ -222,44 +224,133 @@
 	ModuleNames.Reset();
 	DisallowedHeaderNames.Empty();
 
-	// Find all the build rules within the game and engine directories
-	FindRootFilesRecursive(ModuleNames, *(FPaths::EngineDir() / TEXT("Source") / TEXT("Developer")), TEXT("*.Build.cs"));
-	FindRootFilesRecursive(ModuleNames, *(FPaths::EngineDir() / TEXT("Source") / TEXT("Editor")), TEXT("*.Build.cs"));
-	FindRootFilesRecursive(ModuleNames, *(FPaths::EngineDir() / TEXT("Source") / TEXT("Runtime")), TEXT("*.Build.cs"));
-	FindRootFilesRecursive(ModuleNames, *(FPaths::ProjectDir() / TEXT("Source")), TEXT("*.Build.cs"));
+	auto AppendConvertedPaths = [](TArray<FString>& OutDirs, TArrayView<const FString> BaseDirs, FPaths::EPathConversion Method, const TCHAR* ExtraData = nullptr)
+	{
+		OutDirs.Reserve(OutDirs.Num() + BaseDirs.Num());
+		for (const FString& BaseDir : BaseDirs)
+		{
+			OutDirs.Add(FPaths::ConvertPath(BaseDir, Method, ExtraData));
+		}
+	};
 
+	auto AppendEnginePaths = [&AppendConvertedPaths](TArray<FString>& OutDirs, TArrayView<const FString> BaseDirs)
+	{
+		OutDirs.Append(BaseDirs);
+
+		for (const FString& PlatformName : FDataDrivenPlatformInfoRegistry::GetPlatformDirectoryNames(false))
+		{
+			AppendConvertedPaths(OutDirs, BaseDirs, FPaths::EPathConversion::Engine_PlatformExtension, *PlatformName);
+		}
+
+		if (FEngineBuildSettings::IsInternalBuild())
+		{
+			AppendConvertedPaths(OutDirs, BaseDirs, FPaths::EPathConversion::Engine_NoRedist);
+			AppendConvertedPaths(OutDirs, BaseDirs, FPaths::EPathConversion::Engine_NotForLicensees);
+		}
+	};
+
+	auto AppendProjectPaths = [&AppendConvertedPaths](TArray<FString>& OutDirs, TArrayView<const FString> BaseDirs)
+	{
+		OutDirs.Append(BaseDirs);
+
+		for (const FString& PlatformName : FDataDrivenPlatformInfoRegistry::GetPlatformDirectoryNames(false))
+		{
+			AppendConvertedPaths(OutDirs, BaseDirs, FPaths::EPathConversion::Project_PlatformExtension, *PlatformName);
+		}
+
+		if (FEngineBuildSettings::IsInternalBuild())
+		{
+			AppendConvertedPaths(OutDirs, BaseDirs, FPaths::EPathConversion::Project_NoRedist);
+			AppendConvertedPaths(OutDirs, BaseDirs, FPaths::EPathConversion::Project_NotForLicensees);
+		}
+	};
+
+	// Find all the non-plugin modules
+	{
+		TArray<FString> ModuleSearchDirs;
+
+		// Build the list of Engine directories to search for modules
+		{
+			const FString EngineModuleSearchDirs[] = {
+				FPaths::Combine(FPaths::EngineDir(), TEXT("Source"), TEXT("Developer")),
+				FPaths::Combine(FPaths::EngineDir(), TEXT("Source"), TEXT("Editor")),
+				FPaths::Combine(FPaths::EngineDir(), TEXT("Source"), TEXT("Runtime")),
+			};
+			AppendEnginePaths(ModuleSearchDirs, EngineModuleSearchDirs);
+		}
+
+		// Build the list of Project directories to search for modules
+		if (FApp::HasProjectName())
+		{
+			const FString ProjectModuleSearchDirs[] = {
+				FPaths::Combine(FPaths::ProjectDir(), TEXT("Source")),
+			};
+			AppendProjectPaths(ModuleSearchDirs, ProjectModuleSearchDirs);
+		}
+
+		// Find all the build rules within the module search directories
+		for (const FString& ModuleSearchDir : ModuleSearchDirs)
+		{
+			FindRootFilesRecursive(ModuleNames, ModuleSearchDir, TEXT("*.Build.cs"));
+		}
+	}
+
 	// Find list of disallowed header names in native (non-plugin) directories
-	TArray<FString> HeaderFiles;
-	for (const FString& ModuleName : ModuleNames)
 	{
-		IFileManager::Get().FindFilesRecursive(HeaderFiles, *(FPaths::GetPath(ModuleName) / TEXT("Classes")), TEXT("*.h"), true, false, false);
-		IFileManager::Get().FindFilesRecursive(HeaderFiles, *(FPaths::GetPath(ModuleName) / TEXT("Public")), TEXT("*.h"), true, false, false);
+		TArray<FString> HeaderFiles;
+		for (const FString& ModuleName : ModuleNames)
+		{
+			IFileManager::Get().FindFilesRecursive(HeaderFiles, *FPaths::Combine(FPaths::GetPath(ModuleName), TEXT("Classes")), TEXT("*.h"), true, false, false);
+			IFileManager::Get().FindFilesRecursive(HeaderFiles, *FPaths::Combine(FPaths::GetPath(ModuleName), TEXT("Public")), TEXT("*.h"), true, false, false);
+		}
+
+		for (const FString& HeaderFile : HeaderFiles)
+		{
+			DisallowedHeaderNames.Add(FPaths::GetBaseFilename(HeaderFile));
+		}
+
+		for (TObjectIterator<UClass> ClassIt; ClassIt; ++ClassIt)
+		{
+			DisallowedHeaderNames.Remove(ClassIt->GetName());
+		}
 	}
 
-	for (const FString& HeaderFile : HeaderFiles)
+	// Find all the plugin modules
 	{
-		DisallowedHeaderNames.Add(FPaths::GetBaseFilename(HeaderFile));
+		TArray<FString> PluginSearchDirs;
+
+		// Build the list of Engine directories to search for plugins
+		{
+			const FString EnginePluginSearchDirs[] = {
+				FPaths::Combine(FPaths::EngineDir(), TEXT("Plugins")),
+			};
+			AppendEnginePaths(PluginSearchDirs, EnginePluginSearchDirs);
+		}
+
+		// Build the list of Project directories to search for plugins
+		if (FApp::HasProjectName())
+		{
+			const FString ProjectPluginSearchDirs[] = {
+				FPaths::Combine(FPaths::ProjectDir(), TEXT("Plugins")),
+			};
+			AppendProjectPaths(PluginSearchDirs, ProjectPluginSearchDirs);
+		}
+
+		// Find all the plugin directories
+		TArray<FString> PluginNames;
+		for (const FString& PluginSearchDir : PluginSearchDirs)
+		{
+			FindRootFilesRecursive(PluginNames, PluginSearchDir, TEXT("*uplugin"));
+		}
+
+		// Add all the files within plugin directories
+		for (const FString& PluginName : PluginNames)
+		{
+			FindRootFilesRecursive(ModuleNames, FPaths::Combine(FPaths::GetPath(PluginName), TEXT("Source")), TEXT("*.Build.cs"));
+		}
 	}
-
-	for (TObjectIterator<UClass> ClassIt; ClassIt; ++ClassIt)
-	{
-		DisallowedHeaderNames.Remove(ClassIt->GetName());
-	}
-
-	// Find all the plugin directories
-	TArray<FString> PluginNames;
-
-	FindRootFilesRecursive(PluginNames, *(FPaths::EngineDir() / TEXT("Plugins")), TEXT("*.uplugin"));
-	FindRootFilesRecursive(PluginNames, *(FPaths::ProjectDir() / TEXT("Plugins")), TEXT("*.uplugin"));
-
-	// Add all the files within plugin directories
-	for (const FString& PluginName : PluginNames)
-	{
-		FindRootFilesRecursive(ModuleNames, *(FPaths::GetPath(PluginName) / TEXT("Source")), TEXT("*.Build.cs"));
-	}
 }
 
-
 void FSourceFileDatabase::FindRootFilesRecursive(TArray<FString> &FileNames, const FString &BaseDirectory, const FString &Wildcard)
 {
 	// Find all the files within this directory
