[MaterialEditor] Integrating an old PR that adds a preview material domain setting to the material function editor. Will also be used to generate a 2d thumbnail for material function assets that are set to "User Interface" domain which significantly improves asset discoverability in the editor.

Also fixed an issue where material functions would select a random output for their thumbnail preview material (generating unintended thumbnail renders). Now uses the last previewed output instead.

[RN] minor creator/ui

diff --git a/Engine/Source/Editor/MaterialEditor/Private/MaterialEditor.cpp b/Engine/Source/Editor/MaterialEditor/Private/MaterialEditor.cpp
--- a/Engine/Source/Editor/MaterialEditor/Private/MaterialEditor.cpp
+++ b/Engine/Source/Editor/MaterialEditor/Private/MaterialEditor.cpp
@@ -554,6 +554,7 @@ void FMaterialEditor::InitEditorForMaterialFunction(UMaterialFunction* InMateria
 		Material->Serialize(DummyArchive);
 	}
 	Material->SetShadingModel(MSM_Unlit);
+	Material->MaterialDomain = InMaterialFunction->PreviewMaterialDomain;
 
 	// Propagate all object flags except for RF_Standalone, otherwise the preview material function won't GC once
 	// the material editor releases the reference.
@@ -7487,8 +7488,14 @@ void FMaterialEditor::NotifyPostChange( const FPropertyChangedEvent& PropertyCha
 			}
 		}
 		else if( NameOfPropertyThatChanged == GET_MEMBER_NAME_CHECKED(UMaterial, MaterialDomain) ||
-				 NameOfPropertyThatChanged == GET_MEMBER_NAME_CHECKED(UMaterial, ShadingModel))
+				 NameOfPropertyThatChanged == GET_MEMBER_NAME_CHECKED(UMaterial, ShadingModel) ||
+				 NameOfPropertyThatChanged == GET_MEMBER_NAME_CHECKED(UMaterialFunction, PreviewMaterialDomain))
 		{
+			if (MaterialFunction)
+			{
+				Material->MaterialDomain = MaterialFunction->PreviewMaterialDomain;
+			}
+			
 			Material->MaterialGraph->RebuildGraph();
 			TArray<TWeakObjectPtr<UObject>> SelectedObjects = MaterialDetailsView->GetSelectedObjects();
 			MaterialDetailsView->SetObjects( SelectedObjects, true );
diff --git a/Engine/Source/Editor/MaterialEditor/Private/MaterialEditorDetailCustomization.cpp b/Engine/Source/Editor/MaterialEditor/Private/MaterialEditorDetailCustomization.cpp
--- a/Engine/Source/Editor/MaterialEditor/Private/MaterialEditorDetailCustomization.cpp
+++ b/Engine/Source/Editor/MaterialEditor/Private/MaterialEditorDetailCustomization.cpp
@@ -40,6 +40,7 @@
 #include "PropertyCustomizationHelpers.h"
 #include "Curves/CurveLinearColor.h"
 #include "IPropertyUtilities.h"
+#include "MaterialDomain.h"
 #include "Engine/Texture.h"
 #include "Materials/MaterialExpressionCurveAtlasRowParameter.h"
 #include "Curves/CurveLinearColorAtlas.h"
@@ -953,6 +954,32 @@ void FMaterialFunctionDetailCustomization::CustomizeDetails(IDetailLayoutBuilder
 			}
 		}
 	}
+
+	// Hide Thumbnail category if material function preview domain is set to UserInterface
+	{
+		TArray<TWeakObjectPtr<UObject> > Objects;
+		DetailLayout.GetObjectsBeingCustomized(Objects);
+		bool bUIMaterial = true;
+		for (TWeakObjectPtr<UObject>& Object : Objects)
+		{
+			UMaterialFunction* MaterialFunction = Cast<UMaterialFunction>(Object.Get());
+			if (MaterialFunction)
+			{
+				bUIMaterial &= MaterialFunction->PreviewMaterialDomain == MD_UI;
+			}
+			else
+			{
+				// this shouldn't happen but just in case, let all properties through
+				bUIMaterial = false;
+				break;
+			}
+		}
+
+		if (bUIMaterial)
+		{
+			DetailLayout.HideCategory(TEXT("Thumbnail"));
+		}
+	}
 }
 
 TSharedRef<class IDetailCustomization> FMaterialExpressionLayersParameterDetails::MakeInstance(FOnCollectParameterGroups InCollectGroupsDelegate)
diff --git a/Engine/Source/Editor/UnrealEd/Classes/ThumbnailRendering/MaterialFunctionThumbnailRenderer.h b/Engine/Source/Editor/UnrealEd/Classes/ThumbnailRendering/MaterialFunctionThumbnailRenderer.h
--- a/Engine/Source/Editor/UnrealEd/Classes/ThumbnailRendering/MaterialFunctionThumbnailRenderer.h
+++ b/Engine/Source/Editor/UnrealEd/Classes/ThumbnailRendering/MaterialFunctionThumbnailRenderer.h
@@ -31,5 +31,7 @@ class UMaterialFunctionThumbnailRenderer : public UDefaultSizedThumbnailRenderer
 
 private:
 	class FMaterialThumbnailScene* ThumbnailScene;
+	class FWidgetRenderer* WidgetRenderer;
+	struct FSlateMaterialBrush* UIMaterialBrush;
 };
 
diff --git a/Engine/Source/Editor/UnrealEd/Private/ThumbnailRendering/MaterialFunctionThumbnailRenderer.cpp b/Engine/Source/Editor/UnrealEd/Private/ThumbnailRendering/MaterialFunctionThumbnailRenderer.cpp
--- a/Engine/Source/Editor/UnrealEd/Private/ThumbnailRendering/MaterialFunctionThumbnailRenderer.cpp
+++ b/Engine/Source/Editor/UnrealEd/Private/ThumbnailRendering/MaterialFunctionThumbnailRenderer.cpp
@@ -9,11 +9,19 @@
 #include "Materials/Material.h"
 #include "Materials/MaterialInstanceConstant.h"
 #include "ThumbnailHelpers.h"
+#include "Slate/WidgetRenderer.h"
+#include "Widgets/Images/SImage.h"
+#include "SlateMaterialBrush.h"
+#include "ThumbnailRendering/ThumbnailManager.h"
 #include "Settings/ContentBrowserSettings.h"
+#include "Engine/Texture2D.h"
 
 
 UMaterialFunctionThumbnailRenderer::UMaterialFunctionThumbnailRenderer(const FObjectInitializer& ObjectInitializer)
 	: Super(ObjectInitializer)
+	, ThumbnailScene(nullptr)
+	, WidgetRenderer(nullptr)
+	, UIMaterialBrush(nullptr)
 {
 	ThumbnailScene = nullptr;
 }
@@ -26,40 +34,91 @@ void UMaterialFunctionThumbnailRenderer::Draw(UObject* Object, int32 X, int32 Y,
 
 	if (MatFunc || bIsFunctionInstancePreview)
 	{
-		if (ThumbnailScene == nullptr || ensure(ThumbnailScene->GetWorld() != nullptr) == false)
+		UMaterialInterface* PreviewMaterial = bIsFunctionInstancePreview ? MatFuncInst->GetPreviewMaterial() : MatFunc->GetPreviewMaterial();
+		UMaterial* Mat = PreviewMaterial ? PreviewMaterial->GetMaterial() : nullptr;
+
+		if (Mat != nullptr && Mat->IsUIMaterial())
 		{
-			if (ThumbnailScene)
+			if (WidgetRenderer == nullptr)
 			{
-				FlushRenderingCommands();
-				delete ThumbnailScene;
+				const bool bUseGammaCorrection = true;
+				WidgetRenderer = new FWidgetRenderer(bUseGammaCorrection);
+				check(WidgetRenderer);
 			}
 
-			ThumbnailScene = new FMaterialThumbnailScene();
-		}
+			if (UIMaterialBrush == nullptr)
+			{
+				UIMaterialBrush = new FSlateMaterialBrush(FVector2D(SlateBrushDefs::DefaultImageSize, SlateBrushDefs::DefaultImageSize));
+				check(UIMaterialBrush);
+			}
 
-		UMaterialInterface* PreviewMaterial = bIsFunctionInstancePreview ? MatFuncInst->GetPreviewMaterial() : MatFunc->GetPreviewMaterial();
-		EMaterialFunctionUsage FunctionUsage = bIsFunctionInstancePreview ? MatFuncInst->GetMaterialFunctionUsage() : MatFunc->GetMaterialFunctionUsage();
-		UThumbnailInfo* ThumbnailInfo = bIsFunctionInstancePreview ? MatFuncInst->ThumbnailInfo : MatFunc->ThumbnailInfo;
+			UIMaterialBrush->SetMaterial(PreviewMaterial);
+
+			UTexture2D* CheckerboardTexture = UThumbnailManager::Get().CheckerboardTexture;
+
+			FSlateBrush CheckerboardBrush;
+			CheckerboardBrush.SetResourceObject(CheckerboardTexture);
+			CheckerboardBrush.ImageSize = FVector2D(CheckerboardTexture->GetSizeX(), CheckerboardTexture->GetSizeY());
+			CheckerboardBrush.Tiling = ESlateBrushTileType::Both;
 
-		if (PreviewMaterial)
+			TSharedRef<SWidget> Thumbnail =
+				SNew(SOverlay)
+
+				// Checkerboard
+				+ SOverlay::Slot()
+				[
+					SNew(SImage)
+					.Image(&CheckerboardBrush)
+				]
+
+				+ SOverlay::Slot()
+				[
+					SNew(SImage)
+					.Image(UIMaterialBrush)
+				];
+
+			const FVector2D DrawSize((float)Width, (float)Height);
+			const float DeltaTime = 0.f;
+			WidgetRenderer->DrawWidget(RenderTarget, Thumbnail, DrawSize, DeltaTime);
+
+			UIMaterialBrush->SetMaterial(nullptr);
+		}
+		else
 		{
+			if (ThumbnailScene == nullptr || ensure(ThumbnailScene->GetWorld() != nullptr) == false)
+			{
+				if (ThumbnailScene)
+				{
+					FlushRenderingCommands();
+					delete ThumbnailScene;
+				}
+
+				ThumbnailScene = new FMaterialThumbnailScene();
+			}
+
+			EMaterialFunctionUsage FunctionUsage = bIsFunctionInstancePreview ? MatFuncInst->GetMaterialFunctionUsage() : MatFunc->GetMaterialFunctionUsage();
+			UThumbnailInfo* ThumbnailInfo = bIsFunctionInstancePreview ? MatFuncInst->ThumbnailInfo : MatFunc->ThumbnailInfo;
+
+			if (PreviewMaterial)
+			{
 				PreviewMaterial->ThumbnailInfo = ThumbnailInfo;
 				if (FunctionUsage == EMaterialFunctionUsage::MaterialLayerBlend)
 				{
 					PreviewMaterial->SetShouldForcePlanePreview(true);
 				}
 				ThumbnailScene->SetMaterialInterface(PreviewMaterial);
-	
-			FSceneViewFamilyContext ViewFamily( FSceneViewFamily::ConstructionValues( RenderTarget, ThumbnailScene->GetScene(), FEngineShowFlags(ESFIM_Game) )
-				.SetTime(UThumbnailRenderer::GetTime())
-				.SetAdditionalViewFamily(bAdditionalViewFamily));
 
-			ViewFamily.EngineShowFlags.DisableAdvancedFeatures();
-			ViewFamily.EngineShowFlags.MotionBlur = 0;
+				FSceneViewFamilyContext ViewFamily( FSceneViewFamily::ConstructionValues( RenderTarget, ThumbnailScene->GetScene(), FEngineShowFlags(ESFIM_Game) )
+					.SetTime(UThumbnailRenderer::GetTime())
+					.SetAdditionalViewFamily(bAdditionalViewFamily));
 
-			RenderViewFamily(Canvas, &ViewFamily, ThumbnailScene->CreateView(&ViewFamily, X, Y, Width, Height));
+				ViewFamily.EngineShowFlags.DisableAdvancedFeatures();
+				ViewFamily.EngineShowFlags.MotionBlur = 0;
 
-			ThumbnailScene->SetMaterialInterface(nullptr);
+				RenderViewFamily(Canvas, &ViewFamily, ThumbnailScene->CreateView(&ViewFamily, X, Y, Width, Height));
+
+				ThumbnailScene->SetMaterialInterface(nullptr);
+			}
 		}
 	}
 }
@@ -70,12 +129,24 @@ bool UMaterialFunctionThumbnailRenderer::CanVisualizeAsset(UObject* Object)
 }
 
 void UMaterialFunctionThumbnailRenderer::BeginDestroy()
-{ 	
-	if (ThumbnailScene)
+{
+	if (ThumbnailScene != nullptr)
 	{
 		delete ThumbnailScene;
 		ThumbnailScene = nullptr;
 	}
 
+	if (WidgetRenderer != nullptr)
+	{
+		BeginCleanup(WidgetRenderer);
+		WidgetRenderer = nullptr;
+	}
+
+	if (UIMaterialBrush != nullptr)
+	{
+		delete UIMaterialBrush;
+		UIMaterialBrush = nullptr;
+	}
+
 	Super::BeginDestroy();
 }
diff --git a/Engine/Source/Runtime/Engine/Private/Materials/MaterialExpressions.cpp b/Engine/Source/Runtime/Engine/Private/Materials/MaterialExpressions.cpp
--- a/Engine/Source/Runtime/Engine/Private/Materials/MaterialExpressions.cpp
+++ b/Engine/Source/Runtime/Engine/Private/Materials/MaterialExpressions.cpp
@@ -13536,16 +13536,26 @@ UMaterialInterface* UMaterialFunction::GetPreviewMaterial()
 		//Update cached expression data to ensure function calls are populated for resolving the preview
 		PreviewMaterial->UpdateCachedExpressionData();
 
-		//Find the first output expression and use that. 
+		// Find the output with bLastPreviewed set to true preferably expression.
+		UMaterialExpressionFunctionOutput* PreviewOutput = nullptr;
 		for (UMaterialExpression* Expression : GetExpressions())
 		{
 			UMaterialExpressionFunctionOutput* Output = Cast<UMaterialExpressionFunctionOutput>(Expression);
-			if (Output)
+			if (Output && (!PreviewOutput || Output->bLastPreviewed))
 			{
-				Output->ConnectToPreviewMaterial(PreviewMaterial, 0);
+				PreviewOutput = Output;
+				if (Output->bLastPreviewed)
+				{
+					break;
+				}
 			}
 		}
 
+		// Set the chosen output as preview.
+		PreviewOutput->ConnectToPreviewMaterial(PreviewMaterial, 0);
+
+		PreviewMaterial->MaterialDomain = PreviewMaterialDomain;
+
 		//Compile the material.
 		PreviewMaterial->PreEditChange(nullptr);
 		PreviewMaterial->PostEditChange();
diff --git a/Engine/Source/Runtime/Engine/Public/Materials/MaterialFunction.h b/Engine/Source/Runtime/Engine/Public/Materials/MaterialFunction.h
--- a/Engine/Source/Runtime/Engine/Public/Materials/MaterialFunction.h
+++ b/Engine/Source/Runtime/Engine/Public/Materials/MaterialFunction.h
@@ -9,6 +9,7 @@
 #include "Templates/Casts.h"
 #include "Materials/MaterialFunctionInterface.h"
 #include "StaticParameterSet.h"
+#include "MaterialDomain.h"
 #include "MaterialExpression.h"
 #include "MaterialFunction.generated.h"
 
@@ -104,6 +105,10 @@ class UMaterialFunction : public UMaterialFunctionInterface
 	UPROPERTY(EditAnywhere, Category = Preview, AssetRegistrySearchable)
 	TEnumAsByte<enum EBlendMode> PreviewBlendMode = BLEND_Opaque;
 
+	/** The domain to use when previewing a material function. */
+	UPROPERTY(EditAnywhere, Category = Preview, AssetRegistrySearchable)
+	TEnumAsByte<enum EMaterialDomain> PreviewMaterialDomain;
+
 	class UMaterialGraph* MaterialGraph = nullptr;
 
 	/* Whether all expressions in the function loaded correctly. */
