[Transforms] Updating viewport grid drawing to pull from AxisDisplayInfo colors

diff --git a/Engine/Source/Editor/UnrealEd/Private/EditorComponents.cpp b/Engine/Source/Editor/UnrealEd/Private/EditorComponents.cpp
--- a/Engine/Source/Editor/UnrealEd/Private/EditorComponents.cpp
+++ b/Engine/Source/Editor/UnrealEd/Private/EditorComponents.cpp
@@ -14,6 +14,7 @@
 #include "SceneView.h"
 #include "EditorModeManager.h"
 #include "GameFramework/WorldSettings.h"
+#include "Misc/AxisDisplayInfo.h"
 
 static TAutoConsoleVariable<int32> CVarEditorNewLevelGrid(
 	TEXT("r.Editor.NewLevelGrid"),
@@ -117,25 +118,19 @@ UMaterialInstanceDynamic* FGridWidget::GetActiveLevelGridMID()
 	return LevelGridMaterialInst;
 }
 
-static void GetAxisColors(FLinearColor Out[3], bool b3D)
+static void GetAxisColors(FLinearColor Out[3], bool bIsPerspective)
 {
-	Out[0] = FLinearColor::Red;
-	Out[1] = FLinearColor::Green;
-	Out[2] = FLinearColor::Blue;
+	Out[0] = AxisDisplayInfo::GetAxisColor(EAxisList::X);
+	Out[1] = AxisDisplayInfo::GetAxisColor(EAxisList::Y);
+	Out[2] = AxisDisplayInfo::GetAxisColor(EAxisList::Z);
 
-	// less prominent axis lines
-	for(int i = 0; i < 3; ++i)
+	// less prominent axis lines for perspective views
+	if (bIsPerspective)
 	{
-		// darker
-		if(b3D)
+		for (int i = 0; i < 3; ++i)
 		{
 			Out[i] += FLinearColor(0.2f, 0.2f, 0.2f, 0);
-			Out[i] *= 0.1f;
-		}
-		else
-		{
-//			Out[i] += FLinearColor(0.5f, 0.5f, 0.5f, 0);
-			Out[i] *= 0.5f;
+			Out[i] *= 0.3f;
 		}
 	}
 }
@@ -269,7 +264,7 @@ void FGridWidget::DrawNewGrid(const FSceneView* View, FPrimitiveDrawInterface* P
 	ObjectToWorld.SetOrigin(UVCameraPos);
 
 	FLinearColor AxisColors[3];
-	GetAxisColors(AxisColors, true);
+	GetAxisColors(AxisColors, bIsPerspective);
 
 	FLinearColor UAxisColor = AxisColors[1];
 	FLinearColor VAxisColor = AxisColors[0];
@@ -316,7 +311,7 @@ void FGridWidget::DrawNewGrid(const FSceneView* View, FPrimitiveDrawInterface* P
 	if(bIsPerspective)
 	{
 		// the higher we get the larger we make the geometry to give the illusion of an infinite grid while maintains the precision nearby
-		Radii *= FMath::Max<double>( 1.0, FMath::Abs(CameraPos.Z) / 1000.0 );
+		Radii *= FMath::Max( 1.0, FMath::Abs(CameraPos.Z) / 1000.0 );
 	}
 	else
 	{
@@ -484,7 +479,7 @@ void FEditorCommonDrawHelper::DrawOldGrid(const FSceneView* View,FPrimitiveDrawI
 		const bool bIsOrthoYZ = ( FMath::Abs(View->ViewMatrices.GetViewMatrix().M[0][2]) > 0 );
 
 		FLinearColor AxisColors[3];
-		GetAxisColors(AxisColors, false);
+		GetAxisColors(AxisColors, bIsPerspective);
 
 		if (bIsOrthoXY)
 		{
diff --git a/Engine/Source/Editor/WorldPartitionEditor/Private/WorldPartition/SWorldPartitionEditorGrid2D.cpp b/Engine/Source/Editor/WorldPartitionEditor/Private/WorldPartition/SWorldPartitionEditorGrid2D.cpp
--- a/Engine/Source/Editor/WorldPartitionEditor/Private/WorldPartition/SWorldPartitionEditorGrid2D.cpp
+++ b/Engine/Source/Editor/WorldPartitionEditor/Private/WorldPartition/SWorldPartitionEditorGrid2D.cpp
@@ -47,6 +47,7 @@
 #include "WorldPartition/WorldPartitionMiniMapVolume.h"
 #include "WorldPartition/WorldPartitionSubsystem.h"
 #include "HAL/PlatformApplicationMisc.h"
+#include "Misc/AxisDisplayInfo.h"
 
 #include "IAssetViewport.h"
 #include "Kismet2/DebuggerCommands.h"
@@ -1162,7 +1163,7 @@ int32 SWorldPartitionEditorGrid2D::PaintGrid(const FGeometry& AllottedGeometry,
 		LinePoints[0] = FVector2D(ScreenWorldOrigin.X, ScreenRect.Min.Y);
 		LinePoints[1] = FVector2D(ScreenWorldOrigin.X, ScreenRect.Max.Y);
 
-		FLinearColor YAxisColor = FLinearColor::Green;
+		FLinearColor YAxisColor = AxisDisplayInfo::GetAxisColor(EAxisList::Y);
 		YAxisColor.A = 0.4f;
 		
 		FSlateDrawElement::MakeLines(
@@ -1183,7 +1184,7 @@ int32 SWorldPartitionEditorGrid2D::PaintGrid(const FGeometry& AllottedGeometry,
 		LinePoints[0] = FVector2D(ScreenRect.Min.X, ScreenWorldOrigin.Y);
 		LinePoints[1] = FVector2D(ScreenRect.Max.X, ScreenWorldOrigin.Y);
 
-		FLinearColor XAxisColor = FLinearColor::Red;
+		FLinearColor XAxisColor = AxisDisplayInfo::GetAxisColor(EAxisList::X);
 		XAxisColor.A = 0.4f;
 		
 		FSlateDrawElement::MakeLines(
