Exposed some extra DLC and validation options to ICookInfo

diff --git a/Engine/Source/Editor/UnrealEd/Classes/CookOnTheSide/CookOnTheFlyServer.h b/Engine/Source/Editor/UnrealEd/Classes/CookOnTheSide/CookOnTheFlyServer.h
--- a/Engine/Source/Editor/UnrealEd/Classes/CookOnTheSide/CookOnTheFlyServer.h
+++ b/Engine/Source/Editor/UnrealEd/Classes/CookOnTheSide/CookOnTheFlyServer.h
@@ -672,7 +672,9 @@ public:
 	UNREALED_API virtual TArray<UE::Cook::FInstigator> GetInstigatorChain(FName PackageName) override;
 	UNREALED_API virtual UE::Cook::ECookType GetCookType() override;
 	UNREALED_API virtual UE::Cook::ECookingDLC GetCookingDLC() override;
+	UNREALED_API virtual FString GetDLCName() override;
 	UNREALED_API virtual UE::Cook::EProcessType GetProcessType() override;
+	UNREALED_API virtual UE::Cook::ECookValidationOptions GetCookValidationOptions() override;
 	UNREALED_API virtual bool IsIterative() override;
 	UNREALED_API virtual TArray<const ITargetPlatform*> GetSessionPlatforms() override;
 	UNREALED_API virtual FString GetCookOutputFolder(const ITargetPlatform* TargetPlatform) override;
diff --git a/Engine/Source/Editor/UnrealEd/Private/CookOnTheFlyServer.cpp b/Engine/Source/Editor/UnrealEd/Private/CookOnTheFlyServer.cpp
--- a/Engine/Source/Editor/UnrealEd/Private/CookOnTheFlyServer.cpp
+++ b/Engine/Source/Editor/UnrealEd/Private/CookOnTheFlyServer.cpp
@@ -6487,6 +6487,11 @@ UE::Cook::ECookingDLC UCookOnTheFlyServer::GetCookingDLC()
 	return IsCookingDLC() ? ECookingDLC::Yes : ECookingDLC::No;
 }
 
+FString UCookOnTheFlyServer::GetDLCName()
+{
+	return CookByTheBookOptions->DlcName;
+}
+
 UE::Cook::EProcessType UCookOnTheFlyServer::GetProcessType()
 {
 	using namespace UE::Cook;
@@ -6504,6 +6509,25 @@ UE::Cook::EProcessType UCookOnTheFlyServer::GetProcessType()
 	}
 }
 
+UE::Cook::ECookValidationOptions UCookOnTheFlyServer::GetCookValidationOptions()
+{
+	using namespace UE::Cook;
+	ECookValidationOptions ValidationOptions = ECookValidationOptions::None;
+	if (EnumHasAnyFlags(CookByTheBookOptions->StartupOptions, ECookByTheBookOptions::RunAssetValidation))
+	{
+		ValidationOptions |= ECookValidationOptions::RunAssetValidation;
+	}
+	if (EnumHasAnyFlags(CookByTheBookOptions->StartupOptions, ECookByTheBookOptions::RunMapValidation))
+	{
+		ValidationOptions |= ECookValidationOptions::RunMapValidation;
+	}
+	if (EnumHasAnyFlags(CookByTheBookOptions->StartupOptions, ECookByTheBookOptions::ValidationErrorsAreFatal))
+	{
+		ValidationOptions |= ECookValidationOptions::ValidationErrorsAreFatal;
+	}
+	return ValidationOptions;
+}
+
 bool UCookOnTheFlyServer::IsIterative()
 {
 	// TODO: For simplicity, we provide a single bool for all platforms in multiprocess cooks
diff --git a/Engine/Source/Runtime/CoreUObject/Public/UObject/CookEnums.h b/Engine/Source/Runtime/CoreUObject/Public/UObject/CookEnums.h
--- a/Engine/Source/Runtime/CoreUObject/Public/UObject/CookEnums.h
+++ b/Engine/Source/Runtime/CoreUObject/Public/UObject/CookEnums.h
@@ -3,6 +3,7 @@
 #pragma once
 
 #include "HAL/PlatformMath.h"
+#include "Misc/EnumClassFlags.h"
 
 namespace UE::Cook
 {
@@ -54,5 +55,19 @@ enum class ECookResult : uint8
 	NumBits = FPlatformMath::ConstExprCeilLogTwo(ECookResult::Count),
 };
 
+/** Validation options used when cooking (@see ECookByTheBookOptions for the validation flags used when initializing the cooker) */
+enum class ECookValidationOptions
+{
+	/** No flags */
+	None = 0,
+	/** Run asset validation (EditorValidatorSubsystem) on assets loaded during cook */
+	RunAssetValidation = 1<<0,
+	/** Run map validation (MapCheck) on maps loaded during cook */
+	RunMapValidation = 1<<1,
+	/** Consider validation errors (from RunAssetValidation or RunMapValidation) as fatal (preventing the package from being cooked) */
+	ValidationErrorsAreFatal = 1<<2,
+};
+ENUM_CLASS_FLAGS(ECookValidationOptions);
+
 }
 
diff --git a/Engine/Source/Runtime/CoreUObject/Public/UObject/ICookInfo.h b/Engine/Source/Runtime/CoreUObject/Public/UObject/ICookInfo.h
--- a/Engine/Source/Runtime/CoreUObject/Public/UObject/ICookInfo.h
+++ b/Engine/Source/Runtime/CoreUObject/Public/UObject/ICookInfo.h
@@ -212,10 +212,14 @@ public:
 
 	/** The type (e.g. CookByTheBook) of the running cook. This function will not return ECookType::Unknown. */
 	virtual UE::Cook::ECookType GetCookType() = 0;
-	/** Whether dlc is being cooked (e.g. via "-dlcname=<PluginName>"). This function will not return ECookingDLC::Unknown. */
+	/** Whether DLC is being cooked (e.g. via "-dlcname=<PluginName>"). This function will not return ECookingDLC::Unknown. */
 	virtual UE::Cook::ECookingDLC GetCookingDLC() = 0;
+	/** When DLC is being cooked (@see GetCookingDLC), this will return the name of the DLC plugin. */
+	virtual FString GetDLCName() = 0;
 	/** The role the current process plays in its MPCook session, or EProcessType::SingleProcess if it is running standalone. */
 	virtual UE::Cook::EProcessType GetProcessType() = 0;
+	/** Get the validation options used by the running cook, if any. */
+	virtual UE::Cook::ECookValidationOptions GetCookValidationOptions() = 0;
 	/**
 	 * Returns true if the cooker is cooking after a previous cook session and is cooking only the changed files.
 	 * Returns false if the cooker is doing a recook of all packages discovered in the session.
