Fix an issue in the viewport toolbar where the "Not Realtime" warning would be hidden when Unreal switches to the background.

diff --git a/Engine/Source/Editor/UnrealEd/Private/EditorViewportClient.cpp b/Engine/Source/Editor/UnrealEd/Private/EditorViewportClient.cpp
--- a/Engine/Source/Editor/UnrealEd/Private/EditorViewportClient.cpp
+++ b/Engine/Source/Editor/UnrealEd/Private/EditorViewportClient.cpp
@@ -684,6 +684,11 @@ void FEditorViewportClient::SetRealtime(bool bInRealtime)
 	bShouldInvalidateViewportWidget = true;
 }
 
+bool FEditorViewportClient::DoRealtimeAndOverridesMatch(bool bInValue) const
+{
+	return bIsRealtime == bInValue && (RealtimeOverrides.Num() == 0 || RealtimeOverrides.Last().bIsRealtime == bInValue);
+}
+
 FText FEditorViewportClient::GetRealtimeOverrideMessage() const
 {
 	return RealtimeOverrides.Num() > 0 ? RealtimeOverrides.Last().SystemDisplayName : FText::GetEmpty();
diff --git a/Engine/Source/Editor/UnrealEd/Private/ViewportToolbar/UnrealEdViewportToolbar.cpp b/Engine/Source/Editor/UnrealEd/Private/ViewportToolbar/UnrealEdViewportToolbar.cpp
--- a/Engine/Source/Editor/UnrealEd/Private/ViewportToolbar/UnrealEdViewportToolbar.cpp
+++ b/Engine/Source/Editor/UnrealEd/Private/ViewportToolbar/UnrealEdViewportToolbar.cpp
@@ -3421,7 +3421,14 @@ FText GetCameraSpeedTooltip()
 
 bool ShouldShowViewportRealtimeWarning(const FEditorViewportClient& ViewportClient)
 {
-	return !ViewportClient.IsRealtime() && !ViewportClient.IsRealtimeOverrideSet() && ViewportClient.IsPerspective();
+	// Almost all usages of viewport realtime overrides are for overriding realtime to be true.
+	// There are two exceptions:
+	//  - Editor tools running in PIE, where they configure their client to be realtime and then override it to be true.
+	//  - The Editor itself overriding realtime to false all viewport clients when the app goes into the background.
+	// When realtime is disabled and Unreal goes into the background, the `false` override will match.
+	// The realtime warning should be visible in this case, otherwise it will appear and disappear depending on whether
+	// Unreal is the foreground app or not.
+	return ViewportClient.DoRealtimeAndOverridesMatch(false) && ViewportClient.IsPerspective();
 }
 
 FToolMenuEntry CreatePerformanceAndScalabilitySubmenu()
@@ -3540,8 +3547,7 @@ FToolMenuEntry CreateToggleRealtimeEntry()
 		FNewToolMenuSectionDelegate::CreateLambda(
 			[](FToolMenuSection& InnerSection) -> void
 			{
-				UUnrealEdViewportToolbarContext* const EditorViewportContext =
-					InnerSection.FindContext<UUnrealEdViewportToolbarContext>();
+				UUnrealEdViewportToolbarContext* const EditorViewportContext = InnerSection.FindContext<UUnrealEdViewportToolbarContext>();
 				if (!EditorViewportContext)
 				{
 					return;
@@ -3579,8 +3585,7 @@ FToolMenuEntry CreateToggleRealtimeEntry()
 						{
 							if (TSharedPtr<SEditorViewport> EditorViewport = EditorViewportWeak.Pin())
 							{
-								return EditorViewport->IsRealtime() ? ECheckBoxState::Checked
-																	: ECheckBoxState::Unchecked;
+								return EditorViewport->IsRealtime() ? ECheckBoxState::Checked : ECheckBoxState::Unchecked;
 							}
 							return ECheckBoxState::Undetermined;
 						}
@@ -3589,8 +3594,7 @@ FToolMenuEntry CreateToggleRealtimeEntry()
 
 				TAttribute<FText> Tooltip;
 				{
-					const FText RealtimeTooltip =
-						LOCTEXT("ToggleRealtimeTooltip", "Toggle realtime rendering of the viewport");
+					const FText RealtimeTooltip = LOCTEXT("ToggleRealtimeTooltip", "Toggle realtime rendering of the viewport");
 
 					// If we can find a context with a viewport, use that to adjust the tooltip
 					// based on the viewport's realtime status.
diff --git a/Engine/Source/Editor/UnrealEd/Public/EditorViewportClient.h b/Engine/Source/Editor/UnrealEd/Public/EditorViewportClient.h
--- a/Engine/Source/Editor/UnrealEd/Public/EditorViewportClient.h
+++ b/Engine/Source/Editor/UnrealEd/Public/EditorViewportClient.h
@@ -424,7 +424,12 @@ public:
 	 * @return true if there are any temp realtime overrides set
 	 */
 	bool IsRealtimeOverrideSet() const { return RealtimeOverrides.Num() > 0; }
-
+	
+	/**
+	 * @return true if the underlying realtime mode and the last override (if any) match the given value.
+	 */
+	bool DoRealtimeAndOverridesMatch(bool bInValue) const;
+	
 	/**
 	 * @return  If an override is set this returns the message indicating what set it
 	 */
