IncrementalCook: Create CookPackageArtifact's LoadDependencies for use by the upcoming auto-transitive build dependencies.
#rnx
#incrementalcook

diff --git a/Engine/Source/Editor/UnrealEd/Private/CookOnTheFlyServer.cpp b/Engine/Source/Editor/UnrealEd/Private/CookOnTheFlyServer.cpp
--- a/Engine/Source/Editor/UnrealEd/Private/CookOnTheFlyServer.cpp
+++ b/Engine/Source/Editor/UnrealEd/Private/CookOnTheFlyServer.cpp
@@ -3188,6 +3188,7 @@ void UCookOnTheFlyServer::LoadPackageInQueue(UE::Cook::FPackageData& PackageData
 
 	PostLoadPackageFixup(PackageData, LoadedPackage);
 	PackageData.SetPackage(LoadedPackage);
+	PackageData.CreateLoadDependencies();
 	PackageData.SendToState(EPackageState::SaveActive, ESendFlags::QueueAddAndRemove, EStateChangeReason::Loaded);
 	++OutNumPushed;
 }
diff --git a/Engine/Source/Editor/UnrealEd/Private/Cooker/CookFunctionLibrary.cpp b/Engine/Source/Editor/UnrealEd/Private/Cooker/CookFunctionLibrary.cpp
--- a/Engine/Source/Editor/UnrealEd/Private/Cooker/CookFunctionLibrary.cpp
+++ b/Engine/Source/Editor/UnrealEd/Private/Cooker/CookFunctionLibrary.cpp
@@ -176,6 +176,23 @@ void UCookFunctionLibrary::CookAsset(UObject* Object, const FString& ForPlatform
 		SaveArgs.SaveFlags |= SAVE_Unversioned;
 	}
 
+	TArray<TPair<ELogVerbosity::Type, FString>> Messages;
+	FBuildDependencySet LoadDependencies = FBuildDependencySet::CollectLoadedPackage(Package, &Messages);
+	if (!LoadDependencies.IsValid())
+	{
+		TStringBuilder<256> ErrorMessage;
+		ErrorMessage << TEXT("Error collecting cook load dependencies:");
+		for (TPair<ELogVerbosity::Type, FString>& MessagePair : Messages)
+		{
+			ErrorMessage << TEXT("\n\t") << MessagePair.Value;
+		}
+		UE_LOG(LogCook, Warning, TEXT("%s"), *ErrorMessage);
+	}
+
+	// TODO: Execute the asynchronous transformation steps that packages go through after loading and before saving
+	// BeginCacheForCookedPlatformData(); // Call once on each object
+	// IsCachedCookedPlatformDataLoaded(); // Call on each object until it returns true
+
 	// We'll need to support this if we want to support cook diffing or iterative validate, 
 	// but for now it's not needed:
 	// SaveContext.PackageWriter->UpdateSaveArguments(SaveArgs);
@@ -210,12 +227,13 @@ void UCookFunctionLibrary::CookAsset(UObject* Object, const FString& ForPlatform
 		// streamingcells from WorldPartition. Supporting streaming cells would require 
 		// several other changes to this function, and is not the priority of this 
 		// routine at this time.
-		UE::Cook::FGenerationHelper* Todo = nullptr;
+		UE::Cook::FGenerationHelper* GenerationHelper = nullptr;
 		bool bGenerated = false;
 		TArray<FName> RuntimeDepsTodo; // I don't believe we need valid data in here at this time
+		Messages.Reset();
 
-		TArray<TPair<ELogVerbosity::Type, FString>> Messages;
-		FPackageArtifacts Deps = FPackageArtifacts::Collect(Package, TargetPlatform, &Result, Todo, bGenerated, MoveTemp(RuntimeDepsTodo), &Messages);
+		FPackageArtifacts Deps = FPackageArtifacts::Collect(Package, TargetPlatform, &LoadDependencies, &Result,
+			GenerationHelper, bGenerated, MoveTemp(RuntimeDepsTodo), &Messages);
 		if (!Deps.IsValid())
 		{
 			TStringBuilder<256> ErrorMessage;
diff --git a/Engine/Source/Editor/UnrealEd/Private/Cooker/CookGarbageCollect.cpp b/Engine/Source/Editor/UnrealEd/Private/Cooker/CookGarbageCollect.cpp
--- a/Engine/Source/Editor/UnrealEd/Private/Cooker/CookGarbageCollect.cpp
+++ b/Engine/Source/Editor/UnrealEd/Private/Cooker/CookGarbageCollect.cpp
@@ -936,6 +936,20 @@ void UCookOnTheFlyServer::PostGarbageCollect()
 		{
 			Preloader->PostGarbageCollect();
 		}
+
+		// Free memory used by GetLoadDependencies for packages that have been garbage collected.
+		// To avoid the expense of calling FindPackage on every package, only do this for packages that
+		// are no longer in progress but still have loaddependencies.
+		// We can not free LoadDependencies for PackageDatas that still have their package loaded, because
+		// the package might need to be saved later for an additional platform, and we cannot correctly recreate the
+		// package's LoadDependencies until after the package is garbagecollected and reexecutes Load.
+		if (PackageData->GetLoadDependencies() && !PackageData->IsInProgress())
+		{
+			if (!FindPackage(nullptr, *WriteToString<256>(PackageData->GetPackageName())))
+			{
+				PackageData->ClearLoadDependencies();
+			}
+		}
 	});
 
 	// Only after running all possible callbacks that need our links for diagnostics, clear the list of temporary
diff --git a/Engine/Source/Editor/UnrealEd/Private/Cooker/CookPackageArtifacts.cpp b/Engine/Source/Editor/UnrealEd/Private/Cooker/CookPackageArtifacts.cpp
--- a/Engine/Source/Editor/UnrealEd/Private/Cooker/CookPackageArtifacts.cpp
+++ b/Engine/Source/Editor/UnrealEd/Private/Cooker/CookPackageArtifacts.cpp
@@ -347,7 +347,7 @@ enum class EPackageMountPoint
 	Content,
 	GeneratedContent,
 };
-EPackageMountPoint GetPackageMountPoint(FName PackageName)
+static EPackageMountPoint GetPackageMountPoint(FName PackageName)
 {
 	TStringBuilder<256> StringBuffer;
 	PackageName.ToString(StringBuffer);
@@ -391,47 +391,35 @@ static void JoinMessagesIntoErrorReason(FStringBuilderBase& OutText,
 	}
 }
 
-FPackageArtifacts FPackageArtifacts::Collect(UPackage* Package, const ITargetPlatform* TargetPlatform,
-	FSavePackageResultStruct* SaveResult, FGenerationHelper* GenerationHelper, bool bGenerated,
-	TArray<FName>&& InRuntimeDependencies, TArray<TPair<ELogVerbosity::Type, FString>>* OutMessages)
+bool FBuildDependencySet::TryCollectInternal(TArray<FCookDependency>& OutBuildDependencies,
+	TArray<FName>& InOutRuntimeDependencies, const UPackage* Package, const ITargetPlatform* TargetPlatform,
+	const FSavePackageResultStruct* SaveResult, FGenerationHelper* GenerationHelper,
+	bool bGenerated, TArray<TPair<ELogVerbosity::Type, FString>>* OutMessages)
 {
 	if (!Package)
 	{
 		if (OutMessages) OutMessages->Emplace(ELogVerbosity::Error, TEXT("Invalid null package."));
-		return FPackageArtifacts();
+		return false;
 	}
 	IAssetRegistry* AssetRegistry = IAssetRegistry::Get();
 	if (!AssetRegistry)
 	{
 		if (OutMessages) OutMessages->Emplace(ELogVerbosity::Error, TEXT("AssetRegistry is unavailable."));
-		return FPackageArtifacts();
+		return false;
 	}
 	FEditorDomain* EditorDomain = FEditorDomain::Get();
 	if (!EditorDomain)
 	{
 		if (OutMessages) OutMessages->Emplace(ELogVerbosity::Error, TEXT("EditorDomain is unavailable."));
-		return FPackageArtifacts();
+		return false;
 	}
 
-	FPackageArtifacts Result;
-	// Skip the multiple reallocations for an array that grows from 0 to 128, but then reallocate according
-	// to normal TArray growth.
-	TArray<FCookDependency> BuildDependencies;
-	BuildDependencies.Reserve(128);
-	Result.PackageName = Package->GetFName();
-	Result.bHasSaveResults = SaveResult != nullptr;
-	BuildDependencies.Add(FCookDependency::Package(Result.PackageName));
-	Result.RuntimeDependencies = MoveTemp(InRuntimeDependencies);
+	// Skip the multiple reallocations for an array that grows from 0 to 128, for performance, but then
+	// reallocate according to normal TArray growth to reduce spike memory use.
+	OutBuildDependencies.Empty(128);
 
-	TArray<FName> AssetDependencies;
-	// The equivalent of AssetRegistry dependencies for generated packages come from their generator's TryPopulate
-	// function and this function receives them via InRuntimeDependencies.
-	if (!bGenerated)
-	{
-		AssetRegistry->GetDependencies(Result.PackageName, AssetDependencies,
-			UE::AssetRegistry::EDependencyCategory::Package, UE::AssetRegistry::EDependencyQuery::Game);
-		Result.RuntimeDependencies.Append(MoveTemp(AssetDependencies));
-	}
+	FName PackageName = Package->GetFName();
+	OutBuildDependencies.Add(FCookDependency::Package(PackageName));
 
 #if UE_WITH_PACKAGE_ACCESS_TRACKING
 	FPackageBuildDependencyTracker& Tracker = FPackageBuildDependencyTracker::Get();
@@ -439,7 +427,7 @@ FPackageArtifacts FPackageArtifacts::Collect(UPackage* Package, const ITargetPla
 	if (Tracker.IsEnabled())
 	{
 		TArray<TPair<FBuildDependencyAccessData, FResultProjectionList>> AccessDatas
-			= Tracker.GetAccessDatas(Result.PackageName);
+			= Tracker.GetAccessDatas(PackageName);
 
 		for (TPair<FBuildDependencyAccessData, FResultProjectionList>& Pair : AccessDatas)
 		{
@@ -453,34 +441,38 @@ FPackageArtifacts FPackageArtifacts::Collect(UPackage* Package, const ITargetPla
 					// We have not yet enabled marking the auto-added dependencies from TObjectPtr resolve
 					// as transitive, because it causes a performance regression and we are still working on fixing the
 					// regression.
-					BuildDependencies.Add(FCookDependency::Package(AccessData.ReferencedPackage));
+					OutBuildDependencies.Add(FCookDependency::Package(AccessData.ReferencedPackage));
 				}
 				else if (ProjectionList.bHasAll)
 				{
 					// TObjectPtr UE::BuildProjection::All dependencies are added as transitive build dependencies.
 					// We have to do this to be conservative, since we do not know which bytes from the target are
 					// dependended upon and which of the target's build dependencies influence those bytes.
-					BuildDependencies.Add(FCookDependency::TransitiveBuild(AccessData.ReferencedPackage));
+					OutBuildDependencies.Add(FCookDependency::TransitiveBuild(AccessData.ReferencedPackage));
 				}
 				else
 				{
 					for (FTopLevelAssetPath ClassPath : ProjectionList.Classes)
 					{
-						BuildDependencies.Add(FCookDependency::NativeClass(WriteToString<256>(ClassPath)));
+						OutBuildDependencies.Add(FCookDependency::NativeClass(WriteToString<256>(ClassPath)));
 					}
 
 					for (FName ResultProjection : ProjectionList.ResultProjections)
 					{
 						if (ResultProjection == UE::Cook::ResultProjection::PackageAndClass)
 						{
-							BuildDependencies.Add(FCookDependency::Package(AccessData.ReferencedPackage));
+							OutBuildDependencies.Add(FCookDependency::Package(AccessData.ReferencedPackage));
 						}
 						else
 						{
-							UE_LOG(LogCook, Error,
-								TEXT("When saving %s, found ResultProjection %s, which is a system-specific ResultProjection, and this is not yet implemented. ")
-								TEXT("Find the call to UE_COOK_RESULTPROJECTION_SCOPED passing in this name and remove it."),
-								*Result.PackageName.ToString(), *ResultProjection.ToString());
+							if (OutMessages)
+							{
+								OutMessages->Emplace(ELogVerbosity::Error, FString::Printf(
+									TEXT("When saving %s, found ResultProjection %s, which is a system-specific ResultProjection, and this is not yet implemented. ")
+									TEXT("Find the call to UE_COOK_RESULTPROJECTION_SCOPED passing in this name and remove it."),
+									*PackageName.ToString(), *ResultProjection.ToString()));
+							}
+							return false;
 						}
 					}
 				}
@@ -492,9 +484,15 @@ FPackageArtifacts FPackageArtifacts::Collect(UPackage* Package, const ITargetPla
 	{
 		// When PackageAccessTracking is disabled, defensively treat all asset dependencies as transitive build
 		// dependencies.
-		for (FName AssetDependency : AssetDependencies)
+		TArray<FName> AssetDependencies;
+		if (!bGenerated)
 		{
-			BuildDependencies.Add(FCookDependency::TransitiveBuild(AssetDependency));
+			AssetRegistry->GetDependencies(PackageName, AssetDependencies,
+				UE::AssetRegistry::EDependencyCategory::Package, UE::AssetRegistry::EDependencyQuery::Game);
+			for (FName AssetDependency : AssetDependencies)
+			{
+				OutBuildDependencies.Add(FCookDependency::TransitiveBuild(AssetDependency));
+			}
 		}
 	}
 
@@ -504,20 +502,20 @@ FPackageArtifacts FPackageArtifacts::Collect(UPackage* Package, const ITargetPla
 		FCookConfigAccessTracker& ConfigTracker = FCookConfigAccessTracker::Get();
 		if (ConfigTracker.IsEnabled())
 		{
-			TArray<FConfigAccessData> ConfigKeys = ConfigTracker.GetPackageRecords(Result.PackageName, TargetPlatform);
+			TArray<FConfigAccessData> ConfigKeys = ConfigTracker.GetPackageRecords(PackageName, TargetPlatform);
 			for (const FConfigAccessData& ConfigKey : ConfigKeys)
 			{
-				BuildDependencies.Add(FCookDependency::Config(ConfigKey));
+				OutBuildDependencies.Add(FCookDependency::Config(ConfigKey));
 			}
 		}
 	}
 #endif
 	if (SaveResult)
 	{
-		BuildDependencies.Append(SaveResult->CookDependencies);
+		OutBuildDependencies.Append(SaveResult->CookDependencies);
 		for (FName SoftPackageReference : SaveResult->UntrackedSoftPackageReferences)
 		{
-			BuildDependencies.Add(FCookDependency::RedirectionTarget(SoftPackageReference));
+			OutBuildDependencies.Add(FCookDependency::RedirectionTarget(SoftPackageReference));
 		}
 	}
 
@@ -526,7 +524,7 @@ FPackageArtifacts FPackageArtifacts::Collect(UPackage* Package, const ITargetPla
 	// Settings dependencies - Expand transitive dependencies on SettingsObjects into the list of dependencies
 	// recorded for that settings object.
 	TSet<const UObject*> SettingsDependencies;
-	BuildDependencies.RemoveAllSwap([&SettingsDependencies](const FCookDependency& Dependency)
+	OutBuildDependencies.RemoveAllSwap([&SettingsDependencies](const FCookDependency& Dependency)
 		{
 			if (Dependency.GetType() == ECookDependency::SettingsObject)
 			{
@@ -561,7 +559,7 @@ FPackageArtifacts FPackageArtifacts::Collect(UPackage* Package, const ITargetPla
 				JoinMessagesIntoErrorReason(ErrorText, IncludeDependencies.Messages);
 				OutMessages->Emplace(ELogVerbosity::Error, FString(*ErrorText));
 			}
-			return FPackageArtifacts();
+			return false;
 		}
 
 		for (const FCookDependency& IncludeDependency : IncludeDependencies.Dependencies.GetDependencies())
@@ -576,15 +574,15 @@ FPackageArtifacts FPackageArtifacts::Collect(UPackage* Package, const ITargetPla
 						TEXT("Settings dependency on object %s, but that object has a recursive Settings dependency on %s, and recursive Settings dependencies are not supported."),
 						*SettingsObject->GetPathName(), *IncludeDependency.GetSettingsObject()->GetPathName()));
 				}
-				return FPackageArtifacts();
+				return false;
 			}
-			BuildDependencies.Add(IncludeDependency);
+			OutBuildDependencies.Add(IncludeDependency);
 		}
 	}
 
 	// Process some rules for Package dependencies
 	TSet<FName> RedirectionTargets;
-	for (TArray<FCookDependency>::TIterator Iter(BuildDependencies); Iter; ++Iter)
+	for (TArray<FCookDependency>::TIterator Iter(OutBuildDependencies); Iter; ++Iter)
 	{
 		FCookDependency& Dependency = *Iter;
 		if (Dependency.GetType() == ECookDependency::TransitiveBuild ||
@@ -593,7 +591,7 @@ FPackageArtifacts FPackageArtifacts::Collect(UPackage* Package, const ITargetPla
 			FName DependencyName = Dependency.GetPackageName();
 			// Remove transitive dependencies to self, for performance. But keep the package dependency to
 			// self; every cooked package has its EditorDomain package as a dependency.
-			if (Dependency.GetType() == ECookDependency::TransitiveBuild && DependencyName == Result.PackageName)
+			if (Dependency.GetType() == ECookDependency::TransitiveBuild && DependencyName == PackageName)
 			{
 				Iter.RemoveCurrentSwap();
 				continue;
@@ -613,7 +611,7 @@ FPackageArtifacts FPackageArtifacts::Collect(UPackage* Package, const ITargetPla
 			// the savepackage and dependency collection of a generated package or a generator; 
 			// the digests only become available when the target generated package is saved, which can happen after
 			// the save of the packages that refer to it.
-			if (MountPoint == EPackageMountPoint::GeneratedContent && DependencyName != Result.PackageName)
+			if (MountPoint == EPackageMountPoint::GeneratedContent && DependencyName != PackageName)
 			{
 				Iter.RemoveCurrentSwap();
 				continue;
@@ -622,29 +620,26 @@ FPackageArtifacts FPackageArtifacts::Collect(UPackage* Package, const ITargetPla
 			bool bPackageExistOnDisk = AssetRegistry->DoesPackageExistOnDisk(DependencyName);
 			if (!bPackageExistOnDisk)
 			{
-				TStringBuilder<256> PackagePath(EInPlace::InPlace, DependencyName);
-				UPackage* PackageObject = FindPackage(nullptr, *PackagePath); // The package doesn't exist on disk, maybe it's an in-memory package.
-				if (!PackageObject)
+				TStringBuilder<256> DependencyNameStr(InPlace, DependencyName);
+				UPackage* DependencyPackage = FindPackage(nullptr, *DependencyNameStr);
+				if (!DependencyPackage)
 				{
-					// The package does not exist anywhere so error out.
 					if (OutMessages)
 					{
-						TStringBuilder<256> ErrorText;
-						ErrorText << TEXT("Package ") << PackagePath << TEXT(" does not exist.");
-						OutMessages->Emplace(ELogVerbosity::Error, FString(*ErrorText));
+						OutMessages->Emplace(ELogVerbosity::Error,
+							FString::Printf(TEXT("Package %s does not exist."), *DependencyNameStr));
 					}
 
-					return FPackageArtifacts();
+					return false;
 				}
-				else if (PackageObject->HasAnyPackageFlags(PKG_NewlyCreated))
+				else if (DependencyPackage->HasAnyPackageFlags(PKG_NewlyCreated))
 				{
 					// If the package is a newly created package (in-memory package) then ignore it.
 					// In-memory packages are ignored because we can't compute their digest. Only packages on disk have a digest.
 					Iter.RemoveCurrentSwap();
 					continue;
 				}
-
-				// Here the package is not on disk, in-memory and not newly created. It's a strange edge case but let's register it to the dependencies.
+				// else the package is not on disk, in-memory and not newly created. It's a strange edge case but let's register it to the dependencies.
 			}
 
 			// Package dependencies of all kinds (Runtime, Build, TransitiveBuild) also cause RedirectionTarget
@@ -658,7 +653,7 @@ FPackageArtifacts FPackageArtifacts::Collect(UPackage* Package, const ITargetPla
 				PRAGMA_DISABLE_DEPRECATION_WARNINGS;
 				if (Dependency.IsAlsoAddRuntimeDependency())
 				{
-					Result.RuntimeDependencies.Add(DependencyName);
+					InOutRuntimeDependencies.Add(DependencyName);
 					// Remove the IsAlsoRuntimeDependency flag
 					Dependency = FCookDependency::TransitiveBuild(DependencyName);
 				}
@@ -670,17 +665,17 @@ FPackageArtifacts FPackageArtifacts::Collect(UPackage* Package, const ITargetPla
 	// Pull transient packages out of the runtime dependencies for performance; we don't need them for
 	// deciding what gets cooked. 
 	// Runtime and script dependencies also cause RedirectionTarget dependencies, so record those.
-	for (TArray<FName>::TIterator Iter(Result.RuntimeDependencies); Iter; ++Iter)
+	for (TArray<FName>::TIterator Iter(InOutRuntimeDependencies); Iter; ++Iter)
 	{
-		FName PackageName = *Iter;
-		EPackageMountPoint MountPoint = GetPackageMountPoint(PackageName);
+		FName DependencyPackageName = *Iter;
+		EPackageMountPoint MountPoint = GetPackageMountPoint(DependencyPackageName);
 		switch (MountPoint)
 		{
 		case EPackageMountPoint::GeneratedContent:
 		case EPackageMountPoint::Content:
 		case EPackageMountPoint::Script:
 			// Keep it
-			RedirectionTargets.Add(PackageName);
+			RedirectionTargets.Add(DependencyPackageName);
 			break;
 		default:
 			Iter.RemoveCurrentSwap();
@@ -691,30 +686,12 @@ FPackageArtifacts FPackageArtifacts::Collect(UPackage* Package, const ITargetPla
 	// Put all the extra redirection dependencies into BuildDependencies
 	for (FName RedirectionTarget : RedirectionTargets)
 	{
-		BuildDependencies.Add(FCookDependency::RedirectionTarget(RedirectionTarget));
+		OutBuildDependencies.Add(FCookDependency::RedirectionTarget(RedirectionTarget));
 	}
-
-	// Sort and remove duplicates
-	Algo::Sort(BuildDependencies);
-	BuildDependencies.SetNum(Algo::Unique(BuildDependencies), EAllowShrinking::Yes);
-	Algo::Sort(Result.RuntimeDependencies, FNameLexicalLess());
-	Result.RuntimeDependencies.SetNum(Algo::Unique(Result.RuntimeDependencies), EAllowShrinking::Yes);
-
-	Result.SaveBuildDependencies.SetNormalizedDependencies(MoveTemp(BuildDependencies));
-	FBuildDependencySet::ECurrentKeyResult
-		CurrentKeyResult = Result.TryCalculateCurrentKey(GenerationHelper, OutMessages);
-	if (CurrentKeyResult == FBuildDependencySet::ECurrentKeyResult::Error)
-	{
-		return FPackageArtifacts();
-	}
-	Result.SaveBuildDependencies.StoreCurrentKey();
-	Result.SaveBuildDependencies.SetValid(true);
-	Result.bValid = true;
-
-	return Result;
+	return true;
 }
 
-FBuildDependencySet FPackageArtifacts::CollectSettingsObject(const UObject* Object,
+FBuildDependencySet FBuildDependencySet::CollectSettingsObject(const UObject* Object,
 	TArray<TPair<ELogVerbosity::Type, FString>>* OutMessages)
 {
 	if (!Object)
@@ -876,7 +853,7 @@ FCookDependencyGroups::FRecordedDependencies& FCookDependencyGroups::FindOrCreat
 	return Groups.FindOrAdd(Key);
 }
 
-FBuildDefinitionList FBuildDefinitionList::Collect(UPackage* Package, const ITargetPlatform* TargetPlatform,
+FBuildDefinitionList FBuildDefinitionList::Collect(const UPackage* Package, const ITargetPlatform* TargetPlatform,
 	TArray<TPair<ELogVerbosity::Type, FString>>* OutMessages)
 {
 	using namespace UE::DerivedData;
@@ -977,16 +954,17 @@ void FIncrementalCookAttachments::AppendCommitAttachments(TArray<IPackageWriter:
 	}
 }
 
-FIncrementalCookAttachments FIncrementalCookAttachments::Collect(UPackage* Package,
-	const ITargetPlatform* TargetPlatform, FSavePackageResultStruct* SaveResult,
-	FGenerationHelper* GenerationHelper, bool bGenerated, TArray<FName>&& RuntimeDependencies)
+FIncrementalCookAttachments FIncrementalCookAttachments::Collect(const UPackage* Package,
+	const ITargetPlatform* TargetPlatform, const FBuildDependencySet* LoadDependencies,
+	const FSavePackageResultStruct* SaveResult, FGenerationHelper* GenerationHelper, bool bGenerated,
+	TArray<FName>&& RuntimeDependencies)
 {
 	FIncrementalCookAttachments Result;
 
 	Result.CommitStatus = IPackageWriter::ECommitStatus::NotCommitted;
 
 	TArray<TPair<ELogVerbosity::Type, FString>> Messages;
-	Result.Artifacts = FPackageArtifacts::Collect(Package, TargetPlatform, SaveResult,
+	Result.Artifacts = FPackageArtifacts::Collect(Package, TargetPlatform, LoadDependencies, SaveResult,
 		GenerationHelper, bGenerated, MoveTemp(RuntimeDependencies), &Messages);
 	if (!Result.Artifacts.IsValid())
 	{
@@ -1076,4 +1054,114 @@ void FIncrementalCookAttachments::Fetch(TArrayView<FName> PackageNames, const IT
 	check(InProgressResults.Num() == 0);
 }
 
+FBuildDependencySet FBuildDependencySet::CollectLoadedPackage(const UPackage* Package,
+	TArray<TPair<ELogVerbosity::Type, FString>>* OutMessages)
+{
+	TArray<FCookDependency> BuildDependencies;
+	TArray<FName> UnusedRuntimeDependencies;
+	if (!TryCollectInternal(BuildDependencies, UnusedRuntimeDependencies, Package, nullptr /* TargetPlatform */,
+		nullptr /* SaveResult */, nullptr /* GenerationHelper */, false /* bGenerated */, OutMessages))
+	{
+		return FBuildDependencySet();
+	}
+
+	// Sort and remove duplicates in the results from TryCollectInternal.
+	Algo::Sort(BuildDependencies);
+	BuildDependencies.SetNum(Algo::Unique(BuildDependencies), EAllowShrinking::Yes);
+
+	// Store the collected dependencies on the result
+	FBuildDependencySet Result;
+	Result.SetNormalizedDependencies(MoveTemp(BuildDependencies));
+	FBuildDependencySet::ECurrentKeyResult CurrentKeyResult = Result.TryCalculateCurrentKey(
+		Package->GetFName(), nullptr /* GenerationHelper */, OutMessages);
+	if (CurrentKeyResult == FBuildDependencySet::ECurrentKeyResult::Error)
+	{
+		return FBuildDependencySet();
+	}
+	Result.StoreCurrentKey();
+	Result.SetValid(true);
+	return Result;
+}
+
+FPackageArtifacts FPackageArtifacts::Collect(const UPackage* Package, const ITargetPlatform* TargetPlatform,
+	const FBuildDependencySet* InLoadDependencies, const FSavePackageResultStruct* SaveResult,
+	FGenerationHelper* GenerationHelper, bool bGenerated,
+	TArray<FName>&& InRuntimeDependencies, TArray<TPair<ELogVerbosity::Type, FString>>* OutMessages)
+{
+	if (!Package)
+	{
+		if (OutMessages) OutMessages->Emplace(ELogVerbosity::Error, TEXT("Invalid null package."));
+		return FPackageArtifacts();
+	}
+	IAssetRegistry* AssetRegistry = IAssetRegistry::Get();
+	if (!AssetRegistry)
+	{
+		if (OutMessages) OutMessages->Emplace(ELogVerbosity::Error, TEXT("AssetRegistry is unavailable."));
+		return FPackageArtifacts();
+	}
+	FName PackageName = Package->GetFName();
+
+	// Append AssetRegistry dependencies as runtime dependencies, only for non-generated pacakges. The equivalent for
+	// generated packages comes from the generator's ICookPackageSplitter functions and this function receives them via
+	// InRuntimeDependencies.
+	if (!bGenerated)
+	{
+		TArray<FName> AssetDependencies;
+		AssetRegistry->GetDependencies(PackageName, AssetDependencies,
+			UE::AssetRegistry::EDependencyCategory::Package, UE::AssetRegistry::EDependencyQuery::Game);
+		InRuntimeDependencies.Append(MoveTemp(AssetDependencies));
+	}
+
+	// Collect the save's BuildDependencies, and pass in our runtimedependencies for read/write.
+	TArray<FCookDependency> BuildDependencies;
+	if (!FBuildDependencySet::TryCollectInternal(BuildDependencies, InRuntimeDependencies, Package,
+		TargetPlatform, SaveResult, GenerationHelper, bGenerated, OutMessages))
+	{
+		return FPackageArtifacts();
+	}
+
+	// Sort and remove duplicates in the results from TryCollectInternal.
+	Algo::Sort(BuildDependencies);
+	BuildDependencies.SetNum(Algo::Unique(BuildDependencies), EAllowShrinking::Yes);
+	Algo::Sort(InRuntimeDependencies, FNameLexicalLess());
+	InRuntimeDependencies.SetNum(Algo::Unique(InRuntimeDependencies), EAllowShrinking::Yes);
+
+	FPackageArtifacts Result;
+	Result.PackageName = PackageName;
+	Result.bHasSaveResults = SaveResult != nullptr;
+
+	// Store input+collected RuntimeDependencies on the result
+	Result.RuntimeDependencies = MoveTemp(InRuntimeDependencies);
+
+	// Store input LoadDependencies on the result
+	if (InLoadDependencies && InLoadDependencies->IsValid())
+	{
+		Result.LoadBuildDependencies = *InLoadDependencies;
+	}
+	else
+	{
+		// An empty set of dependencies, create a normalized version of it.
+		Result.LoadBuildDependencies.TryCalculateCurrentKey(PackageName,
+			nullptr /* GenerationHelper */, nullptr /* OutMessages */);
+		Result.LoadBuildDependencies.StoreCurrentKey();
+		Result.LoadBuildDependencies.SetValid(true);
+	}
+	Result.LoadBuildDependencies.SetName(NAME_Load);
+	check(Result.LoadBuildDependencies.IsValid());
+
+	// Store the collected SaveDependencies on the result
+	Result.SaveBuildDependencies.SetNormalizedDependencies(MoveTemp(BuildDependencies));
+	FBuildDependencySet::ECurrentKeyResult CurrentKeyResult = Result.SaveBuildDependencies.TryCalculateCurrentKey(
+		PackageName, GenerationHelper, OutMessages);
+	if (CurrentKeyResult == FBuildDependencySet::ECurrentKeyResult::Error)
+	{
+		return FPackageArtifacts();
+	}
+	Result.SaveBuildDependencies.StoreCurrentKey();
+	Result.SaveBuildDependencies.SetValid(true);
+	Result.bValid = true;
+
+	return Result;
+}
+
 } // namespace UE::Cook
diff --git a/Engine/Source/Editor/UnrealEd/Private/Cooker/CookPackageArtifacts.h b/Engine/Source/Editor/UnrealEd/Private/Cooker/CookPackageArtifacts.h
--- a/Engine/Source/Editor/UnrealEd/Private/Cooker/CookPackageArtifacts.h
+++ b/Engine/Source/Editor/UnrealEd/Private/Cooker/CookPackageArtifacts.h
@@ -121,6 +121,36 @@ public:
 	bool TryLoad(FCbFieldView FieldView);
 	void Save(FCbWriter& Writer) const;
 
+	/**
+	 * Read dependencies for the given targetplatform of the given package out of global dependency trackers
+	 * that have recorded its data during the package's load operations in the current cook session.
+	 */
+	static FBuildDependencySet CollectLoadedPackage(const UPackage* Package,
+		TArray<TPair<ELogVerbosity::Type, FString>>* OutMessages = nullptr);
+
+	/**
+	 * Internal helper for FBuildDependencySet::Collect and FPackageArtifacts::Collect. Handles all arguments
+	 * used by either of them, and returns an unsorted,non-unique array of CookDependencies rather than a normalized
+	 * FBuildDependencySet.
+	 * 
+	 * @param TargetPlatform - If null, collects dependencies reported by all platforms. If non-null, only collects
+	 *                         dependencies reported for the given platform or reported as platform-agnostic.
+	 * @param SaveResult - If non-null, contains data reported from the SavePackage call.
+	 * @param GenerationHelper - If non-null, provides the lookup for the AssetPackageData of generated packages.
+	 *                           Must be provided if any generated packages are in the dependencies.
+	 */
+	static bool TryCollectInternal(TArray<FCookDependency>& OutBuildDependencies, TArray<FName>& OutRuntimeDependencies,
+		const UPackage* Package, const ITargetPlatform* TargetPlatform, const FSavePackageResultStruct* SaveResult,
+		FGenerationHelper* GenerationHelper, bool bGenerated,
+		TArray<TPair<ELogVerbosity::Type, FString>>* OutMessages);
+
+	/**
+	 * Collect the dependencies referenced by a given settingsobject from e.g. config. Globally cached for the current
+	 * process.
+	 */
+	static FBuildDependencySet CollectSettingsObject(const UObject* Object,
+		TArray<TPair<ELogVerbosity::Type, FString>>* OutMessages);
+
 	// Hidden friend operators
 private:
 	friend bool LoadFromCompactBinary(FCbFieldView FieldView, UE::Cook::FBuildDependencySet& DependencySet)
@@ -201,17 +231,18 @@ public:
 
 	// Construction functions
 	/**
-	 * Read dependencies for the given targetplatform of the given package out of global dependency trackers
-	 * that have recorded its data during the package's load/other/save operations in the current cook session.
-	 * 
+	 * Read dependencies for the given targetplatform of the given package out of global dependency trackers that have
+	 * recorded its data for the package's save operations, and combine those with the given previously recorded
+	 * LoadDependencies to create the complete PackageArtifacts.
+	 *
+	 * @param LoadDependencies - If non-null, will be copied into LoadDependencies on the result.
 	 * @param GenerationHelper - If non-null, provides the lookup for the AssetPackageData of generated packages.
 	 *                           Must be provided if any generated packages are in the dependencies.
 	 */
-	static FPackageArtifacts Collect(UPackage* Package, const ITargetPlatform* TargetPlatform,
-		FSavePackageResultStruct* SaveResult, FGenerationHelper* GenerationHelper, bool bGenerated,
-		TArray<FName>&& RuntimeDependencies, TArray<TPair<ELogVerbosity::Type, FString>>* OutMessages = nullptr);
-	static FBuildDependencySet CollectSettingsObject(const UObject* Object,
-		TArray<TPair<ELogVerbosity::Type, FString>>* OutMessages);
+	static FPackageArtifacts Collect(const UPackage* Package, const ITargetPlatform* TargetPlatform,
+		const FBuildDependencySet* InLoadDependencies, const FSavePackageResultStruct* SaveResult,
+		FGenerationHelper* GenerationHelper, bool bGenerated, TArray<FName>&& InRuntimeDependencies,
+		TArray<TPair<ELogVerbosity::Type, FString>>* OutMessages = nullptr);
 
 	// Fetch function to load the dependencies from a PackageStore is not yet implemented independently for
 	// this structure. Use FIncrementalCookAttachments instead.
@@ -271,7 +302,7 @@ public:
 	void Empty();
 
 	/** Collect DDC BuildDefinitions that were issued from the load/save of the given package and platform. */
-	static FBuildDefinitionList Collect(UPackage* Package, const ITargetPlatform* TargetPlatform,
+	static FBuildDefinitionList Collect(const UPackage* Package, const ITargetPlatform* TargetPlatform,
 		TArray<TPair<ELogVerbosity::Type, FString>>* OutMessages = nullptr);
 
 private:
@@ -292,9 +323,10 @@ struct FIncrementalCookAttachments
 	static void Fetch(TArrayView<FName> PackageNames, const ITargetPlatform* TargetPlatform,
 		ICookedPackageWriter* PackageWriter,
 		TUniqueFunction<void(FName PackageName, FIncrementalCookAttachments&& Result)>&& Callback);
-	static FIncrementalCookAttachments Collect(UPackage* Package,
-		const ITargetPlatform* TargetPlatform, FSavePackageResultStruct* SaveResult,
-		FGenerationHelper* GenerationHelper, bool bGenerated, TArray<FName>&& RuntimeDependencies);
+	static FIncrementalCookAttachments Collect(const UPackage* Package,
+		const ITargetPlatform* TargetPlatform, const FBuildDependencySet* LoadDependencies,
+		const FSavePackageResultStruct* SaveResult, FGenerationHelper* GenerationHelper, bool bGenerated,
+		TArray<FName>&& RuntimeDependencies);
 };
 
 } // namespace UE::Cook
diff --git a/Engine/Source/Editor/UnrealEd/Private/Cooker/CookPackageData.cpp b/Engine/Source/Editor/UnrealEd/Private/Cooker/CookPackageData.cpp
--- a/Engine/Source/Editor/UnrealEd/Private/Cooker/CookPackageData.cpp
+++ b/Engine/Source/Editor/UnrealEd/Private/Cooker/CookPackageData.cpp
@@ -1132,6 +1132,7 @@ void FPackageData::OnExitInProgress(EStateChangeReason StateChangeReason)
 void FPackageData::OnEnterSaving()
 {
 	check(GetPackage() != nullptr && GetPackage()->IsFullyLoaded());
+	check(GetLoadDependencies() != nullptr);
 
 	check(!HasPrepareSaveFailed());
 	CheckObjectCacheEmpty();
@@ -1238,6 +1239,28 @@ void FPackageData::OnPackagePreloaderDestroyed(FPackagePreloader& InPackagePrelo
 	PackagePreloader = nullptr;
 }
 
+FBuildDependencySet* FPackageData::GetLoadDependencies()
+{
+	return LoadDependencies.Get();
+}
+
+void FPackageData::CreateLoadDependencies()
+{
+	if (!LoadDependencies)
+	{
+		UPackage* LocalPackage = Package.Get();
+		checkf(LocalPackage != nullptr,
+			TEXT("CreateLoadDependencies failed for package %s because this->Package == nullptr. It can only be called after the Package has been set."),
+			*GetPackageName().ToString());
+		LoadDependencies.Reset(new FBuildDependencySet(FBuildDependencySet::CollectLoadedPackage(LocalPackage)));
+	}
+}
+
+void FPackageData::ClearLoadDependencies()
+{
+	LoadDependencies.Reset();
+}
+
 TArray<FCachedObjectInOuter>& FPackageData::GetCachedObjectsInOuter()
 {
 	return CachedObjectsInOuter;
diff --git a/Engine/Source/Editor/UnrealEd/Private/Cooker/CookPackageData.h b/Engine/Source/Editor/UnrealEd/Private/Cooker/CookPackageData.h
--- a/Engine/Source/Editor/UnrealEd/Private/Cooker/CookPackageData.h
+++ b/Engine/Source/Editor/UnrealEd/Private/Cooker/CookPackageData.h
@@ -35,6 +35,7 @@ class FCbWriter;
 class UCookOnTheFlyServer;
 class UObject;
 class UPackage;
+namespace UE::Cook { class FBuildDependencySet; }
 namespace UE::Cook { class FCookWorkerClient; }
 namespace UE::Cook { class FPackagePreloader; }
 namespace UE::Cook { class FRequestCluster; }
@@ -571,6 +572,16 @@ public:
 	bool GetIsVisited() const;
 	void SetIsVisited(bool bValue);
 
+	/**
+	 * Return the LoadDependencies if they have been created, otherwise nullptr. LoadDependencies are stored in the
+	 * incremental cook oplog and are used by CookRequestCluster to test invalidation of previous cookresults for the
+	 * package and for other packages that have the package as a transitive build dependency.
+	 */
+	FBuildDependencySet* GetLoadDependencies();
+	/** Calculate the LoadDependencies if not already created. */
+	void CreateLoadDependencies();
+	void ClearLoadDependencies();
+
 	/**
 	 * The list of objects inside the package.  Only non-empty during saving; it is populated on demand by
 	 * TryCreateObjectCache and is cleared when leaving the save state.
@@ -834,6 +845,7 @@ private:
 	TArray<FCachedObjectInOuter> CachedObjectsInOuter;
 	FCompletionCallback CompletionCallback;
 	TUniquePtr<TMap<FPackageData*, EInstigator>> Unsolicited;
+	TUniquePtr<UE::Cook::FBuildDependencySet> LoadDependencies;
 	FName PackageName;
 	FName FileName;
 	FName ParentGenerator;
diff --git a/Engine/Source/Editor/UnrealEd/Private/Cooker/CookSavePackage.cpp b/Engine/Source/Editor/UnrealEd/Private/Cooker/CookSavePackage.cpp
--- a/Engine/Source/Editor/UnrealEd/Private/Cooker/CookSavePackage.cpp
+++ b/Engine/Source/Editor/UnrealEd/Private/Cooker/CookSavePackage.cpp
@@ -750,9 +750,14 @@ void FSaveCookedPackageContext::CalculateCookDependencies(FGenerationHelper* Gen
 				PlatformRuntimeDependencies.Add(Dependency.AssetId.PackageName);
 			}
 		}
-		PlatformCookAttachments = FIncrementalCookAttachments::Collect(Package, TargetPlatform,
-			HasSavePackageResult() ? &SavePackageResult : nullptr, GenerationHelper, bGenerated,
-			MoveTemp(PlatformRuntimeDependencies));
+		FBuildDependencySet* LoadDependencies = PackageData.GetLoadDependencies();
+		checkf(LoadDependencies,
+			TEXT("LoadDependencies not found during save of package %s. LoadDependencies are supposted to be created by LoadPackageInQueue before entering the Save state."),
+			*PackageData.GetPackageName().ToString());
+		FSavePackageResultStruct* LocalSavePackageResult = HasSavePackageResult() ? &SavePackageResult : nullptr;
+
+		PlatformCookAttachments = FIncrementalCookAttachments::Collect(Package, TargetPlatform, LoadDependencies,
+			LocalSavePackageResult, GenerationHelper, bGenerated, MoveTemp(PlatformRuntimeDependencies));
 	}
 }
 
diff --git a/Engine/Source/Editor/UnrealEd/Private/EditorDomain/EditorDomainUtils.cpp b/Engine/Source/Editor/UnrealEd/Private/EditorDomain/EditorDomainUtils.cpp
--- a/Engine/Source/Editor/UnrealEd/Private/EditorDomain/EditorDomainUtils.cpp
+++ b/Engine/Source/Editor/UnrealEd/Private/EditorDomain/EditorDomainUtils.cpp
@@ -1982,7 +1982,7 @@ bool TrySavePackage(UPackage* Package)
 	{
 		UE::Cook::FIncrementalCookAttachments CookAttachments =
 			UE::Cook::FIncrementalCookAttachments::Collect(
-				Package, nullptr /* TargetPlatform */, nullptr /* SaveResult */,
+				Package, nullptr /* TargetPlatform */, nullptr /* LoadDependencies */, nullptr /* SaveResult */,
 				nullptr /* GenerationHelper */, false /* bGenerated */, TArray<FName>());
 		TArray<IPackageWriter::FCommitAttachmentInfo> Attachments;
 		CookAttachments.AppendCommitAttachments(Attachments);
