Fix an issue where disabling `stat fps` after manually removing the viewport override could cause an ensure.

diff --git a/Engine/Source/Editor/UnrealEd/Private/EditorViewportClient.cpp b/Engine/Source/Editor/UnrealEd/Private/EditorViewportClient.cpp
--- a/Engine/Source/Editor/UnrealEd/Private/EditorViewportClient.cpp
+++ b/Engine/Source/Editor/UnrealEd/Private/EditorViewportClient.cpp
@@ -2357,7 +2357,7 @@ void FEditorViewportClient::HandleViewportStatDisabled(const TCHAR* InName)
 		if (SetStatEnabled(InName, false) == 0)
 		{
 			SetShowStats(false);
-			RemoveRealtimeOverride(LOCTEXT("RealtimeOverrideMessage_Stats", "Stats Display"));
+			RemoveRealtimeOverride(LOCTEXT("RealtimeOverrideMessage_Stats", "Stats Display"), /*bCheckMissingOverride*/false);
 		}
 	}
 }
diff --git a/Engine/Source/Editor/UnrealEd/Public/EditorViewportClient.h b/Engine/Source/Editor/UnrealEd/Public/EditorViewportClient.h
--- a/Engine/Source/Editor/UnrealEd/Public/EditorViewportClient.h
+++ b/Engine/Source/Editor/UnrealEd/Public/EditorViewportClient.h
@@ -371,7 +371,8 @@ public:
 	/**
 	 * Removes the most recent realtime override registered with the given system name.
 	 *
-	 * @param Whether to assert if no matching realtime override was found.
+	 * @param SystemDisplayName	The display name of the override to be removed.
+	 * @param bCheckMissingOverride Whether to assert if no matching realtime override was found. Should be `false` if users can operate within the viewport between enable/disable calls as they can remove overrides.
 	 * @return Whether any matching realtime override was found.
 	 */
 	UNREALED_API bool RemoveRealtimeOverride(FText SystemDisplayName, bool bCheckMissingOverride = true);
@@ -428,7 +429,7 @@ public:
 	/**
 	 * @return true if the underlying realtime mode and the last override (if any) match the given value.
 	 */
-	bool DoRealtimeAndOverridesMatch(bool bInValue) const;
+	UNREALED_API bool DoRealtimeAndOverridesMatch(bool bInValue) const; 
 	
 	/**
 	 * @return  If an override is set this returns the message indicating what set it
