Focus primary buttons for modal windows spawned via FMessageDialog::Open, and FSuppressableWarningDialog

diff --git a/Engine/Source/Editor/UnrealEd/Private/Dialogs/Dialogs.cpp b/Engine/Source/Editor/UnrealEd/Private/Dialogs/Dialogs.cpp
--- a/Engine/Source/Editor/UnrealEd/Private/Dialogs/Dialogs.cpp
+++ b/Engine/Source/Editor/UnrealEd/Private/Dialogs/Dialogs.cpp
@@ -54,7 +54,6 @@ public:
 	void Construct( const FArguments& InArgs )
 	{
 		ParentWindow = InArgs._ParentWindow;
-		ParentWindow->SetWidgetToFocusOnActivate(SharedThis(this));
 		Response = EAppReturnType::Cancel;
 		MessageType = InArgs._MessageType;
 
@@ -162,6 +161,7 @@ public:
 			]
 		];
 
+		TSharedPtr<SButton> PrimaryButton;
 		int32 SlotIndex = 0;
 
 #define ADD_SLOT(Button)\
@@ -177,7 +177,7 @@ public:
 #define ADD_SLOT_PRIMARY(Button)\
 		ButtonBox->AddSlot(SlotIndex++,0)\
 		[\
-			SNew( SButton )\
+			SAssignNew( PrimaryButton, SButton )\
 			.VAlign(VAlign_Center)\
 			.HAlign(HAlign_Center)\
 			.ButtonStyle(&FAppStyle::Get().GetWidgetStyle<FButtonStyle>("PrimaryButton"))\
@@ -186,7 +186,6 @@ public:
 		];
 
 
-
 		switch ( MessageType )
 		{	
 		case EAppMsgType::Ok:
@@ -228,7 +227,20 @@ public:
 			UE_LOG(LogDialogs, Fatal, TEXT("Invalid Message Type"));
 		}
 
+		// Route user focus to the default dialog button. 
+		if (ensure(PrimaryButton.IsValid()))
+		{
+			// Dialog should have a primary slot, allowing users to navigate options via esc, space, enter, and tab keys
+			ParentWindow->SetWidgetToFocusOnActivate(PrimaryButton);
+		}
+		else
+		{
+			// Fallback to entire dialog. This will only support esc key
+			ParentWindow->SetWidgetToFocusOnActivate(SharedThis(this));
+		}
+
 #undef ADD_SLOT
+#undef ADD_SLOT_PRIMARY
 	}
 	END_SLATE_FUNCTION_BUILD_OPTIMIZATION
 
@@ -240,7 +252,7 @@ public:
 
 	virtual	FReply OnKeyDown( const FGeometry& MyGeometry, const FKeyEvent& InKeyEvent )
 	{
-		//see if we pressed the Enter or Spacebar keys
+		//see if we pressed the Escape to cancel choice
 		if( InKeyEvent.GetKey() == EKeys::Escape )
 		{
 			return HandleButtonClicked(EAppReturnType::Cancel);
@@ -614,9 +626,7 @@ public:
 	void Construct( const FArguments& InArgs )
 	{
 		bCheckboxResult = InArgs._bDefaultCheckValue;
-		// Set this widget as focused, to allow users to hit ESC to cancel.
 		MyWindow = InArgs._ParentWindow.Get();
-		MyWindow.Pin()->SetWidgetToFocusOnActivate(SharedThis(this));
 		MyMessage = InArgs._Message;
 		MyCheckboxMessage = InArgs._CheckboxMessage;
 
@@ -684,6 +694,7 @@ public:
 	 */
 	TSharedRef<SHorizontalBox> ConstructConditionalInternals( const FArguments& InArgs ) 
 	{
+		TSharedPtr<SPrimaryButton> PrimaryButton;
 		TSharedRef<SHorizontalBox> HorizontalBox = SNew(SHorizontalBox);
 		
 		// checkbox with user specified text
@@ -706,10 +717,9 @@ public:
 		.Padding(FMargin(16.f, 0.f, 0.f, 0.f)) // currently hardcoded until we adjust StandardDialog.SlotPadding
 		.AutoWidth()
 		[
-			SNew(SPrimaryButton)
+			SAssignNew(PrimaryButton, SPrimaryButton)
 			.Text(InArgs._ConfirmText)
-			.OnClicked(this, &SModalDialogWithCheckbox::OnConfirmClicked)
-			
+			.OnClicked(this, &SModalDialogWithCheckbox::OnConfirmClicked)			
 		];
 
 
@@ -728,6 +738,12 @@ public:
 			];
 		}
 	
+		// Route user focus to the default dialog button. 
+		if (ensure(MyWindow.Pin()))
+		{
+			MyWindow.Pin()->SetWidgetToFocusOnActivate(PrimaryButton);
+		}
+
 		return HorizontalBox;
 	}
 
