Editor: fix jittering when moving camera actors whose camera component have a local offset

When dragging an object, the gizmo-moving code was eventually calling back into MoveCameraToLockedActor() on the editor viewport. However, that code simply set the viewport transform to the actor transform. This is not like UpdateViewForLockedActor() which does that but also checks if the actor implements GetEditorPreviewInfo() or if the actor is attached and so on.

Because of throttling, realtime viewports don't actually tick every frame, so some frames would not have UpdateViewForLockedActor() called. This would end up with the viewport transform only set to the actor transform for one frame, missing the child camera component offset.

diff --git a/Engine/Source/Editor/UnrealEd/Private/LevelEditorViewport.cpp b/Engine/Source/Editor/UnrealEd/Private/LevelEditorViewport.cpp
--- a/Engine/Source/Editor/UnrealEd/Private/LevelEditorViewport.cpp
+++ b/Engine/Source/Editor/UnrealEd/Private/LevelEditorViewport.cpp
@@ -3022,20 +3022,7 @@ void FLevelEditorViewportClient::UpdateViewForLockedActor(float DeltaTime)
 		if (!ViewTransform.IsPlaying())
 		{
 			// Update transform
-			if (Actor->GetAttachParentActor() != NULL)
-			{
-				// Actor is parented, so use the actor to world matrix for translation and rotation information.
-				SetViewLocation(Actor->GetActorLocation());
-				SetViewRotation(Actor->GetActorRotation());
-			}
-			else if (Actor->GetRootComponent() != NULL)
-			{
-				// No attachment, so just use the relative location, so that we don't need to
-				// convert from a quaternion, which loses winding information.
-				SetViewLocation(Actor->GetRootComponent()->GetRelativeLocation());
-				SetViewRotation(Actor->GetRootComponent()->GetRelativeRotation());
-			}
-
+			bool bGotCameraView = false;
 			if (bLockedCameraView)
 			{
 				// If this is a camera actor, then inherit some other settings
@@ -3058,9 +3045,26 @@ void FLevelEditorViewportClient::UpdateViewForLockedActor(float DeltaTime)
 						AspectRatio = ControllingActorViewInfo.AspectRatio;
 						SetViewLocation(ControllingActorViewInfo.Location);
 						SetViewRotation(ControllingActorViewInfo.Rotation);
+						bGotCameraView = true;
 					}
 				}
 			}
+			if (!bGotCameraView)
+			{
+				if (Actor->GetAttachParentActor() != NULL)
+				{
+					// Actor is parented, so use the actor to world matrix for translation and rotation information.
+					SetViewLocation(Actor->GetActorLocation());
+					SetViewRotation(Actor->GetActorRotation());
+				}
+				else if (Actor->GetRootComponent() != NULL)
+				{
+					// No attachment, so just use the relative location, so that we don't need to
+					// convert from a quaternion, which loses winding information.
+					SetViewLocation(Actor->GetRootComponent()->GetRelativeLocation());
+					SetViewRotation(Actor->GetRootComponent()->GetRelativeRotation());
+				}
+			}
 
 			const double DistanceToCurrentLookAt = FVector::Dist( GetViewLocation() , GetLookAtLocation() );
 
@@ -4317,11 +4321,47 @@ UTypedElementSelectionSet* FLevelEditorViewportClient::GetMutableSelectionSet()
 
 void FLevelEditorViewportClient::MoveCameraToLockedActor()
 {
-	// If turned on, move cameras location/rotation to the selected actors
-	if( GetActiveActorLock().IsValid() )
+	// Do a slightly shorter version of the code in UpdateViewForLockedActor where we don't care about
+	// post-process settings.
+	const FLevelViewportActorLock& ActiveLock = ActorLocks.GetLock();
+	AActor* Actor = ActiveLock.GetLockedActor();
+	if (Actor)
 	{
-		SetViewLocation( GetActiveActorLock()->GetActorLocation() );
-		SetViewRotation( GetActiveActorLock()->GetActorRotation() );
+		bool bGotCameraView = false;
+		if (bLockedCameraView)
+		{
+			// If this is a camera actor, then inherit some other settings
+			UActorComponent* const ViewComponent = FindViewComponentForActor(Actor);
+			if (ViewComponent != nullptr)
+			{
+				FMinimalViewInfo ViewInfo;
+				if (ensure(ViewComponent->GetEditorPreviewInfo(0.f, ViewInfo)))
+				{
+					ViewFOV = ControllingActorViewInfo.FOV;
+					AspectRatio = ControllingActorViewInfo.AspectRatio;
+					SetViewLocation(ControllingActorViewInfo.Location);
+					SetViewRotation(ControllingActorViewInfo.Rotation);
+					bGotCameraView = true;
+				}
+			}
+		}
+		if (!bGotCameraView)
+		{
+			if (Actor->GetAttachParentActor() != NULL)
+			{
+				// Actor is parented, so use the actor to world matrix for translation and rotation information.
+				SetViewLocation(Actor->GetActorLocation());
+				SetViewRotation(Actor->GetActorRotation());
+			}
+			else if (Actor->GetRootComponent() != NULL)
+			{
+				// No attachment, so just use the relative location, so that we don't need to
+				// convert from a quaternion, which loses winding information.
+				SetViewLocation(Actor->GetRootComponent()->GetRelativeLocation());
+				SetViewRotation(Actor->GetRootComponent()->GetRelativeRotation());
+			}
+		}
+
 		Invalidate();
 	}
 }
