Minor improvement to file reading to ensure when reading fails and have fallback for relative path reading.

diff --git a/Engine/Plugins/Runtime/GameFeatures/Source/GameFeatures/Private/GameFeaturesSubsystem.cpp b/Engine/Plugins/Runtime/GameFeatures/Source/GameFeatures/Private/GameFeaturesSubsystem.cpp
--- a/Engine/Plugins/Runtime/GameFeatures/Source/GameFeatures/Private/GameFeaturesSubsystem.cpp
+++ b/Engine/Plugins/Runtime/GameFeatures/Source/GameFeatures/Private/GameFeaturesSubsystem.cpp
@@ -4135,10 +4135,11 @@ void UGameFeaturesSubsystem::GetPluginsToCook(TSet<FString>& OutPlugins)
 			// check if it's a filename
 			if (CookPluginsStr.EndsWith(".txt"))
 			{
-				FFileHelper::LoadFileToStringWithLineVisitor(*CookPluginsStr, [&ReturnList](FStringView Line)
+				const bool bSuccess = FFileHelper::LoadFileToStringWithLineVisitor(*CookPluginsStr, [&ReturnList](FStringView Line)
 				{
 					ReturnList.Add(FString(Line));
 				});
+				ensureMsgf(bSuccess, TEXT("Failed to read - %s"), *CookPluginsStr);
 			}
 			else
 			{
@@ -4149,7 +4150,7 @@ void UGameFeaturesSubsystem::GetPluginsToCook(TSet<FString>& OutPlugins)
 		return ReturnList;
 	}();
 	
-	OutPlugins.Append(PluginsList);
+	OutPlugins.Append(PluginsList);	
 }
 
 void UGameFeaturesSubsystem::SetExplanationForNotMountingPlugin(const FString& PluginURL, const FString& Explanation)
diff --git a/Engine/Source/Editor/UnrealEd/Private/CookOnTheFlyServer.cpp b/Engine/Source/Editor/UnrealEd/Private/CookOnTheFlyServer.cpp
--- a/Engine/Source/Editor/UnrealEd/Private/CookOnTheFlyServer.cpp
+++ b/Engine/Source/Editor/UnrealEd/Private/CookOnTheFlyServer.cpp
@@ -11633,7 +11633,14 @@ void UCookOnTheFlyServer::RecordDLCPackagesFromBaseGame(FBeginCookContext& Begin
 	{
 		// read AssetPaths out of the file and add them as already-cooked PackageDatas
 		TArray<FString> OutAssetPaths;
-		FFileHelper::LoadFileToStringArray(OutAssetPaths, *ExtraReleaseVersionAssetsFile);
+		FPaths::MakePlatformFilename(ExtraReleaseVersionAssetsFile);
+		FString FullPathFromBaseDir = FPaths::ConvertRelativePathToFull(FPlatformProcess::BaseDir(), ExtraReleaseVersionAssetsFile);
+		FString FullPathFromRootDir = FPaths::ConvertRelativePathToFull(FPaths::RootDir(), ExtraReleaseVersionAssetsFile);
+		if (!FFileHelper::LoadFileToStringArray(OutAssetPaths, *FullPathFromBaseDir))
+		{
+			ensureMsgf(FFileHelper::LoadFileToStringArray(OutAssetPaths, *FullPathFromRootDir), TEXT("Failed to load from %s or %s"), *FullPathFromBaseDir, *FullPathFromRootDir);
+		}
+		
 		for (const FString& AssetPath : OutAssetPaths)
 		{
 			if (UE::Cook::FPackageData* PackageData = PackageDatas->TryAddPackageDataByFileName(FName(*AssetPath)))
