IncrementalCook: Fix "CookWorker initialization failure: Director sent a different message before sending an InitialConfigMessage." when cooking a single package. The message to start the build event phase is sent in immediate mode, and there was no guard preventing immediatemode messages from being sent before the initialconfigmessage. Added one.
#rnx
#incrementalcook

diff --git a/Engine/Source/Editor/UnrealEd/Private/Cooker/CookDirector.cpp b/Engine/Source/Editor/UnrealEd/Private/Cooker/CookDirector.cpp
--- a/Engine/Source/Editor/UnrealEd/Private/Cooker/CookDirector.cpp
+++ b/Engine/Source/Editor/UnrealEd/Private/Cooker/CookDirector.cpp
@@ -669,7 +669,14 @@ void FCookDirector::BroadcastMessage(UE::CompactBinaryTCP::FMarshalledMessage&&
 		{
 		case ECookBroadcastTiming::Immediate:
 		{
-			RemoteWorker->SendMessage(MoveTemp(CopyMessage), ECookDirectorThread::SchedulerThread);
+			if (RemoteWorker->IsConnected())
+			{
+				RemoteWorker->SendMessage(MoveTemp(CopyMessage), ECookDirectorThread::SchedulerThread);
+			}
+			else
+			{
+				RemoteWorker->AppendMessage(MoveTemp(CopyMessage), ECookDirectorThread::SchedulerThread);
+			}
 			break;
 		}
 		case ECookBroadcastTiming::AfterAssignPackages:
