Cooker: Disable MPCook when requested by config but commandline requested -cooksinglepackagenorefs, because it is a waste of time to wait for CookWorkers to be created in that case.
#rnx

diff --git a/Engine/Source/Editor/UnrealEd/Private/CookOnTheFlyServer.cpp b/Engine/Source/Editor/UnrealEd/Private/CookOnTheFlyServer.cpp
--- a/Engine/Source/Editor/UnrealEd/Private/CookOnTheFlyServer.cpp
+++ b/Engine/Source/Editor/UnrealEd/Private/CookOnTheFlyServer.cpp
@@ -5693,7 +5693,7 @@ void UCookOnTheFlyServer::Initialize( ECookMode::Type DesiredCookMode, ECookInit
 	{
 		bool bLaunchedByEditor = !FPlatformMisc::GetEnvironmentVariable(GEditorUIPidVariable).IsEmpty();
 		int32 CookProcessCount=-1;
-		FParse::Value(FCommandLine::Get(), TEXT("-CookProcessCount="), CookProcessCount);
+		bool bSetByCommandLine = FParse::Value(FCommandLine::Get(), TEXT("-CookProcessCount="), CookProcessCount);
 		if (CookProcessCount < 0 && bLaunchedByEditor)
 		{
 			GConfig->GetInt(TEXT("CookSettings"), TEXT("CookProcessCountFromEditor"), CookProcessCount, GEditorIni);
@@ -5712,7 +5712,7 @@ void UCookOnTheFlyServer::Initialize( ECookMode::Type DesiredCookMode, ECookInit
 		}
 		if (CookProcessCount > 1)
 		{
-			CookDirector = MakeUnique<UE::Cook::FCookDirector>(*this, CookProcessCount);
+			CookDirector = MakeUnique<UE::Cook::FCookDirector>(*this, CookProcessCount, bSetByCommandLine);
 			if (!CookDirector->IsMultiprocessAvailable())
 			{
 				CookDirector.Reset();
diff --git a/Engine/Source/Editor/UnrealEd/Private/Cooker/CookDirector.cpp b/Engine/Source/Editor/UnrealEd/Private/Cooker/CookDirector.cpp
--- a/Engine/Source/Editor/UnrealEd/Private/Cooker/CookDirector.cpp
+++ b/Engine/Source/Editor/UnrealEd/Private/Cooker/CookDirector.cpp
@@ -115,9 +115,11 @@ private:
 	ERetractionState RetractionState = ERetractionState::Idle;
 };
 
-FCookDirector::FCookDirector(UCookOnTheFlyServer& InCOTFS, int32 CookProcessCount)
+FCookDirector::FCookDirector(UCookOnTheFlyServer& InCOTFS, int32 CookProcessCount,
+	bool bInCookProcessCountSetByCommandLine)
 	: RunnableShunt(*this) 
 	, COTFS(InCOTFS)
+	, bCookProcessCountSetByCommandLine(bInCookProcessCountSetByCommandLine)
 {
 	check(CookProcessCount > 1);
 	WorkersStalledStartTimeSeconds = MAX_flt;
@@ -145,10 +147,6 @@ FCookDirector::FCookDirector(UCookOnTheFlyServer& InCOTFS, int32 CookProcessCoun
 	}
 	bMultiprocessAvailable = true;
 
-	UE_LOG(LogCook, Display, TEXT("CookProcessCount=%d. CookMultiprocess is enabled with 1 CookDirector and %d %s."),
-		RequestedCookWorkerCount+1, RequestedCookWorkerCount,
-		RequestedCookWorkerCount > 1 ? TEXT("CookWorkers") : TEXT("CookWorker"));
-
 	Register(new FLogMessagesMessageHandler());
 	Register(new TMPCollectorServerMessageCallback<FRetractionResultsMessage>([this]
 	(FMPCollectorServerMessageContext& Context, bool bReadSuccessful, FRetractionResultsMessage&& Message)
@@ -324,13 +322,11 @@ void FCookDirector::StartCook(const FBeginCookContext& InBeginContext)
 void FCookDirector::AssignRequests(TArrayView<FPackageData*> Requests, TArray<FWorkerId>& OutAssignments,
 	TMap<FPackageData*, TArray<FPackageData*>>&& RequestGraph)
 {
-	ActivateMachineResourceReduction();
-
 	TArray<FWorkerId> WorkerIds;
 	TArray<TRefCountPtr<FCookWorkerServer>> LocalRemoteWorkers;
 	{
 		FScopeLock CommunicationScopeLock(&CommunicationLock);
-		InitializeWorkers();
+		InitializeWorkers(Requests.Num());
 	}
 	LocalRemoteWorkers = CopyRemoteWorkers();;
 
@@ -659,7 +655,7 @@ void FCookDirector::BroadcastMessage(UE::CompactBinaryTCP::FMarshalledMessage&&
 
 	{
 		FScopeLock CommunicationScopeLock(&CommunicationLock);
-		InitializeWorkers();
+		InitializeWorkers(0);
 	}
 	TArray<TRefCountPtr<FCookWorkerServer>> LocalRemoteWorkers = CopyRemoteWorkers();
 	for (TRefCountPtr<FCookWorkerServer>& RemoteWorker : LocalRemoteWorkers)
@@ -1241,7 +1237,7 @@ bool FCookDirector::TryCreateWorkerConnectSocket()
 }
 
 
-void FCookDirector::InitializeWorkers()
+void FCookDirector::InitializeWorkers(int32 InitialRequestCount)
 {
 	if (bWorkersInitialized)
 	{
@@ -1260,6 +1256,23 @@ void FCookDirector::InitializeWorkers()
 		}
 	};
 
+	constexpr int32 MinRequestsForMPCookWithCookSinglePackage = 100;
+	UE::Cook::FCookByTheBookOptions& Options = *COTFS.CookByTheBookOptions;
+	if (!bCookProcessCountSetByCommandLine
+		&& (Options.bSkipHardReferences
+			|| (Options.bSkipSoftReferences && InitialRequestCount < MinRequestsForMPCookWithCookSinglePackage)))
+	{
+		UE_LOG(LogCook, Display,
+			TEXT("CookDirector initialization skipped: -CookSinglePackage was requested. CookMultiprocess is disabled and the cooker is running as a single process."));
+		return;
+	}
+
+	UE_LOG(LogCook, Display, TEXT("CookProcessCount=%d. CookMultiprocess is enabled with 1 CookDirector and %d %s."),
+		RequestedCookWorkerCount + 1, RequestedCookWorkerCount,
+		RequestedCookWorkerCount > 1 ? TEXT("CookWorkers") : TEXT("CookWorker"));
+
+	ActivateMachineResourceReduction();
+
 	if (!TryCreateWorkerConnectSocket())
 	{
 		return;
diff --git a/Engine/Source/Editor/UnrealEd/Private/Cooker/CookDirector.h b/Engine/Source/Editor/UnrealEd/Private/Cooker/CookDirector.h
--- a/Engine/Source/Editor/UnrealEd/Private/Cooker/CookDirector.h
+++ b/Engine/Source/Editor/UnrealEd/Private/Cooker/CookDirector.h
@@ -73,7 +73,7 @@ class FCookDirector
 {
 public:
 
-	FCookDirector(UCookOnTheFlyServer& InCOTFS, int32 CookProcessCount);
+	FCookDirector(UCookOnTheFlyServer& InCOTFS, int32 CookProcessCount, bool bInCookProcessCountSetByCommandLine);
 	~FCookDirector();
 
 	bool IsMultiprocessAvailable() const;
@@ -191,7 +191,7 @@ private:
 	 * Construct CookWorkerServers and communication thread if not yet constructed.
 	 * The CookWorkerServers are constructed to Uninitialized; the worker process is created later.
 	 */
-	void InitializeWorkers();
+	void InitializeWorkers(int32 InitialRequestCount);
 	/**
 	 * Copy to snapshot variables the data required on the communication thread that can only be read from the
 	 * scheduler thread.
@@ -286,6 +286,7 @@ private:
 	bool bMultiprocessAvailable = false;
 	bool bReceivingMessages = false;
 	bool bForceNextHeartbeat = false;
+	bool bCookProcessCountSetByCommandLine = false;
 
 	// Data that is read-only while the CommunicationThread is active and is readable from any thread
 	FBeginCookContextForWorker BeginCookContext;
