Output additional informations in the diff asset  commandlet and the diff handler

diff --git a/Engine/Source/Editor/UnrealEd/Private/Commandlets/DiffAssetRegistriesCommandlet.cpp b/Engine/Source/Editor/UnrealEd/Private/Commandlets/DiffAssetRegistriesCommandlet.cpp
--- a/Engine/Source/Editor/UnrealEd/Private/Commandlets/DiffAssetRegistriesCommandlet.cpp
+++ b/Engine/Source/Editor/UnrealEd/Private/Commandlets/DiffAssetRegistriesCommandlet.cpp
@@ -933,12 +933,18 @@ void UDiffAssetRegistriesCommandlet::LogClassSummary(FArchive *CSVFile, const FS
 		});
 	}
 
+	int64 TotalChangesSize = 0;
+	int64 TotalSize = 0;
+
 	//Overall Class Summary
 	bool bChangesPastThreshold = false;
 	for (FName ClassName : ClassNames)
 	{
 		const FChangeInfo& Changes = InChangeInfoByAsset[ClassName];
 
+		TotalSize += Changes.GetTotalSize();
+		TotalChangesSize += Changes.GetTotalChangeSize();
+
 		if (Changes.GetTotalChangeSize() == 0)
 		{
 			continue;
@@ -1040,6 +1046,10 @@ void UDiffAssetRegistriesCommandlet::LogClassSummary(FArchive *CSVFile, const FS
 		}
 	}
 
+	// log summary & change
+	UE_LOG(LogDiffAssets, Display, TEXT("%s%s FullSummary: "), *HeaderPrefix, *HeaderSpacer);
+	UE_LOG(LogDiffAssets, Display, TEXT("  : %.02f%% changes (%.02f MB Changed out of %.02f MB Total)"), TotalSize ? (float)TotalChangesSize / (float)TotalSize * 100.0 : 0.0f, TotalChangesSize * InvToMB, TotalSize * InvToMB);
+
 	//If we didn't find any changes of sufficient size, note as much instead of the header
 	if (!bChangesPastThreshold)
 	{
