- Remove unecessary menu from Material Editor Viewport
- Removes the label from the Editor Preview Scene control

diff --git a/Engine/Source/Editor/MaterialEditor/Private/SMaterialEditorViewport.cpp b/Engine/Source/Editor/MaterialEditor/Private/SMaterialEditorViewport.cpp
--- a/Engine/Source/Editor/MaterialEditor/Private/SMaterialEditorViewport.cpp
+++ b/Engine/Source/Editor/MaterialEditor/Private/SMaterialEditorViewport.cpp
@@ -984,7 +984,6 @@ TSharedPtr<SWidget> SMaterialEditor3DPreviewViewport::BuildViewportToolbar()
 			}
 
 			RightSection.AddEntry(UE::UnrealEd::CreateViewModesSubmenu());
-			RightSection.AddEntry(UE::MaterialEditor::CreateShowSubmenu());
 			RightSection.AddEntry(UE::UnrealEd::CreatePerformanceAndScalabilitySubmenu());
 
 			// Add Preview Scene Submenu
diff --git a/Engine/Source/Editor/UnrealEd/Private/ViewportToolbar/UnrealEdViewportToolbar.cpp b/Engine/Source/Editor/UnrealEd/Private/ViewportToolbar/UnrealEdViewportToolbar.cpp
--- a/Engine/Source/Editor/UnrealEd/Private/ViewportToolbar/UnrealEdViewportToolbar.cpp
+++ b/Engine/Source/Editor/UnrealEd/Private/ViewportToolbar/UnrealEdViewportToolbar.cpp
@@ -2722,19 +2722,18 @@ FToolMenuEntry CreateAssetViewerProfileSubmenu()
 
 				FToolMenuEntry& Entry = InnerSection.AddSubMenu(
 					"AssetViewerProfile",
-					TAttribute<FText>::CreateLambda(
-						[PreviewProfileControllerWeak]()
+					FText::GetEmpty(),
+					TAttribute<FText>::CreateLambda([PreviewProfileControllerWeak]
+					{
+						if (TSharedPtr<IPreviewProfileController> PreviewProfileController = PreviewProfileControllerWeak.Pin())
 						{
-							if (TSharedPtr<IPreviewProfileController> PreviewProfileController =
-									PreviewProfileControllerWeak.Pin())
-							{
-								return FText::FromString(PreviewProfileController->GetActiveProfile());
-							}
-
-							return LOCTEXT("AssetViewerDefaultProfileLabel", "Profile");
+							return FText::Format(
+								LOCTEXT("AssetViewerProfileSelectionSectionTooltipFormat", "Preview Scene: {0}"),
+								FText::FromString(PreviewProfileController->GetActiveProfile())
+							);
 						}
-					),
-					LOCTEXT("AssetViewerProfileSelectionSectionTooltip", "Select the Preview Scene Profile for this viewport."),
+						return LOCTEXT("AssetViewerProfileSelectionSectionNoneTooltip", "Select the Preview Scene Profile for this viewport.");	
+					}),
 					FNewToolMenuDelegate::CreateLambda(
 						[PreviewProfileControllerWeak](UToolMenu* Submenu) -> void
 						{
