Material Editor: On Shift+Enter or Shift+Click, wire up a newly-created node to the selected node.

This behavior is similar to what you'd get when dragging a new connection off a node and summoning the node-creation menu that way.

diff --git a/Engine/Source/Editor/UnrealEd/Private/MaterialGraphSchema.cpp b/Engine/Source/Editor/UnrealEd/Private/MaterialGraphSchema.cpp
--- a/Engine/Source/Editor/UnrealEd/Private/MaterialGraphSchema.cpp
+++ b/Engine/Source/Editor/UnrealEd/Private/MaterialGraphSchema.cpp
@@ -10,11 +10,13 @@
 #include "UObject/UnrealType.h"
 #include "UObject/PropertyPortFlags.h"
 #include "Textures/SlateIcon.h"
+#include "Framework/Application/SlateApplication.h"
 #include "Framework/Commands/UIAction.h"
 #include "Framework/MultiBox/MultiBoxBuilder.h"
 #include "ToolMenus.h"
 #include "EdGraph/EdGraph.h"
 #include "HAL/IConsoleManager.h"
+#include "IMaterialEditor.h"
 #include "Materials/MaterialExpression.h"
 #include "Materials/MaterialExpressionMaterialFunctionCall.h"
 #include "Materials/MaterialFunction.h"
@@ -62,6 +64,26 @@ UEdGraphNode* FMaterialGraphSchemaAction_NewNode::PerformAction(class UEdGraph*
 
 	const FScopedTransaction Transaction(NSLOCTEXT("UnrealEd", "MaterialEditorNewExpression", "Material Editor: New Expression"));
 
+	// When FromPin is nullptr and Shift is pressed, we attempt to use the selected node's output pin as our FromPin
+	if (FromPin == nullptr && FSlateApplication::Get().GetModifierKeys().IsShiftDown())
+	{
+		// Determine Previously Selected Node beforce creating the new node (which can select itself if bSelectNewNode is true)
+		if (TSharedPtr<IMaterialEditor> MaterialEditor = FMaterialEditorUtilities::GetIMaterialEditorForObject(ParentGraph))
+		{
+			TSet<UObject*> SelectedNodes = MaterialEditor->GetSelectedNodes();
+			if (SelectedNodes.Num() == 1)
+			{
+				for (UObject* SelectedNode : SelectedNodes)
+				{
+					if (UMaterialGraphNode_Base* SelectedMaterialGraphNode_Base = Cast<UMaterialGraphNode_Base>(SelectedNode))
+					{
+						FromPin = SelectedMaterialGraphNode_Base->GetOutputPin(0);
+					}
+				}
+			}
+		}
+	}
+
 	UMaterialExpression* NewExpression = FMaterialEditorUtilities::CreateNewMaterialExpression(ParentGraph, MaterialExpressionClass, Location, bSelectNewNode, /*bAutoAssignResource*/true);
 
 	if (NewExpression)
