Add automatic upgrading of SCommonEditorViewportToolbarBase to the new UToolMenu tool system.
- Toolbars only automatically upgrade if they stay within a restricted set of customizations.
- Modifies SEditorViewport to support automatic relocating of automatically upgraded toolbars.

diff --git a/Engine/Source/Editor/UnrealEd/Private/SCommonEditorViewportToolbarBase.cpp b/Engine/Source/Editor/UnrealEd/Private/SCommonEditorViewportToolbarBase.cpp
--- a/Engine/Source/Editor/UnrealEd/Private/SCommonEditorViewportToolbarBase.cpp
+++ b/Engine/Source/Editor/UnrealEd/Private/SCommonEditorViewportToolbarBase.cpp
@@ -72,8 +72,45 @@ void AddScreenPercentageMenu(FMenuBuilder& InMenuBuilder, const FEditorViewportC
 
 	InMenuBuilder.AddWidget(UToolMenus::Get()->GenerateWidget(OldScreenPercentageMenuName, MenuContext), FText());
 }
+
+TSharedPtr<const SCommonEditorViewportToolbarBase> FindToolbarWidget(UToolMenu* ToolMenu)
+{
+	if (const UCommonViewportToolbarBaseMenuContext* Context = ToolMenu->FindContext<UCommonViewportToolbarBaseMenuContext>())
+	{
+		if (TSharedPtr<const SCommonEditorViewportToolbarBase> ToolbarWidget = Context->ToolbarWidget.Pin())
+		{
+			return ToolbarWidget;
+		}
+	}
+	return nullptr;
+}
+
+TSharedPtr<const SCommonEditorViewportToolbarBase> FindToolbarWidget(const FToolMenuSection& ToolSection)
+{
+	if (const UCommonViewportToolbarBaseMenuContext* Context = ToolSection.FindContext<UCommonViewportToolbarBaseMenuContext>())
+	{
+		if (TSharedPtr<const SCommonEditorViewportToolbarBase> ToolbarWidget = Context->ToolbarWidget.Pin())
+		{
+			return ToolbarWidget;
+		}
+	}
+	return nullptr;
+}
+
 } // namespace UE::UnrealEd::Private
 
+TSharedPtr<IPreviewProfileController> UCommonViewportToolbarBaseMenuContext::GetPreviewProfileController() const
+{
+	if (TSharedPtr<const SCommonEditorViewportToolbarBase> Toolbar = ToolbarWidget.Pin())
+	{
+		if (TSharedPtr<IPreviewProfileController> Controller = Toolbar->GetPreviewProfileController())
+		{
+			return Controller;
+		}
+	}
+	return UUnrealEdViewportToolbarContext::GetPreviewProfileController();
+}
+
 //////////////////////////////////////////////////////////////////////////
 // SPreviewSceneProfileSelector
 
@@ -189,24 +226,17 @@ void SCommonEditorViewportToolbarBase::Construct(const FArguments& InArgs, TShar
 {
 	InfoProviderPtr = InInfoProvider;
 
-	TSharedRef<SEditorViewport> ViewportRef = GetInfoProvider().GetViewportWidget();
-	TSharedPtr<SHorizontalBox> MainBoxPtr;
+	// Create a blank menu to be treated as "null" for the purposes of detecting whether a custom view menu has been defined.
+	BlankViewMenu = MakeShared<SEditorViewportViewMenu>();
 
 	const FMargin ToolbarSlotPadding(4.0f, 1.0f);
 	const FMargin ToolbarButtonPadding(4.0f, 0.0f);
 
-	ChildSlot
-	[
-		SNew( SBorder )
-		.BorderImage(FAppStyle::Get().GetBrush("EditorViewportToolBar.Background"))
-		.Cursor(EMouseCursor::Default)
-		[
-			SAssignNew( MainBoxPtr, SHorizontalBox )
-		]
-	];
+	TSharedRef<SEditorViewport> ViewportRef = GetInfoProvider().GetViewportWidget();
+	TSharedRef<SHorizontalBox> MainBox = SNew(SHorizontalBox);
 
 	// Options menu
-	MainBoxPtr->AddSlot()
+	MainBox->AddSlot()
 		.AutoWidth()
 		.Padding(ToolbarSlotPadding)
 		[
@@ -218,7 +248,7 @@ void SCommonEditorViewportToolbarBase::Construct(const FArguments& InArgs, TShar
 		];
 
 	// Camera mode menu
-	MainBoxPtr->AddSlot()
+	MainBox->AddSlot()
 		.AutoWidth()
 		.Padding(ToolbarSlotPadding)
 		[
@@ -230,7 +260,7 @@ void SCommonEditorViewportToolbarBase::Construct(const FArguments& InArgs, TShar
 		];
 
 	// View menu
-	MainBoxPtr->AddSlot()
+	MainBox->AddSlot()
 		.AutoWidth()
 		.Padding(ToolbarSlotPadding)
 		[
@@ -238,7 +268,7 @@ void SCommonEditorViewportToolbarBase::Construct(const FArguments& InArgs, TShar
 		];
 
 	// Show menu
-	MainBoxPtr->AddSlot()
+	MainBox->AddSlot()
 		.AutoWidth()
 		.Padding(ToolbarSlotPadding)
 		[
@@ -252,18 +282,19 @@ void SCommonEditorViewportToolbarBase::Construct(const FArguments& InArgs, TShar
 	// Profile menu (Controls the Preview Scene Settings)
 	if (InArgs._PreviewProfileController)
 	{
-		MainBoxPtr->AddSlot()
+		PreviewProfileController = InArgs._PreviewProfileController;
+		MainBox->AddSlot()
 			.AutoWidth()
 			.Padding(ToolbarSlotPadding)
 			[
-				SNew(SPreviewSceneProfileSelector).PreviewProfileController(InArgs._PreviewProfileController)
+				SNew(SPreviewSceneProfileSelector).PreviewProfileController(PreviewProfileController)
 			];
 	}
 
 	// Realtime button
 	if (InArgs._AddRealtimeButton)
 	{
-		MainBoxPtr->AddSlot()
+		MainBox->AddSlot()
 			.AutoWidth()
 			.Padding(ToolbarSlotPadding)
 			[
@@ -283,7 +314,7 @@ void SCommonEditorViewportToolbarBase::Construct(const FArguments& InArgs, TShar
 			];
 	}
 
-	MainBoxPtr->AddSlot()
+	MainBox->AddSlot()
 		.AutoWidth()
 		.Padding(ToolbarSlotPadding)
 		[
@@ -295,7 +326,7 @@ void SCommonEditorViewportToolbarBase::Construct(const FArguments& InArgs, TShar
 			.OnGetMenuContent(this, &SCommonEditorViewportToolbarBase::GenerateViewModeOptionsMenu)
 		];
 
-	MainBoxPtr->AddSlot()
+	MainBox->AddSlot()
 		.AutoWidth()
 		.Padding(ToolbarSlotPadding)
 		[
@@ -310,10 +341,10 @@ void SCommonEditorViewportToolbarBase::Construct(const FArguments& InArgs, TShar
 		];
 
 	// Add optional toolbar slots to be added by child classes inherited from this common viewport toolbar
-	ExtendLeftAlignedToolbarSlots(MainBoxPtr, SharedThis(this));
+	ExtendLeftAlignedToolbarSlots(MainBox, SharedThis(this));
 
 	// Transform toolbar
-	MainBoxPtr->AddSlot()
+	MainBox->AddSlot()
 		.Padding(ToolbarSlotPadding)
 		.HAlign(HAlign_Right)
 		[
@@ -324,6 +355,238 @@ void SCommonEditorViewportToolbarBase::Construct(const FArguments& InArgs, TShar
 			.Visibility(ViewportRef, &SEditorViewport::GetTransformToolbarVisibility)
 		];
 
+	// Custom view menus and widgets added to the left side of the menu will retain
+	// the old toolbar widget appearance. Simply including them alongside new elements
+	// will mix styles undesirably, ultimately looking more broken than just using the old
+	// toolbar.
+	// Thus, the automatic toolbar upgrade is only enabled when customizations (e.g. show menus)
+	// can be incorporated into the new design.
+	const bool bUseUpgradedToolbar = !bHasExtendedLeftSide && bUsesDefaultViewMenu;
+	if (bUseUpgradedToolbar)
+	{
+		const FName ViewportToolbarName = "UnrealEd.ViewportToolbar";
+		
+		if (!UToolMenus::Get()->IsMenuRegistered(ViewportToolbarName))
+		{
+			UToolMenu* const ViewportToolbarMenu = UToolMenus::Get()->RegisterMenu(
+				ViewportToolbarName, NAME_None, EMultiBoxType::SlimHorizontalToolBar
+			);
+			
+			ViewportToolbarMenu->StyleName = "ViewportToolbar";
+			
+			FToolMenuSection& LeftSection = ViewportToolbarMenu->FindOrAddSection("Left");
+			{
+				LeftSection.AddEntry(UE::UnrealEd::CreateTransformsSubmenu());
+				LeftSection.AddEntry(UE::UnrealEd::CreateSnappingSubmenu());
+			}
+			
+			FToolMenuSection& RightSection = ViewportToolbarMenu->FindOrAddSection("Right");
+			RightSection.Alignment = EToolMenuSectionAlign::Last;
+			{
+				RightSection.AddEntry(UE::UnrealEd::CreateAssetViewerProfileSubmenu());
+				
+				// Camera Menu
+				{
+					RightSection.AddEntry(UE::UnrealEd::CreateCameraSubmenu());
+					const FName CameraSubmenuName = UToolMenus::JoinMenuPaths(ViewportToolbarName, "Camera");
+					UE::UnrealEd::ExtendCameraSubmenu(CameraSubmenuName);
+					
+					UToolMenu* CameraSubMenu = UToolMenus::Get()->ExtendMenu(CameraSubmenuName);
+					FToolMenuSection& Section = CameraSubMenu->FindOrAddSection(NAME_None);
+					Section.InsertPosition.Name = "Exposure";
+					Section.InsertPosition.Position = EToolMenuInsertType::Before;
+					
+					Section.AddEntry(UE::UnrealEd::CreateCameraSpeedMenu());
+				}
+				
+				// View Menu
+				{
+					// Include backwards-compatability with earlier toolbars
+					// Create our grandparent menu.
+					if (!UToolMenus::Get()->IsMenuRegistered("UnrealEd.ViewportToolbar.View"))
+					{
+						UToolMenus::Get()->RegisterMenu("UnrealEd.ViewportToolbar.View");
+					}
+					
+					// Create our menu.
+					UToolMenus::Get()->RegisterMenu(
+						"UnrealEd.ViewportToolbar.ViewModes", "UnrealEd.ViewportToolbar.View"
+					);
+					
+					RightSection.AddEntry(UE::UnrealEd::CreateViewModesSubmenu());
+				}
+				
+				// Show Menu
+				{
+					// Include backwards-compatability with earlier toolbars
+					if (!UToolMenus::Get()->IsMenuRegistered("ViewportToolbarBase.Show"))
+					{
+						UToolMenus::Get()->RegisterMenu("ViewportToolbarBase.Show");
+					}
+					UToolMenus::Get()->RegisterMenu(UToolMenus::JoinMenuPaths(ViewportToolbarName, "Show"), "ViewportToolbarBase.Show");
+				
+					FToolMenuEntry& ShowMenu = RightSection.AddSubMenu(
+						"Show",
+						LOCTEXT("ShowSubmenuLabel", "Show"),
+						LOCTEXT("ShowSubmenuTooltip", "Show flags related to the current viewport"),
+						FNewToolMenuDelegate::CreateLambda([](UToolMenu* Submenu)
+						{
+							Submenu->AddDynamicSection("Flags", FNewToolMenuDelegate::CreateLambda([](UToolMenu* Menu)
+							{
+								TSharedPtr<const SCommonEditorViewportToolbarBase> ToolbarWidget = UE::UnrealEd::Private::FindToolbarWidget(Menu);
+								if (!ToolbarWidget)
+								{
+									return;
+								}
+								
+								if (!ToolbarWidget->bIsGeneratingToolMenuWidget)
+								{
+									if (TSharedPtr<SWidget> LegacyWidget = ToolbarWidget->MakeLegacyShowMenu())
+									{
+										// Display legacy menu
+										Menu->AddSection("LegacyWidget").AddEntry(FToolMenuEntry::InitWidget(
+											"LegacyWidget",
+											LegacyWidget.ToSharedRef(),
+											FText::GetEmpty(), // No label
+											true, // No indent
+											true, // Searchable
+											true // No padding
+										));
+									}
+									else
+									{
+										ToolbarWidget->FillShowFlagsMenu(Menu);	
+									}
+								}
+							}));
+						})
+					);
+					ShowMenu.Icon = FSlateIcon(FAppStyle::GetAppStyleSetName(), "Level.VisibleHighlightIcon16x");
+					ShowMenu.InsertPosition.Position = EToolMenuInsertType::Last;
+					ShowMenu.ToolBarData.LabelOverride = FText();
+					ShowMenu.ToolBarData.ResizeParams.ClippingPriority = 800;
+				}
+				
+				RightSection.AddEntry(UE::UnrealEd::CreatePerformanceAndScalabilitySubmenu());
+					
+				RightSection.AddDynamicEntry("LegacyOptionsMenu", FNewToolMenuSectionDelegate::CreateLambda([](FToolMenuSection& Section)
+				{
+					TSharedPtr<const SCommonEditorViewportToolbarBase> ToolbarWidget = UE::UnrealEd::Private::FindToolbarWidget(Section);
+					if (!ToolbarWidget)
+					{
+						return;
+					}
+					
+					if (ToolbarWidget->ShouldCreateOptionsMenu())
+					{
+						Section.AddSubMenu(
+							"Settings",
+							LOCTEXT("SettingsSubmenuLabel", "Settings"),
+							LOCTEXT("SettingsSubmenuTooltip", "Viewport-related settings"),
+							FNewToolMenuDelegate::CreateLambda([](UToolMenu* Submenu)
+							{
+								Submenu->AddDynamicSection("Settings", FNewToolMenuDelegate::CreateLambda([](UToolMenu* Menu)
+								{
+									TSharedPtr<const SCommonEditorViewportToolbarBase> Toolbar = UE::UnrealEd::Private::FindToolbarWidget(Menu);
+									if (!Toolbar)
+									{
+										return;
+									}
+								
+									FToolMenuSection& Section = Menu->FindOrAddSection("Settings", LOCTEXT("SettingsSectionLabel", "Settings"));	
+									FMenuBuilder LegacyMenuBuilder(true, nullptr);
+									Toolbar->ExtendOptionsMenu(LegacyMenuBuilder);
+									Section.AddEntry(FToolMenuEntry::InitWidget(
+										"LegacySettingsMenus",
+										LegacyMenuBuilder.MakeWidget(),
+										FText::GetEmpty(), // No label
+										true, // No indent
+										true, // Searchable
+										true // No padding
+									));
+								}));
+							})
+						);
+					}
+				}));
+			}
+		}
+		
+		FToolMenuContext ViewportToolbarContext;
+		{
+			ViewportToolbarContext.AppendCommandList(GetInfoProvider().GetViewportWidget()->GetCommandList());
+		
+			UCommonViewportToolbarBaseMenuContext* Context = NewObject<UCommonViewportToolbarBaseMenuContext>();
+			Context->ToolbarWidget = SharedThis(this);
+			
+			Context->Viewport = GetInfoProvider().GetViewportWidget();
+			Context->IsViewModeSupported.BindSP(this, &SCommonEditorViewportToolbarBase::IsViewModeSupported);
+			
+			ViewportToolbarContext.AddObject(Context);
+			ViewportToolbarContext.AddExtender(GetViewMenuExtender());
+		}
+		
+		bIsGeneratingToolMenuWidget = true;
+		TSharedRef<SWidget> ToolMenuWidget = UToolMenus::Get()->GenerateWidget(ViewportToolbarName, ViewportToolbarContext);
+		bIsGeneratingToolMenuWidget = false;
+
+		// Allow the new toolbar to fall back to the old look & behavior
+		ChildSlot
+		[
+			SNew( SBorder )
+			.BorderImage(FAppStyle::Get().GetBrush("EditorViewportToolBar.Background"))
+			.Cursor(EMouseCursor::Default)
+			[
+				SNew(SVerticalBox)
+				+ SVerticalBox::Slot()
+				.AutoHeight()
+				[
+					SNew(SBox)
+					.Visibility_Lambda([]
+					{
+						return UE::UnrealEd::ShowNewViewportToolbars() ? EVisibility::Visible: EVisibility::Collapsed;
+					})
+					[
+						ToolMenuWidget
+					]
+				]
+				+ SVerticalBox::Slot()
+				.AutoHeight()
+				[
+					SNew(SBox)
+					.Visibility_Lambda([]
+					{
+						return UE::UnrealEd::ShowOldViewportToolbars() ? EVisibility::Visible: EVisibility::Collapsed;
+					})
+					[
+						MainBox
+					]
+				]
+			]
+		];
+		
+		// Register the child widget as automatically upgradeable in the viewport
+		if (InInfoProvider)
+		{
+			if (TSharedPtr<SEditorViewport> Viewport = InInfoProvider->GetViewportWidget())
+            {
+            	Viewport->MarkLegacyToolbarChildAsAutomaticallyUpgradable(ChildSlot.GetWidget());
+            }
+		}
+	}
+	else
+	{
+		ChildSlot
+		[
+			SNew( SBorder )
+			.BorderImage(FAppStyle::Get().GetBrush("EditorViewportToolBar.Background"))
+			.Cursor(EMouseCursor::Default)
+			[
+				MainBox	
+			]
+		];
+	}
+
 	SViewportToolBar::Construct(SViewportToolBar::FArguments());
 }
 
@@ -337,6 +600,10 @@ FText SCommonEditorViewportToolbarBase::GetCameraMenuLabel() const
 	return UE::UnrealEd::GetCameraSubmenuLabelFromViewportType(GetViewportClient().GetViewportType());
 }
 
+FSlateIcon SCommonEditorViewportToolbarBase::GetCameraMenuIcon() const
+{
+	return FSlateIcon(FAppStyle::GetAppStyleSetName(), UE::UnrealEd::GetCameraSubmenuIconFNameFromViewportType(GetViewportClient().GetViewportType()));
+}
 
 EVisibility SCommonEditorViewportToolbarBase::GetViewModeOptionsVisibility() const
 {
@@ -416,6 +683,12 @@ TSharedRef<SWidget> SCommonEditorViewportToolbarBase::GenerateCameraMenu() const
 
 TSharedRef<SWidget> SCommonEditorViewportToolbarBase::GenerateShowMenu() const
 {
+	if (bIsBuildingToolMenu)
+	{
+		// Defer to the newer system
+		return SNullWidget::NullWidget;
+	}
+
 	GetInfoProvider().OnFloatingButtonClicked();
 
 	static const FName MenuName("ViewportToolbarBase.Show");
@@ -595,6 +868,9 @@ FEditorViewportClient& SCommonEditorViewportToolbarBase::GetViewportClient() con
 
 TSharedRef<SEditorViewportViewMenu> SCommonEditorViewportToolbarBase::MakeViewMenu()
 {
+	// Mark that the viewport uses the default view menu, and is potentially upgradable.
+	bUsesDefaultViewMenu = true;
+	
 	TSharedRef<SEditorViewport> ViewportRef = GetInfoProvider().GetViewportWidget();
 
 	return SNew(SEditorViewportViewMenu, ViewportRef, SharedThis(this))
@@ -629,4 +905,38 @@ TSharedRef<SWidget> SCommonEditorViewportToolbarBase::GetScalabilityWarningMenuC
 		];
 }
 
+//////////////////////////////////////////////////////////////////////////
+/// Automatic Legacy Upgrade support
+/// --------------------------------------------------------------------
+/// These functions are seriously hairy so that clients of this widget
+/// with simple needs are converted directly to the new form.
+//////////////////////////////////////////////////////////////////////////
+TSharedPtr<SWidget> SCommonEditorViewportToolbarBase::MakeLegacyShowMenu() const
+{
+	bIsBuildingToolMenu = true;
+	const TSharedRef<SWidget> Menu = GenerateShowMenu();
+	bIsBuildingToolMenu = false;
+	
+	return Menu == SNullWidget::NullWidget ? TSharedPtr<SWidget>(nullptr) : Menu.ToSharedPtr();
+}
+
+void SCommonEditorViewportToolbarBase::ExtendOptionsMenu(FMenuBuilder& OptionsMenuBuilder) const
+{
+	// This flag allows the new toolbar to detect whether a settings menu needs to be created.
+	bHasExtendedSettingsMenu = false;
+}
+
+void SCommonEditorViewportToolbarBase::ExtendLeftAlignedToolbarSlots(TSharedPtr<SHorizontalBox> MainBoxPtr, TSharedPtr<SViewportToolBar> ParentToolBarPtr) const
+{
+	// This flag allows detection on whether the client intends to extend the left side. 
+	bHasExtendedLeftSide = false;
+}
+
+bool SCommonEditorViewportToolbarBase::ShouldCreateOptionsMenu() const
+{
+	FMenuBuilder LegacyMenuBuilder(true, nullptr);
+	ExtendOptionsMenu(LegacyMenuBuilder);
+	return bHasExtendedSettingsMenu;
+}
+
 #undef LOCTEXT_NAMESPACE
diff --git a/Engine/Source/Editor/UnrealEd/Private/SEditorViewport.cpp b/Engine/Source/Editor/UnrealEd/Private/SEditorViewport.cpp
--- a/Engine/Source/Editor/UnrealEd/Private/SEditorViewport.cpp
+++ b/Engine/Source/Editor/UnrealEd/Private/SEditorViewport.cpp
@@ -135,7 +135,24 @@ void SEditorViewport::Construct( const FArguments& InArgs )
 	// Set up viewport toolbars.
 	{
 		const TSharedPtr<SWidget> OldViewportToolbar = MakeViewportToolbar();
-		const TSharedPtr<SWidget> NewViewportToolbar = BuildViewportToolbar();
+		TSharedPtr<SWidget> NewViewportToolbar = BuildViewportToolbar();
+		
+		// Allow programatically-migrated viewports (e.g. SCommonEditorViewportToolbarBase) to register themselves on construction as supporting either position
+		// This indirection is done to minimize API changes (e.g. using a specific toolbar interface). The goal is that clients do nothing and get upgrades.
+		if (!NewViewportToolbar && OldViewportToolbar)
+		{
+			// Overrides of SCommonEditorViewportToolbarBase can construct the full base toolbar and then overwrite the child slot with their own content.
+			// This check ensures that the automatic upgrade is only allowed if we end up with a widget containing the child that we expected to have
+			// when the upgrade was requested.
+			if (TSharedPtr<SWidget> ExpectedChild = AutoUpgradeWidgetChild.Pin())
+			{
+				if (OldViewportToolbar->GetChildren()->Num() == 1 && OldViewportToolbar->GetChildren()->GetChildAt(0) == ExpectedChild)
+				{
+					bLegacyToolbarIsAutomaticallyUpgradable = true;
+					NewViewportToolbar = OldViewportToolbar;
+				}
+			}
+		}
 
 		if (OldViewportToolbar)
 		{
@@ -155,7 +172,7 @@ void SEditorViewport::Construct( const FArguments& InArgs )
 								return EVisibility::Visible;
 							}
 							
-							return  UE::UnrealEd::ShowOldViewportToolbars() ? EVisibility::Visible: EVisibility::Collapsed;
+							return UE::UnrealEd::ShowOldViewportToolbars() ? EVisibility::Visible: EVisibility::Collapsed;
 						}
 					)
 					[
@@ -176,15 +193,15 @@ void SEditorViewport::Construct( const FArguments& InArgs )
 			[
 				SNew(SBox)
 				.Visibility_Lambda(
-					[bHasOldViewportToolbar]() -> EVisibility
+					[bHasOldViewportToolbar, bHasLegacyUpgradedToolbar = bLegacyToolbarIsAutomaticallyUpgradable]() -> EVisibility
 					{
 						// Always show the new viewport toolbar if there is no old one.
 						if (!bHasOldViewportToolbar)
 						{
 							return EVisibility::Visible;
 						}
-						
-						return  UE::UnrealEd::ShowNewViewportToolbars() ? EVisibility::Visible: EVisibility::Collapsed;
+						// In the case of the old toolbar being upgraded, don't show it twice
+						return (UE::UnrealEd::ShowNewViewportToolbars() && (!bHasLegacyUpgradedToolbar || !UE::UnrealEd::ShowOldViewportToolbars())) ? EVisibility::Visible: EVisibility::Collapsed;
 					}
 				)
 				[
@@ -1050,6 +1067,11 @@ TSharedPtr<IPreviewProfileController> SEditorViewport::GetPreviewProfileControll
 	return PreviewProfileController;
 }
 
+void SEditorViewport::MarkLegacyToolbarChildAsAutomaticallyUpgradable(const TSharedRef<SWidget>& ExpectedChild)
+{
+	AutoUpgradeWidgetChild = ExpectedChild;
+}
+
 EActiveTimerReturnType SEditorViewport::EnsureTick( double InCurrentTime, float InDeltaTime )
 {
 	// Keep the timer going if we're realtime or were invalidated this frame
diff --git a/Engine/Source/Editor/UnrealEd/Private/ViewportToolbar/UnrealEdViewportToolbar.cpp b/Engine/Source/Editor/UnrealEd/Private/ViewportToolbar/UnrealEdViewportToolbar.cpp
--- a/Engine/Source/Editor/UnrealEd/Private/ViewportToolbar/UnrealEdViewportToolbar.cpp
+++ b/Engine/Source/Editor/UnrealEd/Private/ViewportToolbar/UnrealEdViewportToolbar.cpp
@@ -2705,14 +2705,7 @@ FToolMenuEntry CreateAssetViewerProfileSubmenu()
 					return;
 				}
 
-				const TSharedPtr<SEditorViewport> EditorViewport = EditorViewportContext->Viewport.Pin();
-				if (!EditorViewport)
-				{
-					return;
-				}
-
-				TSharedPtr<IPreviewProfileController> PreviewProfileController =
-					EditorViewport->GetPreviewProfileController();
+				TSharedPtr<IPreviewProfileController> PreviewProfileController = EditorViewportContext->GetPreviewProfileController();
 				if (!PreviewProfileController)
 				{
 					return;
@@ -3432,8 +3425,7 @@ FToolMenuEntry CreateScalabilitySubmenu()
 					FSlateIcon(FAppStyle::GetAppStyleSetName(), "Icons.ViewportScalability")
 				);
 
-				ScalabilitySubmenu.SetShowInToolbarTopLevel(TAttribute<bool>::Create(&UE::UnrealEd::IsScalabilityWarningVisible
-				));
+				ScalabilitySubmenu.SetShowInToolbarTopLevel(TAttribute<bool>::Create(&UE::UnrealEd::IsScalabilityWarningVisible));
 
 				ScalabilitySubmenu.ToolBarData.ResizeParams.AllowClipping = TAttribute<bool>::CreateLambda(
 					[]() -> bool
diff --git a/Engine/Source/Editor/UnrealEd/Private/ViewportToolbar/UnrealEdViewportToolbarContext.cpp b/Engine/Source/Editor/UnrealEd/Private/ViewportToolbar/UnrealEdViewportToolbarContext.cpp
new file mode 100644
--- /dev/null
+++ b/Engine/Source/Editor/UnrealEd/Private/ViewportToolbar/UnrealEdViewportToolbarContext.cpp
@@ -0,0 +1,14 @@
+﻿// Copyright Epic Games, Inc. All Rights Reserved.
+
+#include "ViewportToolbar/UnrealEdViewportToolbarContext.h"
+
+#include "SEditorViewport.h"
+
+TSharedPtr<IPreviewProfileController> UUnrealEdViewportToolbarContext::GetPreviewProfileController() const
+{
+	if (const TSharedPtr<SEditorViewport> EditorViewport = Viewport.Pin())
+	{
+		return EditorViewport->GetPreviewProfileController();	
+	}
+	return nullptr;
+}
diff --git a/Engine/Source/Editor/UnrealEd/Public/SCommonEditorViewportToolbarBase.h b/Engine/Source/Editor/UnrealEd/Public/SCommonEditorViewportToolbarBase.h
--- a/Engine/Source/Editor/UnrealEd/Public/SCommonEditorViewportToolbarBase.h
+++ b/Engine/Source/Editor/UnrealEd/Public/SCommonEditorViewportToolbarBase.h
@@ -13,6 +13,7 @@
 #include "Templates/SharedPointer.h"
 #include "UObject/Object.h"
 #include "UObject/ObjectMacros.h"
+#include "ViewportToolbar/UnrealEdViewportToolbarContext.h"
 
 #include "SCommonEditorViewportToolbarBase.generated.h"
 
@@ -71,12 +72,14 @@ namespace CommonEditorViewportUtils
 }
 
 UCLASS()
-class UNREALED_API UCommonViewportToolbarBaseMenuContext : public UObject
+class UNREALED_API UCommonViewportToolbarBaseMenuContext : public UUnrealEdViewportToolbarContext
 {
 	GENERATED_BODY()
 
 public:
 	TWeakPtr<const class SCommonEditorViewportToolbarBase> ToolbarWidget;
+	
+	virtual TSharedPtr<IPreviewProfileController> GetPreviewProfileController() const override;
 };
 
 class UNREALED_API SPreviewSceneProfileSelector : public SCompoundWidget
@@ -145,6 +148,9 @@ public:
 	UE_DEPRECATED(5.6, "Use the version taking UToolMenu as argument: UE::UnrealEd::ConstructScreenPercentageMenu(InMenu)")
 	static UNREALED_API void ConstructScreenPercentageMenu(FMenuBuilder& MenuBuilder, class FEditorViewportClient* ViewportClient);
 
+	/** Gets the preview profile controller explicitly set to this toolbar */
+	UNREALED_API const TSharedPtr<IPreviewProfileController>& GetPreviewProfileController() const { return PreviewProfileController; }
+
 private:
 	/**
 	 * Returns the label for the "Camera" tool bar menu, which changes depending on the viewport type
@@ -152,6 +158,8 @@ private:
 	 * @return	Label to use for this menu label
 	 */
 	UNREALED_API FText GetCameraMenuLabel() const;
+	
+	UNREALED_API FSlateIcon GetCameraMenuIcon() const;
 
 	/**
 	 * Generates the toolbar options menu content 
@@ -228,10 +236,10 @@ protected:
 	UNREALED_API void CreateViewMenuExtensions(FMenuBuilder& MenuBuilder);
 
 	/** Extension allowing derived classes to add to the options menu.*/	
-	virtual void ExtendOptionsMenu(FMenuBuilder& OptionsMenuBuilder) const {}
+	UNREALED_API virtual void ExtendOptionsMenu(FMenuBuilder& OptionsMenuBuilder) const;
 
 	/** Extension allowing derived classes to add to left-aligned portion of the toolbar slots.*/
-	virtual void ExtendLeftAlignedToolbarSlots(TSharedPtr<SHorizontalBox> MainBoxPtr, TSharedPtr<SViewportToolBar> ParentToolBarPtr) const {}
+	UNREALED_API virtual void ExtendLeftAlignedToolbarSlots(TSharedPtr<SHorizontalBox> MainBoxPtr, TSharedPtr<SViewportToolBar> ParentToolBarPtr) const;
 
 	UNREALED_API virtual void FillShowFlagsMenu(class UToolMenu* InMenu) const;
 
@@ -260,5 +268,20 @@ protected:
 private:
 	/** The viewport that we are in */
 	TWeakPtr<class ICommonEditorViewportToolbarInfoProvider> InfoProviderPtr;
+	TSharedPtr<IPreviewProfileController> PreviewProfileController;
+	
+	TSharedPtr<class SEditorViewportViewMenu> BlankViewMenu;
+	
+	/// Automatic Legacy Upgrade Support
+	// True when generating the tool menu widget for the first time.
+	// Used to avoid calling menu generation functions right on construction, as by default Tool Menus could call those function earlier than before. 
+	bool bIsGeneratingToolMenuWidget = false;
+	// Allows functions to modulate behavior depending on whether the toolbar is in the new context.
+	mutable bool bIsBuildingToolMenu = false;
+	mutable bool bHasExtendedSettingsMenu = true;
+	mutable bool bHasExtendedLeftSide = true;
+	mutable bool bUsesDefaultViewMenu = false;
+	TSharedPtr<SWidget> MakeLegacyShowMenu() const;
+	bool ShouldCreateOptionsMenu() const;
 };
 
diff --git a/Engine/Source/Editor/UnrealEd/Public/SEditorViewport.h b/Engine/Source/Editor/UnrealEd/Public/SEditorViewport.h
--- a/Engine/Source/Editor/UnrealEd/Public/SEditorViewport.h
+++ b/Engine/Source/Editor/UnrealEd/Public/SEditorViewport.h
@@ -116,6 +116,12 @@ public:
 
 	UNREALED_API TSharedPtr<IPreviewProfileController> GetPreviewProfileController();
 
+	/**
+	 * A legacy patch built to allow SCommonEditorViewportToolbarBase to automatically switch to the new viewport style.
+	 * This function will go away in the future.
+	 */
+	void MarkLegacyToolbarChildAsAutomaticallyUpgradable(const TSharedRef<SWidget>& ExpectedChild);
+
 ///////////////////////////////////////////////////////////////////////////////
 // begin feature level control functions block
 ///////////////////////////////////////////////////////////////////////////////
@@ -155,10 +161,7 @@ protected:
 	virtual TSharedRef<FEditorViewportClient> MakeEditorViewportClient() = 0;
 
 	// Implement this to add a viewport toolbar above the viewport that does not cover the viewport.
-	virtual TSharedPtr<SWidget> BuildViewportToolbar()
-	{
-		return TSharedPtr<SWidget>(nullptr);
-	}
+	virtual TSharedPtr<SWidget> BuildViewportToolbar() { return TSharedPtr<SWidget>(nullptr); }
 
 	// Implement this to add a viewport toolbar to the inside top of the viewport
 	virtual TSharedPtr<SWidget> MakeViewportToolbar() { return TSharedPtr<SWidget>(nullptr); }
@@ -332,4 +335,9 @@ private:
 
 	/** Whether the viewport needs to update, even without input or realtime (e.g. inertial camera movement) */
 	bool bInvalidated;
+	
+	/** For automatic legacy toolbar upgrades. Registers the child widget to expect in order to allow an automatic upgrade. */
+	TWeakPtr<SWidget> AutoUpgradeWidgetChild = nullptr;
+	/** A flag indicating that the automatic toolbar upgrade has been approved */
+	bool bLegacyToolbarIsAutomaticallyUpgradable = false;
 };
diff --git a/Engine/Source/Editor/UnrealEd/Public/ViewportToolbar/UnrealEdViewportToolbarContext.h b/Engine/Source/Editor/UnrealEd/Public/ViewportToolbar/UnrealEdViewportToolbarContext.h
--- a/Engine/Source/Editor/UnrealEd/Public/ViewportToolbar/UnrealEdViewportToolbarContext.h
+++ b/Engine/Source/Editor/UnrealEd/Public/ViewportToolbar/UnrealEdViewportToolbarContext.h
@@ -11,6 +11,7 @@
 #include "UnrealEdViewportToolbarContext.generated.h"
 
 class FAssetEditorToolkit;
+class IPreviewProfileController;
 
 namespace UE::UnrealEd
 {
@@ -64,4 +65,6 @@ public:
 	 * Identifier for the Preview Settings Tab
 	 */
 	FTabId PreviewSettingsTabId;
+	
+	virtual TSharedPtr<IPreviewProfileController> GetPreviewProfileController() const;
 };
