Fixed properties that affect a notify event's name, via GetNotifyName(), not triggering a UI rebuild. Additionally, ensure that notify states behave the same way.

#rnx

diff --git a/Engine/Source/Editor/UnrealEd/Classes/Animation/EditorNotifyObject.h b/Engine/Source/Editor/UnrealEd/Classes/Animation/EditorNotifyObject.h
--- a/Engine/Source/Editor/UnrealEd/Classes/Animation/EditorNotifyObject.h
+++ b/Engine/Source/Editor/UnrealEd/Classes/Animation/EditorNotifyObject.h
@@ -30,4 +30,13 @@ class UEditorNotifyObject : public UEditorAnimBaseObj
 	/** The notify event to modify */
 	UPROPERTY(EditAnywhere, Category=Event)
 	FAnimNotifyEvent Event;
+	
+	/** 
+	 * Query and cache the generated notify name.
+	 * @return true if the name changed from previous cached value.
+	 */
+	bool TryToCacheNotifyName();
+	
+	/** Used to determine if we should rebuild the notify ui when its properties are changed. If it affects the name is the exception and a rebuild will be kicked. */
+	FName CachedNotifyName;
 };
diff --git a/Engine/Source/Editor/UnrealEd/Private/Animation/EditorNotifyObject.cpp b/Engine/Source/Editor/UnrealEd/Private/Animation/EditorNotifyObject.cpp
--- a/Engine/Source/Editor/UnrealEd/Private/Animation/EditorNotifyObject.cpp
+++ b/Engine/Source/Editor/UnrealEd/Private/Animation/EditorNotifyObject.cpp
@@ -46,16 +46,53 @@ void UEditorNotifyObject::InitialiseNotify(const FAnimNotifyEvent& InNotify)
 	if(AnimObject)
 	{
 		Event = InNotify;
+		
+		TryToCacheNotifyName();
 	}
 }
 
 bool UEditorNotifyObject::PropertyChangeRequiresRebuild(FPropertyChangedEvent& PropertyChangedEvent)
 {
-	// We dont need to rebuild the track UI when we change the properties of a notify
-	if(PropertyChangedEvent.GetMemberPropertyName() == GET_MEMBER_NAME_CHECKED(UEditorNotifyObject, Event) && PropertyChangedEvent.GetPropertyName() == GET_MEMBER_NAME_CHECKED(FAnimNotifyEvent, Notify))
+	// We don't need to rebuild the track UI when we change the properties of a notify
+	if(PropertyChangedEvent.GetMemberPropertyName() == GET_MEMBER_NAME_CHECKED(UEditorNotifyObject, Event) && (PropertyChangedEvent.GetPropertyName() == GET_MEMBER_NAME_CHECKED(FAnimNotifyEvent, Notify) || PropertyChangedEvent.GetPropertyName() == GET_MEMBER_NAME_CHECKED(FAnimNotifyEvent, NotifyStateClass)))
 	{
-		return false;
+		const bool bIsNotifyNameUnchanged = !TryToCacheNotifyName();
+		if (bIsNotifyNameUnchanged)
+		{
+			return false;
+		}
 	}
 
 	return Super::PropertyChangeRequiresRebuild(PropertyChangedEvent);
-}
\ No newline at end of file
+}
+
+bool UEditorNotifyObject::TryToCacheNotifyName()
+{
+	bool bChanged = false;
+
+	const auto TrySetName = [&bChanged, this](FName InName)
+	{
+		if (InName != CachedNotifyName)
+		{
+			CachedNotifyName = InName;
+			bChanged = true;
+		}
+	};
+	
+	if (Event.Notify)
+	{
+		FName NewName = FName(Event.Notify->GetNotifyName());
+		TrySetName(NewName);
+	}
+	else if (Event.NotifyStateClass)
+	{
+		FName NewName = FName(Event.NotifyStateClass->GetNotifyName());
+		TrySetName(NewName);
+	}
+	else
+	{
+		TrySetName(Event.NotifyName);
+	}
+
+	return bChanged;
+}
