Default the Top/Bottom ortho view rotation to be on so that screenspace Up corresponds to X/Forward.
Swap Perspective and Top view to align 2x2 axis with synced ortho mode
Reordered default 2x2 layout to align with other familiar CAD tooling
Make the Top/Bottom ortho rotate apply everywhere

Some renaming in the 2x2 layout to match to the viewtype rather than the layout location of the view.
Fixing comments for the UnrealTypes.h enum. Semantic meaning in FEditorViewportCommands followed.

diff --git a/Engine/Source/Editor/UnrealEd/Classes/Editor/UnrealEdTypes.h b/Engine/Source/Editor/UnrealEd/Classes/Editor/UnrealEdTypes.h
--- a/Engine/Source/Editor/UnrealEd/Classes/Editor/UnrealEdTypes.h
+++ b/Engine/Source/Editor/UnrealEd/Classes/Editor/UnrealEdTypes.h
@@ -112,19 +112,34 @@ enum ELevelViewportType : int
 {
 	/** Top */
 	LVT_OrthoXY = 0,
-	/** Front */
+	/** Right */
 	LVT_OrthoXZ = 1,
-	/** Left */
+	/** Back */
 	LVT_OrthoYZ = 2,
+	
 	LVT_Perspective = 3,
 	LVT_OrthoFreelook = 4,
 	/** Bottom */
 	LVT_OrthoNegativeXY = 5,
-	/** Back */
+	/** Left */
 	LVT_OrthoNegativeXZ = 6,
-	/** Right */
+	/** Front */
 	LVT_OrthoNegativeYZ = 7,
-	LVT_MAX,
+
+	/** Alias semantic meaning of viewport types for orthographic views */
+	/** Top */
+	LVT_OrthoTop = LVT_OrthoXY,
+	/** Right */
+	LVT_OrthoRight = LVT_OrthoXZ,
+	/** Front */
+	LVT_OrthoFront = LVT_OrthoNegativeYZ,
+	/** Back */
+	LVT_OrthoBack = LVT_OrthoYZ,
+	/** Bottom */
+	LVT_OrthoBottom = LVT_OrthoNegativeXY,
+	/** Left */
+	LVT_OrthoLeft = LVT_OrthoNegativeXZ,
+	LVT_MAX = 8,
 
 	LVT_None = 255,
 };
diff --git a/Engine/Source/Editor/UnrealEd/Private/EditorViewportClient.cpp b/Engine/Source/Editor/UnrealEd/Private/EditorViewportClient.cpp
--- a/Engine/Source/Editor/UnrealEd/Private/EditorViewportClient.cpp
+++ b/Engine/Source/Editor/UnrealEd/Private/EditorViewportClient.cpp
@@ -112,12 +112,6 @@ static TAutoConsoleVariable<float> CVarOrthoEditorDebugClipPlaneScale(
 	TEXT("Other light artefacts may appear when this value changes, this is unavoidable for now.\n"),
 	ECVF_RenderThreadSafe);
 
-static TAutoConsoleVariable<bool> CVarOrthoTopBottomViewRotated(
-	TEXT("Editor.TopBottomViewRotate"),
-	false,
-	TEXT("Rotates the top and bottom ortho views for +X to map to screenspace up and +Y to map to screenspace right"),
-	ECVF_Default);
-
 static bool GetDefaultLowDPIPreviewValue()
 {
 	static auto CVarEditorViewportHighDPIPtr = IConsoleManager::Get().FindConsoleVariable(TEXT("r.Editor.Viewport.HighDPI"));
@@ -393,12 +387,12 @@ void InitViewOptionsArray()
 {
 	ViewOptions.Empty();
 
-	ELevelViewportType Front = ELevelViewportType::LVT_OrthoXZ;
-	ELevelViewportType Back = ELevelViewportType::LVT_OrthoNegativeXZ;
+	ELevelViewportType Right = ELevelViewportType::LVT_OrthoXZ;
+	ELevelViewportType Left = ELevelViewportType::LVT_OrthoNegativeXZ;
 	ELevelViewportType Top = ELevelViewportType::LVT_OrthoXY;
 	ELevelViewportType Bottom = ELevelViewportType::LVT_OrthoNegativeXY;
-	ELevelViewportType Left = ELevelViewportType::LVT_OrthoYZ;
-	ELevelViewportType Right = ELevelViewportType::LVT_OrthoNegativeYZ;
+	ELevelViewportType Back = ELevelViewportType::LVT_OrthoYZ;
+	ELevelViewportType Front = ELevelViewportType::LVT_OrthoNegativeYZ;
 
 	ViewOptions.Add(Front);
 	ViewOptions.Add(Back);
@@ -1196,22 +1190,11 @@ FSceneView* FEditorViewportClient::CalcSceneView(FSceneViewFamily* ViewFamily, c
 
 			if (EffectiveViewportType == LVT_OrthoXY)
 			{
-				if (!CVarOrthoTopBottomViewRotated.GetValueOnAnyThread())
-				{
-					ViewInitOptions.ViewRotationMatrix = FMatrix(
-						FPlane(1, 0, 0, 0),
-						FPlane(0, -1, 0, 0),
-						FPlane(0, 0, -1, 0),
-						FPlane(0, 0, 0, 1));
-				}
-				else
-				{
-					ViewInitOptions.ViewRotationMatrix = FMatrix(
-						FPlane(0, 1, 0, 0),
-						FPlane(1, 0, 0, 0),
-						FPlane(0, 0, 1, 0),
-						FPlane(0, 0, 0, 1));
-				}
+				ViewInitOptions.ViewRotationMatrix = FMatrix(
+					FPlane(0, -1, 0, 0),
+					FPlane(-1, 0, 0, 0),
+					FPlane(0, 0, -1, 0),
+					FPlane(0, 0, 0, 1));
 			}
 			else if (EffectiveViewportType == LVT_OrthoXZ)
 			{
@@ -1231,22 +1214,11 @@ FSceneView* FEditorViewportClient::CalcSceneView(FSceneViewFamily* ViewFamily, c
 			}
 			else if (EffectiveViewportType == LVT_OrthoNegativeXY)
 			{
-				if (!CVarOrthoTopBottomViewRotated.GetValueOnAnyThread())
-				{
-					ViewInitOptions.ViewRotationMatrix = FMatrix(
-						FPlane(-1, 0, 0, 0),
-						FPlane(0, -1, 0, 0),
-						FPlane(0, 0, -1, 0),
-						FPlane(0, 0, 0, 1));
-				}
-				else
-				{
-					ViewInitOptions.ViewRotationMatrix = FMatrix(
+				ViewInitOptions.ViewRotationMatrix = FMatrix(
 						FPlane(0, -1, 0, 0),
-						FPlane(-1, 0, 0, 0),
+						FPlane(1, 0, 0, 0),
 						FPlane(0, 0, 1, 0),
 						FPlane(0, 0, 0, 1));
-				}
 			}
 			else if (EffectiveViewportType == LVT_OrthoNegativeXZ)
 			{
@@ -4483,8 +4455,7 @@ void FEditorViewportClient::Draw(FViewport* InViewport, FCanvas* Canvas)
 			{
 				const FRotator XYRot = []()
 				{
-					if (!CVarOrthoTopBottomViewRotated.GetValueOnAnyThread()) { return FRotator(-90.0f, -90.0f, 0.0f); }
-					else { return FRotator(-90.0f, -90.0f, 90.0f); }
+					return FRotator(-90.0f, -90.0f, -90.0f);
 				}();
 				DrawAxes(Viewport, Canvas, &XYRot, EAxisList::XY);
 				if (View)
@@ -4515,11 +4486,7 @@ void FEditorViewportClient::Draw(FViewport* InViewport, FCanvas* Canvas)
 			}
 		case LVT_OrthoNegativeXY:
 			{
-				const FRotator XYRot = []()
-				{
-					if (!CVarOrthoTopBottomViewRotated.GetValueOnAnyThread()) { return FRotator(90.0f, 90.0f, 0.0f); }
-					else { return FRotator(90.0f, 90.0f, 90.0f); }
-				}();
+				const FRotator XYRot = FRotator(90.0f, 90.0f, 90.0f);
 				DrawAxes(Viewport, Canvas, &XYRot, EAxisList::XY);
 				if (View)
 				{
@@ -4770,7 +4737,7 @@ FVector FEditorViewportClient::TranslateDelta( FKey InKey, float InDelta, bool I
 					switch( GetViewportType() )
 					{
 					case LVT_OrthoXY:
-						vec.Y *= -1.0f;
+						vec = FVector (-Y * UnitsPerPixel, -X * UnitsPerPixel, 0.0f);
 						break;
 					case LVT_OrthoXZ:
 						vec = FVector(X * UnitsPerPixel, 0.f, Y * UnitsPerPixel);
@@ -4779,10 +4746,10 @@ FVector FEditorViewportClient::TranslateDelta( FKey InKey, float InDelta, bool I
 						vec = FVector(0.f, X * UnitsPerPixel, Y * UnitsPerPixel);
 						break;
  					case LVT_OrthoNegativeXY:
- 						vec = FVector(-X * UnitsPerPixel, -Y * UnitsPerPixel, 0.0f);
+ 						vec = FVector(-Y * UnitsPerPixel, X * UnitsPerPixel, 0.0f);
  						break;
 					case LVT_OrthoNegativeXZ:
-						vec = FVector(-X * UnitsPerPixel, 0.f, Y * UnitsPerPixel);
+						vec = FVector(X * UnitsPerPixel, 0.f, Y * UnitsPerPixel);
 						break;
 					case LVT_OrthoNegativeYZ:
 						vec = FVector(0.f, -X * UnitsPerPixel, Y * UnitsPerPixel);
@@ -5530,7 +5497,7 @@ void FEditorViewportClient::ShowWidget(const bool bShow)
 
 void FEditorViewportClient::MoveViewportCamera(const FVector& InDrag, const FRotator& InRot, bool bDollyCamera)
 {
-	ELevelViewportType EffectiveViewportType = GetViewportType();
+	const ELevelViewportType EffectiveViewportType = GetViewportType();
 	switch( EffectiveViewportType )
 	{
 	case LVT_OrthoXY:
@@ -5540,21 +5507,6 @@ void FEditorViewportClient::MoveViewportCamera(const FVector& InDrag, const FRot
 	case LVT_OrthoNegativeXZ:
 	case LVT_OrthoNegativeYZ:
 		{
-			auto MapDragToOrthoViewLocation = [](const FVector& InDrag, ELevelViewportType ViewportType_)
-			{
-				if (CVarOrthoTopBottomViewRotated.GetValueOnAnyThread())
-				{
-					if (ViewportType_ == LVT_OrthoXY)
-					{
-						return FVector(-InDrag.Y, InDrag.X, InDrag.Z);
-					}
-					else if (ViewportType_ == LVT_OrthoNegativeXY)
-					{
-						return FVector(InDrag.Y, InDrag.X, InDrag.Z);
-					}
-				}
-				return InDrag;
-			};
 			const bool LeftMouseButtonDown = Viewport->KeyState(EKeys::LeftMouseButton);
 			const bool RightMouseButtonDown = Viewport->KeyState(EKeys::RightMouseButton);
 			const bool bIsUsingTrackpad = FSlateApplication::Get().IsUsingTrackpad();
@@ -5566,7 +5518,7 @@ void FEditorViewportClient::MoveViewportCamera(const FVector& InDrag, const FRot
 			}
 			else
 			{
-				const FVector ViewLocationAdjustment = MapDragToOrthoViewLocation(InDrag, EffectiveViewportType);
+				const FVector ViewLocationAdjustment = InDrag;
 				SetViewLocation( GetViewLocation() + ViewLocationAdjustment );
 			}
 
diff --git a/Engine/Source/Editor/UnrealEd/Private/EditorViewportLayout2x2.cpp b/Engine/Source/Editor/UnrealEd/Private/EditorViewportLayout2x2.cpp
--- a/Engine/Source/Editor/UnrealEd/Private/EditorViewportLayout2x2.cpp
+++ b/Engine/Source/Editor/UnrealEd/Private/EditorViewportLayout2x2.cpp
@@ -18,9 +18,9 @@ namespace ViewportLayout2x2Defs
 
 TSharedRef<SWidget> FEditorViewportLayout2x2::MakeViewportLayout(TSharedPtr<FAssetEditorViewportLayout> InParentLayout, const FString& LayoutString)
 {
-	FString TopLeftKey, BottomLeftKey, TopRightKey, BottomRightKey;
+	FString FrontConfigKey, PerspectiveConfigKey, RightConfigKey, TopConfigKey;
 
-	FString TopLeftType = TEXT("Default"), BottomLeftType = TEXT("Default"), TopRightType = TEXT("Default"), BottomRightType = TEXT("Default");
+	FString RightTypeName = TEXT("Default"), PerspectiveTypeName = TEXT("Default"), FrontTypeName = TEXT("Default"), TopTypeName = TEXT("Default");
 
 	TArray<FVector2D> SplitterPercentages;
 	
@@ -30,15 +30,15 @@ TSharedRef<SWidget> FEditorViewportLayout2x2::MakeViewportLayout(TSharedPtr<FAss
 		// The Layout String only holds the unique ID of the Additional Layout Configs to use
 		const FString& IniSection = FLayoutSaveRestore::GetAdditionalLayoutConfigIni();
 
-		TopLeftKey = SpecificLayoutString + TEXT(".Viewport0");
-		BottomLeftKey = SpecificLayoutString + TEXT(".Viewport1");
-		TopRightKey = SpecificLayoutString + TEXT(".Viewport2");
-		BottomRightKey = SpecificLayoutString + TEXT(".Viewport3");
+		FrontConfigKey = SpecificLayoutString + TEXT(".Viewport0");
+		PerspectiveConfigKey = SpecificLayoutString + TEXT(".Viewport1");
+		RightConfigKey = SpecificLayoutString + TEXT(".Viewport2");
+		TopConfigKey = SpecificLayoutString + TEXT(".Viewport3");
 
-		GConfig->GetString(*IniSection, *(TopLeftKey + TEXT(".TypeWithinLayout")), TopLeftType, GEditorPerProjectIni);
-		GConfig->GetString(*IniSection, *(TopRightKey + TEXT(".TypeWithinLayout")), TopRightType, GEditorPerProjectIni);
-		GConfig->GetString(*IniSection, *(BottomLeftKey + TEXT(".TypeWithinLayout")), BottomLeftType, GEditorPerProjectIni);
-		GConfig->GetString(*IniSection, *(BottomRightKey + TEXT(".TypeWithinLayout")), BottomRightType, GEditorPerProjectIni);
+		GConfig->GetString(*IniSection, *(FrontConfigKey + TEXT(".TypeWithinLayout")), FrontTypeName, GEditorPerProjectIni);
+		GConfig->GetString(*IniSection, *(PerspectiveConfigKey + TEXT(".TypeWithinLayout")), PerspectiveTypeName, GEditorPerProjectIni);
+		GConfig->GetString(*IniSection, *(RightConfigKey + TEXT(".TypeWithinLayout")), RightTypeName, GEditorPerProjectIni);
+		GConfig->GetString(*IniSection, *(TopConfigKey + TEXT(".TypeWithinLayout")), TopTypeName, GEditorPerProjectIni);
 
 		for (int32 i = 0; i < 4; ++i)
 		{
@@ -56,49 +56,49 @@ TSharedRef<SWidget> FEditorViewportLayout2x2::MakeViewportLayout(TSharedPtr<FAss
 	FAssetEditorViewportConstructionArgs Args;
  	Args.ParentLayout = InParentLayout;
 
-	// Left viewport
+	// Front viewport
 	Args.bRealtime = false;
-	Args.ConfigKey = *TopLeftKey;
-	Args.ViewportType = LVT_OrthoYZ;
-	TSharedRef< SWidget > ViewportTL = InParentLayout->FactoryViewport(*TopLeftType, Args);
+	Args.ConfigKey = *FrontConfigKey;
+	Args.ViewportType = LVT_OrthoFront;
+	TSharedRef< SWidget > FrontViewport = InParentLayout->FactoryViewport(*FrontTypeName, Args);
 
 	// Persp viewport
 	Args.bRealtime = !FPlatformMisc::IsRemoteSession();
-	Args.ConfigKey = *BottomLeftKey;
+	Args.ConfigKey = *PerspectiveConfigKey;
 	Args.ViewportType = LVT_Perspective;
-	TSharedRef< SWidget > ViewportBL = InParentLayout->FactoryViewport(*BottomLeftType, Args);
-	PerspectiveViewportConfigKey = *BottomLeftKey;
+	TSharedRef< SWidget > PerspectiveViewport = InParentLayout->FactoryViewport(*PerspectiveTypeName, Args);
+	PerspectiveViewportConfigKey = *PerspectiveConfigKey;
 
-	// Front viewport
+	// Right viewport
 	Args.bRealtime = false;
-	Args.ConfigKey = *TopRightKey;
-	Args.ViewportType = LVT_OrthoXZ;
-	TSharedRef< SWidget > ViewportTR = InParentLayout->FactoryViewport(*TopRightType, Args);
+	Args.ConfigKey = *RightConfigKey;
+	Args.ViewportType = LVT_OrthoRight;
+	TSharedRef< SWidget > RightViewport = InParentLayout->FactoryViewport(*RightTypeName, Args);
 
 	// Top Viewport
 	Args.bRealtime = false;
-	Args.ConfigKey = *BottomRightKey;
-	Args.ViewportType = LVT_OrthoXY;
-	TSharedRef< SWidget > ViewportBR = InParentLayout->FactoryViewport(*BottomRightType, Args);
+	Args.ConfigKey = *TopConfigKey;
+	Args.ViewportType = LVT_OrthoTop;
+	TSharedRef< SWidget > TopViewport = InParentLayout->FactoryViewport(*TopTypeName, Args);
 
 	// Set up the splitter
 	SplitterWidget = 
 	SNew( SSplitter2x2 )
 	.TopLeft()
 	[
-		ViewportTL
+		PerspectiveViewport
 	]
 	.BottomLeft()
 	[
-		ViewportBL
+		RightViewport
 	]
 	.TopRight()
 	[
-		ViewportTR
+		TopViewport
 	]
 	.BottomRight()
 	[
-		ViewportBR
+		FrontViewport
 	];
 	
 	if (SplitterPercentages.Num() > 0)
