Fix for a case where class Source was renamed to class Target but a package was not resaved so it still contains references to class Source. So when bTargetIterativeEnabled is set to true for class Target, I need to find class Source and set the flag for it as well.

FClassDigestData : new flag bIsRedirected to easily find class being redirected.
FPrecacheClassDigest::GetRecursive : set bIsRedirect to true when a redirected class is found.
ConstructTargetIterativeClassAllowList : Go through all the class digests and if it is the source of a redirection, set the flag bTargetIterativeEnabled to the same value as the target class.

#rnx
#incrementalcook

diff --git a/Engine/Source/Editor/UnrealEd/Private/EditorDomain/EditorDomainUtils.cpp b/Engine/Source/Editor/UnrealEd/Private/EditorDomain/EditorDomainUtils.cpp
--- a/Engine/Source/Editor/UnrealEd/Private/EditorDomain/EditorDomainUtils.cpp
+++ b/Engine/Source/Editor/UnrealEd/Private/EditorDomain/EditorDomainUtils.cpp
@@ -578,6 +578,8 @@ FClassDigestData* FPrecacheClassDigest::GetRecursive(FTopLevelAssetPath ClassPat
 			// The map has possibly been modified so we need to recalculate the address of ClassName's DigestData
 			DigestData = &ClassDigests.FindChecked(ClassPath);
 			*DigestData = *RedirectedDigest;
+			DigestData->bIsRedirected = true;
+
 			return DigestData;
 		}
 	}
@@ -1283,6 +1285,33 @@ void ConstructTargetIterativeClassAllowList()
 				DigestData->bTargetIterativeEnabled = true;
 			}
 		}
+
+		// The previous loop only set the flag bTargetIterativeEnabled to target classes of redirection, not to the source class.
+		// This is an issue in the case where class Source was renamed to class Target but a package was not resaved so it still contains references to class Source.
+		// So when bTargetIterativeEnabled is set to true for class Target, I need to find class Source and set the flag for it as well.
+		for (TPair<FTopLevelAssetPath, FClassDigestData>& Item : ClassDigests.Map)
+		{
+			FClassDigestData& OldClassDigest = Item.Value;
+			if (!OldClassDigest.bIsRedirected)
+			{
+				continue;
+			}
+			
+			FCoreRedirectObjectName OldName(Item.Key);
+			FCoreRedirectObjectName NewName = FCoreRedirects::GetRedirectedName(ECoreRedirectFlags::Type_Class | ECoreRedirectFlags::Type_Struct, OldName);
+				
+			if (NewName == OldName)
+			{
+				continue;
+			}
+
+			FTopLevelAssetPath NewNameClassPath(NewName.PackageName, NewName.ObjectName);
+			FClassDigestData* NewClassDigest = ClassDigests.Map.Find(NewNameClassPath);
+			if (NewClassDigest)
+			{
+				OldClassDigest.bTargetIterativeEnabled = NewClassDigest->bTargetIterativeEnabled;
+			}
+		}
 	}
 }
 
diff --git a/Engine/Source/Editor/UnrealEd/Private/EditorDomain/EditorDomainUtils.h b/Engine/Source/Editor/UnrealEd/Private/EditorDomain/EditorDomainUtils.h
--- a/Engine/Source/Editor/UnrealEd/Private/EditorDomain/EditorDomainUtils.h
+++ b/Engine/Source/Editor/UnrealEd/Private/EditorDomain/EditorDomainUtils.h
@@ -54,6 +54,9 @@ public:
 
 	bool bConstructed : 1 = false;
 	bool bConstructionComplete : 1 = false;
+
+	/** Whether or not this class is the source of a redirection.*/
+	bool bIsRedirected : 1 = false;
 };
 
 /** Threadsafe cache of ClassName -> Digest data for calculating EditorDomain Digests */
