[Curve Editor]
- Fixed editor locking issue when the user resets any property to default on the property chain
- Refactored the other fix code, now utilizing PropertyHandle recursive callback

diff --git a/Engine/Source/Editor/DetailCustomizations/Private/CurveColorCustomization.cpp b/Engine/Source/Editor/DetailCustomizations/Private/CurveColorCustomization.cpp
--- a/Engine/Source/Editor/DetailCustomizations/Private/CurveColorCustomization.cpp
+++ b/Engine/Source/Editor/DetailCustomizations/Private/CurveColorCustomization.cpp
@@ -143,18 +143,29 @@ void FCurveColorCustomization::CustomizeHeader(TSharedRef<IPropertyHandle> InStr
 			];
 
 		check(CurveWidget.IsValid());
+		FCurveOwnerInterface* CurveOwner = nullptr;
+		bool bCanEdit = true;
 		if (RuntimeCurve && RuntimeCurve->ExternalCurve)
 		{
-			CurveWidget->SetCurveOwner(RuntimeCurve->ExternalCurve, false);
+			CurveOwner = RuntimeCurve->ExternalCurve;
+			bCanEdit = false;
 		}
 		else
 		{
-			CurveWidget->SetCurveOwner(this, InStructPropertyHandle->IsEditable());
-			CurveWidget->SetPropertyUtils(StructCustomizationUtils.GetPropertyUtilities());
+			CurveOwner = this;
+		}
+
+		CurveWidget->SetCurveOwner(CurveOwner, bCanEdit);
+		CurveWidget->SetPropertyUtils(StructCustomizationUtils.GetPropertyUtilities());
+
+		// If any property on the chain gets reset while the Curve Owner is set  
+		TSharedPtr<IPropertyHandle> RootHandle = StructPropertyHandle;
+		while (RootHandle->GetParentHandle().IsValid())
+		{
+			RootHandle = RootHandle->GetParentHandle();
 		}
 
-		const FSimpleDelegate OnPropertyReset = FSimpleDelegate::CreateSP(this, &FCurveColorCustomization::HandleCurveColorReset);
-		StructPropertyHandle->SetOnPropertyResetToDefault(OnPropertyReset);
+		CurveWidget->RegisterToPropertyChangedEvent(RootHandle);
 	}
 	else
 	{
@@ -524,9 +535,4 @@ void FCurveColorCustomization::DestroyPopOutWindow()
 	}
 }
 
-void FCurveColorCustomization::HandleCurveColorReset()
-{
-	CurveWidget->SetCurveOwner(this, StructPropertyHandle->IsEditable());
-}
-
 #undef LOCTEXT_NAMESPACE
diff --git a/Engine/Source/Editor/DetailCustomizations/Private/CurveColorCustomization.h b/Engine/Source/Editor/DetailCustomizations/Private/CurveColorCustomization.h
--- a/Engine/Source/Editor/DetailCustomizations/Private/CurveColorCustomization.h
+++ b/Engine/Source/Editor/DetailCustomizations/Private/CurveColorCustomization.h
@@ -116,11 +116,6 @@ private:
 	 */
 	void DestroyPopOutWindow();
 
-	/**
-	 * Reset the curve view models of the curve editor to be this, after resetting the property
-	 */
-	void HandleCurveColorReset();
-
 private:
 	/** Cached RuntimeCurveLinearColor struct handle */
 	TSharedPtr<IPropertyHandle> StructPropertyHandle;
diff --git a/Engine/Source/Editor/DetailCustomizations/Private/CurveStructCustomization.cpp b/Engine/Source/Editor/DetailCustomizations/Private/CurveStructCustomization.cpp
--- a/Engine/Source/Editor/DetailCustomizations/Private/CurveStructCustomization.cpp
+++ b/Engine/Source/Editor/DetailCustomizations/Private/CurveStructCustomization.cpp
@@ -88,7 +88,6 @@ FCurveStructCustomization::FCurveStructCustomization()
 void FCurveStructCustomization::CustomizeHeader( TSharedRef<IPropertyHandle> InStructPropertyHandle, FDetailWidgetRow& HeaderRow, IPropertyTypeCustomizationUtils& StructCustomizationUtils )
 {
 	this->StructPropertyHandle = InStructPropertyHandle;
-
 	TArray<UObject*> OuterObjects;
 	StructPropertyHandle->GetOuterObjects(OuterObjects);
 
@@ -145,18 +144,30 @@ void FCurveStructCustomization::CustomizeHeader( TSharedRef<IPropertyHandle> InS
 			];
 
 		check(CurveWidget.IsValid());
+
+		FCurveOwnerInterface* CurveOwner = nullptr;
+		bool bCanEdit = true;
 		if (RuntimeCurve && RuntimeCurve->ExternalCurve)
 		{
-			CurveWidget->SetCurveOwner(RuntimeCurve->ExternalCurve, false);
+			CurveOwner = RuntimeCurve->ExternalCurve;
+			bCanEdit = false;
 		}
 		else
 		{
-			CurveWidget->SetCurveOwner(this, InStructPropertyHandle->IsEditable());
-			CurveWidget->SetPropertyUtils(StructCustomizationUtils.GetPropertyUtilities());
+			CurveOwner = this;
 		}
 
-		const FSimpleDelegate OnPropertyReset = FSimpleDelegate::CreateSP(this, &FCurveStructCustomization::HandleCurveStructReset);
-		StructPropertyHandle->SetOnPropertyResetToDefault(OnPropertyReset);
+		CurveWidget->SetCurveOwner(CurveOwner, bCanEdit);
+		CurveWidget->SetPropertyUtils(StructCustomizationUtils.GetPropertyUtilities());
+
+		// If any property on the chain gets reset while the Curve Owner is set  
+		TSharedPtr<IPropertyHandle> RootHandle = StructPropertyHandle;
+		while (RootHandle->GetParentHandle().IsValid())
+		{
+			RootHandle = RootHandle->GetParentHandle();
+		}
+
+		CurveWidget->RegisterToPropertyChangedEvent(RootHandle);
 	}
 	else
 	{
@@ -532,10 +543,4 @@ void FCurveStructCustomization::DestroyPopOutWindow()
 		CurveEditorWindow.Reset();
 	}
 }
-
-void FCurveStructCustomization::HandleCurveStructReset()
-{
-	CurveWidget->SetCurveOwner(this, StructPropertyHandle->IsEditable());
-}
-
 #undef LOCTEXT_NAMESPACE
diff --git a/Engine/Source/Editor/DetailCustomizations/Private/CurveStructCustomization.h b/Engine/Source/Editor/DetailCustomizations/Private/CurveStructCustomization.h
--- a/Engine/Source/Editor/DetailCustomizations/Private/CurveStructCustomization.h
+++ b/Engine/Source/Editor/DetailCustomizations/Private/CurveStructCustomization.h
@@ -117,11 +117,6 @@ private:
 	 */
 	void DestroyPopOutWindow();
 
-	/**
-	 * Reset the curve view models of the curve editor to be this, after resetting the property
-	 */
-	void HandleCurveStructReset();
-
 private:
 	/** Cached RuntimeFloatCurve struct handle */
 	TSharedPtr<IPropertyHandle> StructPropertyHandle;
diff --git a/Engine/Source/Editor/DetailCustomizations/Private/CurveVectorCustomization.cpp b/Engine/Source/Editor/DetailCustomizations/Private/CurveVectorCustomization.cpp
--- a/Engine/Source/Editor/DetailCustomizations/Private/CurveVectorCustomization.cpp
+++ b/Engine/Source/Editor/DetailCustomizations/Private/CurveVectorCustomization.cpp
@@ -143,18 +143,29 @@ void FCurveVectorCustomization::CustomizeHeader(TSharedRef<IPropertyHandle> InSt
 			];
 
 		check(CurveWidget.IsValid());
+		FCurveOwnerInterface* CurveOwner = nullptr;
+		bool bCanEdit = true;
 		if (RuntimeCurve && RuntimeCurve->ExternalCurve)
 		{
-			CurveWidget->SetCurveOwner(RuntimeCurve->ExternalCurve, false);
+			CurveOwner = RuntimeCurve->ExternalCurve;
+			bCanEdit = false;
 		}
-		else if (RuntimeCurve)
+		else
 		{
-			CurveWidget->SetCurveOwner(this, InStructPropertyHandle->IsEditable());
-			CurveWidget->SetPropertyUtils(StructCustomizationUtils.GetPropertyUtilities());
+			CurveOwner = this;
 		}
 
-		const FSimpleDelegate OnPropertyReset = FSimpleDelegate::CreateSP(this, &FCurveVectorCustomization::HandleCurveVectorReset);
-		StructPropertyHandle->SetOnPropertyResetToDefault(OnPropertyReset);
+		CurveWidget->SetCurveOwner(CurveOwner, bCanEdit);
+		CurveWidget->SetPropertyUtils(StructCustomizationUtils.GetPropertyUtilities());
+
+		// If any property on the chain gets reset while the Curve Owner is set  
+		TSharedPtr<IPropertyHandle> RootHandle = StructPropertyHandle;
+		while (RootHandle->GetParentHandle().IsValid())
+		{
+			RootHandle = RootHandle->GetParentHandle();
+		}
+
+		CurveWidget->RegisterToPropertyChangedEvent(RootHandle);
 	}
 	else
 	{
@@ -508,9 +519,4 @@ void FCurveVectorCustomization::DestroyPopOutWindow()
 	}
 }
 
-void FCurveVectorCustomization::HandleCurveVectorReset()
-{
-	CurveWidget->SetCurveOwner(this, StructPropertyHandle->IsEditable());
-}
-
 #undef LOCTEXT_NAMESPACE
diff --git a/Engine/Source/Editor/DetailCustomizations/Private/CurveVectorCustomization.h b/Engine/Source/Editor/DetailCustomizations/Private/CurveVectorCustomization.h
--- a/Engine/Source/Editor/DetailCustomizations/Private/CurveVectorCustomization.h
+++ b/Engine/Source/Editor/DetailCustomizations/Private/CurveVectorCustomization.h
@@ -109,11 +109,6 @@ private:
 	 */
 	void DestroyPopOutWindow();
 
-	/**
-	 * Reset the curve view models of the curve editor to be this, after resetting the property
-	 */
-	void HandleCurveVectorReset();
-
 private:
 	/** Cached RuntimeCurveLinearColor struct handle */
 	TSharedPtr<IPropertyHandle> StructPropertyHandle;
diff --git a/Engine/Source/Editor/UnrealEd/Private/SCurveEditor.cpp b/Engine/Source/Editor/UnrealEd/Private/SCurveEditor.cpp
--- a/Engine/Source/Editor/UnrealEd/Private/SCurveEditor.cpp
+++ b/Engine/Source/Editor/UnrealEd/Private/SCurveEditor.cpp
@@ -37,6 +37,7 @@
 #include "Widgets/Notifications/SNotificationList.h"
 #include "Framework/Notifications/NotificationManager.h"
 #include "IPropertyUtilities.h"
+#include "PropertyHandle.h"
 
 #define LOCTEXT_NAMESPACE "SCurveEditor"
  
@@ -1313,6 +1314,17 @@ void SCurveEditor::SetCurveOwner(FCurveOwnerInterface* InCurveOwner, bool bCanEd
 	CurveSelectionWidget.Pin()->SetContent(CreateCurveSelectionWidget());
 }
 
+void SCurveEditor::RegisterToPropertyChangedEvent(const TSharedPtr<IPropertyHandle>& InRootPropertyHandle)
+{
+	RootPropertyHandleWeak = InRootPropertyHandle.ToWeakPtr();
+	const auto OnChildPropertyChanged = TDelegate<void(const FPropertyChangedEvent&)>::CreateSP(this, &SCurveEditor::OnRootPropertyChanged);
+
+	if (const TSharedPtr<IPropertyHandle> RootPropertyHandle = RootPropertyHandleWeak.Pin())
+	{
+		RootPropertyHandle->SetOnChildPropertyValueChangedWithData(OnChildPropertyChanged);
+	}
+}
+
 void SCurveEditor::SetZoomToFit(bool bNewZoomToFitVertical, bool bNewZoomToFitHorizontal)
 {
 	bZoomToFitVertical = bNewZoomToFitVertical;
@@ -3812,44 +3824,67 @@ void SCurveEditor::RedoAction()
 	GEditor->RedoTransaction();
 }
 
+void SCurveEditor::RefreshDetailsView()
+{
+	check(!bIsPendingRebuilt);
+
+	bIsPendingRebuilt = true;
+
+	//To clean up dangling references
+	EmptyAllSelection();
+	EmptyAllCurveViewModels();
+	RegisterActiveTimer(0.f, FWidgetActiveTimerDelegate::CreateSPLambda(
+		this,
+		[this](double, float)
+		{
+			if(const TSharedPtr<IPropertyUtilities> PropertyUtilities = PropertyUtilitiesWeak.Pin())
+			{
+				//Rebuild the widget so it will reference to the new correct address of Curves, and show things correctly
+				PropertyUtilities->ForceRefresh();
+			}
+
+			return EActiveTimerReturnType::Stop;
+		}
+	));
+
+	ValidateSelection();
+}
+
 void SCurveEditor::OnObjectPropertyChanged(UObject* Object, FPropertyChangedEvent& PropertyChangedEvent)
 {
 	if (CurveOwner && CurveOwner->GetOwners().Contains(Object))
 	{
-		// CurveEditor will hold dangling references to CurveData and access them in some cases:
+		if (!bIsPendingRebuilt)
+		{
+			if (GIsTransacting)
+			{
+				RefreshDetailsView();
+			}
+		}
+	}
+}
+
+void SCurveEditor::OnRootPropertyChanged(const FPropertyChangedEvent& PropertyChangedEvent)
+{
+	if (!bIsPendingRebuilt)
+	{
+		// CurveEditor will hold dangling references to CurveData and access them in some cases :
 		// - Curve Data is inline allocated, or wrapped by InstancedStruct, in a container. And we do Array Remove/Clear op, or Add that causes resize
 		// - Curve Data is wrapped by InstancedStruct in a Container, and we do undo/redo(will trigger emptying and refilling the container, causing address changed)
 		// - Curve Data is wrapped by InstancedStruct and InstancedStruct is being reset/replaced, causing actual memory address changed
 		// todo: should only rebuild when it's the InstanceStruct with ValueSet. But currently there is no way to detect if there is an InstanceStruct in the PropertyNode Hierarchy.
-		if (!bIsPendingRebuilt
-			&& (GIsTransacting
-			|| PropertyChangedEvent.ChangeType == EPropertyChangeType::ArrayAdd
-			|| PropertyChangedEvent.ChangeType == EPropertyChangeType::ArrayRemove
-			|| PropertyChangedEvent.ChangeType == EPropertyChangeType::ArrayClear
-			|| PropertyChangedEvent.ChangeType == EPropertyChangeType::ValueSet
-			))
+		if (PropertyChangedEvent.ChangeType
+			& (EPropertyChangeType::ArrayAdd
+			| EPropertyChangeType::ArrayRemove
+			| EPropertyChangeType::ArrayClear
+			| EPropertyChangeType::ValueSet
+			| EPropertyChangeType::ResetToDefault))
 		{
-			bIsPendingRebuilt = true;
-
-			//To clean up dangling references
-			EmptyAllSelection();
-			EmptyAllCurveViewModels();
-			RegisterActiveTimer(0.f, FWidgetActiveTimerDelegate::CreateSPLambda(
-				this,
-				[this](double, float)
-				{
-					if(const TSharedPtr<IPropertyUtilities> PropertyUtilities = PropertyUtilitiesWeak.Pin())
-					{
-						//Rebuild the widget so it will reference to the new correct address of Curves, and show things correctly
-						PropertyUtilities->ForceRefresh();
-					}
-					return EActiveTimerReturnType::Stop;
-				}
-				));
+			RefreshDetailsView();
 		}
-
-		ValidateSelection();
 	}
+
+	ValidateSelection();
 }
 
 void SCurveEditor::HandlePackageReloaded(const EPackageReloadPhase InPackageReloadPhase, FPackageReloadedEvent* InPackageReloadedEvent)
diff --git a/Engine/Source/Editor/UnrealEd/Public/SCurveEditor.h b/Engine/Source/Editor/UnrealEd/Public/SCurveEditor.h
--- a/Engine/Source/Editor/UnrealEd/Public/SCurveEditor.h
+++ b/Engine/Source/Editor/UnrealEd/Public/SCurveEditor.h
@@ -22,6 +22,7 @@
 class FPaintArgs;
 class FSlateWindowElementList;
 class FUICommandList;
+class IPropertyHandle;
 class IPropertyUtilities;
 class IMenu;
 class SBox;
@@ -258,6 +259,8 @@ public:
 		PropertyUtilitiesWeak = InPropertyUtils;
 	}
 
+	UNREALED_API void RegisterToPropertyChangedEvent(const TSharedPtr<IPropertyHandle>& InRootPropertyHandle);
+
 	/** Set new zoom to fit **/
 	UNREALED_API void SetZoomToFit(bool bNewZoomToFitVertical, bool bNewZoomToFitHorizontal);
 
@@ -613,6 +616,9 @@ private:
 	/** Perform redo*/
 	void	RedoAction();
 
+	/** Refresh the Details View and respawn this widget when Curve Data goes stale */
+	void RefreshDetailsView();
+
 	//~ Begin FEditorUndoClient Interface
 	UNREALED_API virtual void PostUndo(bool bSuccess) override;
 	virtual void PostRedo(bool bSuccess) override { PostUndo(bSuccess); }
@@ -647,6 +653,7 @@ private:
 	TSharedPtr<FCurveViewModel> GetViewModelForCurve(FRealCurve* InCurve);
 
 	void OnObjectPropertyChanged(UObject* Object, FPropertyChangedEvent& PropertyChangedEvent);
+	void OnRootPropertyChanged(const FPropertyChangedEvent& PropertyChangedEvent);
 
 	void HandlePackageReloaded(const EPackageReloadPhase InPackageReloadPhase, FPackageReloadedEvent* InPackageReloadedEvent);
 
@@ -873,4 +880,5 @@ protected:
 private:
 	bool bIsPendingRebuilt;
 	TWeakPtr<IPropertyUtilities> PropertyUtilitiesWeak;
+	TWeakPtr<IPropertyHandle> RootPropertyHandleWeak;
 };
