[GitHub PR #12682] Fixed UEditorActorSubsystem linker error introduced in UE 5.5

diff --git a/Engine/Source/Editor/UnrealEd/Public/Subsystems/EditorActorSubsystem.h b/Engine/Source/Editor/UnrealEd/Public/Subsystems/EditorActorSubsystem.h
--- a/Engine/Source/Editor/UnrealEd/Public/Subsystems/EditorActorSubsystem.h
+++ b/Engine/Source/Editor/UnrealEd/Public/Subsystems/EditorActorSubsystem.h
@@ -52,7 +52,7 @@ class UEditorActorSubsystem : public UEditorSubsystem
 
 public:
 
-	virtual void Initialize(FSubsystemCollectionBase& Collection) override;
+	UNREALED_API virtual void Initialize(FSubsystemCollectionBase& Collection) override;
 	UNREALED_API virtual void Deinitialize() override;
 
 	UPROPERTY(BlueprintAssignable, Category = "Editor Scripting | Level Utility")
