Adding multiline Details field to Comment block node.  This info will show up in tooltip as well.

diff --git a/Engine/Source/Editor/UnrealEd/Private/EdGraphNode_Comment.cpp b/Engine/Source/Editor/UnrealEd/Private/EdGraphNode_Comment.cpp
--- a/Engine/Source/Editor/UnrealEd/Private/EdGraphNode_Comment.cpp
+++ b/Engine/Source/Editor/UnrealEd/Private/EdGraphNode_Comment.cpp
@@ -99,7 +99,12 @@ FText UEdGraphNode_Comment::GetTooltipText() const
 {
 	if (CachedTooltip.IsOutOfDate(this))
 	{
-		CachedTooltip.SetCachedText(FText::Format(NSLOCTEXT("K2Node", "CommentBlock_Tooltip", "Comment:\n{0}"), FText::FromString(NodeComment)), this);
+		FText DetailsTooltipText = FText::GetEmpty();
+		if (!NodeDetails.IsEmpty()) 
+		{
+			DetailsTooltipText = FText::Format(FText::FromString("\n\nDetails:\n{0}"), NodeDetails);
+		}
+		CachedTooltip.SetCachedText(FText::Format(NSLOCTEXT("K2Node", "CommentBlock_Tooltip", "Comment:\n{0}{1}"), FText::FromString(NodeComment), DetailsTooltipText), this);
 	}
 	return CachedTooltip;
 }
diff --git a/Engine/Source/Editor/UnrealEd/Public/EdGraphNode_Comment.h b/Engine/Source/Editor/UnrealEd/Public/EdGraphNode_Comment.h
--- a/Engine/Source/Editor/UnrealEd/Public/EdGraphNode_Comment.h
+++ b/Engine/Source/Editor/UnrealEd/Public/EdGraphNode_Comment.h
@@ -70,6 +70,10 @@ public:
 	UPROPERTY(EditAnywhere, Category=Comment)
 	TEnumAsByte<ECommentBoxMode::Type> MoveMode;
 
+	/** Details field if more info is needed to be communicated (will show up in tooltip) */
+	UPROPERTY(EditAnywhere, Category=Comment, meta=(DisplayName="Details", MultiLine=true))
+	FText NodeDetails;
+
 	/** comment Depth */
 	UPROPERTY()
 	int32 CommentDepth;
