Display the name of every package affected by the world partition builder in order to easily retrieve which file was affected by the commandlet in an automated build log

#rnx

diff --git a/Engine/Source/Editor/UnrealEd/Private/WorldPartition/WorldPartitionBuilder.cpp b/Engine/Source/Editor/UnrealEd/Private/WorldPartition/WorldPartitionBuilder.cpp
--- a/Engine/Source/Editor/UnrealEd/Private/WorldPartition/WorldPartitionBuilder.cpp
+++ b/Engine/Source/Editor/UnrealEd/Private/WorldPartition/WorldPartitionBuilder.cpp
@@ -20,6 +20,7 @@
 #include "LevelInstance/LevelInstanceSubsystem.h"
 #include "Math/IntVector.h"
 #include "UObject/SavePackage.h"
+#include "Algo/ForEach.h"
 #include "Algo/Transform.h"
 
 DEFINE_LOG_CATEGORY_STATIC(LogWorldPartitionBuilder, All, All);
@@ -351,13 +352,19 @@ bool UWorldPartitionBuilder::SavePackages(const TArray<UPackage*>& Packages, FPa
 	UE_LOG(LogWorldPartitionBuilder, Display, TEXT("Saving %d packages..."), Packages.Num());
 
 	const TArray<FString> PackageFilenames = SourceControlHelpers::PackageFilenames(Packages);
-	
+
 	TArray<FString> PackagesToCheckout;
 	TArray<FString> PackagesToAdd;
 	TArray<FString> PackageNames;
 
 	PackageNames.Reserve(Packages.Num());
 	Algo::Transform(Packages, PackageNames, [](UPackage* InPackage) { return InPackage->GetName(); });
+
+	// Display all affected files : that can be useful when analyzing what was saved (or at least what we tried to) by this commandlet : 
+	FStringBuilderBase PackageNamesString;
+	Algo::ForEach(PackageNames, [&PackageNamesString](const FString& InPackageName) { PackageNamesString.Appendf(TEXT("\t%s\n"), *InPackageName); });
+	UE_LOG(LogWorldPartitionBuilder, Display, TEXT("%s"), PackageNamesString.ToString());
+
 	bool bSuccess = PackageHelper.GetDesiredStatesForModification(PackageNames, PackagesToCheckout, PackagesToAdd, bErrorsAsWarnings);
 	if (!bSuccess && !bErrorsAsWarnings)
 	{
