Fix incremental cook issues with verse packages :

- Don't call ComputePackageRemoval in incremental cook. The reason is convoluted. ComputePackageRemoval removes the currently cooked packages from the ZenStoreWriter. Later on this prevents from loading the cook package artifacts. Without a cook package artifact, I can't compute the key hash so the package is always cooked.
- IsIterativeEnabled : verse packages are in memory packages, so they can't be found in the AssetBrowser and are always cooked. So check if the current package is an in memory package and if so enable iterative cook.

#incrementalcook
#rnx

diff --git a/Engine/Source/Editor/UnrealEd/Private/CookOnTheFlyServer.cpp b/Engine/Source/Editor/UnrealEd/Private/CookOnTheFlyServer.cpp
--- a/Engine/Source/Editor/UnrealEd/Private/CookOnTheFlyServer.cpp
+++ b/Engine/Source/Editor/UnrealEd/Private/CookOnTheFlyServer.cpp
@@ -7543,14 +7543,10 @@ void UCookOnTheFlyServer::PopulateCookedPackages(TArrayView<const ITargetPlatfor
 			// but checks for out-of-datedness are done by FRequestCluster using the TargetDomainKey (which is built
 			// from dependencies), so we do not need to check for out-of-datedness here
 
-			// Remove packages that no longer exist in the WorkspaceDomain from the TargetDomain. We have to
-			// check this for all packages in the previous cook rather than just the currently referenced
-			// set because when packages are removed the references to them are usually also removed.
 			TArray<FName> TombstonePackages;
 			int32 NumNeverCookPlaceHolderPackages;
 			PlatformAssetRegistry.ComputePackageRemovals(*PreviousAssetRegistry, TombstonePackages,
 				PreviousGeneratorPackages, NumNeverCookPlaceHolderPackages);
-			PackageWriter.RemoveCookedPackages(TombstonePackages);
 
 			// Do not show NeverCookPlaceholder packages in the package counts
 			NumPreviousPackages -= NumNeverCookPlaceHolderPackages;
diff --git a/Engine/Source/Editor/UnrealEd/Private/TargetDomain/TargetDomainUtils.cpp b/Engine/Source/Editor/UnrealEd/Private/TargetDomain/TargetDomainUtils.cpp
--- a/Engine/Source/Editor/UnrealEd/Private/TargetDomain/TargetDomainUtils.cpp
+++ b/Engine/Source/Editor/UnrealEd/Private/TargetDomain/TargetDomainUtils.cpp
@@ -33,15 +33,33 @@ bool IsIterativeEnabled(FName PackageName, bool bAllowAllClasses, const FAssetPa
 	{
 		return false;
 	}
-	TOptional<FAssetPackageData> PackageDataOpt;
+	
 	if (!OverrideAssetPackageData)
 	{
-		PackageDataOpt = AssetRegistry->GetAssetPackageDataCopy(PackageName);
-		if (!PackageDataOpt)
+		TOptional<FAssetPackageData> PackageDataOpt = AssetRegistry->GetAssetPackageDataCopy(PackageName);
+		if(PackageDataOpt)
+		{
+			OverrideAssetPackageData = PackageDataOpt.GetPtrOrNull();
+		}
+		else
 		{
-			return false;
+			TStringBuilder<256> PackageNameStr(EInPlace::InPlace, PackageName);
+			const UPackage* Package = FindPackage(nullptr, *PackageNameStr);
+			if (!Package)
+			{
+				return false;
+			}
+
+			bool bInMemoryPackage = Package->HasAnyPackageFlags(PKG_InMemoryOnly);
+			if (!bInMemoryPackage)
+			{
+				// Strange case where the package does not exist on disk nor in memory.
+				return false;
+			}
+
+			//This is an in memory package
+			return true;
 		}
-		OverrideAssetPackageData = PackageDataOpt.GetPtrOrNull();
 	}
 	const FAssetPackageData& PackageData = *OverrideAssetPackageData;
 
