Cooker: Header moves, whitespace and comment changes.
Move UE_COOK_RESULTPROJECTION_SCOPED into PackageAccessTracking.h so we can remove a dependency on PackageAccessTracking.h and therefore Package.h from the low-level CookDependency.h header.
Move functions in CookPackageArtifacts.cpp to match their refactored location in the header.
Move ArchiveCookContext inline functions outside of the class declaration so we can follow the coding standard's requirements of braces on their own lines without reducing readability of the class declaration.
Variable names and comments fixes.
#rnx

diff --git a/Engine/Source/Editor/UnrealEd/Private/Cooker/CookPackageArtifacts.cpp b/Engine/Source/Editor/UnrealEd/Private/Cooker/CookPackageArtifacts.cpp
--- a/Engine/Source/Editor/UnrealEd/Private/Cooker/CookPackageArtifacts.cpp
+++ b/Engine/Source/Editor/UnrealEd/Private/Cooker/CookPackageArtifacts.cpp
@@ -271,73 +271,33 @@
 	Writer.EndObject();
 }
 
-FName FPackageArtifacts::NAME_Save(TEXT("Save"));
-FName FPackageArtifacts::NAME_Load(TEXT("Load"));
+FBuildDependencySet FBuildDependencySet::CollectLoadedPackage(const UPackage* Package,
+	TArray<TPair<ELogVerbosity::Type, FString>>* OutMessages)
+{
+	TArray<FCookDependency> BuildDependencies;
+	TArray<FName> UnusedRuntimeDependencies;
+	if (!TryCollectInternal(BuildDependencies, UnusedRuntimeDependencies, Package, nullptr /* TargetPlatform */,
+		nullptr /* SaveResult */, nullptr /* GenerationHelper */, false /* bGenerated */, OutMessages))
+	{
+		return FBuildDependencySet();
+	}
 
-FPackageArtifacts::FPackageArtifacts()
-{
-	LoadBuildDependencies.SetName(NAME_Load);
-	SaveBuildDependencies.SetName(NAME_Save);
+	// Sort and remove duplicates in the results from TryCollectInternal.
+	Algo::Sort(BuildDependencies);
+	BuildDependencies.SetNum(Algo::Unique(BuildDependencies), EAllowShrinking::Yes);
+
+	// Store the collected dependencies on the result
+	FBuildDependencySet Result;
+	Result.SetNormalizedDependencies(MoveTemp(BuildDependencies));
+	FBuildDependencySet::ECurrentKeyResult CurrentKeyResult = Result.TryCalculateCurrentKey(
+		Package->GetFName(), nullptr /* GenerationHelper */, OutMessages);
+	if (CurrentKeyResult == FBuildDependencySet::ECurrentKeyResult::Error)
+	{
+		return FBuildDependencySet();
+	}
+	Result.StoreCurrentKey();
+	Result.SetValid(true);
+	return Result;
 }
 
-FBuildDependencySet& FPackageArtifacts::FindOrAddBuildDependencySet(FName ResultName)
-{
-	if (ResultName == NAME_Save)
-	{
-		return SaveBuildDependencies;
-	}
-	else if (ResultName == NAME_Load)
-	{
-		return LoadBuildDependencies;
-	}
-	else
-	{
-		// Not yet implemented
-		check(false);
-		return SaveBuildDependencies;
-	}
-}
-
-FBuildDependencySet* FPackageArtifacts::FindBuildDependencySet(FName ResultName)
-{
-	if (ResultName == NAME_Save)
-	{
-		return &SaveBuildDependencies;
-	}
-	else if (ResultName == NAME_Load)
-	{
-		return &LoadBuildDependencies;
-	}
-	else
-	{
-		return nullptr;
-	}
-}
-
-bool FPackageArtifacts::HasSaveResults() const
-{
-	return bHasSaveResults;
-}
-
-bool FPackageArtifacts::HasKeyMatch(FGenerationHelper* GenerationHelper)
-{
-	return bValid && SaveBuildDependencies.HasKeyMatch(PackageName, GenerationHelper);
-}
-
-FBuildDependencySet::ECurrentKeyResult FPackageArtifacts::TryCalculateCurrentKey(
-	FGenerationHelper* GenerationHelper, TArray<TPair<ELogVerbosity::Type, FString>>* OutMessages)
-{
-	return SaveBuildDependencies.TryCalculateCurrentKey(PackageName, GenerationHelper, OutMessages);
-}
-
-void FPackageArtifacts::Empty()
-{
-	SaveBuildDependencies.Empty();
-	LoadBuildDependencies.Empty();
-	RuntimeDependencies.Empty();
-	PackageName = FName();
-	bHasSaveResults = false;
-	bValid = false;
-}
-
 enum class EPackageMountPoint
@@ -754,6 +714,75 @@
 	return Result;
 }
 
+FName FPackageArtifacts::NAME_Save(TEXT("Save"));
+FName FPackageArtifacts::NAME_Load(TEXT("Load"));
+
+FPackageArtifacts::FPackageArtifacts()
+{
+	LoadBuildDependencies.SetName(NAME_Load);
+	SaveBuildDependencies.SetName(NAME_Save);
 }
 
+FBuildDependencySet& FPackageArtifacts::FindOrAddBuildDependencySet(FName ResultName)
+{
+	if (ResultName == NAME_Save)
+	{
+		return SaveBuildDependencies;
+	}
+	else if (ResultName == NAME_Load)
+	{
+		return LoadBuildDependencies;
+	}
+	else
+	{
+		// Not yet implemented
+		check(false);
+		return SaveBuildDependencies;
+	}
+}
+
+FBuildDependencySet* FPackageArtifacts::FindBuildDependencySet(FName ResultName)
+{
+	if (ResultName == NAME_Save)
+	{
+		return &SaveBuildDependencies;
+	}
+	else if (ResultName == NAME_Load)
+	{
+		return &LoadBuildDependencies;
+	}
+	else
+	{
+		return nullptr;
+	}
+}
+
+bool FPackageArtifacts::HasSaveResults() const
+{
+	return bHasSaveResults;
+}
+
+bool FPackageArtifacts::HasKeyMatch(FGenerationHelper* GenerationHelper)
+{
+	return bValid && SaveBuildDependencies.HasKeyMatch(PackageName, GenerationHelper);
+}
+
+FBuildDependencySet::ECurrentKeyResult FPackageArtifacts::TryCalculateCurrentKey(
+	FGenerationHelper* GenerationHelper, TArray<TPair<ELogVerbosity::Type, FString>>* OutMessages)
+{
+	return SaveBuildDependencies.TryCalculateCurrentKey(PackageName, GenerationHelper, OutMessages);
+}
+
+void FPackageArtifacts::Empty()
+{
+	SaveBuildDependencies.Empty();
+	LoadBuildDependencies.Empty();
+	RuntimeDependencies.Empty();
+	PackageName = FName();
+	bHasSaveResults = false;
+	bValid = false;
+}
+
+}
+
 bool LoadFromCompactBinary(FCbObjectView ObjectView, UE::Cook::FPackageArtifacts& Artifacts)
@@ -1054,35 +1083,6 @@
 	check(InProgressResults.Num() == 0);
 }
 
-FBuildDependencySet FBuildDependencySet::CollectLoadedPackage(const UPackage* Package,
-	TArray<TPair<ELogVerbosity::Type, FString>>* OutMessages)
-{
-	TArray<FCookDependency> BuildDependencies;
-	TArray<FName> UnusedRuntimeDependencies;
-	if (!TryCollectInternal(BuildDependencies, UnusedRuntimeDependencies, Package, nullptr /* TargetPlatform */,
-		nullptr /* SaveResult */, nullptr /* GenerationHelper */, false /* bGenerated */, OutMessages))
-	{
-		return FBuildDependencySet();
-	}
-
-	// Sort and remove duplicates in the results from TryCollectInternal.
-	Algo::Sort(BuildDependencies);
-	BuildDependencies.SetNum(Algo::Unique(BuildDependencies), EAllowShrinking::Yes);
-
-	// Store the collected dependencies on the result
-	FBuildDependencySet Result;
-	Result.SetNormalizedDependencies(MoveTemp(BuildDependencies));
-	FBuildDependencySet::ECurrentKeyResult CurrentKeyResult = Result.TryCalculateCurrentKey(
-		Package->GetFName(), nullptr /* GenerationHelper */, OutMessages);
-	if (CurrentKeyResult == FBuildDependencySet::ECurrentKeyResult::Error)
-	{
-		return FBuildDependencySet();
-	}
-	Result.StoreCurrentKey();
-	Result.SetValid(true);
-	return Result;
-}
-
 FPackageArtifacts FPackageArtifacts::Collect(const UPackage* Package, const ITargetPlatform* TargetPlatform,
 	const FBuildDependencySet* InLoadDependencies, const FSavePackageResultStruct* SaveResult,
 	FGenerationHelper* GenerationHelper, bool bGenerated,
diff --git a/Engine/Source/Runtime/CoreUObject/Private/UObject/SavePackage/SaveContext.cpp b/Engine/Source/Runtime/CoreUObject/Private/UObject/SavePackage/SaveContext.cpp
--- a/Engine/Source/Runtime/CoreUObject/Private/UObject/SavePackage/SaveContext.cpp
+++ b/Engine/Source/Runtime/CoreUObject/Private/UObject/SavePackage/SaveContext.cpp
@@ -6,6 +6,7 @@
 #include "Algo/Unique.h"
 #include "Cooker/CookDependency.h"
 #include "Misc/ConfigCacheIni.h"
+#include "Misc/PackageAccessTracking.h"
 #include "Serialization/PackageWriter.h"
 #include "UObject/ObjectSaveContext.h"
 #include "UObject/UObjectGlobals.h"
diff --git a/Engine/Source/Runtime/CoreUObject/Private/UObject/SavePackage/SavePackageUtilities.h b/Engine/Source/Runtime/CoreUObject/Private/UObject/SavePackage/SavePackageUtilities.h
--- a/Engine/Source/Runtime/CoreUObject/Private/UObject/SavePackage/SavePackageUtilities.h
+++ b/Engine/Source/Runtime/CoreUObject/Private/UObject/SavePackage/SavePackageUtilities.h
@@ -26,7 +26,7 @@ class FPackagePath;
 class FSaveContext;
 class FSavePackageContext;
 class IPackageWriter;
-
+class ITargetPlatform;
 enum class ESavePackageResult;
 
 // Save Time trace
@@ -349,7 +349,7 @@ ESavePackageResult CreatePayloadSidecarFile(FLinkerSave& Linker, const FPackageP
 
 void SaveMetaData(UPackage* InOuter, FLinkerSave* Linker, FStructuredArchive::FRecord Record);
 void SaveWorldLevelInfo(UPackage* InOuter, FLinkerSave* Linker, FStructuredArchive::FRecord Record);
-EObjectMark GetExcludedObjectMarksForTargetPlatform(const class ITargetPlatform* TargetPlatform);
+EObjectMark GetExcludedObjectMarksForTargetPlatform(const ITargetPlatform* TargetPlatform);
 void FindMostLikelyCulprit(const TArray<UObject*>& BadObjects, UObject*& OutMostLikelyCulprit, UObject*& OutReferencer, const FProperty*& OutReferencerProperty, bool& OutIsCulpritArchetype, FSaveContext* InOptionalSaveContext = nullptr);
 	
 /** 
diff --git a/Engine/Source/Runtime/CoreUObject/Private/UObject/SavePackage2.cpp b/Engine/Source/Runtime/CoreUObject/Private/UObject/SavePackage2.cpp
--- a/Engine/Source/Runtime/CoreUObject/Private/UObject/SavePackage2.cpp
+++ b/Engine/Source/Runtime/CoreUObject/Private/UObject/SavePackage2.cpp
@@ -18,6 +18,7 @@
 #include "Misc/ConfigCacheIni.h"
 #include "Misc/FeedbackContext.h"
 #include "Misc/MessageDialog.h"
+#include "Misc/PackageAccessTracking.h"
 #include "Misc/PackageName.h"
 #include "Misc/PathViews.h"
 #include "Misc/RedirectCollector.h"
diff --git a/Engine/Source/Runtime/CoreUObject/Public/Cooker/CookDependency.h b/Engine/Source/Runtime/CoreUObject/Public/Cooker/CookDependency.h
--- a/Engine/Source/Runtime/CoreUObject/Public/Cooker/CookDependency.h
+++ b/Engine/Source/Runtime/CoreUObject/Public/Cooker/CookDependency.h
@@ -7,7 +7,6 @@
 #include "Containers/UnrealString.h"
 #include "HAL/PreprocessorHelpers.h"
 #include "Logging/LogVerbosity.h"
-#include "Misc/PackageAccessTracking.h"
 #include "Serialization/CompactBinary.h"
 #include "Templates/Function.h"
 #include "Templates/UniquePtr.h"
@@ -21,21 +20,6 @@ namespace UE::ConfigAccessTracking { struct FConfigAccessData; }
 struct FARFilter;
 #endif
 
-#if WITH_EDITOR && UE_WITH_PACKAGE_ACCESS_TRACKING
-/**
- * Specify that within the current scope, dereferenced TObjectPtrs cause a dependency on the specified
- * ResultProjection of the dereferenced object and package, instead of the outer scope's projection or the default
- * projection of UE::CooK::ResultProjection::All. ProjectionName must be either one of the generic terms provided
- * in UE::Cook::ResultProjection or it must be a name declared as a resultprojection by one of the objects in
- * the target package.
- */
-#define UE_COOK_RESULTPROJECTION_SCOPED(ProjectionName) \
-	PackageAccessTracking_Private::FPackageAccessRefScope ANONYMOUS_VARIABLE(PackageAccessTracker_)( \
-		PackageAccessTracking_Private::FPackageAccessRefScope::CookResultProjectionType, ProjectionName)
-#else
-#define UE_COOK_RESULTPROJECTION_SCOPED(ProjectionName)
-#endif
-
 #if WITH_EDITOR
 
 namespace UE::Cook
@@ -149,7 +133,7 @@ namespace ResultProjection
 extern COREUOBJECT_API FName All;
 /*
  * FName passed into FCookDependency::ResultProjection or UE_COOK_RESULTPROJECTION_SCOPED to indicate the bytes
- * within the target package, and the native class used by objects referenced from it, can contribute the the
+ * within the target package, and the native class used by objects referenced from it, can contribute to the
  * source package's cook. This is more limited than the transitive dependency indicated by BuildProjectionAll
  * because it does not include e.g. other packages that the target package depends on for its cook.
  * This ResultProjection is applied automatically for all hard imports saved into the source cooked package;
@@ -315,15 +299,15 @@ private:
 	static COREUOBJECT_API bool ARFilterEqual(const FARFilter& A, const FARFilter& B);
 
 	/** Public hidden friend for operator<< into an FCbWriter. */
-	friend FCbWriter& operator<<(FCbWriter& Writer, const FCookDependency& CookDependencies)
+	friend FCbWriter& operator<<(FCbWriter& Writer, const FCookDependency& CookDependency)
 	{
-		CookDependencies.Save(Writer);
+		CookDependency.Save(Writer);
 		return Writer;
 	}
 	/** Public hidden friend for LoadFromCompactBinary. */
-	friend bool LoadFromCompactBinary(FCbFieldView Value, FCookDependency& CookDependencies)
+	friend bool LoadFromCompactBinary(FCbFieldView Value, FCookDependency& CookDependency)
 	{
-		return CookDependencies.Load(Value);
+		return CookDependency.Load(Value);
 	}
 
 private:
diff --git a/Engine/Source/Runtime/CoreUObject/Public/Misc/PackageAccessTracking.h b/Engine/Source/Runtime/CoreUObject/Public/Misc/PackageAccessTracking.h
--- a/Engine/Source/Runtime/CoreUObject/Public/Misc/PackageAccessTracking.h
+++ b/Engine/Source/Runtime/CoreUObject/Public/Misc/PackageAccessTracking.h
@@ -34,6 +34,22 @@ class UPackage;
 #define UE_TRACK_REFERENCING_PACKAGE_DEACTIVATE_SCOPE_VARIABLE(VariableName)
 #endif //UE_WITH_PACKAGE_ACCESS_TRACKING
 
+// Cook-specific AccessTracking Scopes
+#if WITH_EDITOR && UE_WITH_PACKAGE_ACCESS_TRACKING
+/**
+ * Specify that within the current scope, dereferenced TObjectPtrs cause a dependency on the specified
+ * ResultProjection of the dereferenced object and package, instead of the outer scope's projection or the default
+ * projection of UE::CooK::ResultProjection::All. ProjectionName must be either one of the generic terms provided
+ * in UE::Cook::ResultProjection or it must be a name declared as a resultprojection by one of the objects in
+ * the target package.
+ */
+#define UE_COOK_RESULTPROJECTION_SCOPED(ProjectionName) \
+	PackageAccessTracking_Private::FPackageAccessRefScope ANONYMOUS_VARIABLE(PackageAccessTracker_)( \
+		PackageAccessTracking_Private::FPackageAccessRefScope::CookResultProjectionType, ProjectionName)
+#else
+#define UE_COOK_RESULTPROJECTION_SCOPED(ProjectionName)
+#endif //WITH_EDITOR && UE_WITH_PACKAGE_ACCESS_TRACKING
+
 ///////////////////////////////////////////////////////////////////////////////////////////////////////////////////////
 
 #if UE_WITH_PACKAGE_ACCESS_TRACKING
diff --git a/Engine/Source/Runtime/CoreUObject/Public/UObject/ArchiveCookContext.h b/Engine/Source/Runtime/CoreUObject/Public/UObject/ArchiveCookContext.h
--- a/Engine/Source/Runtime/CoreUObject/Public/UObject/ArchiveCookContext.h
+++ b/Engine/Source/Runtime/CoreUObject/Public/UObject/ArchiveCookContext.h
@@ -40,32 +40,81 @@ private:
 public:
 
 	UE_DEPRECATED(5.4, "Call version that takes the TargetPlatform")
-	FArchiveCookContext(UPackage* InPackage, UE::Cook::ECookType InCookType, UE::Cook::ECookingDLC InCookingDLC)
-		: FArchiveCookContext(InPackage, InCookType, InCookingDLC, nullptr)
-	{
-	}
+	FArchiveCookContext(UPackage* InPackage, UE::Cook::ECookType InCookType, UE::Cook::ECookingDLC InCookingDLC);
 
 	FArchiveCookContext(UPackage* InPackage, UE::Cook::ECookType InCookType, UE::Cook::ECookingDLC InCookingDLC,
-		const ITargetPlatform* InTargetPlatform)
-		: CookTagList(InPackage)
-		, TargetPlatform(InTargetPlatform)
-		, bCookTagListEnabled(InPackage && InCookType == UE::Cook::ECookType::ByTheBook)
-		, CookType(InCookType)
-		, CookingDLC(InCookingDLC)
-	{
-	}
-
-	void Reset()
-	{
-		CookTagList.Reset();
-	}
-
-	FCookTagList* GetCookTagList() { return bCookTagListEnabled ? &CookTagList : nullptr; }
-	const ITargetPlatform* GetTargetPlatform() const { return TargetPlatform; }
-
-	bool IsCookByTheBook() const { return CookType == UE::Cook::ECookType::ByTheBook; }
-	bool IsCookOnTheFly() const { return CookType == UE::Cook::ECookType::OnTheFly; }
-	bool IsCookTypeUnknown() const { return CookType == UE::Cook::ECookType::Unknown; }
-	UE::Cook::ECookType GetCookType() const { return CookType; }
-	UE::Cook::ECookingDLC GetCookingDLC() const { return CookingDLC; }
-};
\ No newline at end of file
+		const ITargetPlatform* InTargetPlatform);
+
+	void Reset();
+
+	FCookTagList* GetCookTagList();
+	const ITargetPlatform* GetTargetPlatform() const;
+
+	bool IsCookByTheBook() const;
+	bool IsCookOnTheFly() const;
+	bool IsCookTypeUnknown() const;
+	UE::Cook::ECookType GetCookType() const;
+	UE::Cook::ECookingDLC GetCookingDLC() const;
+};
+
+
+///////////////////////////////////////////////////////
+// Inline implementations
+///////////////////////////////////////////////////////
+
+
+inline FArchiveCookContext::FArchiveCookContext(UPackage* InPackage, UE::Cook::ECookType InCookType, 
+	UE::Cook::ECookingDLC InCookingDLC)
+	: FArchiveCookContext(InPackage, InCookType, InCookingDLC, nullptr)
+{
+}
+
+inline FArchiveCookContext::FArchiveCookContext(UPackage* InPackage, UE::Cook::ECookType InCookType,
+	UE::Cook::ECookingDLC InCookingDLC, const ITargetPlatform* InTargetPlatform)
+	: CookTagList(InPackage)
+	, TargetPlatform(InTargetPlatform)
+	, bCookTagListEnabled(InPackage&& InCookType == UE::Cook::ECookType::ByTheBook)
+	, CookType(InCookType)
+	, CookingDLC(InCookingDLC)
+{
+}
+
+inline void FArchiveCookContext::Reset()
+{
+	CookTagList.Reset();
+}
+
+inline FCookTagList* FArchiveCookContext::GetCookTagList()
+{
+	return bCookTagListEnabled ? &CookTagList : nullptr;
+}
+
+inline const ITargetPlatform* FArchiveCookContext::GetTargetPlatform() const
+{
+	return TargetPlatform;
+}
+
+inline bool FArchiveCookContext::IsCookByTheBook() const
+{
+	return CookType == UE::Cook::ECookType::ByTheBook;
+}
+
+inline bool FArchiveCookContext::IsCookOnTheFly() const
+{
+	return CookType == UE::Cook::ECookType::OnTheFly;
+}
+
+inline bool FArchiveCookContext::IsCookTypeUnknown() const
+{
+	return CookType == UE::Cook::ECookType::Unknown;
+}
+
+inline UE::Cook::ECookType FArchiveCookContext::GetCookType() const
+{
+	return CookType;
+}
+
+inline UE::Cook::ECookingDLC FArchiveCookContext::GetCookingDLC() const
+{
+	return CookingDLC;
+}
\ No newline at end of file
