[ContentBrowser] Thumbnail Action button is not working properly

diff --git a/Engine/Source/Editor/ContentBrowser/Private/SAssetView.cpp b/Engine/Source/Editor/ContentBrowser/Private/SAssetView.cpp
--- a/Engine/Source/Editor/ContentBrowser/Private/SAssetView.cpp
+++ b/Engine/Source/Editor/ContentBrowser/Private/SAssetView.cpp
@@ -1369,6 +1369,7 @@ void SAssetView::Construct( const FArguments& InArgs )
 			.Padding(FMargin(0, 14, 0, 0))
 			[
 				SNew(SScrollBox)
+				.Visibility( this, &SAssetView::IsAssetShowWarningTextVisible )
 				+SScrollBox::Slot()
 				[
 					// A warning to display when there are no assets to show
diff --git a/Engine/Source/Editor/UnrealEd/Private/SAssetThumbnailEditModeTools.cpp b/Engine/Source/Editor/UnrealEd/Private/SAssetThumbnailEditModeTools.cpp
--- a/Engine/Source/Editor/UnrealEd/Private/SAssetThumbnailEditModeTools.cpp
+++ b/Engine/Source/Editor/UnrealEd/Private/SAssetThumbnailEditModeTools.cpp
@@ -215,8 +215,7 @@ FReply SAssetThumbnailEditModeTools::OnMouseButtonDown(const FGeometry& MyGeomet
 				bModifiedThumbnailWhileDragging = false;
 				DragStartLocation = FIntPoint(FMath::TruncToInt32(MouseEvent.GetScreenSpacePosition().X), FMath::TruncToInt32(MouseEvent.GetScreenSpacePosition().Y));
 
-				// Keep the widget hovered so that the view can maintain a hovered state on the whole widget while dragging the mouse even outside its bounds
-				SetHover(true);
+				bIsEditing = true;
 				return FReply::Handled().CaptureMouse(AsShared()).UseHighPrecisionMouseMovement(AsShared()).PreventThrottling();
 			}
 		}
@@ -243,9 +242,7 @@ FReply SAssetThumbnailEditModeTools::OnMouseButtonUp(const FGeometry& MyGeometry
 			bModifiedThumbnailWhileDragging = false;
 		}
 
-		// Reset the Hover to being changed automatically
-		// since the mouse position is reset as well inside the thumbnail there will be no need to force it changing here
-		SetHover(TAttribute<bool>());
+		bIsEditing = false;
 		return FReply::Handled().ReleaseMouseCapture().SetMousePos(DragStartLocation);
 	}
 
@@ -317,6 +314,11 @@ FCursorReply SAssetThumbnailEditModeTools::OnCursorQuery(const FGeometry& MyGeom
 		FCursorReply::Cursor(EMouseCursor::Default);
 }
 
+bool SAssetThumbnailEditModeTools::IsEditingThumbnail() const
+{
+	return bIsEditing;
+}
+
 USceneThumbnailInfo* SAssetThumbnailEditModeTools::GetSceneThumbnailInfo() const
 {
 	USceneThumbnailInfo* SceneThumbnailInfo = SceneThumbnailInfoPtr.Get();
diff --git a/Engine/Source/Editor/UnrealEd/Private/SAssetThumbnailEditModeTools.h b/Engine/Source/Editor/UnrealEd/Private/SAssetThumbnailEditModeTools.h
--- a/Engine/Source/Editor/UnrealEd/Private/SAssetThumbnailEditModeTools.h
+++ b/Engine/Source/Editor/UnrealEd/Private/SAssetThumbnailEditModeTools.h
@@ -38,6 +38,9 @@ public:
 	virtual FReply OnMouseMove(const FGeometry& MyGeometry, const FPointerEvent& MouseEvent) override;
 	virtual FCursorReply OnCursorQuery(const FGeometry& MyGeometry, const FPointerEvent& CursorEvent) const override;
 
+	/** True if currently editing the thumbnail, false otherwise */
+	bool IsEditingThumbnail() const;
+
 protected:
 	/** Gets the visibility for the primitives toolbar */
 	EVisibility GetPrimitiveToolsVisibility() const;
@@ -65,6 +68,9 @@ protected:
 	TWeakPtr<FAssetThumbnail> AssetThumbnail;
 
 private:
+	/** True if currently editing the thumbnail, false otherwise */
+	bool bIsEditing = false;
+
 	// Never access this directly.
 	mutable TWeakObjectPtr<USceneThumbnailInfo> SceneThumbnailInfoPtr;
 };
