IsIterativeEnabled : Remove unused parameter.

#incrementalcook
#rnx

diff --git a/Engine/Source/Editor/UnrealEd/Private/Cooker/CookRequestCluster.cpp b/Engine/Source/Editor/UnrealEd/Private/Cooker/CookRequestCluster.cpp
--- a/Engine/Source/Editor/UnrealEd/Private/Cooker/CookRequestCluster.cpp
+++ b/Engine/Source/Editor/UnrealEd/Private/Cooker/CookRequestCluster.cpp
@@ -1535,7 +1535,6 @@ bool FRequestCluster::FGraphSearch::FExploreEdgesContext::TryCalculateIterativel
 		FPackagePlatformData& PackagePlatformData = PackageData->FindOrAddPlatformData(TargetPlatform);
 
 		UE::Cook::FPackageArtifacts& Artifacts = QueryPlatformData.CookAttachments.Artifacts;
-		const FAssetPackageData* OverrideAssetPackageData = nullptr;
 		FPackageData* ParentPackageData = nullptr;
 		TRefCountPtr<FGenerationHelper> GenerationHelper;
 		if (!PackageData->IsGenerated())
@@ -1572,16 +1571,13 @@ bool FRequestCluster::FGraphSearch::FExploreEdgesContext::TryCalculateIterativel
 			continue;
 		}
 
-		if (!IsIterativeEnabled(PackageName, Cluster.COTFS.bHybridIterativeAllowAllClasses, OverrideAssetPackageData))
-		{
-			SetIsIterativelyUnmodified(PlatformIndex, false, PackagePlatformData);
-			continue;
-		}
 		// Generated packages of a generator that is not IterativelyEnabled are also not iteratively enabled, even
 		// if they would otherwise qualify for iterative on their own. e.g. if worlds are iteratively disallowed,
 		// then streamingobject generated packages of the world are also disallowed.
-		if (ParentPackageData && !IsIterativeEnabled(ParentPackageData->GetPackageName(),
-			Cluster.COTFS.bHybridIterativeAllowAllClasses))
+		FName PackageNameForIterativeTest = ParentPackageData ? ParentPackageData->GetPackageName() : PackageName;
+		bool bIsIterativeEnabled = IsIterativeEnabled(PackageNameForIterativeTest, Cluster.COTFS.bHybridIterativeAllowAllClasses);
+		
+		if (!bIsIterativeEnabled)
 		{
 			SetIsIterativelyUnmodified(PlatformIndex, false, PackagePlatformData);
 			continue;
diff --git a/Engine/Source/Editor/UnrealEd/Private/TargetDomain/TargetDomainUtils.cpp b/Engine/Source/Editor/UnrealEd/Private/TargetDomain/TargetDomainUtils.cpp
--- a/Engine/Source/Editor/UnrealEd/Private/TargetDomain/TargetDomainUtils.cpp
+++ b/Engine/Source/Editor/UnrealEd/Private/TargetDomain/TargetDomainUtils.cpp
@@ -26,7 +26,7 @@
 namespace UE::TargetDomain
 {
 
-bool IsIterativeEnabled(FName PackageName, bool bAllowAllClasses, const FAssetPackageData* OverrideAssetPackageData)
+bool IsIterativeEnabled(FName PackageName, bool bAllowAllClasses)
 {
 	IAssetRegistry* AssetRegistry = IAssetRegistry::Get();
 	if (!AssetRegistry)
@@ -34,34 +34,28 @@ bool IsIterativeEnabled(FName PackageName, bool bAllowAllClasses, const FAssetPa
 		return false;
 	}
 	
-	if (!OverrideAssetPackageData)
+	TOptional<FAssetPackageData> PackageDataOpt = AssetRegistry->GetAssetPackageDataCopy(PackageName);
+	if(!PackageDataOpt)
 	{
-		TOptional<FAssetPackageData> PackageDataOpt = AssetRegistry->GetAssetPackageDataCopy(PackageName);
-		if(PackageDataOpt)
+		TStringBuilder<256> PackageNameStr(EInPlace::InPlace, PackageName);
+		const UPackage* Package = FindPackage(nullptr, *PackageNameStr);
+		if (!Package)
 		{
-			OverrideAssetPackageData = PackageDataOpt.GetPtrOrNull();
+			return false;
 		}
-		else
-		{
-			TStringBuilder<256> PackageNameStr(EInPlace::InPlace, PackageName);
-			const UPackage* Package = FindPackage(nullptr, *PackageNameStr);
-			if (!Package)
-			{
-				return false;
-			}
 
-			bool bInMemoryPackage = Package->HasAnyPackageFlags(PKG_InMemoryOnly);
-			if (!bInMemoryPackage)
-			{
-				// Strange case where the package does not exist on disk nor in memory.
-				return false;
-			}
-
-			//This is an in memory package
-			return true;
+		bool bInMemoryPackage = Package->HasAnyPackageFlags(PKG_InMemoryOnly);
+		if (!bInMemoryPackage)
+		{
+			// Strange case where the package does not exist on disk nor in memory.
+			return false;
 		}
+
+		//This is an in memory package
+		return true;
 	}
-	const FAssetPackageData& PackageData = *OverrideAssetPackageData;
+
+	const FAssetPackageData& PackageData = *PackageDataOpt.GetPtrOrNull();
 
 	if (!bAllowAllClasses)
 	{
diff --git a/Engine/Source/Editor/UnrealEd/Private/TargetDomain/TargetDomainUtils.h b/Engine/Source/Editor/UnrealEd/Private/TargetDomain/TargetDomainUtils.h
--- a/Engine/Source/Editor/UnrealEd/Private/TargetDomain/TargetDomainUtils.h
+++ b/Engine/Source/Editor/UnrealEd/Private/TargetDomain/TargetDomainUtils.h
@@ -24,8 +24,7 @@ namespace UE::TargetDomain
 void CookInitialize();
 
 /** Return whether iterative cook is enabled for the given packagename, based on used-class allowlist/blocklist. */
-bool IsIterativeEnabled(FName PackageName, bool bAllowAllClasses,
-	const FAssetPackageData* OverrideAssetPackageData = nullptr);
+bool IsIterativeEnabled(FName PackageName, bool bAllowAllClasses);
 
 /** Store extra information derived during save and used by the cooker for the given EditorDomain package. */
 void CommitEditorDomainCookAttachments(FName PackageName, TArrayView<IPackageWriter::FCommitAttachmentInfo> Attachments);
