FUnrealEdMisc : Don't register the EditorAxisDisplayInfo modular feature if another implementation is already registered.

diff --git a/Engine/Source/Editor/UnrealEd/Private/UnrealEdMisc.cpp b/Engine/Source/Editor/UnrealEd/Private/UnrealEdMisc.cpp
--- a/Engine/Source/Editor/UnrealEd/Private/UnrealEdMisc.cpp
+++ b/Engine/Source/Editor/UnrealEd/Private/UnrealEdMisc.cpp
@@ -257,8 +257,11 @@ void FUnrealEdMisc::OnInit()
 	IModularFeatures::Get().RegisterModularFeature(IConsoleCommandExecutor::ModularFeatureName(), CmdExec.Get());
 
 	// Register the EditorAxisDisplayInfo modular feature
-	EditorAxisDisplayInfo = MakeUnique<FEditorAxisDisplayInfo>();
-	IModularFeatures::Get().RegisterModularFeature(IAxisDisplayInfo::GetModularFeatureName(), EditorAxisDisplayInfo.Get());
+	if (!IModularFeatures::Get().IsModularFeatureAvailable(IAxisDisplayInfo::GetModularFeatureName()))
+	{
+		EditorAxisDisplayInfo = MakeUnique<FEditorAxisDisplayInfo>();
+		IModularFeatures::Get().RegisterModularFeature(IAxisDisplayInfo::GetModularFeatureName(), EditorAxisDisplayInfo.Get());
+	}
 
 	// Register all callback notifications
 	FEditorDelegates::SelectedProps.AddRaw(this, &FUnrealEdMisc::CB_SelectedProps);
@@ -1144,8 +1147,11 @@ void FUnrealEdMisc::OnExit()
 	}
 
 	// Unregister the EditorAxisDisplayInfo modular feature
-	IModularFeatures::Get().UnregisterModularFeature(IAxisDisplayInfo::GetModularFeatureName(), EditorAxisDisplayInfo.Get());
-	EditorAxisDisplayInfo.Reset();
+	if (EditorAxisDisplayInfo.IsValid())
+	{
+		IModularFeatures::Get().UnregisterModularFeature(IAxisDisplayInfo::GetModularFeatureName(), EditorAxisDisplayInfo.Get());
+		EditorAxisDisplayInfo.Reset();
+	}
 
 	// Unregister the command executor
 	IModularFeatures::Get().UnregisterModularFeature(IConsoleCommandExecutor::ModularFeatureName(), CmdExec.Get());
diff --git a/Engine/Source/Editor/UnrealEd/Public/EditorAxisDisplayInfo.h b/Engine/Source/Editor/UnrealEd/Public/EditorAxisDisplayInfo.h
--- a/Engine/Source/Editor/UnrealEd/Public/EditorAxisDisplayInfo.h
+++ b/Engine/Source/Editor/UnrealEd/Public/EditorAxisDisplayInfo.h
@@ -4,7 +4,7 @@
 
 #include "Misc/AxisDisplayInfo.h"
 
-class FEditorAxisDisplayInfo final : public IAxisDisplayInfo
+class UNREALED_API FEditorAxisDisplayInfo : public IAxisDisplayInfo
 {
 public:
 
