[Transforms] Flip left / right ortho camera to align to camera view space

Update 2x2 default to flip orientation of panes vertically

diff --git a/Engine/Source/Editor/UnrealEd/Classes/Editor/UnrealEdTypes.h b/Engine/Source/Editor/UnrealEd/Classes/Editor/UnrealEdTypes.h
--- a/Engine/Source/Editor/UnrealEd/Classes/Editor/UnrealEdTypes.h
+++ b/Engine/Source/Editor/UnrealEd/Classes/Editor/UnrealEdTypes.h
@@ -129,16 +129,16 @@ enum ELevelViewportType : int
 	/** Alias semantic meaning of viewport types for orthographic views */
 	/** Top */
 	LVT_OrthoTop = LVT_OrthoXY,
-	/** Right */
-	LVT_OrthoRight = LVT_OrthoXZ,
+	/** Left */
+	LVT_OrthoLeft = LVT_OrthoXZ,
 	/** Front */
 	LVT_OrthoFront = LVT_OrthoNegativeYZ,
 	/** Back */
 	LVT_OrthoBack = LVT_OrthoYZ,
 	/** Bottom */
 	LVT_OrthoBottom = LVT_OrthoNegativeXY,
-	/** Left */
-	LVT_OrthoLeft = LVT_OrthoNegativeXZ,
+	/** Right */
+	LVT_OrthoRight = LVT_OrthoNegativeXZ,
 	LVT_MAX = 8,
 
 	LVT_None = 255,
diff --git a/Engine/Source/Editor/UnrealEd/Private/EditorViewportClient.cpp b/Engine/Source/Editor/UnrealEd/Private/EditorViewportClient.cpp
--- a/Engine/Source/Editor/UnrealEd/Private/EditorViewportClient.cpp
+++ b/Engine/Source/Editor/UnrealEd/Private/EditorViewportClient.cpp
@@ -387,12 +387,12 @@ void InitViewOptionsArray()
 {
 	ViewOptions.Empty();
 
-	ELevelViewportType Right = ELevelViewportType::LVT_OrthoXZ;
-	ELevelViewportType Left = ELevelViewportType::LVT_OrthoNegativeXZ;
-	ELevelViewportType Top = ELevelViewportType::LVT_OrthoXY;
-	ELevelViewportType Bottom = ELevelViewportType::LVT_OrthoNegativeXY;
-	ELevelViewportType Back = ELevelViewportType::LVT_OrthoYZ;
-	ELevelViewportType Front = ELevelViewportType::LVT_OrthoNegativeYZ;
+	ELevelViewportType Left = ELevelViewportType::LVT_OrthoLeft;
+	ELevelViewportType Right = ELevelViewportType::LVT_OrthoRight;
+	ELevelViewportType Top = ELevelViewportType::LVT_OrthoTop;
+	ELevelViewportType Bottom = ELevelViewportType::LVT_OrthoBottom;
+	ELevelViewportType Back = ELevelViewportType::LVT_OrthoBack;
+	ELevelViewportType Front = ELevelViewportType::LVT_OrthoFront;
 
 	ViewOptions.Add(Front);
 	ViewOptions.Add(Back);
diff --git a/Engine/Source/Editor/UnrealEd/Private/EditorViewportLayout2x2.cpp b/Engine/Source/Editor/UnrealEd/Private/EditorViewportLayout2x2.cpp
--- a/Engine/Source/Editor/UnrealEd/Private/EditorViewportLayout2x2.cpp
+++ b/Engine/Source/Editor/UnrealEd/Private/EditorViewportLayout2x2.cpp
@@ -86,19 +86,19 @@ TSharedRef<SWidget> FEditorViewportLayout2x2::MakeViewportLayout(TSharedPtr<FAss
 	SNew( SSplitter2x2 )
 	.TopLeft()
 	[
-		PerspectiveViewport
+		TopViewport
 	]
 	.BottomLeft()
 	[
-		RightViewport
+		FrontViewport
 	]
 	.TopRight()
 	[
-		TopViewport
+		PerspectiveViewport
 	]
 	.BottomRight()
 	[
-		FrontViewport
+		RightViewport
 	];
 	
 	if (SplitterPercentages.Num() > 0)
diff --git a/Engine/Source/Editor/UnrealEd/Private/SEditorViewport.cpp b/Engine/Source/Editor/UnrealEd/Private/SEditorViewport.cpp
--- a/Engine/Source/Editor/UnrealEd/Private/SEditorViewport.cpp
+++ b/Engine/Source/Editor/UnrealEd/Private/SEditorViewport.cpp
@@ -318,39 +318,39 @@ void SEditorViewport::BindCommands()
 
 	CommandListRef.MapAction( 
 		Commands.Front,
-		FExecuteAction::CreateSP( ClientRef, &FEditorViewportClient::SetViewportType, LVT_OrthoNegativeYZ ),
+		FExecuteAction::CreateSP( ClientRef, &FEditorViewportClient::SetViewportType, LVT_OrthoFront ),
 		FCanExecuteAction(),
-		FIsActionChecked::CreateSP(ClientRef, &FEditorViewportClient::IsActiveViewportType, LVT_OrthoNegativeYZ));
+		FIsActionChecked::CreateSP(ClientRef, &FEditorViewportClient::IsActiveViewportType, LVT_OrthoFront));
 
 	CommandListRef.MapAction( 
 		Commands.Left,
-		FExecuteAction::CreateSP( ClientRef, &FEditorViewportClient::SetViewportType, LVT_OrthoNegativeXZ ),
+		FExecuteAction::CreateSP(ClientRef, &FEditorViewportClient::SetViewportType, LVT_OrthoLeft),
 		FCanExecuteAction(),
-		FIsActionChecked::CreateSP(ClientRef, &FEditorViewportClient::IsActiveViewportType, LVT_OrthoNegativeXZ));
+		FIsActionChecked::CreateSP(ClientRef, &FEditorViewportClient::IsActiveViewportType, LVT_OrthoLeft));
 
 	CommandListRef.MapAction( 
 		Commands.Top,
-		FExecuteAction::CreateSP( ClientRef, &FEditorViewportClient::SetViewportType, LVT_OrthoXY ),
+		FExecuteAction::CreateSP( ClientRef, &FEditorViewportClient::SetViewportType, LVT_OrthoTop ),
 		FCanExecuteAction(),
-		FIsActionChecked::CreateSP(ClientRef, &FEditorViewportClient::IsActiveViewportType, LVT_OrthoXY));
+		FIsActionChecked::CreateSP(ClientRef, &FEditorViewportClient::IsActiveViewportType, LVT_OrthoTop));
 
 	CommandListRef.MapAction(
 		Commands.Back,
-		FExecuteAction::CreateSP(ClientRef, &FEditorViewportClient::SetViewportType, LVT_OrthoYZ),
+		FExecuteAction::CreateSP(ClientRef, &FEditorViewportClient::SetViewportType, LVT_OrthoBack),
 		FCanExecuteAction(),
-		FIsActionChecked::CreateSP(ClientRef, &FEditorViewportClient::IsActiveViewportType, LVT_OrthoYZ));
+		FIsActionChecked::CreateSP(ClientRef, &FEditorViewportClient::IsActiveViewportType, LVT_OrthoBack));
 
 	CommandListRef.MapAction(
 		Commands.Right,
-		FExecuteAction::CreateSP(ClientRef, &FEditorViewportClient::SetViewportType, LVT_OrthoXZ),
+		FExecuteAction::CreateSP(ClientRef, &FEditorViewportClient::SetViewportType, LVT_OrthoRight),
 		FCanExecuteAction(),
-		FIsActionChecked::CreateSP(ClientRef, &FEditorViewportClient::IsActiveViewportType, LVT_OrthoXZ));
+		FIsActionChecked::CreateSP(ClientRef, &FEditorViewportClient::IsActiveViewportType, LVT_OrthoRight));
 
 	CommandListRef.MapAction(
 		Commands.Bottom,
-		FExecuteAction::CreateSP(ClientRef, &FEditorViewportClient::SetViewportType, LVT_OrthoNegativeXY),
+		FExecuteAction::CreateSP(ClientRef, &FEditorViewportClient::SetViewportType, LVT_OrthoBottom),
 		FCanExecuteAction(),
-		FIsActionChecked::CreateSP(ClientRef, &FEditorViewportClient::IsActiveViewportType, LVT_OrthoNegativeXY));
+		FIsActionChecked::CreateSP(ClientRef, &FEditorViewportClient::IsActiveViewportType, LVT_OrthoBottom));
 
 	CommandListRef.MapAction(
 		Commands.Next,
diff --git a/Engine/Source/Editor/UnrealEd/Private/ViewportToolbar/UnrealEdViewportToolbar.cpp b/Engine/Source/Editor/UnrealEd/Private/ViewportToolbar/UnrealEdViewportToolbar.cpp
--- a/Engine/Source/Editor/UnrealEd/Private/ViewportToolbar/UnrealEdViewportToolbar.cpp
+++ b/Engine/Source/Editor/UnrealEd/Private/ViewportToolbar/UnrealEdViewportToolbar.cpp
@@ -2477,27 +2477,27 @@ FText GetCameraSubmenuLabelFromViewportType(const ELevelViewportType ViewportTyp
 		Label = LOCTEXT("CameraMenuTitle_Perspective", "Perspective");
 		break;
 
-	case LVT_OrthoXY:
+	case LVT_OrthoTop:
 		Label = LOCTEXT("CameraMenuTitle_Top", "Top");
 		break;
 
-	case LVT_OrthoNegativeXZ:
+	case LVT_OrthoLeft:
 		Label = LOCTEXT("CameraMenuTitle_Left", "Left");
 		break;
 
-	case LVT_OrthoNegativeYZ:
+	case LVT_OrthoFront:
 		Label = LOCTEXT("CameraMenuTitle_Front", "Front");
 		break;
 
-	case LVT_OrthoNegativeXY:
+	case LVT_OrthoBottom:
 		Label = LOCTEXT("CameraMenuTitle_Bottom", "Bottom");
 		break;
 
-	case LVT_OrthoXZ:
+	case LVT_OrthoRight:
 		Label = LOCTEXT("CameraMenuTitle_Right", "Right");
 		break;
 
-	case LVT_OrthoYZ:
+	case LVT_OrthoBack:
 		Label = LOCTEXT("CameraMenuTitle_Back", "Back");
 		break;
 	case LVT_OrthoFreelook:
@@ -2525,27 +2525,27 @@ FName GetCameraSubmenuIconFNameFromViewportType(const ELevelViewportType Viewpor
 		Icon = PerspectiveIcon;
 		break;
 
-	case LVT_OrthoXY:
+	case LVT_OrthoTop:
 		Icon = TopIcon;
 		break;
 
-	case LVT_OrthoNegativeXZ:
+	case LVT_OrthoLeft:
 		Icon = LeftIcon;
 		break;
 
-	case LVT_OrthoNegativeYZ:
+	case LVT_OrthoFront:
 		Icon = FrontIcon;
 		break;
 
-	case LVT_OrthoNegativeXY:
+	case LVT_OrthoBottom:
 		Icon = BottomIcon;
 		break;
 
-	case LVT_OrthoXZ:
+	case LVT_OrthoRight:
 		Icon = RightIcon;
 		break;
 
-	case LVT_OrthoYZ:
+	case LVT_OrthoBack:
 		Icon = BackIcon;
 		break;
 	case LVT_OrthoFreelook:
