IncrementalCook: Change the API of CookPackageSplitter to take a single FPopulateContext& argument instead of multiple arguments, so that for future API changes we can add on inputs and outputs without needing to write the difficult boilerplate required for deprecation of a virtual function.
[RN] Minor, Cooking
#incrementalcook

diff --git a/Engine/Plugins/Mutable/Source/CustomizableObjectEditor/Private/MuCOE/CustomizableObjectCookPackageSplitter.cpp b/Engine/Plugins/Mutable/Source/CustomizableObjectEditor/Private/MuCOE/CustomizableObjectCookPackageSplitter.cpp
--- a/Engine/Plugins/Mutable/Source/CustomizableObjectEditor/Private/MuCOE/CustomizableObjectCookPackageSplitter.cpp
+++ b/Engine/Plugins/Mutable/Source/CustomizableObjectEditor/Private/MuCOE/CustomizableObjectCookPackageSplitter.cpp
@@ -210,11 +210,7 @@ TArray<ICookPackageSplitter::FGeneratedPackage> FCustomizableObjectCookPackageSp
 	return Result;
 }
 
-bool FCustomizableObjectCookPackageSplitter::PreSaveGeneratorPackage(
-	UPackage* OwnerPackage,
-	UObject* OwnerObject,
-	const TArray<ICookPackageSplitter::FGeneratedPackageForPreSave>& PlaceholderPackages,
-	TArray<UPackage*>& OutKeepReferencedPackages)
+bool FCustomizableObjectCookPackageSplitter::PreSaveGeneratorPackage(FPopulateContext& PopulateContext)
 {
 	// The CO is just about to be saved (i.e. produce the cooked version of the asset), so this
 	// function needs to:
@@ -226,7 +222,8 @@ bool FCustomizableObjectCookPackageSplitter::PreSaveGeneratorPackage(
 	//		the CO is loaded
 	//
 	// 3.	Fix up the soft references to point to the streamed data's new location
-	const auto& PreSavePackage = [] (const ICookPackageSplitter::FGeneratedPackageForPreSave& GeneratedPackage,
+	TConstArrayView<ICookPackageSplitter::FGeneratedPackageForPopulate>& PlaceholderPackages = PopulateContext.GetGeneratedPackages();
+	const auto& PreSavePackage = [] (const ICookPackageSplitter::FGeneratedPackageForPopulate& GeneratedPackage,
 		TArray<FCustomizableObjectStreamedResourceData>& StreamedResources
 		) -> bool
 	{
@@ -254,7 +251,7 @@ bool FCustomizableObjectCookPackageSplitter::PreSaveGeneratorPackage(
 		return true;
 	};
 
-	UCustomizableObject* Object = CastChecked<UCustomizableObject>(OwnerObject);
+	UCustomizableObject* Object = CastChecked<UCustomizableObject>(PopulateContext.GetOwnerObject());
 
 	UModelResources* ModelResources = FindModelResources(*Object);
 	if (!ModelResources)
@@ -278,7 +275,7 @@ bool FCustomizableObjectCookPackageSplitter::PreSaveGeneratorPackage(
 
 	for (int32 Index = 0; Index < NumStreamedData; ++Index)
 	{
-		const ICookPackageSplitter::FGeneratedPackageForPreSave& GeneratedPackage = PlaceholderPackages[Index];
+		const ICookPackageSplitter::FGeneratedPackageForPopulate& GeneratedPackage = PlaceholderPackages[Index];
 		if (!PreSavePackage(GeneratedPackage, ModelResources->StreamedResourceData))
 		{
 			return false;
@@ -289,7 +286,7 @@ bool FCustomizableObjectCookPackageSplitter::PreSaveGeneratorPackage(
 
 	for (int32 Index = 0; Index < NumStreamedExtensionData; ++Index)
 	{
-		const ICookPackageSplitter::FGeneratedPackageForPreSave& GeneratedPackage = PlaceholderPackages[NumStreamedData + Index];
+		const ICookPackageSplitter::FGeneratedPackageForPopulate& GeneratedPackage = PlaceholderPackages[NumStreamedData + Index];
 		if (!PreSavePackage(GeneratedPackage, ModelResources->StreamedExtensionData))
 		{
 			return false;
@@ -301,12 +298,12 @@ bool FCustomizableObjectCookPackageSplitter::PreSaveGeneratorPackage(
 	return true;
 }
 
-void FCustomizableObjectCookPackageSplitter::PostSaveGeneratorPackage(UPackage* OwnerPackage, UObject* OwnerObject)
+void FCustomizableObjectCookPackageSplitter::PostSaveGeneratorPackage(FPopulateContext& PopulateContext)
 {
 	// Move the streamed data back into the CO's package and restore the StreamedResourceData and StreamedExtensionData
 	// array on the CO to how it was before PreSaveGeneratorPackage.
 
-	UCustomizableObject* Object = CastChecked<UCustomizableObject>(OwnerObject);
+	UCustomizableObject* Object = CastChecked<UCustomizableObject>(PopulateContext.GetOwnerObject());
 	UModelResources* ModelResources = FindModelResources(*Object);
 	if (!ModelResources)
 	{
@@ -329,8 +326,8 @@ void FCustomizableObjectCookPackageSplitter::PostSaveGeneratorPackage(UPackage*
 		}
 
 		UCustomizableObjectResourceDataContainer* Container = nullptr;
-		EMoveContainerError Error = MoveContainerToNewOuter(OwnerObject, ResourceData, Container);
-		UE_CLOG(Error != EMoveContainerError::None, LogMutable, Warning, TEXT("Failed to move container %s back to %s - %s"), *ContainerName, *GetPathNameSafe(OwnerObject), LexToString(Error));
+		EMoveContainerError Error = MoveContainerToNewOuter(Object, ResourceData, Container);
+		UE_CLOG(Error != EMoveContainerError::None, LogMutable, Warning, TEXT("Failed to move container %s back to %s - %s"), *ContainerName, *GetPathNameSafe(Object), LexToString(Error));
 
 		NewArray.Emplace(Container);
 	}
@@ -351,8 +348,8 @@ void FCustomizableObjectCookPackageSplitter::PostSaveGeneratorPackage(UPackage*
 		}
 
 		UCustomizableObjectResourceDataContainer* Container = nullptr;
-		EMoveContainerError Error = MoveContainerToNewOuter(OwnerObject, ResourceData, Container);
-		UE_CLOG(Error != EMoveContainerError::None, LogMutable, Warning, TEXT("Failed to move container %s back to %s - %s"), *ContainerName, *GetPathNameSafe(OwnerObject), LexToString(Error));
+		EMoveContainerError Error = MoveContainerToNewOuter(Object, ResourceData, Container);
+		UE_CLOG(Error != EMoveContainerError::None, LogMutable, Warning, TEXT("Failed to move container %s back to %s - %s"), *ContainerName, *GetPathNameSafe(Object), LexToString(Error));
 
 		NewArray.Emplace(Container);
 	}
@@ -360,16 +357,12 @@ void FCustomizableObjectCookPackageSplitter::PostSaveGeneratorPackage(UPackage*
 	ModelResources->StreamedExtensionData = NewArray;
 }
 
-bool FCustomizableObjectCookPackageSplitter::PopulateGeneratedPackage(
-	UPackage* OwnerPackage,
-	UObject* OwnerObject,
-	const FGeneratedPackageForPopulate& GeneratedPackage,
-	TArray<UObject*>& OutObjectsToMove,
-	TArray<UPackage*>& OutKeepReferencedPackages)
+bool FCustomizableObjectCookPackageSplitter::PopulateGeneratedPackage(FPopulateContext& PopulateContext)
 {
 	// Move the container into its newly generated package
 
-	UCustomizableObject* Object = CastChecked<UCustomizableObject>(OwnerObject);
+	const FGeneratedPackageForPopulate& GeneratedPackage = *PopulateContext.GetTargetGeneratedPackage();
+	UCustomizableObject* Object = CastChecked<UCustomizableObject>(PopulateContext.GetOwnerObject());
 	UModelResources* ModelResources = FindModelResources(*Object);
 
 	FCustomizableObjectStreamedResourceData* ResourceData = FindStreamedResourceData(ModelResources->StreamedResourceData, GeneratedPackage.RelativePath);
@@ -404,20 +397,18 @@ bool FCustomizableObjectCookPackageSplitter::PopulateGeneratedPackage(
 		return false;
 	}
 
-	OutObjectsToMove.Add(Container);
+	PopulateContext.ReportObjectToMove(Container);
 
 	return true;
 }
 
-void FCustomizableObjectCookPackageSplitter::PostSaveGeneratedPackage(
-	UPackage* OwnerPackage,
-	UObject* OwnerObject,
-	const FGeneratedPackageForPopulate& GeneratedPackage)
+void FCustomizableObjectCookPackageSplitter::PostSaveGeneratedPackage(FPopulateContext& PopulateContext)
 {
 	// Now that the generated package has been saved/cooked, move the container back to the CO, so
 	// that everything is the same as it was before cooking.
 
-	UCustomizableObject* Object = CastChecked<UCustomizableObject>(OwnerObject);
+	const FGeneratedPackageForPopulate& GeneratedPackage = *PopulateContext.GetTargetGeneratedPackage();
+	UCustomizableObject* Object = CastChecked<UCustomizableObject>(PopulateContext.GetOwnerObject());
 	UModelResources* ModelResources = FindModelResources(*Object);
 
 	FCustomizableObjectStreamedResourceData* ResourceData = FindStreamedResourceData(ModelResources->StreamedResourceData, GeneratedPackage.RelativePath);
@@ -435,9 +426,9 @@ void FCustomizableObjectCookPackageSplitter::PostSaveGeneratedPackage(
 	}
 
 	UCustomizableObjectResourceDataContainer* Container = nullptr;
-	EMoveContainerError Error = MoveContainerToNewOuter(OwnerObject, ResourceData, Container);
+	EMoveContainerError Error = MoveContainerToNewOuter(Object, ResourceData, Container);
 	UE_CLOG(Error != EMoveContainerError::None, LogMutable, Warning, 
-		TEXT("Failed to move container %s back to %s - %s"), *ResourceData->GetPath().ToSoftObjectPath().ToString(), *GetPathNameSafe(OwnerObject), LexToString(Error));
+		TEXT("Failed to move container %s back to %s - %s"), *ResourceData->GetPath().ToSoftObjectPath().ToString(), *GetPathNameSafe(Object), LexToString(Error));
 }
 
 void FCustomizableObjectCookPackageSplitter::Teardown(ETeardown Status)
diff --git a/Engine/Plugins/Mutable/Source/CustomizableObjectEditor/Private/MuCOE/CustomizableObjectCookPackageSplitter.h b/Engine/Plugins/Mutable/Source/CustomizableObjectEditor/Private/MuCOE/CustomizableObjectCookPackageSplitter.h
--- a/Engine/Plugins/Mutable/Source/CustomizableObjectEditor/Private/MuCOE/CustomizableObjectCookPackageSplitter.h
+++ b/Engine/Plugins/Mutable/Source/CustomizableObjectEditor/Private/MuCOE/CustomizableObjectCookPackageSplitter.h
@@ -19,25 +19,13 @@ public:
 		const UPackage* OwnerPackage,
 		const UObject* OwnerObject) override;
 	
-	virtual bool PreSaveGeneratorPackage(
-		UPackage* OwnerPackage,
-		UObject* OwnerObject,
-		const TArray<ICookPackageSplitter::FGeneratedPackageForPreSave>& PlaceholderPackages,
-		TArray<UPackage*>& OutKeepReferencedPackages) override;
+	virtual bool PreSaveGeneratorPackage(FPopulateContext& PopulateContext) override;
 
-	virtual void PostSaveGeneratorPackage(UPackage* OwnerPackage, UObject* OwnerObject) override;
+	virtual void PostSaveGeneratorPackage(FPopulateContext& PopulateContext) override;
 
-	virtual bool PopulateGeneratedPackage(
-		UPackage* OwnerPackage,
-		UObject* OwnerObject,
-		const FGeneratedPackageForPopulate& GeneratedPackage,
-		TArray<UObject*>& OutObjectsToMove,
-		TArray<UPackage*>& OutKeepReferencedPackages) override;
+	virtual bool PopulateGeneratedPackage(FPopulateContext& PopulateContext) override;
 
-	virtual void PostSaveGeneratedPackage(
-		UPackage* OwnerPackage,
-		UObject* OwnerObject,
-		const FGeneratedPackageForPopulate& GeneratedPackage) override;
+	virtual void PostSaveGeneratedPackage(FPopulateContext& PopulateContext) override;
 
 	/** Do teardown actions after all packages have saved, or when the cook is cancelled. Always called before destruction. */
 	virtual void Teardown(ETeardown Status) override;
diff --git a/Engine/Source/Editor/UnrealEd/Classes/CookOnTheSide/CookOnTheFlyServer.h b/Engine/Source/Editor/UnrealEd/Classes/CookOnTheSide/CookOnTheFlyServer.h
--- a/Engine/Source/Editor/UnrealEd/Classes/CookOnTheSide/CookOnTheFlyServer.h
+++ b/Engine/Source/Editor/UnrealEd/Classes/CookOnTheSide/CookOnTheFlyServer.h
@@ -1427,15 +1427,15 @@ private:
 	 */
 	UE::Cook::EPollStatus BeginCacheObjectsToMove(UE::Cook::FGenerationHelper& GenerationHelper,
 		UE::Cook::FCookGenerationInfo& Info, UE::Cook::FCookerTimer& Timer,
-		TArray<ICookPackageSplitter::FGeneratedPackageForPreSave>& GeneratedPackagesForPresave);
+		TArray<ICookPackageSplitter::FGeneratedPackageForPopulate>& GeneratedPackagesForPopulate);
 	/** Call the CookPackageSplitter's PreSaveGeneratorPackage to create/move objects into its main UPackage. */
 	UE::Cook::EPollStatus PreSaveGeneratorPackage(UE::Cook::FPackageData& PackageData,
 		UE::Cook::FGenerationHelper& GenerationHelper, UE::Cook::FCookGenerationInfo& Info,
-		TArray<ICookPackageSplitter::FGeneratedPackageForPreSave>& GeneratedPackagesForPresave);
+		TArray<ICookPackageSplitter::FGeneratedPackageForPopulate>& GeneratedPackagesForPopulate);
 	/** Construct the list of generated packages that is required for some CookPackageSplitter interface calls. */
-	bool TryConstructGeneratedPackagesForPresave(UE::Cook::FPackageData& PackageData,
+	bool TryConstructGeneratedPackagesForPopulate(UE::Cook::FPackageData& PackageData,
 		UE::Cook::FGenerationHelper& GenerationHelper,
-		TArray<ICookPackageSplitter::FGeneratedPackageForPreSave>& GeneratedPackagesForPresave);
+		TArray<ICookPackageSplitter::FGeneratedPackageForPopulate>& GeneratedPackagesForPopulate);
 	/**
 	 * Call BeginCacheForCookedPlatformData on any undeclared objects in the CookPackageSplitter's main UPackage
 	 * after the move.
diff --git a/Engine/Source/Editor/UnrealEd/Private/CookOnTheFlyServer.cpp b/Engine/Source/Editor/UnrealEd/Private/CookOnTheFlyServer.cpp
--- a/Engine/Source/Editor/UnrealEd/Private/CookOnTheFlyServer.cpp
+++ b/Engine/Source/Editor/UnrealEd/Private/CookOnTheFlyServer.cpp
@@ -3603,13 +3603,13 @@ UE::Cook::EPollStatus UCookOnTheFlyServer::PrepareSaveGenerationPackage(UE::Cook
 		PackageData.SetSaveSubStateComplete(ESaveSubState::Generation_PreMoveCookedPlatformData_WaitingForIsLoaded);
 	}
 
-	// GeneratedPackagesForPresave is used by multiple steps, recreate it when needed each time we come in to this function
-	TArray<ICookPackageSplitter::FGeneratedPackageForPreSave> GeneratedPackagesForPresave;
+	// GeneratedPackagesForPopulate is used by multiple steps, recreate it when needed each time we come in to this function
+	TArray<ICookPackageSplitter::FGeneratedPackageForPopulate> GeneratedPackagesForPopulate;
 	if (PackageData.GetSaveSubState() <= ESaveSubState::Generation_FinishCacheObjectsToMove)
 	{
 		if (PackageData.GetSaveSubState() <= ESaveSubState::Generation_BeginCacheObjectsToMove)
 		{
-			EPollStatus Result = BeginCacheObjectsToMove(GenerationHelper, Info, Timer, GeneratedPackagesForPresave);
+			EPollStatus Result = BeginCacheObjectsToMove(GenerationHelper, Info, Timer, GeneratedPackagesForPopulate);
 			if (Result != EPollStatus::Success)
 			{
 				return Result;
@@ -3650,7 +3650,7 @@ UE::Cook::EPollStatus UCookOnTheFlyServer::PrepareSaveGenerationPackage(UE::Cook
 		EPollStatus Result;
 		if (Info.IsGenerator())
 		{
-			Result = PreSaveGeneratorPackage(PackageData, GenerationHelper, Info, GeneratedPackagesForPresave);
+			Result = PreSaveGeneratorPackage(PackageData, GenerationHelper, Info, GeneratedPackagesForPopulate);
 		}
 		else
 		{
@@ -3702,7 +3702,7 @@ UE::Cook::EPollStatus UCookOnTheFlyServer::PrepareSaveGenerationPackage(UE::Cook
 
 UE::Cook::EPollStatus UCookOnTheFlyServer::BeginCacheObjectsToMove(UE::Cook::FGenerationHelper& GenerationHelper,
 	UE::Cook::FCookGenerationInfo& Info, UE::Cook::FCookerTimer& Timer,
-	TArray<ICookPackageSplitter::FGeneratedPackageForPreSave>& GeneratedPackagesForPresave)
+	TArray<ICookPackageSplitter::FGeneratedPackageForPopulate>& GeneratedPackagesForPopulate)
 {
 	using namespace UE::Cook;
 
@@ -3722,7 +3722,7 @@ UE::Cook::EPollStatus UCookOnTheFlyServer::BeginCacheObjectsToMove(UE::Cook::FGe
 		if (Info.IsGenerator() || GenerationHelper.DoesGeneratedRequireGenerator()
 			>= ICookPackageSplitter::EGeneratedRequiresGenerator::Populate)
 		{
-			if (!GenerationHelper.TryCallPopulateGeneratorPackage(GeneratedPackagesForPresave))
+			if (!GenerationHelper.TryCallPopulateGeneratorPackage(GeneratedPackagesForPopulate))
 			{
 				return EPollStatus::Error;
 			}
@@ -3763,7 +3763,7 @@ UE::Cook::EPollStatus UCookOnTheFlyServer::BeginCacheObjectsToMove(UE::Cook::FGe
 
 UE::Cook::EPollStatus UCookOnTheFlyServer::PreSaveGeneratorPackage(UE::Cook::FPackageData& PackageData,
 	UE::Cook::FGenerationHelper& GenerationHelper, UE::Cook::FCookGenerationInfo& Info,
-	TArray<ICookPackageSplitter::FGeneratedPackageForPreSave>& GeneratedPackagesForPresave)
+	TArray<ICookPackageSplitter::FGeneratedPackageForPopulate>& GeneratedPackagesForPopulate)
 {
 	using namespace UE::Cook;
 
@@ -3779,9 +3779,9 @@ UE::Cook::EPollStatus UCookOnTheFlyServer::PreSaveGeneratorPackage(UE::Cook::FPa
 	}
 
 	TArray<UPackage*> KeepReferencedPackages;
-	if (!TryConstructGeneratedPackagesForPresave(PackageData, GenerationHelper, GeneratedPackagesForPresave))
+	if (!TryConstructGeneratedPackagesForPopulate(PackageData, GenerationHelper, GeneratedPackagesForPopulate))
 	{
-		UE_LOG(LogCook, Error, TEXT("PackageSplitter unexpected failure: could not ConstructGeneratedPackagesForPreSave. Splitter=%s"),
+		UE_LOG(LogCook, Error, TEXT("PackageSplitter unexpected failure: could not Construct GeneratedPackagesForPopulate. Splitter=%s"),
 			*GenerationHelper.GetSplitDataObjectName().ToString());
 		return EPollStatus::Error;
 	}
@@ -3794,7 +3794,19 @@ UE::Cook::EPollStatus UCookOnTheFlyServer::PreSaveGeneratorPackage(UE::Cook::FPa
 			FName()
 #endif
 		);
-		if (!Splitter->PreSaveGeneratorPackage(Package, SplitDataObject, GeneratedPackagesForPresave, KeepReferencedPackages))
+		UE::Cook::CookPackageSplitter::FPopulateContextData PopulateData;
+		ICookPackageSplitter::FPopulateContext PopulateContext(PopulateData);
+		PopulateData.OwnerPackage = Package;
+		PopulateData.OwnerObject = SplitDataObject;
+		PopulateData.GeneratedPackages = GeneratedPackagesForPopulate;
+
+		PRAGMA_DISABLE_DEPRECATION_WARNINGS;
+		bool bPreSaveSucceeded = Splitter->PreSaveGeneratorPackage(PopulateData.OwnerPackage, PopulateData.OwnerObject,
+			GeneratedPackagesForPopulate, PopulateData.KeepReferencedPackages);
+		Splitter->WarnIfDeprecatedVirtualNotCalled(TEXT("PreSaveGeneratorPackage"));
+		PRAGMA_ENABLE_DEPRECATION_WARNINGS;
+		bPreSaveSucceeded = Splitter->PreSaveGeneratorPackage(PopulateContext) && bPreSaveSucceeded;
+		if (!bPreSaveSucceeded)
 		{
 			UE_LOG(LogCook, Error, TEXT("PackageSplitter returned false from PreSaveGeneratorPackage. Splitter=%s"),
 				*GenerationHelper.GetSplitDataObjectName().ToString());
@@ -3806,12 +3818,12 @@ UE::Cook::EPollStatus UCookOnTheFlyServer::PreSaveGeneratorPackage(UE::Cook::FPa
 	return EPollStatus::Success;
 }
 
-bool UCookOnTheFlyServer::TryConstructGeneratedPackagesForPresave(UE::Cook::FPackageData& PackageData, UE::Cook::FGenerationHelper& GenerationHelper,
-	TArray<ICookPackageSplitter::FGeneratedPackageForPreSave>& GeneratedPackagesForPresave)
+bool UCookOnTheFlyServer::TryConstructGeneratedPackagesForPopulate(UE::Cook::FPackageData& PackageData, UE::Cook::FGenerationHelper& GenerationHelper,
+	TArray<ICookPackageSplitter::FGeneratedPackageForPopulate>& GeneratedPackagesForPopulate)
 {
 	using namespace UE::Cook;
 
-	if (GeneratedPackagesForPresave.Num() > 0)
+	if (GeneratedPackagesForPopulate.Num() > 0)
 	{
 		// Already constructed, save time by early exiting
 		return true;
@@ -3820,11 +3832,10 @@ bool UCookOnTheFlyServer::TryConstructGeneratedPackagesForPresave(UE::Cook::FPac
 	// We need to find or (create empty stub packages for) each of the PackagesToGenerate so that PreSaveGeneratorPackage
 	// can refer to them to create hardlinks in the cooked Generator package
 	TArrayView<FCookGenerationInfo> PackagesToGenerate = GenerationHelper.GetPackagesToGenerate();
-	TArray<ICookPackageSplitter::FGeneratedPackageForPreSave> SplitterDatas;
-	SplitterDatas.Reserve(PackagesToGenerate.Num());
+	GeneratedPackagesForPopulate.Reserve(PackagesToGenerate.Num());
 	for (FCookGenerationInfo& Info : PackagesToGenerate)
 	{
-		ICookPackageSplitter::FGeneratedPackageForPreSave& SplitterData = GeneratedPackagesForPresave.Emplace_GetRef();
+		ICookPackageSplitter::FGeneratedPackageForPopulate& SplitterData = GeneratedPackagesForPopulate.Emplace_GetRef();
 		SplitterData.RelativePath = Info.RelativePath;
 		SplitterData.GeneratedRootPath = Info.GeneratedRootPath;
 		SplitterData.bCreatedAsMap = Info.IsCreateAsMap();
@@ -3936,11 +3947,11 @@ UE::Cook::EPollStatus UCookOnTheFlyServer::TryPopulateGeneratedPackage(UE::Cook:
 
 	// Populate package using CookPackageSplitterInstance and pass GeneratedPackage's cooked name for it to
 	// properly setup any internal reference to this package (SoftObjectPaths or others)
-	ICookPackageSplitter::FGeneratedPackageForPopulate PopulateData;
-	PopulateData.RelativePath = GeneratedInfo.RelativePath;
-	PopulateData.GeneratedRootPath = GeneratedInfo.GeneratedRootPath;
-	PopulateData.Package = GeneratedPackage;
-	PopulateData.bCreatedAsMap = GeneratedInfo.IsCreateAsMap();
+	ICookPackageSplitter::FGeneratedPackageForPopulate GeneratedPackagePopulateData;
+	GeneratedPackagePopulateData.RelativePath = GeneratedInfo.RelativePath;
+	GeneratedPackagePopulateData.GeneratedRootPath = GeneratedInfo.GeneratedRootPath;
+	GeneratedPackagePopulateData.Package = GeneratedPackage;
+	GeneratedPackagePopulateData.bCreatedAsMap = GeneratedInfo.IsCreateAsMap();
 	TArray<UPackage*> KeepReferencedPackages;
 	{
 		FScopedActivePackage ScopedActivePackage(*this, GenerationHelper.GetOwner().GetPackageName(),
@@ -3950,7 +3961,19 @@ UE::Cook::EPollStatus UCookOnTheFlyServer::TryPopulateGeneratedPackage(UE::Cook:
 			FName()
 #endif
 		);
-		if (!Splitter->PreSaveGeneratedPackage(OwnerPackage, OwnerObject, PopulateData, KeepReferencedPackages))
+		UE::Cook::CookPackageSplitter::FPopulateContextData PopulateData;
+		ICookPackageSplitter::FPopulateContext PopulateContext(PopulateData);
+		PopulateData.OwnerPackage = OwnerPackage;
+		PopulateData.OwnerObject = OwnerObject;
+		PopulateData.TargetGeneratedPackage = &GeneratedPackagePopulateData;
+
+		PRAGMA_DISABLE_DEPRECATION_WARNINGS;
+		bool bPreSaveSucceeded = Splitter->PreSaveGeneratedPackage(PopulateData.OwnerPackage, PopulateData.OwnerObject,
+			*PopulateData.TargetGeneratedPackage, PopulateData.KeepReferencedPackages);
+		Splitter->WarnIfDeprecatedVirtualNotCalled(TEXT("PreSaveGeneratedPackage"));
+		PRAGMA_ENABLE_DEPRECATION_WARNINGS;
+		bPreSaveSucceeded = Splitter->PreSaveGeneratedPackage(PopulateContext) && bPreSaveSucceeded;
+		if (!bPreSaveSucceeded)
 		{
 			UE_LOG(LogCook, Error, TEXT("PackageSplitter returned false from PreSaveGeneratedPackage. Splitter=%s, Generated=%s."),
 				*GenerationHelper.GetSplitDataObjectName().ToString(), *GeneratedPackageName);
diff --git a/Engine/Source/Editor/UnrealEd/Private/Cooker/CookGenerationHelper.cpp b/Engine/Source/Editor/UnrealEd/Private/Cooker/CookGenerationHelper.cpp
--- a/Engine/Source/Editor/UnrealEd/Private/Cooker/CookGenerationHelper.cpp
+++ b/Engine/Source/Editor/UnrealEd/Private/Cooker/CookGenerationHelper.cpp
@@ -20,6 +20,130 @@
 #include "TargetDomain/TargetDomainUtils.h"
 #include "UObject/ReferenceChainSearch.h"
 
+//////////////////////////////////////////////////////////////////////////
+// FPopulateContext
+
+ICookPackageSplitter::FPopulateContext::FPopulateContext(UE::Cook::CookPackageSplitter::FPopulateContextData& InData)
+	: Data(InData)
+{
+}
+
+UPackage* ICookPackageSplitter::FPopulateContext::GetOwnerPackage() const
+{
+	return Data.OwnerPackage;
+}
+
+UObject* ICookPackageSplitter::FPopulateContext::GetOwnerObject() const
+{
+	return Data.OwnerObject;
+}
+
+TConstArrayView<ICookPackageSplitter::FGeneratedPackageForPopulate>& ICookPackageSplitter::FPopulateContext::GetGeneratedPackages()
+{
+	return Data.GeneratedPackages;
+}
+
+bool ICookPackageSplitter::FPopulateContext::IsCalledOnGenerator() const
+{
+	return Data.TargetGeneratedPackage == nullptr;
+}
+
+UPackage* ICookPackageSplitter::FPopulateContext::GetTargetPackage() const
+{
+	return Data.TargetGeneratedPackage ? Data.TargetGeneratedPackage->Package : Data.OwnerPackage;
+}
+
+const ICookPackageSplitter::FGeneratedPackageForPopulate* ICookPackageSplitter::FPopulateContext::GetTargetGeneratedPackage() const
+{
+	return Data.TargetGeneratedPackage;
+}
+
+void ICookPackageSplitter::FPopulateContext::ReportObjectToMove(UObject* Object)
+{
+	Data.ObjectsToMove.Add(Object);
+}
+
+void ICookPackageSplitter::FPopulateContext::ReportObjectsToMove(TConstArrayView<UObject*> Objects)
+{
+	Data.ObjectsToMove.Append(Objects);
+}
+
+void ICookPackageSplitter::FPopulateContext::ReportKeepReferencedPackage(UPackage* Package)
+{
+	Data.KeepReferencedPackages.Add(Package);
+}
+
+void ICookPackageSplitter::FPopulateContext::ReportKeepReferencedPackages(TConstArrayView<UPackage*> Packages)
+{
+	Data.KeepReferencedPackages.Append(Packages);
+}
+
+bool ICookPackageSplitter::PopulateGeneratorPackage(UPackage* OwnerPackage, UObject* OwnerObject,
+	const TArray<ICookPackageSplitter::FGeneratedPackageForPopulate>& GeneratedPackages, TArray<UObject*>& OutObjectsToMove,
+	TArray<UPackage*>& OutKeepReferencedPackages)
+{
+	PRAGMA_DISABLE_DEPRECATION_WARNINGS;
+	bDeprecatedVirtualCalledAsExpected = true;
+	PRAGMA_ENABLE_DEPRECATION_WARNINGS;
+	return true;
+}
+
+bool ICookPackageSplitter::PreSaveGeneratorPackage(UPackage* OwnerPackage, UObject* OwnerObject,
+	const TArray<FGeneratedPackageForPopulate>& PlaceholderPackages, TArray<UPackage*>& OutKeepReferencedPackages)
+{
+	PRAGMA_DISABLE_DEPRECATION_WARNINGS;
+	bDeprecatedVirtualCalledAsExpected = true;
+	PRAGMA_ENABLE_DEPRECATION_WARNINGS;
+	return true;
+}
+
+void ICookPackageSplitter::PostSaveGeneratorPackage(UPackage* OwnerPackage, UObject* OwnerObject)
+{
+	PRAGMA_DISABLE_DEPRECATION_WARNINGS;
+	bDeprecatedVirtualCalledAsExpected = true;
+	PRAGMA_ENABLE_DEPRECATION_WARNINGS;
+}
+
+bool ICookPackageSplitter::PopulateGeneratedPackage(UPackage* OwnerPackage, UObject* OwnerObject,
+	const FGeneratedPackageForPopulate& GeneratedPackage, TArray<UObject*>& OutObjectsToMove,
+	TArray<UPackage*>& OutKeepReferencedPackages)
+{
+	PRAGMA_DISABLE_DEPRECATION_WARNINGS;
+	bDeprecatedVirtualCalledAsExpected = true;
+	PRAGMA_ENABLE_DEPRECATION_WARNINGS;
+	return true;
+}
+
+bool ICookPackageSplitter::PreSaveGeneratedPackage(UPackage* OwnerPackage, UObject* OwnerObject,
+	const FGeneratedPackageForPopulate& GeneratedPackage, TArray<UPackage*>& OutKeepReferencedPackages)
+{
+	PRAGMA_DISABLE_DEPRECATION_WARNINGS;
+	bDeprecatedVirtualCalledAsExpected = true;
+	PRAGMA_ENABLE_DEPRECATION_WARNINGS;
+	return true;
+};
+
+void ICookPackageSplitter::PostSaveGeneratedPackage(UPackage* OwnerPackage, UObject* OwnerObject,
+	const FGeneratedPackageForPopulate& GeneratedPackage)
+{
+	PRAGMA_DISABLE_DEPRECATION_WARNINGS;
+	bDeprecatedVirtualCalledAsExpected = true;
+	PRAGMA_ENABLE_DEPRECATION_WARNINGS;
+}
+
+void ICookPackageSplitter::WarnIfDeprecatedVirtualNotCalled(const TCHAR* FunctionName)
+{
+	PRAGMA_DISABLE_DEPRECATION_WARNINGS;
+	if (!bDeprecatedVirtualCalledAsExpected)
+	{
+		UE_LOG(LogCook, Warning,
+			TEXT("ICookPackageSplitter::%s has been deprecated. Implement version that takes FPopulateContext instead."),
+			FunctionName);
+	}
+	bDeprecatedVirtualCalledAsExpected = false;
+	PRAGMA_ENABLE_DEPRECATION_WARNINGS;
+}
+
 namespace UE::Cook
 {
 
@@ -901,7 +1025,7 @@ bool FGenerationHelper::TryGenerateList()
 }
 
 bool FGenerationHelper::TryCallPopulateGeneratorPackage(
-	TArray<ICookPackageSplitter::FGeneratedPackageForPreSave>& InOutGeneratedPackagesForPresave)
+	TArray<ICookPackageSplitter::FGeneratedPackageForPopulate>& InOutGeneratedPackagesForPopulate)
 {
 	if (OwnerInfo.HasCalledPopulate())
 	{
@@ -928,7 +1052,7 @@ bool FGenerationHelper::TryCallPopulateGeneratorPackage(
 		return false;
 	}
 	UPackage* LocalOwnerPackage = LocalSplitDataObject->GetPackage();
-	if (!COTFS.TryConstructGeneratedPackagesForPresave(OwnerPackageData, *this, InOutGeneratedPackagesForPresave))
+	if (!COTFS.TryConstructGeneratedPackagesForPopulate(OwnerPackageData, *this, InOutGeneratedPackagesForPopulate))
 	{
 		UE_LOG(LogCook, Error,
 			TEXT("PackageSplitter unexpected failure: could not ConstructGeneratedPackagesForPreSave. Splitter=%s"),
@@ -942,19 +1066,29 @@ bool FGenerationHelper::TryCallPopulateGeneratorPackage(
 		FName()
 #endif
 	);
-	TArray<UPackage*> KeepReferencedPackages;
-	TArray<UObject*> ObjectsToMove;
-	bool bPopulateSucceeded = CookPackageSplitterInstance->PopulateGeneratorPackage(LocalOwnerPackage,
-		LocalSplitDataObject, InOutGeneratedPackagesForPresave, ObjectsToMove, KeepReferencedPackages);
+	UE::Cook::CookPackageSplitter::FPopulateContextData PopulateData;
+	ICookPackageSplitter::FPopulateContext PopulateContext(PopulateData);
+	PopulateData.OwnerPackage = LocalOwnerPackage;
+	PopulateData.OwnerObject = LocalSplitDataObject;
+	PopulateData.GeneratedPackages = InOutGeneratedPackagesForPopulate;
+
+	PRAGMA_DISABLE_DEPRECATION_WARNINGS;
+	bool bPopulateSucceeded = CookPackageSplitterInstance->PopulateGeneratorPackage(PopulateData.OwnerPackage,
+		PopulateData.OwnerObject, InOutGeneratedPackagesForPopulate, PopulateData.ObjectsToMove,
+		PopulateData.KeepReferencedPackages);
+	CookPackageSplitterInstance->WarnIfDeprecatedVirtualNotCalled(TEXT("PopulateGeneratorPackage"));
+	PRAGMA_ENABLE_DEPRECATION_WARNINGS;
+	bPopulateSucceeded = CookPackageSplitterInstance->PopulateGeneratorPackage(PopulateContext) && bPopulateSucceeded;
+
 	if (!bPopulateSucceeded)
 	{
 		UE_LOG(LogCook, Error, TEXT("CookPackageSplitter returned false from PopulateGeneratorPackage. Splitter=%s"),
 			*GetSplitDataObjectName().ToString());
 		return false;
 	}
-	OwnerInfo.AddKeepReferencedPackages(*this, KeepReferencedPackages);
-	OwnerObjectsToMove.Reserve(ObjectsToMove.Num());
-	for (UObject* Object : ObjectsToMove)
+	OwnerInfo.AddKeepReferencedPackages(*this, PopulateData.KeepReferencedPackages);
+	OwnerObjectsToMove.Reserve(PopulateData.ObjectsToMove.Num());
+	for (UObject* Object : PopulateData.ObjectsToMove)
 	{
 		if (Object)
 		{
@@ -1010,9 +1144,23 @@ bool FGenerationHelper::TryCallPopulateGeneratedPackage(UE::Cook::FCookGeneratio
 		FName()
 #endif
 	);
-	TArray<UPackage*> KeepReferencedPackages;
-	bool bPopulateSucceeded = CookPackageSplitterInstance->PopulateGeneratedPackage(Package,
-		LocalSplitDataObject, SplitterInfo, OutObjectsToMove, KeepReferencedPackages);
+	UE::Cook::CookPackageSplitter::FPopulateContextData PopulateData;
+	ICookPackageSplitter::FPopulateContext PopulateContext(PopulateData);
+	PopulateData.OwnerPackage = GetOwnerPackage();
+	PopulateData.OwnerObject = LocalSplitDataObject;
+	PopulateData.TargetGeneratedPackage = &SplitterInfo;
+
+	PRAGMA_DISABLE_DEPRECATION_WARNINGS;
+	bool bPopulateSucceeded = CookPackageSplitterInstance->PopulateGeneratedPackage(
+		// Bug in the deprecated API: First argument was supposed to be OwnerPackage, but we were passing in the target generated package
+		// Keep passing in the target generated package in the deprecated version. This is fixed in the new API function.
+		PopulateData.TargetGeneratedPackage->Package,
+		PopulateData.OwnerObject, *PopulateData.TargetGeneratedPackage, PopulateData.ObjectsToMove,
+		PopulateData.KeepReferencedPackages);
+	CookPackageSplitterInstance->WarnIfDeprecatedVirtualNotCalled(TEXT("PopulateGeneratedPackage"));
+	PRAGMA_ENABLE_DEPRECATION_WARNINGS;
+	bPopulateSucceeded = CookPackageSplitterInstance->PopulateGeneratedPackage(PopulateContext) && bPopulateSucceeded;
+
 	if (!bPopulateSucceeded)
 	{
 		UE_LOG(LogCook, Error,
@@ -1022,7 +1170,8 @@ bool FGenerationHelper::TryCallPopulateGeneratedPackage(UE::Cook::FCookGeneratio
 		return false;
 	}
 
-	Info.AddKeepReferencedPackages(*this, KeepReferencedPackages);
+	OutObjectsToMove.Append(PopulateData.ObjectsToMove);
+	Info.AddKeepReferencedPackages(*this, PopulateData.KeepReferencedPackages);
 
 	// Contract Point 1: We will not call populate again until the splitter has been destroyed
 	// Contract Point 2: We will not call populate again without garbage collecting the generator package
@@ -1410,9 +1559,19 @@ void FGenerationHelper::ResetSaveState(FCookGenerationInfo& Info, UPackage* Pack
 				FName()
 #endif
 			);
+
+			UE::Cook::CookPackageSplitter::FPopulateContextData PopulateData;
+			ICookPackageSplitter::FPopulateContext PopulateContext(PopulateData);
+			PopulateData.OwnerPackage = LocalOwnerPackage;
+			PopulateData.OwnerObject = SplitObject;
+			ICookPackageSplitter* LocalSplitterInstance = GetCookPackageSplitterInstance();
 			if (Info.IsGenerator())
 			{
-				GetCookPackageSplitterInstance()->PostSaveGeneratorPackage(Package, SplitObject);
+				PRAGMA_DISABLE_DEPRECATION_WARNINGS;
+				LocalSplitterInstance->PostSaveGeneratorPackage(PopulateData.OwnerPackage, PopulateData.OwnerObject);
+				LocalSplitterInstance->WarnIfDeprecatedVirtualNotCalled(TEXT("PostSaveGeneratorPackage"));
+				PRAGMA_ENABLE_DEPRECATION_WARNINGS;
+				LocalSplitterInstance->PostSaveGeneratorPackage(PopulateContext);
 			}
 			else
 			{
@@ -1421,8 +1580,13 @@ void FGenerationHelper::ResetSaveState(FCookGenerationInfo& Info, UPackage* Pack
 				PopulateInfo.GeneratedRootPath = Info.GeneratedRootPath;
 				PopulateInfo.bCreatedAsMap = Info.IsCreateAsMap();
 				PopulateInfo.Package = Package;
-				GetCookPackageSplitterInstance()->PostSaveGeneratedPackage(LocalOwnerPackage, SplitObject,
-					PopulateInfo);
+				PopulateData.TargetGeneratedPackage = &PopulateInfo;
+				PRAGMA_DISABLE_DEPRECATION_WARNINGS;
+				LocalSplitterInstance->PostSaveGeneratedPackage(PopulateData.OwnerPackage, PopulateData.OwnerObject,
+					*PopulateData.TargetGeneratedPackage);
+				LocalSplitterInstance->WarnIfDeprecatedVirtualNotCalled(TEXT("PostSaveGeneratedPackage"));
+				PRAGMA_ENABLE_DEPRECATION_WARNINGS;
+				LocalSplitterInstance->PostSaveGeneratedPackage(PopulateContext);
 			}
 		}
 	}
diff --git a/Engine/Source/Editor/UnrealEd/Private/Cooker/CookGenerationHelper.h b/Engine/Source/Editor/UnrealEd/Private/Cooker/CookGenerationHelper.h
--- a/Engine/Source/Editor/UnrealEd/Private/Cooker/CookGenerationHelper.h
+++ b/Engine/Source/Editor/UnrealEd/Private/Cooker/CookGenerationHelper.h
@@ -28,6 +28,22 @@ class UCookOnTheFlyServer;
 namespace UE::Cook { class FCookGCDiagnosticContext; }
 namespace UE::EditorDomain { struct FPackageDigest; }
 
+namespace UE::Cook::CookPackageSplitter
+{
+
+/** Data necessary to support the ICookPackageSplitter::FPopulateContext interface. */
+struct FPopulateContextData
+{
+	TArray<UPackage*> KeepReferencedPackages;
+	TArray<UObject*> ObjectsToMove;
+	TConstArrayView<ICookPackageSplitter::FGeneratedPackageForPopulate> GeneratedPackages;
+	UPackage* OwnerPackage = nullptr;
+	UObject* OwnerObject = nullptr;
+	const ICookPackageSplitter::FGeneratedPackageForPopulate* TargetGeneratedPackage = nullptr;
+};
+
+} // namespace UE::Cook::CookPackageSplitter::Private
+
 namespace UE::Cook
 {
 
@@ -256,7 +272,7 @@ public:
 	 * errors and returns false on failure.
 	 */
 	bool TryCallPopulateGeneratorPackage(
-		TArray<ICookPackageSplitter::FGeneratedPackageForPreSave>& InOutGeneratedPackagesForPresave);
+		TArray<ICookPackageSplitter::FGeneratedPackageForPopulate>& InOutGeneratedPackagesForPresave);
 	/**
 	 * Call the Splitter's PopulateGeneratedPackage if not yet called. Assumes GenerateList has been called. Logs
 	 * errors and returns false on failure.
diff --git a/Engine/Source/Editor/UnrealEd/Public/CookPackageSplitter.h b/Engine/Source/Editor/UnrealEd/Public/CookPackageSplitter.h
--- a/Engine/Source/Editor/UnrealEd/Public/CookPackageSplitter.h
+++ b/Engine/Source/Editor/UnrealEd/Public/CookPackageSplitter.h
@@ -6,6 +6,7 @@
 
 #include "AssetRegistry/IAssetRegistry.h"
 #include "Containers/Array.h"
+#include "Containers/ArrayView.h"
 #include "Containers/List.h"
 #include "Containers/StringView.h"
 #include "Containers/UnrealString.h"
@@ -24,6 +25,8 @@ class UPackage;
 struct FAssetDependency;
 template <typename FuncType> class TFunctionRef;
 
+namespace UE::Cook::CookPackageSplitter { struct FPopulateContextData; }
+
 /**
  * This class is used for packages that need to be split into multiple runtime packages.
  * It provides the instructions to the cooker for how to split the package.
@@ -159,41 +162,103 @@ public:
 	/** Return the list of packages to generate. */
 	virtual TArray<FGeneratedPackage> GetGenerateList(const UPackage* OwnerPackage, const UObject* OwnerObject) = 0;
 
-	/** Representation of a generated package that is provided when populating the generator package. */
-	struct FGeneratedPackageForPreSave
+	/** Representation of generated packages prepared by the cooker. */
+	struct FGeneratedPackageForPopulate
 	{
 		/** RelativePath returned from GetGenerateList */
 		FString RelativePath;
 		/** Root returned from GetGenerateList */
 		FString GeneratedRootPath;
 		/**
-		 * Non-null UPackage. Possibly an empty placeholder package. Provided so that the generator package
+		 * Non-null UPackage. Possibly an empty placeholder package, but may contain modifications
+		 * that were made during PopulateGeneratorPackage. Provided so that the generator package
 		 * can create import references to objects that will be stored in the generated package.
 		 */
 		UPackage* Package = nullptr;
 		/** *GetCreateAsMap returned from GetGenerateList. The package filename extension has already been set based on this. */
 		bool bCreatedAsMap = false;
 	};
+
+	/**
+	 * Context passed into Populate, PreSave, and PostSave functions, on the Generator package and on the Generated
+	 * packages. Some functions are only applicable in certain calls, see the description of each function.
+	 */
+	struct FPopulateContext
+	{
+	public:
+		explicit FPopulateContext(UE::Cook::CookPackageSplitter::FPopulateContextData& InData);
+
+		/** The generator package being split. */
+		UNREALED_API UPackage* GetOwnerPackage() const;
+		/** The SplitDataClass instance that this CookPackageSplitter instance was created for. */
+		UNREALED_API UObject* GetOwnerObject() const;
+		/**
+		 * Placeholder UPackage and relative path information for all packages that will be generated.
+		 * 
+		 * This function is only available in Populate and PreSave calls. It returns an empty array during PostSave.
+		 * This function is only available in calls on the generator package. It returns an empty array during calls to
+		 * generated packages.
+		 */
+		UNREALED_API TConstArrayView<ICookPackageSplitter::FGeneratedPackageForPopulate>& GetGeneratedPackages();
+
+		/**
+		 * Returns true during calls on the generator package (e.g. PopulateGeneratorPackage). Returns false during
+		 * calls on the generated packages (e.g. PopulateGeneratedPackage).
+		 */
+		UNREALED_API bool IsCalledOnGenerator() const;
+
+		/**
+		 * Returns the UPackage for which the event is being called. Returns the OwnerPackage if IsCalledOnGenerator,
+		 * returns the targetgenerated package if !IsCalledOnGenerator.
+		 * Guaranteed to not return nullptr.
+		 */
+		UNREALED_API UPackage* GetTargetPackage() const;
+
+		/**
+		 * Returns the FGeneratedPackageForPopulate for the package for which the event is being called.
+		 * 
+		 * Guaranteed to return non-null if !IsCalledOnGenerator().
+		 * Returns null if IsCalledOnGenerator().
+		 */
+		UNREALED_API const FGeneratedPackageForPopulate* GetTargetGeneratedPackage() const;
+
+		/**
+		 * Report objects that will be moved into the Generator or Generated package during its save.
+		 * This is optional - these reported objects are processed (BeginCacheForCookPlatformData) asynchronously
+		 * instead of synchronously during save.
+		 *
+		 * This callback is only valid during Populate functions. It is ignored during PreSave and PostSave functions.
+		 */
+		UNREALED_API void ReportObjectToMove(UObject* Object);
+		UNREALED_API void ReportObjectsToMove(TConstArrayView<UObject*> Objects);
+
+		/**
+		 * Report a package to keep referenced until the generator/generated package finishes save.
+		 * When called for a Generator, if DoesGeneratedRequireGenerator() >= Populate, these will also be kept
+		 * referenced until all generated packages finish saving or the splitter is destroyed.
+		 *
+		 * This is partially optional; the CookPackageSplitter can also manage the lifetime of the objects
+		 * internally. But allowing objects necessary for the save to be garbage collected will cause performance
+		 * problems and possibly errors, so either this method or some other internal method must be used.
+		 * 
+		 * This callback is only valid during Populate and PreSave functions. It is ignored during PostSave functions.
+		 */
+		UNREALED_API void ReportKeepReferencedPackage(UPackage* Package);
+		UNREALED_API void ReportKeepReferencedPackages(TConstArrayView<UPackage*> Packages);
+	private:
+		UE::Cook::CookPackageSplitter::FPopulateContextData& Data;
+	};
+
 	/**
 	 * Called before presaving the parent generator package, to give the generator a chance to inform the cooker which
 	 * objects will be moved into the generator package that are not already present in it.
-	 * 
+	 *
 	 * PopulateGeneratorPackage is guaranteed to not be called again until the splitter has been destroyed and the
 	 * generator package has been garbage collected.
 	 *
-	 * @param OwnerPackage				The generator package being split
-	 * @param OwnerObject				The SplitDataClass instance that this CookPackageSplitter instance was created for
-	 * @param GeneratedPackages			Placeholder UPackage and relative path information for all packages that will be generated
-	 * @param OutObjectsToMove			List of all the objects that will be moved into the Generator package during its save
-	 * @param OutKeepReferencedPackages Packages to keep referenced until the generator package finishes save.
-	 *                                  If DoesGeneratedRequireGenerator() >= Populate, these will also be kept referenced until
-	 *                                  all generated packages finish saving or the splitter is destroyed.
-	 *
-	 * @return							True if successfully populated, false on error (this will cause a cook error).
+	 * @return True if successfully populated, false on error (this will cause a cook error).
 	 */
-	virtual bool PopulateGeneratorPackage(UPackage* OwnerPackage, UObject* OwnerObject,
-		const TArray<ICookPackageSplitter::FGeneratedPackageForPreSave>& GeneratedPackages, TArray<UObject*>& OutObjectsToMove,
-		TArray<UPackage*>& OutKeepReferencedPackages)
+	virtual bool PopulateGeneratorPackage(FPopulateContext& PopulateContext)
 	{
 		return true;
 	}
@@ -201,18 +266,9 @@ public:
 	/**
 	 * Called before saving the parent generator package, after PopulateGeneratorPackage but before PopulateGeneratedPackage for any
 	 * generated packages. Make any required adjustments to the parent package before it is saved into the target domain.
-	 *
-	 * @param OwnerPackage				The generator package being split
-	 * @param OwnerObject				The SplitDataClass instance that this CookPackageSplitter instance was created for
-	 * @param GeneratedPackages			Placeholder UPackage and relative path information for all packages that will be generated
-	 * @param OutKeepReferencedPackages Packages to keep referenced until the generator package finishes save.
-	 *                                  If DoesGeneratedRequireGenerator() >= Populate, these will also be kept referenced until
-	 *                                  all generated packages finish saving or the splitter is destroyed.
-	 *
-	 * @return							True if successfully presaved, false on error (this will cause a cook error).
+	 * @return True if successfully presaved, false on error (this will cause a cook error).
 	 */
-	virtual bool PreSaveGeneratorPackage(UPackage* OwnerPackage, UObject* OwnerObject,
-		const TArray<FGeneratedPackageForPreSave>& PlaceholderPackages, TArray<UPackage*>& OutKeepReferencedPackages)
+	virtual bool PreSaveGeneratorPackage(FPopulateContext& PopulateContext)
 	{
 		return true;
 	}
@@ -222,26 +278,10 @@ public:
 	 * were made in PreSaveGeneratorPackage, so that the package is once again ready for use in the editor or in
 	 * future GetGenerateList or PreSaveGeneratedPackage calls
 	 */
-	virtual void PostSaveGeneratorPackage(UPackage* OwnerPackage, UObject* OwnerObject)
+	virtual void PostSaveGeneratorPackage(FPopulateContext& PopulateContext)
 	{
 	}
 
-	/** Representation of a generated package when it itself is being populated. */
-	struct FGeneratedPackageForPopulate
-	{
-		/** RelativePath returned from GetGenerateList */
-		FString RelativePath;
-		/** Root returned from GetGenerateList */
-		FString GeneratedRootPath;
-		/**
-		 * The UPackage that has been created for the package. Possibly empty, but may also contain still have modifications
-		 * that were made during PopulateGeneratorPackage.
-		 */
-		UPackage* Package = nullptr;
-		/** *GetCreateAsMap returned from GetGenerateList. The package filename extension has already been set based on this. */
-		bool bCreatedAsMap = false;
-	};
-
 	/**
 	 * Try to populate a generated package.
 	 *
@@ -252,18 +292,10 @@ public:
 	 * 
 	 * PopulateGeneratedPackage is guaranteed to not be called again on the same generated package until the splitter
 	 * has been destroyed and the generator package has been garbage collected.
-	 *
-	 * @param OwnerPackage				The parent package being split
-	 * @param OwnerObject				The SplitDataClass instance that this CookPackageSplitter instance was created for
-	 * @param GeneratedPackage			Pointer and information about the package to populate
-	 * @param OutObjectsToMove			List of all the objects that will be moved into the generated package during its save
-	 * @param OutKeepReferencedPackages Packages to keep referenced until the generated package finishes save.
 	 * 
-	 * @return							True if successfully populated, false on error (this will cause a cook error).
+	 * @return True if successfully populated, false on error (this will cause a cook error).
 	 */
-	virtual bool PopulateGeneratedPackage(UPackage* OwnerPackage, UObject* OwnerObject,
-		const FGeneratedPackageForPopulate& GeneratedPackage, TArray<UObject*>& OutObjectsToMove,
-		TArray<UPackage*>& OutKeepReferencedPackages)
+	virtual bool PopulateGeneratedPackage(FPopulateContext& PopulateContext)
 	{
 		return true;
 	}
@@ -272,15 +304,9 @@ public:
 	 * Called before saving a generated package, after PopulateGeneratedPackage. Make any required adjustments to the
 	 * generated package before it is saved into the target domain.
 	 * 
-	 * @param OwnerPackage				The parent package being split
-	 * @param OwnerObject				The SplitDataClass instance that this CookPackageSplitter instance was created for
-	 * @param GeneratedPackage			Pointer and information about the package to populate
-	 * @param OutKeepReferencedPackages Packages to keep referenced until the generated package finishes save.
-	 * 
-	 * @return							True if successfully presaved, false on error (this will cause a cook error).
+	 * @return True if successfully presaved, false on error (this will cause a cook error).
 	 */
-	virtual bool PreSaveGeneratedPackage(UPackage* OwnerPackage, UObject* OwnerObject,
-		const FGeneratedPackageForPopulate& GeneratedPackage, TArray<UPackage*>& OutKeepReferencedPackages)
+	virtual bool PreSaveGeneratedPackage(FPopulateContext& PopulateContext)
 	{
 		return true;
 	};
@@ -290,8 +316,7 @@ public:
 	 * were made in PreSaveGeneratedPackage, so that the parent package is once again ready for use in the editor or in
 	 * future PreSaveGeneratedPackage calls.
 	 */
-	virtual void PostSaveGeneratedPackage(UPackage* OwnerPackage, UObject* OwnerObject,
-		const FGeneratedPackageForPopulate& GeneratedPackage)
+	virtual void PostSaveGeneratedPackage(FPopulateContext& PopulateContext)
 	{
 	}
 
@@ -312,6 +337,38 @@ public:
 	 */
 	UNREALED_API static FString ConstructGeneratedPackageName(FName OwnerPackageName, FStringView RelPath,
 		FStringView GeneratedRootOverride = FStringView());
+
+
+	// FGeneratedPackageForPreSave and FGeneratedPackageForPopulate were two different structs that were so far identical,
+	// given separate identical definitions because we wanted to reserve the ability to change one but not the other.
+	// But to simplify the API we have given up on that flexibility, and now use FGeneratedPackageForPopulate everywhere.
+	UE_DEPRECATED(5.6, "Use FGeneratedPackageForPopulate instead.")
+	typedef FGeneratedPackageForPopulate FGeneratedPackageForPreSave;
+	UE_DEPRECATED(5.6, "Implement version that takes an FPopulateContext instead.")
+	UNREALED_API virtual bool PopulateGeneratorPackage(UPackage* OwnerPackage, UObject* OwnerObject,
+		const TArray<ICookPackageSplitter::FGeneratedPackageForPopulate>& GeneratedPackages,
+		TArray<UObject*>& OutObjectsToMove, TArray<UPackage*>& OutKeepReferencedPackages);
+	UE_DEPRECATED(5.6, "Implement version that takes an FPopulateContext instead.")
+	UNREALED_API virtual bool PreSaveGeneratorPackage(UPackage* OwnerPackage, UObject* OwnerObject,
+		const TArray<FGeneratedPackageForPopulate>& PlaceholderPackages, TArray<UPackage*>& OutKeepReferencedPackages);
+	UE_DEPRECATED(5.6, "Implement version that takes an FPopulateContext instead.")
+	UNREALED_API virtual void PostSaveGeneratorPackage(UPackage* OwnerPackage, UObject* OwnerObject);
+	UE_DEPRECATED(5.6, "Implement version that takes an FPopulateContext instead.")
+	UNREALED_API virtual bool PopulateGeneratedPackage(UPackage* OwnerPackage, UObject* OwnerObject,
+		const FGeneratedPackageForPopulate& GeneratedPackage, TArray<UObject*>& OutObjectsToMove,
+		TArray<UPackage*>& OutKeepReferencedPackages);
+	UE_DEPRECATED(5.6, "Implement version that takes an FPopulateContext instead.")
+	UNREALED_API virtual bool PreSaveGeneratedPackage(UPackage* OwnerPackage, UObject* OwnerObject,
+		const FGeneratedPackageForPopulate& GeneratedPackage, TArray<UPackage*>& OutKeepReferencedPackages);
+	UE_DEPRECATED(5.6, "Implement version that takes an FPopulateContext instead.")
+	UNREALED_API virtual void PostSaveGeneratedPackage(UPackage* OwnerPackage, UObject* OwnerObject,
+		const FGeneratedPackageForPopulate& GeneratedPackage);
+	UE_DEPRECATED(5.6, "Deprecation support, do not call outside of cooker code.")
+	UNREALED_API void WarnIfDeprecatedVirtualNotCalled(const TCHAR* FunctionName);
+
+private:
+	UE_DEPRECATED(5.6, "Deprecation support, do not read/write outside of cooker code.")
+	bool bDeprecatedVirtualCalledAsExpected = false;
 };
 
 namespace UE::Cook::Private
diff --git a/Engine/Source/Runtime/Engine/Private/WorldPartition/WorldPartitionCookPackageSplitter.cpp b/Engine/Source/Runtime/Engine/Private/WorldPartition/WorldPartitionCookPackageSplitter.cpp
--- a/Engine/Source/Runtime/Engine/Private/WorldPartition/WorldPartitionCookPackageSplitter.cpp
+++ b/Engine/Source/Runtime/Engine/Private/WorldPartition/WorldPartitionCookPackageSplitter.cpp
@@ -109,10 +109,10 @@ TArray<ICookPackageSplitter::FGeneratedPackage> FWorldPartitionCookPackageSplitt
 	return PackagesToGenerate;
 }
 
-bool FWorldPartitionCookPackageSplitter::PopulateGeneratedPackage(UPackage* OwnerPackage, UObject* OwnerObject,
-	const FGeneratedPackageForPopulate& GeneratedPackage, TArray<UObject*>& OutObjectsToMove,
-	TArray<UPackage*>& OutModifiedPackages)
+bool FWorldPartitionCookPackageSplitter::PopulateGeneratedPackage(FPopulateContext& PopulateContext)
 {
+	const FGeneratedPackageForPopulate& GeneratedPackage = *PopulateContext.GetTargetGeneratedPackage();
+
 	UE_LOG(LogWorldPartition, Verbose, TEXT("[Cook][PopulateGeneratedPackage] Processing %s"), *FWorldPartitionCookPackage::MakeGeneratedFullPath(GeneratedPackage.GeneratedRootPath, GeneratedPackage.RelativePath));
 
 	bool bIsSuccess = true;
@@ -131,15 +131,15 @@ bool FWorldPartitionCookPackageSplitter::PopulateGeneratedPackage(UPackage* Owne
 	}
 
 	UE_LOG(LogWorldPartition, Verbose, TEXT("[Cook][PopulateGeneratedPackage] Gathered %u modified packages for %s"), ModifiedPackages.Num() , *FWorldPartitionCookPackage::MakeGeneratedFullPath(GeneratedPackage.GeneratedRootPath, GeneratedPackage.RelativePath));
-	OutModifiedPackages = MoveTemp(ModifiedPackages);
+	PopulateContext.ReportKeepReferencedPackages(ModifiedPackages);
 
 	return bIsSuccess;
 }
 
-bool FWorldPartitionCookPackageSplitter::PopulateGeneratorPackage(UPackage* OwnerPackage, UObject* OwnerObject,
-	const TArray<ICookPackageSplitter::FGeneratedPackageForPreSave>& GeneratedPackages, TArray<UObject*>& OutObjectsToMove,
-	TArray<UPackage*>& OutModifiedPackages)
+bool FWorldPartitionCookPackageSplitter::PopulateGeneratorPackage(FPopulateContext& PopulateContext)
 {
+	TConstArrayView<ICookPackageSplitter::FGeneratedPackageForPopulate> GeneratedPackages = PopulateContext.GetGeneratedPackages();
+
 	UE_LOG(LogWorldPartition, Display, TEXT("[Cook][PopulateGeneratorPackage] Processing %u packages"), GeneratedPackages.Num());
 
 	bool bIsSuccess = true;
@@ -160,7 +160,7 @@ bool FWorldPartitionCookPackageSplitter::PopulateGeneratorPackage(UPackage* Owne
 	}
 
 	UE_LOG(LogWorldPartition, Display, TEXT("[Cook][PopulateGeneratorPackage] Gathered %u modified packages"), ModifiedPackages.Num());
-	OutModifiedPackages = MoveTemp(ModifiedPackages);
+	PopulateContext.ReportKeepReferencedPackages(ModifiedPackages);
 
 	return bIsSuccess;
 }
diff --git a/Engine/Source/Runtime/Engine/Private/WorldPartition/WorldPartitionCookPackageSplitter.h b/Engine/Source/Runtime/Engine/Private/WorldPartition/WorldPartitionCookPackageSplitter.h
--- a/Engine/Source/Runtime/Engine/Private/WorldPartition/WorldPartitionCookPackageSplitter.h
+++ b/Engine/Source/Runtime/Engine/Private/WorldPartition/WorldPartitionCookPackageSplitter.h
@@ -34,11 +34,8 @@ public:
 	}
 
 	virtual TArray<ICookPackageSplitter::FGeneratedPackage> GetGenerateList(const UPackage* OwnerPackage, const UObject* OwnerObject) override;
-	virtual bool PopulateGeneratedPackage(UPackage* OwnerPackage, UObject* OwnerObject,
-		const FGeneratedPackageForPopulate& GeneratedPackage, TArray<UObject*>& OutObjectsToMove, TArray<UPackage*>& OutModifiedPackages) override;
-	virtual bool PopulateGeneratorPackage(UPackage* OwnerPackage, UObject* OwnerObject,
-		const TArray<ICookPackageSplitter::FGeneratedPackageForPreSave>& GeneratedPackages, TArray<UObject*>& OutObjectsToMove,
-		TArray<UPackage*>& OutModifiedPackages) override;
+	virtual bool PopulateGeneratedPackage(FPopulateContext& PopulateContext) override;
+	virtual bool PopulateGeneratorPackage(FPopulateContext& PopulateContext) override;
 	virtual void OnOwnerReloaded(UPackage* OwnerPackage, UObject* OwnerObject) override;
 	//~ End of ICookPackageSplitter
 
@@ -50,7 +47,6 @@ private:
 	virtual FString GetReferencerName() const override;
 
 	void BuildPackagesToGenerateList(TArray<ICookPackageSplitter::FGeneratedPackage>& PackagesToGenerate) const;
-	bool MapGeneratePackageToCookPackage(const TArray<ICookPackageSplitter::FGeneratedPackageForPreSave>& GeneratedPackages);
 	
 	TObjectPtr<UWorld> ReferencedWorld = nullptr;
 
