IncrementalCook: Add new virtual UObject::OnCookEvent to manually report cookdependencies that is called even when the package is not saved, so we can support manually recorded dependencies for the transitive build dependencies stored on packages that are not cooked but are referenced as build dependencies.
OnCookEvent will be used in the future for any other feedback that the cooker needs from UObject classes, to reduce the number of virtuals we have to add onto UObject for all of the cook-related information.
The OnCookEvent API allows specifying Load dependencies, Save dependencies, and system-specific BuildResult dependencies. System-specific BuildResults will be expanded on in future changes.
To support reporting dependencies from structs, expand the existing FObjectSavePackageSerializeContext to support being called to collect the dependencies from the same point as the new OnCookEvent(PlatformCookDependencies).
Add access to the ICookInfo interface of the cooker in OnCookEvent and the existing ObjectSaveContextData. This is a replacement for the more limited information about the cook (GetCookType,GetCookingDLC) previously available in the PreSave and Serialize APIs.
#rnx
#incrementalcook

diff --git a/Engine/Plugins/FX/Niagara/Source/Niagara/Classes/NiagaraSystem.h b/Engine/Plugins/FX/Niagara/Source/Niagara/Classes/NiagaraSystem.h
--- a/Engine/Plugins/FX/Niagara/Source/Niagara/Classes/NiagaraSystem.h
+++ b/Engine/Plugins/FX/Niagara/Source/Niagara/Classes/NiagaraSystem.h
@@ -228,6 +228,7 @@ public:
 	NIAGARA_API virtual bool IsReadyForFinishDestroy() override;
 	NIAGARA_API virtual void PreSave(FObjectPreSaveContext ObjectSaveContext) override;
 #if WITH_EDITOR
+	NIAGARA_API virtual void OnCookEvent(UE::Cook::ECookEvent CookEvent, UE::Cook::FCookEventContext& CookContext) override;
 	NIAGARA_API virtual void PostRename(UObject* OldOuter, const FName OldName) override;
 	NIAGARA_API virtual void PreEditChange(FProperty* PropertyThatWillChange)override;
 	NIAGARA_API virtual void PostEditChangeProperty(struct FPropertyChangedEvent& PropertyChangedEvent) override; 
diff --git a/Engine/Plugins/FX/Niagara/Source/Niagara/Private/NiagaraSystem.cpp b/Engine/Plugins/FX/Niagara/Source/Niagara/Private/NiagaraSystem.cpp
--- a/Engine/Plugins/FX/Niagara/Source/Niagara/Private/NiagaraSystem.cpp
+++ b/Engine/Plugins/FX/Niagara/Source/Niagara/Private/NiagaraSystem.cpp
@@ -39,6 +39,7 @@
 #include "Algo/RemoveIf.h"
 #include "Algo/StableSort.h"
 #include "Cooker/CookDependency.h"
+#include "Cooker/CookEvents.h"
 #include "HAL/LowLevelMemTracker.h"
 #include "Interfaces/ITargetPlatform.h"
 #include "Interfaces/ITargetPlatformManagerModule.h"
@@ -98,7 +99,7 @@ namespace NiagaraSystemPrivate
 #if WITH_EDITOR
 	constexpr int32 CookDependenciesArgsVersion = 1;
 	void HashDependenciesForCook(FCbFieldViewIterator Args, UE::Cook::FCookDependencyContext& Context);
-	void RegisterCookDependencies(FObjectPreSaveContext ObjectSaveContext, bool bIncludeGpuTarget);
+	void RegisterCookDependencies(UE::Cook::FCookEventContext& CookContext, bool bIncludeGpuTarget);
 
 	struct FCookDependency
 	{
@@ -353,17 +354,20 @@ void UNiagaraSystem::PreSave(FObjectPreSaveContext ObjectSaveContext)
 #if WITH_EDITORONLY_DATA
 	WaitForCompilationComplete();
 #endif
+}
 
 #if WITH_EDITOR
-	if (ObjectSaveContext.IsCooking())
+void UNiagaraSystem::OnCookEvent(UE::Cook::ECookEvent CookEvent, UE::Cook::FCookEventContext& CookContext)
+{
+	Super::OnCookEvent(CookEvent, CookContext);
+
+	if (CookEvent == UE::Cook::ECookEvent::PlatformCookDependencies && CookContext.IsCooking())
 	{
 		UpdateHasGPUEmitters();
-		NiagaraSystemPrivate::RegisterCookDependencies(ObjectSaveContext, bHasAnyGPUEmitters);
+		NiagaraSystemPrivate::RegisterCookDependencies(CookContext, bHasAnyGPUEmitters);
 	}
-#endif
 }
 
-#if WITH_EDITOR
 void UNiagaraSystem::BeginCacheForCookedPlatformData(const ITargetPlatform *TargetPlatform)
 {
 	//UE_LOG(LogNiagara, Display, TEXT("UNiagaraSystem::BeginCacheForCookedPlatformData %s %s"), *GetFullName(), GIsSavingPackage ? TEXT("Saving...") : TEXT("Not Saving..."));
@@ -4419,9 +4423,9 @@ void HashDependenciesForCook(FCbFieldViewIterator Args, UE::Cook::FCookDependenc
 	}
 }
 
-void RegisterCookDependencies(FObjectPreSaveContext ObjectSaveContext, bool bIncludeGpuTarget)
+void RegisterCookDependencies(UE::Cook::FCookEventContext& CookContext, bool bIncludeGpuTarget)
 {
-	const ITargetPlatform* TargetPlatform = ObjectSaveContext.GetTargetPlatform();
+	const ITargetPlatform* TargetPlatform = CookContext.GetTargetPlatform();
 
 	FCookDependency CookDependency;
 	if (bIncludeGpuTarget)
@@ -4438,7 +4442,7 @@ void RegisterCookDependencies(FObjectPreSaveContext ObjectSaveContext, bool bInc
 
 	CookDependency.Save(Writer);
 
-	ObjectSaveContext.AddCookBuildDependency(
+	CookContext.AddLoadBuildDependency(
 		UE::Cook::FCookDependency::Function(
 			UE_COOK_DEPENDENCY_FUNCTION_CALL(HashNiagaraSystemDependenciesForCook), Writer.Save()));
 }
diff --git a/Engine/Plugins/Runtime/ComputeFramework/Source/ComputeFramework/Private/ComputeGraph.cpp b/Engine/Plugins/Runtime/ComputeFramework/Source/ComputeFramework/Private/ComputeGraph.cpp
--- a/Engine/Plugins/Runtime/ComputeFramework/Source/ComputeFramework/Private/ComputeGraph.cpp
+++ b/Engine/Plugins/Runtime/ComputeFramework/Source/ComputeFramework/Private/ComputeGraph.cpp
@@ -24,6 +24,10 @@
 #include "UObject/ObjectSaveContext.h"
 #include "UObject/Package.h"
 
+#if WITH_EDITOR
+#include "Cooker/CookEvents.h"
+#endif
+
 #include UE_INLINE_GENERATED_CPP_BY_NAME(ComputeGraph)
 
 UComputeGraph::UComputeGraph() = default;
@@ -149,14 +153,14 @@ UE_COOK_DEPENDENCY_FUNCTION(HashComputeGraphDependenciesForCook, ComputeFramewor
 #endif // WITH_EDITOR
 
 
-void UComputeGraph::PreSave(FObjectPreSaveContext ObjectSaveContext)
-{
-	Super::PreSave(ObjectSaveContext);
 #if WITH_EDITOR
-	if (ObjectSaveContext.IsCooking())
+void UComputeGraph::OnCookEvent(UE::Cook::ECookEvent CookEvent, UE::Cook::FCookEventContext& CookContext)
+{
+	Super::OnCookEvent(CookEvent, CookContext);
+	if (CookEvent == UE::Cook::ECookEvent::PlatformCookDependencies && CookContext.IsCooking())
 	{
 		TArray<FName> ShaderFormats;
-		if (ITargetPlatform const* TargetPlatform = ObjectSaveContext.GetTargetPlatform())
+		if (ITargetPlatform const* TargetPlatform = CookContext.GetTargetPlatform())
 		{
 			TargetPlatform->GetAllTargetedShaderFormats(ShaderFormats);
 		}
@@ -184,10 +188,10 @@ void UComputeGraph::PreSave(FObjectPreSaveContext ObjectSaveContext)
 		Writer << ShaderFormats;
 		Writer << ShaderVirtualPaths;
 
-		ObjectSaveContext.AddCookBuildDependency(UE::Cook::FCookDependency::Function(UE_COOK_DEPENDENCY_FUNCTION_CALL(HashComputeGraphDependenciesForCook), Writer.Save()));
+		CookContext.AddLoadBuildDependency(UE::Cook::FCookDependency::Function(UE_COOK_DEPENDENCY_FUNCTION_CALL(HashComputeGraphDependenciesForCook), Writer.Save()));
 	}
-#endif
 }
+#endif
 
 bool UComputeGraph::ValidateGraph(FString* OutErrors)
 {
diff --git a/Engine/Plugins/Runtime/ComputeFramework/Source/ComputeFramework/Public/ComputeFramework/ComputeGraph.h b/Engine/Plugins/Runtime/ComputeFramework/Source/ComputeFramework/Public/ComputeFramework/ComputeGraph.h
--- a/Engine/Plugins/Runtime/ComputeFramework/Source/ComputeFramework/Public/ComputeFramework/ComputeGraph.h
+++ b/Engine/Plugins/Runtime/ComputeFramework/Source/ComputeFramework/Public/ComputeFramework/ComputeGraph.h
@@ -140,8 +140,8 @@ protected:
 	//~ Begin UObject Interface.
 	void Serialize(FArchive& Ar) override;
 	void PostLoad() override;
-	void PreSave(FObjectPreSaveContext ObjectSaveContext);
 #if WITH_EDITOR
+	void OnCookEvent(UE::Cook::ECookEvent CookEvent, UE::Cook::FCookEventContext& CookContext) override;
 	void BeginCacheForCookedPlatformData(ITargetPlatform const* TargetPlatform) override;
 	bool IsCachedCookedPlatformDataLoaded(ITargetPlatform const* TargetPlatform) override;
 	void ClearCachedCookedPlatformData(ITargetPlatform const* TargetPlatform) override;
diff --git a/Engine/Source/Editor/UnrealEd/Private/Commandlets/SavePackageUtilitiesCommandlet.cpp b/Engine/Source/Editor/UnrealEd/Private/Commandlets/SavePackageUtilitiesCommandlet.cpp
--- a/Engine/Source/Editor/UnrealEd/Private/Commandlets/SavePackageUtilitiesCommandlet.cpp
+++ b/Engine/Source/Editor/UnrealEd/Private/Commandlets/SavePackageUtilitiesCommandlet.cpp
@@ -28,6 +28,8 @@ int32 USavePackageUtilitiesCommandlet::Main(const FString& Params)
 {
 	InitParameters(Params);
 
+	UE::Cook::ICookInfo* CookInfo = nullptr; // TODO: NeedSimpleCookInfoClass
+
 	for (const FString& PackageName : PackageNames)
 	{
 		// Load Package
@@ -41,7 +43,7 @@ int32 USavePackageUtilitiesCommandlet::Main(const FString& Params)
 		if (TargetPlatform != nullptr)
 		{
 			CookContext.Emplace(Package, UE::Cook::ECookType::Unknown,
-				UE::Cook::ECookingDLC::Unknown, TargetPlatform);
+				UE::Cook::ECookingDLC::Unknown, TargetPlatform, CookInfo);
 			CookData.Emplace(*TargetPlatform, *CookContext);
 		}
 
diff --git a/Engine/Source/Editor/UnrealEd/Private/Cooker/CookFunctionLibrary.cpp b/Engine/Source/Editor/UnrealEd/Private/Cooker/CookFunctionLibrary.cpp
--- a/Engine/Source/Editor/UnrealEd/Private/Cooker/CookFunctionLibrary.cpp
+++ b/Engine/Source/Editor/UnrealEd/Private/Cooker/CookFunctionLibrary.cpp
@@ -134,10 +134,13 @@ void UCookFunctionLibrary::CookAsset(UObject* Object, const FString& ForPlatform
 		FPaths::ProjectSavedDir() +
 		FString(TEXT("Temp/")) + DestinationSubfolder + Package->GetName() + TEXT(".uasset"));
 
+	// TODO: Create an ICookInfo provider that does not require UCookOnTheFlyServer and pass it in so we can collect
+	// manual BuildDependencies.
+	ICookInfo* CookInfo = nullptr; // TODO: NeedSimpleCookInfoClass
 	FArchiveCookContext ArchiveCookContext(Package,
 		UE::Cook::ECookType::ByTheBook,
 		UE::Cook::ECookingDLC::No, // used only by UMaterialInterface::Serialize, we could expose via command line arguments
-		TargetPlatform);
+		TargetPlatform, CookInfo);
 	FArchiveCookData CookData(*TargetPlatform, ArchiveCookContext);
 
 	// The AsyncIODelete object optimizes the cleaning of the saved/cook directory. Happily not relevant
@@ -177,17 +180,7 @@ void UCookFunctionLibrary::CookAsset(UObject* Object, const FString& ForPlatform
 	}
 
 	TArray<TPair<ELogVerbosity::Type, FString>> Messages;
-	FBuildDependencySet LoadDependencies = FBuildDependencySet::CollectLoadedPackage(Package, &Messages);
-	if (!LoadDependencies.IsValid())
-	{
-		TStringBuilder<256> ErrorMessage;
-		ErrorMessage << TEXT("Error collecting cook load dependencies:");
-		for (TPair<ELogVerbosity::Type, FString>& MessagePair : Messages)
-		{
-			ErrorMessage << TEXT("\n\t") << MessagePair.Value;
-		}
-		UE_LOG(LogCook, Warning, TEXT("%s"), *ErrorMessage);
-	}
+	FBuildResultDependenciesMap LoadDependencies = FBuildDependencySet::CollectLoadedPackage(Package, &Messages);
 
 	// TODO: Execute the asynchronous transformation steps that packages go through after loading and before saving
 	// BeginCacheForCookedPlatformData(); // Call once on each object
@@ -229,11 +222,20 @@ void UCookFunctionLibrary::CookAsset(UObject* Object, const FString& ForPlatform
 		// routine at this time.
 		UE::Cook::FGenerationHelper* GenerationHelper = nullptr;
 		bool bGenerated = false;
-		TArray<FName> RuntimeDepsTodo; // I don't believe we need valid data in here at this time
 		Messages.Reset();
 
-		FPackageArtifacts Deps = FPackageArtifacts::Collect(Package, TargetPlatform, &LoadDependencies, &Result,
-			GenerationHelper, bGenerated, MoveTemp(RuntimeDepsTodo), &Messages);
+		// TODO: Call OnCookEvent(ECookEvent::PlatformLoadDependencies) on every object in the package, as is done in
+		// FSaveCookedPackageContext::CalculateCookDependencies.
+		// Doing this requires implementing the NeedSimpleCookInfoClass TODO first.
+		FBuildResultDependenciesMap& BuildResultDependencies = Result.BuildResultDependencies;
+		BuildResultDependencies.Append(LoadDependencies);
+		// TODO: Populate RuntimeDependencies in the same way they are calculated in
+		// FSaveCookedPackageContext::CalculateCookDependencies
+		TArray<FName> RuntimeDependencies;
+
+		FPackageArtifacts Deps = FPackageArtifacts::Collect(Package, TargetPlatform,
+			MoveTemp(BuildResultDependencies), true /* bHasSaveResult */, Result.UntrackedSoftPackageReferences,
+			GenerationHelper, bGenerated, MoveTemp(RuntimeDependencies), &Messages);
 		if (!Deps.IsValid())
 		{
 			TStringBuilder<256> ErrorMessage;
diff --git a/Engine/Source/Editor/UnrealEd/Private/Cooker/CookPackageArtifacts.cpp b/Engine/Source/Editor/UnrealEd/Private/Cooker/CookPackageArtifacts.cpp
--- a/Engine/Source/Editor/UnrealEd/Private/Cooker/CookPackageArtifacts.cpp
+++ b/Engine/Source/Editor/UnrealEd/Private/Cooker/CookPackageArtifacts.cpp
@@ -7,6 +7,7 @@
 #include "AssetRegistry/IAssetRegistry.h"
 #include "Cooker/CookConfigAccessTracker.h"
 #include "Cooker/CookDependency.h"
+#include "Cooker/CookEvents.h"
 #include "Cooker/CookGenerationHelper.h"
 #include "Cooker/PackageBuildDependencyTracker.h"
 #include "CookOnTheSide/CookLog.h"
@@ -25,6 +26,7 @@
 #include "Serialization/CompactBinaryWriter.h"
 #include "TargetDomain/TargetDomainUtils.h"
 #include "UObject/CoreRedirects.h"
+#include "UObject/ICookInfo.h"
 
 namespace UE::Cook
 {
@@ -271,33 +273,26 @@ void FBuildDependencySet::Save(FCbWriter& Writer) const
 	Writer.EndObject();
 }
 
-FBuildDependencySet FBuildDependencySet::CollectLoadedPackage(const UPackage* Package,
+FBuildResultDependenciesMap FBuildDependencySet::CollectLoadedPackage(const UPackage* Package,
 	TArray<TPair<ELogVerbosity::Type, FString>>* OutMessages)
 {
-	TArray<FCookDependency> BuildDependencies;
+	FBuildResultDependenciesMap ResultDependencies;
 	TArray<FName> UnusedRuntimeDependencies;
-	if (!TryCollectInternal(BuildDependencies, UnusedRuntimeDependencies, Package, nullptr /* TargetPlatform */,
-		nullptr /* SaveResult */, nullptr /* GenerationHelper */, false /* bGenerated */, OutMessages))
+	if (!TryCollectInternal(ResultDependencies, UnusedRuntimeDependencies, OutMessages,
+		BuildResult::NAME_Load, Package, nullptr /* TargetPlatform */,
+		TConstArrayView<FName>() /* UntrackedSoftPackageReferences */, nullptr /* GenerationHelper */,
+		false /* bGenerated */))
 	{
-		return FBuildDependencySet();
+		return FBuildResultDependenciesMap();
 	}
 
 	// Sort and remove duplicates in the results from TryCollectInternal.
-	Algo::Sort(BuildDependencies);
-	BuildDependencies.SetNum(Algo::Unique(BuildDependencies), EAllowShrinking::Yes);
-
-	// Store the collected dependencies on the result
-	FBuildDependencySet Result;
-	Result.SetNormalizedDependencies(MoveTemp(BuildDependencies));
-	FBuildDependencySet::ECurrentKeyResult CurrentKeyResult = Result.TryCalculateCurrentKey(
-		Package->GetFName(), nullptr /* GenerationHelper */, OutMessages);
-	if (CurrentKeyResult == FBuildDependencySet::ECurrentKeyResult::Error)
+	for (TPair<FName, TArray<FCookDependency>>& ResultPair : ResultDependencies)
 	{
-		return FBuildDependencySet();
+		Algo::Sort(ResultPair.Value);
+		ResultPair.Value.SetNum(Algo::Unique(ResultPair.Value), EAllowShrinking::Yes);
 	}
-	Result.StoreCurrentKey();
-	Result.SetValid(true);
-	return Result;
+	return ResultDependencies;
 }
 
 enum class EPackageMountPoint
@@ -351,10 +346,10 @@ static void JoinMessagesIntoErrorReason(FStringBuilderBase& OutText,
 	}
 }
 
-bool FBuildDependencySet::TryCollectInternal(TArray<FCookDependency>& OutBuildDependencies,
-	TArray<FName>& InOutRuntimeDependencies, const UPackage* Package, const ITargetPlatform* TargetPlatform,
-	const FSavePackageResultStruct* SaveResult, FGenerationHelper* GenerationHelper,
-	bool bGenerated, TArray<TPair<ELogVerbosity::Type, FString>>* OutMessages)
+bool FBuildDependencySet::TryCollectInternal(FBuildResultDependenciesMap& InOutResultDependencies,
+	TArray<FName>& InOutRuntimeDependencies, TArray<TPair<ELogVerbosity::Type, FString>>* OutMessages,
+	FName DefaultBuildResult, const UPackage* Package, const ITargetPlatform* TargetPlatform,
+	TConstArrayView<FName> UntrackedSoftPackageReferences, FGenerationHelper* GenerationHelper, bool bGenerated)
 {
 	if (!Package)
 	{
@@ -374,12 +369,13 @@ bool FBuildDependencySet::TryCollectInternal(TArray<FCookDependency>& OutBuildDe
 		return false;
 	}
 
+	TArray<FCookDependency> DefaultResultDependencies;
 	// Skip the multiple reallocations for an array that grows from 0 to 128, for performance, but then
 	// reallocate according to normal TArray growth to reduce spike memory use.
-	OutBuildDependencies.Empty(128);
+	DefaultResultDependencies.Empty(128);
 
 	FName PackageName = Package->GetFName();
-	OutBuildDependencies.Add(FCookDependency::Package(PackageName));
+	DefaultResultDependencies.Add(FCookDependency::Package(PackageName));
 
 #if UE_WITH_PACKAGE_ACCESS_TRACKING
 	FPackageBuildDependencyTracker& Tracker = FPackageBuildDependencyTracker::Get();
@@ -401,27 +397,27 @@ bool FBuildDependencySet::TryCollectInternal(TArray<FCookDependency>& OutBuildDe
 					// We have not yet enabled marking the auto-added dependencies from TObjectPtr resolve
 					// as transitive, because it causes a performance regression and we are still working on fixing the
 					// regression.
-					OutBuildDependencies.Add(FCookDependency::Package(AccessData.ReferencedPackage));
+					DefaultResultDependencies.Add(FCookDependency::Package(AccessData.ReferencedPackage));
 				}
 				else if (ProjectionList.bHasAll)
 				{
 					// TObjectPtr UE::BuildProjection::All dependencies are added as transitive build dependencies.
 					// We have to do this to be conservative, since we do not know which bytes from the target are
 					// dependended upon and which of the target's build dependencies influence those bytes.
-					OutBuildDependencies.Add(FCookDependency::TransitiveBuild(AccessData.ReferencedPackage));
+					DefaultResultDependencies.Add(FCookDependency::TransitiveBuild(AccessData.ReferencedPackage));
 				}
 				else
 				{
 					for (FTopLevelAssetPath ClassPath : ProjectionList.Classes)
 					{
-						OutBuildDependencies.Add(FCookDependency::NativeClass(WriteToString<256>(ClassPath)));
+						DefaultResultDependencies.Add(FCookDependency::NativeClass(WriteToString<256>(ClassPath)));
 					}
 
 					for (FName ResultProjection : ProjectionList.ResultProjections)
 					{
 						if (ResultProjection == UE::Cook::ResultProjection::PackageAndClass)
 						{
-							OutBuildDependencies.Add(FCookDependency::Package(AccessData.ReferencedPackage));
+							DefaultResultDependencies.Add(FCookDependency::Package(AccessData.ReferencedPackage));
 						}
 						else
 						{
@@ -451,7 +447,7 @@ bool FBuildDependencySet::TryCollectInternal(TArray<FCookDependency>& OutBuildDe
 				UE::AssetRegistry::EDependencyCategory::Package, UE::AssetRegistry::EDependencyQuery::Game);
 			for (FName AssetDependency : AssetDependencies)
 			{
-				OutBuildDependencies.Add(FCookDependency::TransitiveBuild(AssetDependency));
+				DefaultResultDependencies.Add(FCookDependency::TransitiveBuild(AssetDependency));
 			}
 		}
 	}
@@ -465,189 +461,206 @@ bool FBuildDependencySet::TryCollectInternal(TArray<FCookDependency>& OutBuildDe
 			TArray<FConfigAccessData> ConfigKeys = ConfigTracker.GetPackageRecords(PackageName, TargetPlatform);
 			for (const FConfigAccessData& ConfigKey : ConfigKeys)
 			{
-				OutBuildDependencies.Add(FCookDependency::Config(ConfigKey));
+				DefaultResultDependencies.Add(FCookDependency::Config(ConfigKey));
 			}
 		}
 	}
 #endif
-	if (SaveResult)
+	if (!UntrackedSoftPackageReferences.IsEmpty())
 	{
-		OutBuildDependencies.Append(SaveResult->CookDependencies);
-		for (FName SoftPackageReference : SaveResult->UntrackedSoftPackageReferences)
+		TArray<FCookDependency>& SaveDependencies = InOutResultDependencies.FindOrAdd(BuildResult::NAME_Save);
+		for (FName SoftPackageReference : UntrackedSoftPackageReferences)
 		{
-			OutBuildDependencies.Add(FCookDependency::RedirectionTarget(SoftPackageReference));
+			SaveDependencies.Add(FCookDependency::RedirectionTarget(SoftPackageReference));
 		}
 	}
 
+	// Put the Dependencies we have collected onto the requested DefaultBuildResult
+	InOutResultDependencies.FindOrAdd(DefaultBuildResult).Append(MoveTemp(DefaultResultDependencies));
+
+	// If we have any RuntimeDependencies, they will cause some BuildDependencies in the SaveBuildResult, so add a SaveBuildResult
+	// output if we don't already have one.
+	if (!InOutRuntimeDependencies.IsEmpty())
+	{
+		InOutResultDependencies.FindOrAdd(BuildResult::NAME_Save);
+	}
+
 	// All Dependencies have been gathered. Format the lists for TryCalculateCurrentKey and for storage.
 
-	// Settings dependencies - Expand transitive dependencies on SettingsObjects into the list of dependencies
-	// recorded for that settings object.
-	TSet<const UObject*> SettingsDependencies;
-	OutBuildDependencies.RemoveAllSwap([&SettingsDependencies](const FCookDependency& Dependency)
-		{
-			if (Dependency.GetType() == ECookDependency::SettingsObject)
-			{
-				SettingsDependencies.Add(Dependency.GetSettingsObject());
-				return true;
-			}
-			return false;
-		},
-		EAllowShrinking::No);
-	for (const UObject* SettingsObject : SettingsDependencies)
+	for (TPair<FName, TArray<FCookDependency>>& ResultPair : InOutResultDependencies)
 	{
-		// We rely on the object to be rooted because we use its pointer as a key for the lifetime of
-		// the cook process, so it being garbage collected and something else allocated on the same
-		// pointer would break our key. IsRooted should have been validated by
-		// FCookDependency::SettingsObject.
-		check(SettingsObject->IsRooted());
-		FCookDependencyGroups::FRecordedDependencies& IncludeDependencies =
-			FCookDependencyGroups::Get().FindOrCreate(reinterpret_cast<UPTRINT>(SettingsObject));
-		if (!IncludeDependencies.bInitialized)
-		{
-			IncludeDependencies.Dependencies = CollectSettingsObject(SettingsObject,
-				&IncludeDependencies.Messages);
-			IncludeDependencies.bInitialized = true;
-		}
-		if (!IncludeDependencies.Dependencies.IsValid())
+		TArray<FCookDependency>& ResultDependencies = ResultPair.Value;
+		// Settings dependencies - Expand transitive dependencies on SettingsObjects into the list of dependencies
+		// recorded for that settings object.
+		TSet<const UObject*> SettingsDependencies;
+		ResultDependencies.RemoveAllSwap([&SettingsDependencies](const FCookDependency& Dependency)
+			{
+				if (Dependency.GetType() == ECookDependency::SettingsObject)
+				{
+					SettingsDependencies.Add(Dependency.GetSettingsObject());
+					return true;
+				}
+				return false;
+			},
+			EAllowShrinking::No);
+		for (const UObject* SettingsObject : SettingsDependencies)
 		{
-			if (OutMessages)
+			// We rely on the object to be rooted because we use its pointer as a key for the lifetime of
+			// the cook process, so it being garbage collected and something else allocated on the same
+			// pointer would break our key. IsRooted should have been validated by
+			// FCookDependency::SettingsObject.
+			check(SettingsObject->IsRooted());
+			FCookDependencyGroups::FRecordedDependencies& IncludeDependencies =
+				FCookDependencyGroups::Get().FindOrCreate(reinterpret_cast<UPTRINT>(SettingsObject));
+			if (!IncludeDependencies.bInitialized)
 			{
-				TStringBuilder<256> ErrorText;
-				ErrorText << TEXT("Dependencies for SettingsObject ") << SettingsObject->GetPathName()
-					<< TEXT(" are unavailable");
-				JoinMessagesIntoErrorReason(ErrorText, IncludeDependencies.Messages);
-				OutMessages->Emplace(ELogVerbosity::Error, FString(*ErrorText));
+				IncludeDependencies.Dependencies = CollectSettingsObject(SettingsObject,
+					&IncludeDependencies.Messages);
+				IncludeDependencies.bInitialized = true;
 			}
-			return false;
-		}
-
-		for (const FCookDependency& IncludeDependency : IncludeDependencies.Dependencies.GetDependencies())
-		{
-			// Recursive SettingsDependencies are not allowed. We haven't needed them yet, and not supporting them allows
-			// prevents the need for cycle detection.
-			if (IncludeDependency.GetType() == ECookDependency::SettingsObject)
+			if (!IncludeDependencies.Dependencies.IsValid())
 			{
 				if (OutMessages)
 				{
-					OutMessages->Emplace(ELogVerbosity::Error, FString::Printf(
-						TEXT("Settings dependency on object %s, but that object has a recursive Settings dependency on %s, and recursive Settings dependencies are not supported."),
-						*SettingsObject->GetPathName(), *IncludeDependency.GetSettingsObject()->GetPathName()));
+					TStringBuilder<256> ErrorText;
+					ErrorText << TEXT("Dependencies for SettingsObject ") << SettingsObject->GetPathName()
+						<< TEXT(" are unavailable");
+					JoinMessagesIntoErrorReason(ErrorText, IncludeDependencies.Messages);
+					OutMessages->Emplace(ELogVerbosity::Error, FString(*ErrorText));
 				}
 				return false;
 			}
-			OutBuildDependencies.Add(IncludeDependency);
-		}
-	}
 
-	// Process some rules for Package dependencies
-	TSet<FName> RedirectionTargets;
-	for (TArray<FCookDependency>::TIterator Iter(OutBuildDependencies); Iter; ++Iter)
-	{
-		FCookDependency& Dependency = *Iter;
-		if (Dependency.GetType() == ECookDependency::TransitiveBuild ||
-			Dependency.GetType() == ECookDependency::Package)
-		{
-			FName DependencyName = Dependency.GetPackageName();
-			// Remove transitive dependencies to self, for performance. But keep the package dependency to
-			// self; every cooked package has its EditorDomain package as a dependency.
-			if (Dependency.GetType() == ECookDependency::TransitiveBuild && DependencyName == PackageName)
+			for (const FCookDependency& IncludeDependency : IncludeDependencies.Dependencies.GetDependencies())
 			{
-				Iter.RemoveCurrentSwap();
-				continue;
+				// Recursive SettingsDependencies are not allowed. We haven't needed them yet, and not supporting them allows
+				// prevents the need for cycle detection.
+				if (IncludeDependency.GetType() == ECookDependency::SettingsObject)
+				{
+					if (OutMessages)
+					{
+						OutMessages->Emplace(ELogVerbosity::Error, FString::Printf(
+							TEXT("Settings dependency on object %s, but that object has a recursive Settings dependency on %s, and recursive Settings dependencies are not supported."),
+							*SettingsObject->GetPathName(), *IncludeDependency.GetSettingsObject()->GetPathName()));
+					}
+					return false;
+				}
+				ResultDependencies.Add(IncludeDependency);
 			}
+		}
 
-			// We do not hash dependencies to non-content packages (e.g. temp, memory, script),
-			// so remove package or transitive package dependencies to them.
-			EPackageMountPoint MountPoint = GetPackageMountPoint(DependencyName);
-			if (MountPoint != EPackageMountPoint::Content && MountPoint != EPackageMountPoint::GeneratedContent)
+		// Process some rules for Package dependencies
+		TSet<FName> RedirectionTargets;
+		for (TArray<FCookDependency>::TIterator Iter(ResultDependencies); Iter; ++Iter)
+		{
+			FCookDependency& Dependency = *Iter;
+			if (Dependency.GetType() == ECookDependency::TransitiveBuild ||
+				Dependency.GetType() == ECookDependency::Package)
 			{
-				Iter.RemoveCurrentSwap();
-				continue;
-			}
+				FName DependencyName = Dependency.GetPackageName();
+				// Remove transitive dependencies to self, for performance. But keep the package dependency to
+				// self; every cooked package has its EditorDomain package as a dependency.
+				if (Dependency.GetType() == ECookDependency::TransitiveBuild && DependencyName == PackageName)
+				{
+					Iter.RemoveCurrentSwap();
+					continue;
+				}
 
-			// Remove dependencies to generated packages, except for a generated package's dependency to itself.
-			// We do not yet support the availability of the digest of other generated packages when requested from
-			// the savepackage and dependency collection of a generated package or a generator; 
-			// the digests only become available when the target generated package is saved, which can happen after
-			// the save of the packages that refer to it.
-			if (MountPoint == EPackageMountPoint::GeneratedContent && DependencyName != PackageName)
-			{
-				Iter.RemoveCurrentSwap();
-				continue;
-			}
+				// We do not hash dependencies to non-content packages (e.g. temp, memory, script),
+				// so remove package or transitive package dependencies to them.
+				EPackageMountPoint MountPoint = GetPackageMountPoint(DependencyName);
+				if (MountPoint != EPackageMountPoint::Content && MountPoint != EPackageMountPoint::GeneratedContent)
+				{
+					Iter.RemoveCurrentSwap();
+					continue;
+				}
 
-			bool bPackageExistOnDisk = AssetRegistry->DoesPackageExistOnDisk(DependencyName);
-			if (!bPackageExistOnDisk)
-			{
-				TStringBuilder<256> DependencyNameStr(InPlace, DependencyName);
-				UPackage* DependencyPackage = FindPackage(nullptr, *DependencyNameStr);
-				if (!DependencyPackage)
+				// Remove dependencies to generated packages, except for a generated package's dependency to itself.
+				// We do not yet support the availability of the digest of other generated packages when requested from
+				// the savepackage and dependency collection of a generated package or a generator; 
+				// the digests only become available when the target generated package is saved, which can happen after
+				// the save of the packages that refer to it.
+				if (MountPoint == EPackageMountPoint::GeneratedContent && DependencyName != PackageName)
 				{
-					if (OutMessages)
+					Iter.RemoveCurrentSwap();
+					continue;
+				}
+
+				bool bPackageExistOnDisk = AssetRegistry->DoesPackageExistOnDisk(DependencyName);
+				if (!bPackageExistOnDisk)
+				{
+					TStringBuilder<256> DependencyNameStr(InPlace, DependencyName);
+					UPackage* DependencyPackage = FindPackage(nullptr, *DependencyNameStr);
+					if (!DependencyPackage)
 					{
-						OutMessages->Emplace(ELogVerbosity::Error,
-							FString::Printf(TEXT("Package %s does not exist."), *DependencyNameStr));
-					}
+						if (OutMessages)
+						{
+							OutMessages->Emplace(ELogVerbosity::Error,
+								FString::Printf(TEXT("Package %s does not exist."), *DependencyNameStr));
+						}
 
-					return false;
+						return false;
+					}
+					else if (DependencyPackage->HasAnyPackageFlags(PKG_NewlyCreated))
+					{
+						// If the package is a newly created package (in-memory package) then ignore it.
+						// In-memory packages are ignored because we can't compute their digest. Only packages on disk have a digest.
+						Iter.RemoveCurrentSwap();
+						continue;
+					}
+					// else the package is not on disk, in-memory and not newly created. It's a strange edge case but let's register it to the dependencies.
 				}
-				else if (DependencyPackage->HasAnyPackageFlags(PKG_NewlyCreated))
+
+				// Package dependencies of all kinds (Runtime, Build, TransitiveBuild) also cause RedirectionTarget
+				// dependencies.
+				RedirectionTargets.Add(DependencyName);
+
+				// Deprecated TransitiveBuildAndRuntime dependencies can also cause runtime dependencies;
+				// convert them to separate dependencies now.
+				if (Dependency.GetType() == ECookDependency::TransitiveBuild && ResultPair.Key == BuildResult::NAME_Save)
 				{
-					// If the package is a newly created package (in-memory package) then ignore it.
-					// In-memory packages are ignored because we can't compute their digest. Only packages on disk have a digest.
-					Iter.RemoveCurrentSwap();
-					continue;
+					PRAGMA_DISABLE_DEPRECATION_WARNINGS;
+					if (Dependency.IsAlsoAddRuntimeDependency())
+					{
+						InOutRuntimeDependencies.Add(DependencyName);
+						// Remove the IsAlsoRuntimeDependency flag
+						Dependency = FCookDependency::TransitiveBuild(DependencyName);
+					}
+					PRAGMA_ENABLE_DEPRECATION_WARNINGS;
 				}
-				// else the package is not on disk, in-memory and not newly created. It's a strange edge case but let's register it to the dependencies.
 			}
+		}
 
-			// Package dependencies of all kinds (Runtime, Build, TransitiveBuild) also cause RedirectionTarget
-			// dependencies.
-			RedirectionTargets.Add(DependencyName);
-
-			// Deprecated TransitiveBuildAndRuntime dependencies can also cause runtime dependencies;
-			// convert them to separate dependencies now.
-			if (Dependency.GetType() == ECookDependency::TransitiveBuild)
+		if (ResultPair.Key == BuildResult::NAME_Save)
+		{
+			// Pull transient packages out of the runtime dependencies for performance; we don't need them for
+			// deciding what gets cooked. 
+			// Runtime and script dependencies also cause RedirectionTarget dependencies, so record those.
+			for (TArray<FName>::TIterator Iter(InOutRuntimeDependencies); Iter; ++Iter)
 			{
-				PRAGMA_DISABLE_DEPRECATION_WARNINGS;
-				if (Dependency.IsAlsoAddRuntimeDependency())
+				FName DependencyPackageName = *Iter;
+				EPackageMountPoint MountPoint = GetPackageMountPoint(DependencyPackageName);
+				switch (MountPoint)
 				{
-					InOutRuntimeDependencies.Add(DependencyName);
-					// Remove the IsAlsoRuntimeDependency flag
-					Dependency = FCookDependency::TransitiveBuild(DependencyName);
+				case EPackageMountPoint::GeneratedContent:
+				case EPackageMountPoint::Content:
+				case EPackageMountPoint::Script:
+					// Keep it
+					RedirectionTargets.Add(DependencyPackageName);
+					break;
+				default:
+					Iter.RemoveCurrentSwap();
+					break;
 				}
-				PRAGMA_ENABLE_DEPRECATION_WARNINGS;
 			}
 		}
-	}
 
-	// Pull transient packages out of the runtime dependencies for performance; we don't need them for
-	// deciding what gets cooked. 
-	// Runtime and script dependencies also cause RedirectionTarget dependencies, so record those.
-	for (TArray<FName>::TIterator Iter(InOutRuntimeDependencies); Iter; ++Iter)
-	{
-		FName DependencyPackageName = *Iter;
-		EPackageMountPoint MountPoint = GetPackageMountPoint(DependencyPackageName);
-		switch (MountPoint)
+		// Put all the extra redirection dependencies into BuildDependencies
+		for (FName RedirectionTarget : RedirectionTargets)
 		{
-		case EPackageMountPoint::GeneratedContent:
-		case EPackageMountPoint::Content:
-		case EPackageMountPoint::Script:
-			// Keep it
-			RedirectionTargets.Add(DependencyPackageName);
-			break;
-		default:
-			Iter.RemoveCurrentSwap();
-			break;
+			ResultDependencies.Add(FCookDependency::RedirectionTarget(RedirectionTarget));
 		}
 	}
-
-	// Put all the extra redirection dependencies into BuildDependencies
-	for (FName RedirectionTarget : RedirectionTargets)
-	{
-		OutBuildDependencies.Add(FCookDependency::RedirectionTarget(RedirectionTarget));
-	}
 	return true;
 }
 
@@ -714,22 +727,19 @@ FBuildDependencySet FBuildDependencySet::CollectSettingsObject(const UObject* Ob
 	return Result;
 }
 
-FName FPackageArtifacts::NAME_Save(TEXT("Save"));
-FName FPackageArtifacts::NAME_Load(TEXT("Load"));
-
 FPackageArtifacts::FPackageArtifacts()
 {
-	LoadBuildDependencies.SetName(NAME_Load);
-	SaveBuildDependencies.SetName(NAME_Save);
+	LoadBuildDependencies.SetName(BuildResult::NAME_Load);
+	SaveBuildDependencies.SetName(BuildResult::NAME_Save);
 }
 
 FBuildDependencySet& FPackageArtifacts::FindOrAddBuildDependencySet(FName ResultName)
 {
-	if (ResultName == NAME_Save)
+	if (ResultName == BuildResult::NAME_Save)
 	{
 		return SaveBuildDependencies;
 	}
-	else if (ResultName == NAME_Load)
+	else if (ResultName == BuildResult::NAME_Load)
 	{
 		return LoadBuildDependencies;
 	}
@@ -743,11 +753,11 @@ FBuildDependencySet& FPackageArtifacts::FindOrAddBuildDependencySet(FName Result
 
 FBuildDependencySet* FPackageArtifacts::FindBuildDependencySet(FName ResultName)
 {
-	if (ResultName == NAME_Save)
+	if (ResultName == BuildResult::NAME_Save)
 	{
 		return &SaveBuildDependencies;
 	}
-	else if (ResultName == NAME_Load)
+	else if (ResultName == BuildResult::NAME_Load)
 	{
 		return &LoadBuildDependencies;
 	}
@@ -984,17 +994,18 @@ void FIncrementalCookAttachments::AppendCommitAttachments(TArray<IPackageWriter:
 }
 
 FIncrementalCookAttachments FIncrementalCookAttachments::Collect(const UPackage* Package,
-	const ITargetPlatform* TargetPlatform, const FBuildDependencySet* LoadDependencies,
-	const FSavePackageResultStruct* SaveResult, FGenerationHelper* GenerationHelper, bool bGenerated,
-	TArray<FName>&& RuntimeDependencies)
+	const ITargetPlatform* TargetPlatform, FBuildResultDependenciesMap&& InResultDependencies,
+	bool bHasSaveResult, TConstArrayView<FName> UntrackedSoftPackageReferences, FGenerationHelper* GenerationHelper,
+	bool bGenerated, TArray<FName>&& RuntimeDependencies)
 {
 	FIncrementalCookAttachments Result;
 
 	Result.CommitStatus = IPackageWriter::ECommitStatus::NotCommitted;
 
 	TArray<TPair<ELogVerbosity::Type, FString>> Messages;
-	Result.Artifacts = FPackageArtifacts::Collect(Package, TargetPlatform, LoadDependencies, SaveResult,
-		GenerationHelper, bGenerated, MoveTemp(RuntimeDependencies), &Messages);
+	Result.Artifacts = FPackageArtifacts::Collect(Package, TargetPlatform, MoveTemp(InResultDependencies),
+		bHasSaveResult, UntrackedSoftPackageReferences, GenerationHelper, bGenerated,
+		MoveTemp(RuntimeDependencies), &Messages);
 	if (!Result.Artifacts.IsValid())
 	{
 		TStringBuilder<256> LogText;
@@ -1084,8 +1095,8 @@ void FIncrementalCookAttachments::Fetch(TArrayView<FName> PackageNames, const IT
 }
 
 FPackageArtifacts FPackageArtifacts::Collect(const UPackage* Package, const ITargetPlatform* TargetPlatform,
-	const FBuildDependencySet* InLoadDependencies, const FSavePackageResultStruct* SaveResult,
-	FGenerationHelper* GenerationHelper, bool bGenerated,
+	FBuildResultDependenciesMap&& InResultDependencies, bool bHasSaveResult,
+	TConstArrayView<FName> UntrackedSoftPackageReferences, FGenerationHelper* GenerationHelper, bool bGenerated,
 	TArray<FName>&& InRuntimeDependencies, TArray<TPair<ELogVerbosity::Type, FString>>* OutMessages)
 {
 	if (!Package)
@@ -1113,44 +1124,47 @@ FPackageArtifacts FPackageArtifacts::Collect(const UPackage* Package, const ITar
 	}
 
 	// Collect the save's BuildDependencies, and pass in our runtimedependencies for read/write.
-	TArray<FCookDependency> BuildDependencies;
-	if (!FBuildDependencySet::TryCollectInternal(BuildDependencies, InRuntimeDependencies, Package,
-		TargetPlatform, SaveResult, GenerationHelper, bGenerated, OutMessages))
+	if (!FBuildDependencySet::TryCollectInternal(InResultDependencies, InRuntimeDependencies, OutMessages,
+		BuildResult::NAME_Save, Package, TargetPlatform, UntrackedSoftPackageReferences, GenerationHelper, bGenerated))
 	{
 		return FPackageArtifacts();
 	}
-
 	// Sort and remove duplicates in the results from TryCollectInternal.
-	Algo::Sort(BuildDependencies);
-	BuildDependencies.SetNum(Algo::Unique(BuildDependencies), EAllowShrinking::Yes);
-	Algo::Sort(InRuntimeDependencies, FNameLexicalLess());
-	InRuntimeDependencies.SetNum(Algo::Unique(InRuntimeDependencies), EAllowShrinking::Yes);
+	for (TPair<FName, TArray<FCookDependency>>& BuildResultPair : InResultDependencies)
+	{
+		TArray<FCookDependency>& ResultDependencies = BuildResultPair.Value;
+		Algo::Sort(ResultDependencies);
+		ResultDependencies.SetNum(Algo::Unique(ResultDependencies));
+	}
 
 	FPackageArtifacts Result;
 	Result.PackageName = PackageName;
-	Result.bHasSaveResults = SaveResult != nullptr;
+	Result.bHasSaveResults = bHasSaveResult;
 
 	// Store input+collected RuntimeDependencies on the result
+	Algo::Sort(InRuntimeDependencies, FNameLexicalLess());
+	InRuntimeDependencies.SetNum(Algo::Unique(InRuntimeDependencies), EAllowShrinking::Yes);
 	Result.RuntimeDependencies = MoveTemp(InRuntimeDependencies);
 
-	// Store input LoadDependencies on the result
-	if (InLoadDependencies && InLoadDependencies->IsValid())
-	{
-		Result.LoadBuildDependencies = *InLoadDependencies;
-	}
-	else
-	{
-		// An empty set of dependencies, create a normalized version of it.
+	// Store the collected LoadBuildDependencies on the result
+	TArray<FCookDependency>& LoadDependencies = InResultDependencies.FindOrAdd(BuildResult::NAME_Load);
+	Result.LoadBuildDependencies.SetName(BuildResult::NAME_Load);
+	Result.LoadBuildDependencies.SetNormalizedDependencies(MoveTemp(LoadDependencies));
 		Result.LoadBuildDependencies.TryCalculateCurrentKey(PackageName,
 			nullptr /* GenerationHelper */, nullptr /* OutMessages */);
 		Result.LoadBuildDependencies.StoreCurrentKey();
 		Result.LoadBuildDependencies.SetValid(true);
-	}
-	Result.LoadBuildDependencies.SetName(NAME_Load);
-	check(Result.LoadBuildDependencies.IsValid());
+
+	// Copy LoadBuildDependencies onto SaveBuildDependencies
+	// TODO: Add a transitive builddependency from SaveBuildDependencies to LoadBuildDependencies rather than duplicating.
+	TArray<FCookDependency>& SaveDependencies = InResultDependencies.FindOrAdd(BuildResult::NAME_Save);
+	SaveDependencies.Append(Result.LoadBuildDependencies.GetDependencies());
+	Algo::Sort(SaveDependencies);
+	SaveDependencies.SetNum(Algo::Unique(SaveDependencies), EAllowShrinking::Yes);
 
 	// Store the collected SaveDependencies on the result
-	Result.SaveBuildDependencies.SetNormalizedDependencies(MoveTemp(BuildDependencies));
+	Result.SaveBuildDependencies.SetName(BuildResult::NAME_Save);
+	Result.SaveBuildDependencies.SetNormalizedDependencies(MoveTemp(SaveDependencies));
 	FBuildDependencySet::ECurrentKeyResult CurrentKeyResult = Result.SaveBuildDependencies.TryCalculateCurrentKey(
 		PackageName, GenerationHelper, OutMessages);
 	if (CurrentKeyResult == FBuildDependencySet::ECurrentKeyResult::Error)
diff --git a/Engine/Source/Editor/UnrealEd/Private/Cooker/CookPackageArtifacts.h b/Engine/Source/Editor/UnrealEd/Private/Cooker/CookPackageArtifacts.h
--- a/Engine/Source/Editor/UnrealEd/Private/Cooker/CookPackageArtifacts.h
+++ b/Engine/Source/Editor/UnrealEd/Private/Cooker/CookPackageArtifacts.h
@@ -26,6 +26,7 @@ struct FSavePackageResultStruct;
 template <typename FuncType> class TUniqueFunction;
 
 namespace UE::Cook { class FPackageArtifacts; }
+namespace UE::Cook { class ICookInfo; }
 namespace UE::Cook { struct FBuildDefinitionList; }
 namespace UE::Cook { struct FGenerationHelper; }
 namespace UE::Cook { struct FIncrementalCookAttachments; }
@@ -125,24 +126,26 @@ public:
 	 * Read dependencies for the given targetplatform of the given package out of global dependency trackers
 	 * that have recorded its data during the package's load operations in the current cook session.
 	 */
-	static FBuildDependencySet CollectLoadedPackage(const UPackage* Package,
+	static FBuildResultDependenciesMap CollectLoadedPackage(const UPackage* Package,
 		TArray<TPair<ELogVerbosity::Type, FString>>* OutMessages = nullptr);
 
 	/**
 	 * Internal helper for FBuildDependencySet::Collect and FPackageArtifacts::Collect. Handles all arguments
-	 * used by either of them, and returns an unsorted,non-unique array of CookDependencies rather than a normalized
-	 * FBuildDependencySet.
+	 * used by either of them, and returns a map of BuildResultDependencies. Each returned map entry contains the
+	 * dependencies for a BuildResult (e.g. NAME_Save); those dependencies are not yet sorted or unique.
 	 * 
+	 * @param DefaultBuildResult - into which build result (NAME_Load or NAME_Save) detected dependencies should be
+	 *                             added onto, for each dependency that does not already have a buildresult specified.
 	 * @param TargetPlatform - If null, collects dependencies reported by all platforms. If non-null, only collects
 	 *                         dependencies reported for the given platform or reported as platform-agnostic.
 	 * @param SaveResult - If non-null, contains data reported from the SavePackage call.
 	 * @param GenerationHelper - If non-null, provides the lookup for the AssetPackageData of generated packages.
 	 *                           Must be provided if any generated packages are in the dependencies.
 	 */
-	static bool TryCollectInternal(TArray<FCookDependency>& OutBuildDependencies, TArray<FName>& OutRuntimeDependencies,
-		const UPackage* Package, const ITargetPlatform* TargetPlatform, const FSavePackageResultStruct* SaveResult,
-		FGenerationHelper* GenerationHelper, bool bGenerated,
-		TArray<TPair<ELogVerbosity::Type, FString>>* OutMessages);
+	static bool TryCollectInternal(FBuildResultDependenciesMap& InOutResultDependencies,
+		TArray<FName>& InOutRuntimeDependencies, TArray<TPair<ELogVerbosity::Type, FString>>* OutMessages,
+		FName DefaultBuildResult, const UPackage* Package, const ITargetPlatform* TargetPlatform,
+		TConstArrayView<FName> UntrackedSoftPackageReferences, FGenerationHelper* GenerationHelper, bool bGenerated);
 
 	/**
 	 * Collect the dependencies referenced by a given settingsobject from e.g. config. Globally cached for the current
@@ -204,10 +207,6 @@ public:
 
 	FBuildDependencySet& FindOrAddBuildDependencySet(FName ResultName);
 	FBuildDependencySet* FindBuildDependencySet(FName ResultName);
-	/** The fixed name used for the SavePackage BuildResult. */
-	static FName NAME_Save;
-	/** The fixed name used for the LoadPackage BuildResult. */
-	static FName NAME_Load;
 
 	/** Get all of the runtime dependencies reported by the package, both script and content. */
 	const TArray<FName>& GetRuntimeDependencies() const;
@@ -240,9 +239,9 @@ public:
 	 *                           Must be provided if any generated packages are in the dependencies.
 	 */
 	static FPackageArtifacts Collect(const UPackage* Package, const ITargetPlatform* TargetPlatform,
-		const FBuildDependencySet* InLoadDependencies, const FSavePackageResultStruct* SaveResult,
-		FGenerationHelper* GenerationHelper, bool bGenerated, TArray<FName>&& InRuntimeDependencies,
-		TArray<TPair<ELogVerbosity::Type, FString>>* OutMessages = nullptr);
+		FBuildResultDependenciesMap&& InResultDependencies, bool bHasSaveResult,
+		TConstArrayView<FName> UntrackedSoftPackageReferences, FGenerationHelper* GenerationHelper, bool bGenerated,
+		TArray<FName>&& InRuntimeDependencies, TArray<TPair<ELogVerbosity::Type, FString>>* OutMessages = nullptr);
 
 	// Fetch function to load the dependencies from a PackageStore is not yet implemented independently for
 	// this structure. Use FIncrementalCookAttachments instead.
@@ -324,9 +323,9 @@ struct FIncrementalCookAttachments
 		ICookedPackageWriter* PackageWriter,
 		TUniqueFunction<void(FName PackageName, FIncrementalCookAttachments&& Result)>&& Callback);
 	static FIncrementalCookAttachments Collect(const UPackage* Package,
-		const ITargetPlatform* TargetPlatform, const FBuildDependencySet* LoadDependencies,
-		const FSavePackageResultStruct* SaveResult, FGenerationHelper* GenerationHelper, bool bGenerated,
-		TArray<FName>&& RuntimeDependencies);
+		const ITargetPlatform* TargetPlatform, FBuildResultDependenciesMap&& InResultDependencies,
+		bool bHasSaveResult, TConstArrayView<FName> UntrackedSoftPackageReferences, FGenerationHelper* GenerationHelper,
+		bool bGenerated, TArray<FName>&& RuntimeDependencies);
 };
 
 } // namespace UE::Cook
diff --git a/Engine/Source/Editor/UnrealEd/Private/Cooker/CookPackageData.cpp b/Engine/Source/Editor/UnrealEd/Private/Cooker/CookPackageData.cpp
--- a/Engine/Source/Editor/UnrealEd/Private/Cooker/CookPackageData.cpp
+++ b/Engine/Source/Editor/UnrealEd/Private/Cooker/CookPackageData.cpp
@@ -11,6 +11,7 @@
 #include "AssetRegistry/IAssetRegistry.h"
 #include "Async/ParallelFor.h"
 #include "CompactBinaryTCP.h"
+#include "Cooker/BuildResultDependenciesMap.h"
 #include "Cooker/CookDirector.h"
 #include "Cooker/CookGenerationHelper.h"
 #include "Cooker/CookPackagePreloader.h"
@@ -1239,7 +1240,7 @@ void FPackageData::OnPackagePreloaderDestroyed(FPackagePreloader& InPackagePrelo
 	PackagePreloader = nullptr;
 }
 
-FBuildDependencySet* FPackageData::GetLoadDependencies()
+const FBuildResultDependenciesMap* FPackageData::GetLoadDependencies() const
 {
 	return LoadDependencies.Get();
 }
@@ -1252,7 +1253,7 @@ void FPackageData::CreateLoadDependencies()
 		checkf(LocalPackage != nullptr,
 			TEXT("CreateLoadDependencies failed for package %s because this->Package == nullptr. It can only be called after the Package has been set."),
 			*GetPackageName().ToString());
-		LoadDependencies.Reset(new FBuildDependencySet(FBuildDependencySet::CollectLoadedPackage(LocalPackage)));
+		LoadDependencies.Reset(new FBuildResultDependenciesMap(FBuildDependencySet::CollectLoadedPackage(LocalPackage)));
 	}
 }
 
diff --git a/Engine/Source/Editor/UnrealEd/Private/Cooker/CookPackageData.h b/Engine/Source/Editor/UnrealEd/Private/Cooker/CookPackageData.h
--- a/Engine/Source/Editor/UnrealEd/Private/Cooker/CookPackageData.h
+++ b/Engine/Source/Editor/UnrealEd/Private/Cooker/CookPackageData.h
@@ -35,10 +35,10 @@ class FCbWriter;
 class UCookOnTheFlyServer;
 class UObject;
 class UPackage;
-namespace UE::Cook { class FBuildDependencySet; }
 namespace UE::Cook { class FCookWorkerClient; }
 namespace UE::Cook { class FPackagePreloader; }
 namespace UE::Cook { class FRequestCluster; }
+namespace UE::Cook { struct FBuildResultDependenciesMap; }
 namespace UE::Cook { struct FConstructPackageData; }
 namespace UE::Cook { struct FDiscoveredPlatformSet; }
 
@@ -577,7 +577,7 @@ public:
 	 * incremental cook oplog and are used by CookRequestCluster to test invalidation of previous cookresults for the
 	 * package and for other packages that have the package as a transitive build dependency.
 	 */
-	FBuildDependencySet* GetLoadDependencies();
+	const FBuildResultDependenciesMap* GetLoadDependencies() const;
 	/** Calculate the LoadDependencies if not already created. */
 	void CreateLoadDependencies();
 	void ClearLoadDependencies();
@@ -845,7 +845,7 @@ private:
 	TArray<FCachedObjectInOuter> CachedObjectsInOuter;
 	FCompletionCallback CompletionCallback;
 	TUniquePtr<TMap<FPackageData*, EInstigator>> Unsolicited;
-	TUniquePtr<UE::Cook::FBuildDependencySet> LoadDependencies;
+	TUniquePtr<FBuildResultDependenciesMap> LoadDependencies;
 	FName PackageName;
 	FName FileName;
 	FName ParentGenerator;
diff --git a/Engine/Source/Editor/UnrealEd/Private/Cooker/CookSavePackage.cpp b/Engine/Source/Editor/UnrealEd/Private/Cooker/CookSavePackage.cpp
--- a/Engine/Source/Editor/UnrealEd/Private/Cooker/CookSavePackage.cpp
+++ b/Engine/Source/Editor/UnrealEd/Private/Cooker/CookSavePackage.cpp
@@ -8,6 +8,7 @@
 #include "Cooker/CookDependency.h"
 #include "Cooker/CookDeterminismManager.h"
 #include "Cooker/CookDiagnostics.h"
+#include "Cooker/CookEvents.h"
 #include "Cooker/CookGenerationHelper.h"
 #include "Cooker/CookOnTheFlyServerInterface.h"
 #include "Cooker/CookPackageArtifacts.h"
@@ -41,6 +42,7 @@
 #include "Trace/Trace.h"
 #include "UObject/CookEnums.h"
 #include "UObject/ObjectMacros.h"
+#include "UObject/ObjectSaveContext.h"
 #include "UObject/ObjectSaveOverride.h"
 #include "UObject/SavePackage.h"
 #include "UObject/SoftObjectPath.h"
@@ -184,6 +186,10 @@ void UCookOnTheFlyServer::CommitUncookedPackage(UE::Cook::FSaveCookedPackageCont
 	using namespace UE::Cook;
 
 	UPackage* Package = Context.Package;
+	FOnScopeExit ScopedPackageFlags([Package, OriginalPackageFlags = Package->GetPackageFlags()]()
+		{
+			Package->SetPackageFlagsTo(OriginalPackageFlags);
+		});
 	Context.SetupPackage();
 
 	for (int32 PlatformIndex = 0; PlatformIndex < Context.PlatformsForPackage.Num(); ++PlatformIndex)
@@ -255,13 +261,13 @@ void FSaveCookedPackageContext::SetupPlatform(const ITargetPlatform* InTargetPla
 	SavePackageContext = &CookContext->SaveContext;
 	PackageWriter = CookContext->PackageWriter;
 
+	ArchiveCookContext.Emplace(Package,
+		COTFS.IsDirectorCookByTheBook() ? UE::Cook::ECookType::ByTheBook : UE::Cook::ECookType::OnTheFly,
+		COTFS.IsCookingDLC() ? UE::Cook::ECookingDLC::Yes : UE::Cook::ECookingDLC::No,
+		TargetPlatform, &COTFS);
+
 	if (CommitType == EReachability::Runtime)
 	{
-		ArchiveCookContext.Emplace(Package,
-			COTFS.IsDirectorCookByTheBook() ? UE::Cook::ECookType::ByTheBook : UE::Cook::ECookType::OnTheFly,
-			COTFS.IsCookingDLC() ? UE::Cook::ECookingDLC::Yes : UE::Cook::ECookingDLC::No,
-			TargetPlatform);
-
 		// don't save Editor resources from the Engine if the target doesn't have editoronly data
 		if (COTFS.IsCookFlagSet(ECookInitializationFlags::SkipEditorContent) &&
 			(PackageName.StartsWith(TEXT("/Engine/Editor")) || PackageName.StartsWith(TEXT("/Engine/VREditor"))) &&
@@ -453,7 +459,7 @@ void FSaveCookedPackageContext::FinishPlatform()
 	{
 		CalculatePlatformAgnosticRuntimeDependencies();
 		CalculatePlatformRuntimeDependencies();
-		// Note CalculateCookDependencies mutates this->SaveResult; it moves CookDependencies out of it
+		// Note CalculateCookDependencies mutates this->SaveResult; it moves BuildResultDependencies out of it
 		CalculateCookDependencies(GenerationHelper.GetReference(), bGenerated, OverridePackageDependencies.GetPtrOrNull());
 		RecordPlatformBuildDependencies();
 
@@ -518,13 +524,12 @@ void FSaveCookedPackageContext::FinishPlatform()
 		bAnySaveSucceeded |= bSuccessful;
 
 		bHasRetryErrorCode |= bIsRetryErrorCode;
-
-		ArchiveCookContext.Reset();
 	}
 	else
 	{
 		PackageData.SetPlatformCommitted(TargetPlatform);
 	}
+	ArchiveCookContext.Reset();
 	PlatformIndex = -1;
 }
 
@@ -735,13 +740,43 @@ TArray<FName> FSaveCookedPackageContext::GetPlatformRuntimeDependencies() const
 	return PlatformDependencyNames;
 }
 
+/** FArchive used to collect BuildResultDependencies from structs on UObjects in a package being recorded as a BuildDependency. */
+class FBuildDependencyHarvestArchive : public UE::SavePackageUtilities::Private::FArchiveSavePackageCollector
+{
+public:
+	FBuildDependencyHarvestArchive(UPackage* InPackage, FArchiveCookContext& ArchiveCookContext,
+		FObjectSaveContextData& InObjectSaveContextData)
+		: ObjectSaveContextData(InObjectSaveContextData)
+		, ObjectSavePackageSerializeContext(ObjectSaveContextData)
+		, ArchiveSavePackageData(ObjectSavePackageSerializeContext,
+			InObjectSaveContextData.TargetPlatform, &ArchiveCookContext)
+	{
+		SetArchiveFlags(ArchiveSavePackageData, InPackage->HasAnyPackageFlags(PKG_FilterEditorOnly),
+			(ObjectSaveContextData.SaveFlags & SAVE_Unversioned) != 0, true /* bCooking */);
+	}
+
+	// We use the empty operator<< functions defined on the parent class. The function of this class is not to interpret
+	// any serialized data, it is instead only to provide structs and UObjects access to the
+	// FObjectSavePackageSerializeContext API, which they get by calling
+	// Archive.GetSavePackageData()->SavePackageContext.
+
+private:
+	FObjectSaveContextData ObjectSaveContextData;
+	FObjectSavePackageSerializeContext ObjectSavePackageSerializeContext;
+	FArchiveSavePackageData ArchiveSavePackageData;
+};
+
 void FSaveCookedPackageContext::CalculateCookDependencies(FGenerationHelper* GenerationHelper, bool bGenerated,
 	TArray<FAssetDependency>* OverrideDependencies)
 {
 	if (COTFS.bHybridIterativeEnabled)
 	{
 		UE_SCOPED_HIERARCHICAL_COOKTIMER(TargetDomainDependencies);
-		TArray<FName> PlatformRuntimeDependencies = GetPlatformRuntimeDependencies();
+		FBuildResultDependenciesMap BuildResultDependencies;
+		TArray<FName> PlatformRuntimeDependencies;
+		TConstArrayView<FName> UntrackedSoftPackageReferences;
+		
+		PlatformRuntimeDependencies = GetPlatformRuntimeDependencies();
 		if (OverrideDependencies)
 		{
 			PlatformRuntimeDependencies.Reserve(PlatformRuntimeDependencies.Num() + OverrideDependencies->Num());
@@ -750,14 +785,65 @@ void FSaveCookedPackageContext::CalculateCookDependencies(FGenerationHelper* Gen
 				PlatformRuntimeDependencies.Add(Dependency.AssetId.PackageName);
 			}
 		}
-		FBuildDependencySet* LoadDependencies = PackageData.GetLoadDependencies();
+		bool bHasSavePackageResult = HasSavePackageResult();
+		FSavePackageResultStruct* LocalSavePackageResult = bHasSavePackageResult ? &SavePackageResult : nullptr;
+
+		if (!bHasSavePackageResult)
+		{
+			// During cook saves of a Package, SavePackage is responsible for collecting BuildResultDependencies from
+			// objects, but if this SaveCookedPackage call is instead for the recording of a BuildDependency package,
+			// we need to collect them here.
+			// Call OnCookEvent(ECookEvent::LoadDependencies) on every UObject in the package, and to support UStructs,
+			// also serialize each object into an FArchive that identifies as a SavePackage archive and contains an
+			// FArchiveCookContext.
+			// Pass the reported BuildResult dependencies into FIncrementalCookAttachments::Collect.
+			FObjectSaveContextData ObjectSaveContextData;
+			ObjectSaveContextData.Set(Package, TargetPlatform, FPackagePath(), SaveFlags);
+			ObjectSaveContextData.CookType = COTFS.GetCookType();
+			ObjectSaveContextData.CookingDLC = COTFS.GetCookingDLC();
+			ObjectSaveContextData.CookInfo = &COTFS;
+			ObjectSaveContextData.ObjectSaveContextPhase = EObjectSaveContextPhase::CookDependencyHarvest;
+
+			UE::Cook::FCookEventContext CookEventContext(ObjectSaveContextData);
+			FBuildDependencyHarvestArchive Harvester(Package, *ArchiveCookContext, ObjectSaveContextData);
+
+			for (FCachedObjectInOuter& CachedObjectInOuter : PackageData.GetCachedObjectsInOuter())
+			{
+				UObject* Object = CachedObjectInOuter.Object.Get();;
+				ObjectSaveContextData.Object = Object;
+				if (Object)
+				{
+					Object->OnCookEvent(ECookEvent::PlatformCookDependencies, CookEventContext);
+					Object->Serialize(Harvester);
+				}
+			}
+
+			BuildResultDependencies = MoveTemp(ObjectSaveContextData.BuildResultDependencies);
+			for (FSoftObjectPath& RuntimeDependency : ObjectSaveContextData.CookRuntimeDependencies)
+			{
+				FName PackageDependency = RuntimeDependency.GetLongPackageFName();
+				if (!PackageDependency.IsNone())
+				{
+					PlatformRuntimeDependencies.Add(PackageDependency);
+				}
+			}
+		}
+
+		if (bHasSavePackageResult)
+		{
+			BuildResultDependencies.Append(MoveTemp(LocalSavePackageResult->BuildResultDependencies));
+			PlatformRuntimeDependencies.Append(MoveTemp(LocalSavePackageResult->SoftPackageReferences));
+			UntrackedSoftPackageReferences = LocalSavePackageResult->UntrackedSoftPackageReferences;
+		}
+		const FBuildResultDependenciesMap* LoadDependencies = PackageData.GetLoadDependencies();
 		checkf(LoadDependencies,
 			TEXT("LoadDependencies not found during save of package %s. LoadDependencies are supposted to be created by LoadPackageInQueue before entering the Save state."),
 			*PackageData.GetPackageName().ToString());
-		FSavePackageResultStruct* LocalSavePackageResult = HasSavePackageResult() ? &SavePackageResult : nullptr;
+		BuildResultDependencies.Append(*LoadDependencies);
 
-		PlatformCookAttachments = FIncrementalCookAttachments::Collect(Package, TargetPlatform, LoadDependencies,
-			LocalSavePackageResult, GenerationHelper, bGenerated, MoveTemp(PlatformRuntimeDependencies));
+		PlatformCookAttachments = FIncrementalCookAttachments::Collect(Package, TargetPlatform,
+			MoveTemp(BuildResultDependencies), bHasSavePackageResult, UntrackedSoftPackageReferences, GenerationHelper,
+			bGenerated, MoveTemp(PlatformRuntimeDependencies));
 	}
 }
 
diff --git a/Engine/Source/Editor/UnrealEd/Private/EditorDomain/EditorDomainUtils.cpp b/Engine/Source/Editor/UnrealEd/Private/EditorDomain/EditorDomainUtils.cpp
--- a/Engine/Source/Editor/UnrealEd/Private/EditorDomain/EditorDomainUtils.cpp
+++ b/Engine/Source/Editor/UnrealEd/Private/EditorDomain/EditorDomainUtils.cpp
@@ -2009,9 +2009,11 @@ bool TrySavePackage(UPackage* Package)
 
 	if (bStorageResultValid)
 	{
+		UE::Cook::FBuildResultDependenciesMap ExistingResultDependencies;
 		UE::Cook::FIncrementalCookAttachments CookAttachments =
 			UE::Cook::FIncrementalCookAttachments::Collect(
-				Package, nullptr /* TargetPlatform */, nullptr /* LoadDependencies */, nullptr /* SaveResult */,
+				Package, nullptr /* TargetPlatform */, MoveTemp(ExistingResultDependencies),
+				false /* bHasSaveResult */, TConstArrayView<FName>() /* UntrackedSoftPackageReferences */,
 				nullptr /* GenerationHelper */, false /* bGenerated */, TArray<FName>());
 		TArray<IPackageWriter::FCommitAttachmentInfo> Attachments;
 		CookAttachments.AppendCommitAttachments(Attachments);
diff --git a/Engine/Source/Editor/UnrealEd/Private/EditorEngine.cpp b/Engine/Source/Editor/UnrealEd/Private/EditorEngine.cpp
--- a/Engine/Source/Editor/UnrealEd/Private/EditorEngine.cpp
+++ b/Engine/Source/Editor/UnrealEd/Private/EditorEngine.cpp
@@ -4408,6 +4408,7 @@ FSavePackageResultStruct UEditorEngine::Save(UPackage* InOuter, UObject* InAsset
 	{
 		ObjectSaveContext.CookType = InSaveArgs.ArchiveCookData->CookContext.GetCookType();
 		ObjectSaveContext.CookingDLC = InSaveArgs.ArchiveCookData->CookContext.GetCookingDLC();
+		ObjectSaveContext.CookInfo = InSaveArgs.ArchiveCookData->CookContext.GetCookInfo();
 	}
 	UWorld *OriginalOwningWorld = nullptr;
 	if ( World )
diff --git a/Engine/Source/Runtime/CoreUObject/Private/Cooker/BuildResultDependenciesMap.cpp b/Engine/Source/Runtime/CoreUObject/Private/Cooker/BuildResultDependenciesMap.cpp
new file mode 100644
--- /dev/null
+++ b/Engine/Source/Runtime/CoreUObject/Private/Cooker/BuildResultDependenciesMap.cpp
@@ -0,0 +1,42 @@
+// Copyright Epic Games, Inc. All Rights Reserved.
+
+#include "Cooker/BuildResultDependenciesMap.h"
+
+#if WITH_EDITOR
+#include "Cooker/CookDependency.h"
+#endif
+
+#if WITH_EDITOR
+
+namespace UE::Cook
+{
+
+void FBuildResultDependenciesMap::Add(FName Name, UE::Cook::FCookDependency CookDependency)
+{
+	FindOrAdd(Name).Add(MoveTemp(CookDependency));
+}
+
+void FBuildResultDependenciesMap::Append(FName Name, TArray<UE::Cook::FCookDependency> CookDependencies)
+{
+	FindOrAdd(Name).Append(MoveTemp(CookDependencies));
+}
+
+void FBuildResultDependenciesMap::Append(const FBuildResultDependenciesMap& Other)
+{
+	for (const TPair<FName, TArray<FCookDependency>>& Pair: Other)
+	{
+		FindOrAdd(Pair.Key).Append(Pair.Value);
+	}
+}
+
+void FBuildResultDependenciesMap::Append(FBuildResultDependenciesMap&& Other)
+{
+	for (TPair<FName, TArray<FCookDependency>>& Pair : Other)
+	{
+		FindOrAdd(Pair.Key).Append(MoveTemp(Pair.Value));
+	}
+}
+
+} // namespace UE::Cook
+
+#endif // WITH_EDITOR
\ No newline at end of file
diff --git a/Engine/Source/Runtime/CoreUObject/Private/Cooker/CookDependency.cpp b/Engine/Source/Runtime/CoreUObject/Private/Cooker/CookDependency.cpp
--- a/Engine/Source/Runtime/CoreUObject/Private/Cooker/CookDependency.cpp
+++ b/Engine/Source/Runtime/CoreUObject/Private/Cooker/CookDependency.cpp
@@ -51,6 +51,14 @@ FName None(TEXT("UE_Cook_ResultProjection_None"));
 
 }
 
+namespace BuildResult
+{
+
+FName NAME_Save(TEXT("Save"));
+FName NAME_Load(TEXT("Load"));
+
+}
+
 FCookDependency FCookDependency::File(FStringView InFileName)
 {
 	FCookDependency Result(ECookDependency::File);
diff --git a/Engine/Source/Runtime/CoreUObject/Private/Cooker/CookEvents.cpp b/Engine/Source/Runtime/CoreUObject/Private/Cooker/CookEvents.cpp
new file mode 100644
--- /dev/null
+++ b/Engine/Source/Runtime/CoreUObject/Private/Cooker/CookEvents.cpp
@@ -0,0 +1,39 @@
+// Copyright Epic Games, Inc. All Rights Reserved.
+
+#include "Cooker/CookEvents.h"
+
+#if WITH_EDITOR
+#include "Cooker/CookDependency.h"
+#include "UObject/ICookInfo.h"
+#include "UObject/SavePackage/SavePackageUtilities.h"
+#include "UObject/SoftObjectPath.h"
+#endif
+
+#if WITH_EDITOR
+
+namespace UE::Cook
+{
+
+void FCookEventContext::AddLoadBuildDependency(FCookDependency CookDependency)
+{
+	Data.BuildResultDependencies.Add(BuildResult::NAME_Load, MoveTemp(CookDependency));
+}
+
+void FCookEventContext::AddSaveBuildDependency(FCookDependency CookDependency)
+{
+	Data.BuildResultDependencies.Add(BuildResult::NAME_Save, MoveTemp(CookDependency));
+}
+
+void FCookEventContext::AddRuntimeDependency(FSoftObjectPath ObjectName)
+{
+	Data.CookRuntimeDependencies.Add(MoveTemp(ObjectName));
+}
+
+void FCookEventContext::HarvestCookRuntimeDependencies(UObject* HarvestReferencesFrom)
+{
+	UE::SavePackageUtilities::HarvestCookRuntimeDependencies(Data, HarvestReferencesFrom);
+}
+
+} // namespace UE::Cook
+
+#endif // WITH_EDITOR
diff --git a/Engine/Source/Runtime/CoreUObject/Private/UObject/SavePackage.cpp b/Engine/Source/Runtime/CoreUObject/Private/UObject/SavePackage.cpp
--- a/Engine/Source/Runtime/CoreUObject/Private/UObject/SavePackage.cpp
+++ b/Engine/Source/Runtime/CoreUObject/Private/UObject/SavePackage.cpp
@@ -30,7 +30,7 @@ bool UPackage::SavePackage(UPackage* InOuter, UObject* Base, EObjectFlags TopLev
 	if (TargetPlatform != nullptr)
 	{
 		CookContext.Emplace(InOuter, UE::Cook::ECookType::Unknown,
-			UE::Cook::ECookingDLC::Unknown, TargetPlatform);
+			UE::Cook::ECookingDLC::Unknown, TargetPlatform, nullptr);
 		CookData.Emplace(*TargetPlatform, *CookContext);
 	}
 	FSavePackageArgs SaveArgs = { nullptr /* deprecated target platform */, CookData.GetPtrOrNull(), TopLevelFlags, SaveFlags, bForceByteSwapping,
diff --git a/Engine/Source/Runtime/CoreUObject/Private/UObject/SavePackage/PackageHarvester.cpp b/Engine/Source/Runtime/CoreUObject/Private/UObject/SavePackage/PackageHarvester.cpp
--- a/Engine/Source/Runtime/CoreUObject/Private/UObject/SavePackage/PackageHarvester.cpp
+++ b/Engine/Source/Runtime/CoreUObject/Private/UObject/SavePackage/PackageHarvester.cpp
@@ -7,6 +7,7 @@
 #include "Interfaces/ITargetPlatform.h"
 #include "Misc/PackageAccessTracking.h"
 #include "UObject/OverridableManager.h"
+#include "UObject/PropertyPortFlags.h"
 #include "UObject/SavePackage/SaveContext.h"
 #include "UObject/SavePackage/SavePackageUtilities.h"
 #include "UObject/UObjectGlobals.h"
@@ -193,26 +194,38 @@ bool DoesObjectNeedLoadForEditorGame(UObject* InObject)
 	return bNeedsLoadForEditorGame;
 }
 
-FArchiveSavePackageCollector::FArchiveSavePackageCollector(FSaveContext& SaveContext)
+namespace UE::SavePackageUtilities::Private
+{
+
+FArchiveSavePackageCollector::FArchiveSavePackageCollector(FArchiveSavePackageData& InSavePackageData,
+	bool bFilterEditorOnly, bool bSaveUnversioned, bool bCooking)
 	: FArchiveUObject()
 {
-	SetArchiveFlags(SaveContext);
+	SetArchiveFlags(InSavePackageData, bFilterEditorOnly, bSaveUnversioned, bCooking);
 }
 
-void FArchiveSavePackageCollector::SetArchiveFlags(FSaveContext& SaveContext)
+void FArchiveSavePackageCollector::SetArchiveFlags(FArchiveSavePackageData& InSavePackageData,
+	bool bFilterEditorOnly, bool bSaveUnversioned, bool bCooking)
 {
 	this->SetIsSaving(true);
 	this->SetIsPersistent(true);
 	ArIsObjectReferenceCollector = true;
 	ArShouldSkipBulkData = true;
-	ArIgnoreClassGeneratedByRef = SaveContext.IsCooking();
+	ArIgnoreClassGeneratedByRef = bCooking;
+
+	this->SetPortFlags(GetSavePackagePortFlags());
+	this->SetFilterEditorOnly(bFilterEditorOnly);
+	this->SetSavePackageData(&InSavePackageData);
+	this->SetUseUnversionedPropertySerialization(bSaveUnversioned);
+}
 
-	this->SetPortFlags(SaveContext.GetPortFlags());
-	this->SetFilterEditorOnly(SaveContext.IsFilterEditorOnly());
-	this->SetSavePackageData(&SaveContext.GetArchiveSavePackageData());
-	this->SetUseUnversionedPropertySerialization(SaveContext.IsSaveUnversionedProperties());
+uint32 GetSavePackagePortFlags()
+{
+	return PPF_DeepCompareInstances | PPF_DeepCompareDSOsOnly;
 }
 
+} // namespace UE::SavePackageUtilities::Private
+
 FPackageHarvester::FExportScope::FExportScope(FPackageHarvester& InHarvester, const FExportWithContext& InToProcess)
 	: Harvester(InHarvester)
 	, PreviousExportHarvestingRealms(InHarvester.CurrentExportHarvestingRealms)
@@ -456,7 +469,8 @@ bool FPackageHarvester::HasAnyExportHarvestingRealms()
 }
 
 FPackageHarvester::FPackageHarvester(FSaveContext& InContext)
-	: FArchiveSavePackageCollector(InContext)
+	: UE::SavePackageUtilities::Private::FArchiveSavePackageCollector(InContext.GetArchiveSavePackageData(),
+		InContext.IsFilterEditorOnly(), InContext.IsSaveUnversionedProperties(), InContext.IsCooking())
 	, SaveContext(InContext)
 {
 	ResolveOverrides();
diff --git a/Engine/Source/Runtime/CoreUObject/Private/UObject/SavePackage/PackageHarvester.h b/Engine/Source/Runtime/CoreUObject/Private/UObject/SavePackage/PackageHarvester.h
--- a/Engine/Source/Runtime/CoreUObject/Private/UObject/SavePackage/PackageHarvester.h
+++ b/Engine/Source/Runtime/CoreUObject/Private/UObject/SavePackage/PackageHarvester.h
@@ -5,9 +5,9 @@
 #include "Containers/Map.h"
 #include "Containers/Set.h"
 #include "Containers/Queue.h"
-#include "Serialization/ArchiveUObject.h"
 #include "UObject/NameTypes.h"
 #include "UObject/ObjectPtr.h"
+#include "UObject/SavePackage.h"
 
 enum class ESaveableStatus;
 enum class ESaveRealm : uint32;
@@ -19,20 +19,11 @@ struct FWeakObjectPtr;
 struct FLazyObjectPtr;
 struct FSoftObjectPath;
 
-/* Base archive class for FPackageHarvester and FPostSaveArchive used to setup all the common flags. */
-class FArchiveSavePackageCollector : public FArchiveUObject
-{
-public:
-	FArchiveSavePackageCollector(FSaveContext& SaveContext);
-
-	void SetArchiveFlags(FSaveContext& SaveContext);
-};
-
 /**
  * Helper class used to gather all package exports, imports and dependencies to build linker tables when saving packages
  * Gather Export, Imports, Referenced Names, SoftObjectPaths, Custom Object Versions
  */
-class FPackageHarvester : public FArchiveSavePackageCollector
+class FPackageHarvester : public UE::SavePackageUtilities::Private::FArchiveSavePackageCollector
 {
 public:
 	using FExportingRealmsArray = TArray<ESaveRealm, TInlineAllocator<2>>;
diff --git a/Engine/Source/Runtime/CoreUObject/Private/UObject/SavePackage/SaveContext.cpp b/Engine/Source/Runtime/CoreUObject/Private/UObject/SavePackage/SaveContext.cpp
--- a/Engine/Source/Runtime/CoreUObject/Private/UObject/SavePackage/SaveContext.cpp
+++ b/Engine/Source/Runtime/CoreUObject/Private/UObject/SavePackage/SaveContext.cpp
@@ -166,8 +166,9 @@ UObject* FSaveContext::ResolveForSave(TObjectPtr<UObject> InObject, UE::SavePack
 		{
 			// Declare the build dependencies necessary for our use of the object.
 #if WITH_EDITOR
-			ObjectSaveContext.CookBuildDependencies.Add(FCookDependency::Package(Object->GetPackage()->GetFName()));
-			ObjectSaveContext.CookBuildDependencies.Add(FCookDependency::NativeClass(Object->GetClass()));
+			TArray<FCookDependency>& SaveDependencies = ObjectSaveContext.BuildResultDependencies.FindOrAdd(BuildResult::NAME_Save);
+			SaveDependencies.Add(FCookDependency::Package(Object->GetPackage()->GetFName()));
+			SaveDependencies.Add(FCookDependency::NativeClass(Object->GetClass()));
 #endif
 			Status.bDeclaredSaveBuildDependency = true;
 		}
@@ -243,7 +244,7 @@ FSavePackageResultStruct FSaveContext::GetFinalResult()
 			ResultData.SoftPackageReferences.Add(PackageDependency);
 		}
 	}
-	ResultData.CookDependencies = MoveTemp(ObjectSaveContext.CookBuildDependencies);
+	ResultData.BuildResultDependencies = MoveTemp(ObjectSaveContext.BuildResultDependencies);
 #endif
 
 	const TSet<FName>& UntrackedPackageReferenceList = GetUntrackedSoftPackageReferenceList();
@@ -349,24 +350,33 @@ void FSaveContext::UpdateEditorRealmPackageBuildDependencies()
 	}
 
 #if WITH_EDITOR
-	PackageBuildDependencies.Reserve(ObjectSaveContext.CookBuildDependencies.Num());
-	for (UE::Cook::FCookDependency& CookDependency : ObjectSaveContext.CookBuildDependencies)
+	for (const TPair<FName, TArray<FCookDependency>>& BuildResultPair : ObjectSaveContext.BuildResultDependencies)
 	{
-		FName PackageName = NAME_None;
-		switch (CookDependency.GetType())
+		// Only consider Load and Save BuildResults; other BuildResults on the package do not contribute to the runtime references
+		// of the cooked package (or if they do, then they are also added to the Save result), so we don't need to report them to
+		// the AssetManager's chunk assignment graph traversal.
+		if (BuildResultPair.Key != UE::Cook::BuildResult::NAME_Load && BuildResultPair.Key != UE::Cook::BuildResult::NAME_Save)
 		{
-		case ECookDependency::Package: [[fallthrough]];
-		case ECookDependency::TransitiveBuild:
-			PackageName = CookDependency.GetPackageName();
-			break;
-		default:
-			break;
+			continue;
 		}
-		if (PackageName.IsNone())
+		for (const UE::Cook::FCookDependency& CookDependency : BuildResultPair.Value)
 		{
-			continue;
+			FName PackageName = NAME_None;
+			switch (CookDependency.GetType())
+			{
+			case ECookDependency::Package: [[fallthrough]];
+			case ECookDependency::TransitiveBuild:
+				PackageName = CookDependency.GetPackageName();
+				break;
+			default:
+				break;
+			}
+			if (PackageName.IsNone())
+			{
+				continue;
+			}
+			PackageBuildDependencies.Add(PackageName);
 		}
-		PackageBuildDependencies.Add(PackageName);
 	}
 	PackageBuildDependencies.Sort(FNameLexicalLess());
 	PackageBuildDependencies.SetNum(Algo::Unique(PackageBuildDependencies));
@@ -411,10 +421,15 @@ void FSaveContext::AddExportedClassesToDependencies()
 	}
 
 	//Finally add those missing class to the dependencies
-	TArray<UE::Cook::FCookDependency>& CookDependencies = GetObjectSaveContext().CookBuildDependencies;
-	for (const UClass* Class : UniqueClassesToAdd)
-		CookDependencies.Add(UE::Cook::FCookDependency::NativeClass(Class));
-
+	if (!UniqueClassesToAdd.IsEmpty())
+	{
+		TArray<UE::Cook::FCookDependency>& CookSaveDependencies
+			= GetObjectSaveContext().BuildResultDependencies.FindOrAdd(UE::Cook::BuildResult::NAME_Save);
+		for (const UClass* Class : UniqueClassesToAdd)
+		{
+			CookSaveDependencies.Add(UE::Cook::FCookDependency::NativeClass(Class));
+		}
+	}
 #endif // WITH_EDITOR
 }
 
diff --git a/Engine/Source/Runtime/CoreUObject/Private/UObject/SavePackage/SaveContext.h b/Engine/Source/Runtime/CoreUObject/Private/UObject/SavePackage/SaveContext.h
--- a/Engine/Source/Runtime/CoreUObject/Private/UObject/SavePackage/SaveContext.h
+++ b/Engine/Source/Runtime/CoreUObject/Private/UObject/SavePackage/SaveContext.h
@@ -565,6 +565,9 @@ public:
 		{
 			ObjectSaveContext.CookType = SaveArgs.ArchiveCookData->CookContext.GetCookType();
 			ObjectSaveContext.CookingDLC = SaveArgs.ArchiveCookData->CookContext.GetCookingDLC();
+#if WITH_EDITOR
+			ObjectSaveContext.CookInfo = SaveArgs.ArchiveCookData->CookContext.GetCookInfo();
+#endif
 		}
 		ArchiveSavePackageData.TargetPlatform = ObjectSaveContext.TargetPlatform;
 		ArchiveSavePackageData.CookContext = SaveArgs.ArchiveCookData ? &SaveArgs.ArchiveCookData->CookContext : nullptr;
@@ -804,11 +807,6 @@ public:
 		EDLCookChecker = InCookChecker;
 	}
 
-	uint32 GetPortFlags() const
-	{
-		return PPF_DeepCompareInstances | PPF_DeepCompareDSOsOnly;
-	}
-
 	bool GetPostSaveRootRequired() const
 	{
 		return bPostSaveRootRequired;
@@ -972,8 +970,16 @@ public:
 	}
 
 	/**
-	 * PackageBuildDependencies are copied from CookBuildDependencies, and they are only needed
-	 * for the Editor realm.
+	 * We copy manually-declared CookBuildDependencies from native class OnCookEvent functions into the AssetRegistry
+	 * dependencies recorded for the editor save of a package. These build dependencies are used for the management
+	 * database created by the AssetManager (UpdateManagementDatabase) that assigns chunks from PrimaryAssets out to
+	 * all of the Assets that each PrimaryAsset depends on. We need to include CookBuildDependencies in that graph
+	 * traversal because we can have some editoronly references that are read during cook to put data into the cooked
+	 * version of the package that references them. We also use editoronly references to hide references that are used
+	 * in game, but only on some platforms. To follow those references during the AssetManager's graph traversal, we
+	 * need to report them in some other way than as used-in-game dependencies (because used-in-game dependencies
+	 * would incorrectly mark them as used-in-game for all platforms). We implement that with build dependencies.
+	 * See FMaterialOverrideNanite for an example of editor-only dependencies that are used-in-game on some platforms.
 	*/
 	void UpdateEditorRealmPackageBuildDependencies();
 
diff --git a/Engine/Source/Runtime/CoreUObject/Private/UObject/SavePackage/SavePackageUtilities.cpp b/Engine/Source/Runtime/CoreUObject/Private/UObject/SavePackage/SavePackageUtilities.cpp
--- a/Engine/Source/Runtime/CoreUObject/Private/UObject/SavePackage/SavePackageUtilities.cpp
+++ b/Engine/Source/Runtime/CoreUObject/Private/UObject/SavePackage/SavePackageUtilities.cpp
@@ -8,6 +8,7 @@
 #include "AssetRegistry/CookTagList.h"
 #include "Blueprint/BlueprintSupport.h"
 #include "Cooker/CookDependency.h"
+#include "Cooker/CookEvents.h"
 #include "CoreMinimal.h"
 #include "HAL/FileManager.h"
 #include "Interfaces/ITargetPlatform.h"
@@ -688,6 +689,17 @@ void CallPreSave(UObject* Object, FObjectSaveContextData& ObjectSaveContext)
 	}
 }
 
+#if WITH_EDITOR
+void CallCookEventPlatformCookDependencies(UObject* Object, FObjectSaveContextData& ObjectSaveContext)
+{
+	SCOPED_SAVETIMER_TEXT(*WriteToString<256>(GetClassTraceScope(Object), TEXTVIEW("_PlatformCookDependencies")));
+
+	using namespace UE::Cook;
+	FCookEventContext CookEventContext(ObjectSaveContext);
+	Object->OnCookEvent(ECookEvent::PlatformCookDependencies, CookEventContext);
+}
+#endif
+
 void CallPreSaveRoot(UObject* Object, FObjectSaveContextData& ObjectSaveContext)
 {
 	SCOPED_SAVETIMER_TEXT(*WriteToString<256>(GetClassTraceScope(Object), TEXTVIEW("_PreSave")));
@@ -748,11 +760,12 @@ FAddResaveOnDemandPackage OnAddResaveOnDemandPackage;
 
 } // end namespace UE::SavePackageUtilities
 
-// Constructor/Destructor defined here in cpp rather than header so we can 
-// avoid needing the definition of FCookDependency in the header; it is needed
-// for construct/destruct of TArray<FCookDependency>.
+PRAGMA_DISABLE_DEPRECATION_WARNINGS
 FObjectSaveContextData::FObjectSaveContextData() = default;
 FObjectSaveContextData::~FObjectSaveContextData() = default;
+FObjectSaveContextData::FObjectSaveContextData(const FObjectSaveContextData& Other) = default;
+FObjectSaveContextData::FObjectSaveContextData(FObjectSaveContextData&& Other) = default;
+PRAGMA_ENABLE_DEPRECATION_WARNINGS
 
 FObjectSaveContextData::FObjectSaveContextData(UPackage* Package, const ITargetPlatform* InTargetPlatform, const TCHAR* InTargetFilename, uint32 InSaveFlags)
 {
@@ -783,7 +796,7 @@ void HarvestCookRuntimeDependencies(FObjectSaveContextData& Data, UObject* Harve
 	// Don't store the input Data on the ArchiveSavePackageData; we just want to harvest the serialized
 	// FSoftObjectPaths, not allow direct writes to its cookdependencies or other data. We only set
 	// ArchiveSavePackageData to provide access to the CookContext.
-	FArchiveCookContext CookContext(PackageBeingSaved, Data.CookType, Data.CookingDLC, Data.TargetPlatform);
+	FArchiveCookContext CookContext(PackageBeingSaved, Data.CookType, Data.CookingDLC, Data.TargetPlatform, Data.CookInfo);
 	FArchiveSavePackageDataBuffer SavePackageData(CookContext);
 
 	FImportExportCollector Collector(HarvestReferencesFrom->GetPackage());
@@ -810,7 +823,7 @@ void HarvestCookRuntimeDependencies(FObjectSaveContextData& Data, UObject* Harve
 
 void FObjectPreSaveContext::AddCookBuildDependency(UE::Cook::FCookDependency BuildDependency)
 {
-	Data.CookBuildDependencies.Add(MoveTemp(BuildDependency));
+	Data.BuildResultDependencies.Add(UE::Cook::BuildResult::NAME_Load, MoveTemp(BuildDependency));
 }
 void FObjectPreSaveContext::AddCookRuntimeDependency(FSoftObjectPath RuntimeDependency)
 {
@@ -835,23 +848,43 @@ void FObjectPreSaveContext::RegisterDeterminismHelper(
 	}
 }
 
+bool FObjectSavePackageSerializeContext::IsHarvestingCookDependencies() const
+{
+	return (Data.ObjectSaveContextPhase == EObjectSaveContextPhase::Harvest)
+		| (Data.ObjectSaveContextPhase == EObjectSaveContextPhase::CookDependencyHarvest);
+}
 void FObjectSavePackageSerializeContext::AddCookBuildDependency(UE::Cook::FCookDependency BuildDependency)
 {
-	if (GetPhase() != EObjectSaveContextPhase::Harvest)
+	AddCookLoadDependency(MoveTemp(BuildDependency));
+}
+void FObjectSavePackageSerializeContext::AddCookLoadDependency(UE::Cook::FCookDependency BuildDependency)
+{
+	if (!IsHarvestingCookDependencies())
 	{
 		UE_LOG(LogSavePackage, Error,
-			TEXT("AddCookBuildDependency called when GetPhase() != EObjectSaveContextPhase::Harvest. This is invalid and will be ignored."));
+			TEXT("AddCookLoadDependency called when !IsHarvestingCookDependencies(). This is invalid and will be ignored."));
 		FDebug::DumpStackTraceToLog(ELogVerbosity::Warning);
 		return;
 	}
-	Data.CookBuildDependencies.Add(MoveTemp(BuildDependency));
+	Data.BuildResultDependencies.Add(UE::Cook::BuildResult::NAME_Load, MoveTemp(BuildDependency));
+}
+void FObjectSavePackageSerializeContext::AddCookSaveDependency(UE::Cook::FCookDependency BuildDependency)
+{
+	if (!IsHarvestingCookDependencies())
+	{
+		UE_LOG(LogSavePackage, Error,
+			TEXT("AddCookSaveDependency called when !IsHarvestingCookDependencies(). This is invalid and will be ignored."));
+		FDebug::DumpStackTraceToLog(ELogVerbosity::Warning);
+		return;
+	}
+	Data.BuildResultDependencies.Add(UE::Cook::BuildResult::NAME_Save, MoveTemp(BuildDependency));
 }
 void FObjectSavePackageSerializeContext::AddCookRuntimeDependency(FSoftObjectPath RuntimeDependency)
 {
-	if (GetPhase() != EObjectSaveContextPhase::Harvest)
+	if (!IsHarvestingCookDependencies())
 	{
 		UE_LOG(LogSavePackage, Error,
-			TEXT("AddCookRuntimeDependency called when GetPhase() != EObjectSaveContextPhase::Harvest. This is invalid and will be ignored."));
+			TEXT("AddCookRuntimeDependency called when !IsHarvestingCookDependencies(). This is invalid and will be ignored."));
 		FDebug::DumpStackTraceToLog(ELogVerbosity::Warning);
 		return;
 	}
@@ -2433,7 +2466,7 @@ void WritePackageData(FStructuredArchiveRecord& ParentRecord, bool bIsCooking, c
 	if (TargetPlatform)
 	{
 		FArchiveCookContext CookContext(const_cast<UPackage*>(Package), UE::Cook::ECookType::Unknown,
-			UE::Cook::ECookingDLC::Unknown, TargetPlatform);
+			UE::Cook::ECookingDLC::Unknown, TargetPlatform, nullptr /* CookInfo */);
 		WritePackageData(ParentRecord, &CookContext, Package, Linker, ImportsUsedInGame, SoftPackagesUsedInGame,
 			OutAssetDatas, true /* bProceduralSave */);
 	}
diff --git a/Engine/Source/Runtime/CoreUObject/Private/UObject/SavePackage/SavePackageUtilities.h b/Engine/Source/Runtime/CoreUObject/Private/UObject/SavePackage/SavePackageUtilities.h
--- a/Engine/Source/Runtime/CoreUObject/Private/UObject/SavePackage/SavePackageUtilities.h
+++ b/Engine/Source/Runtime/CoreUObject/Private/UObject/SavePackage/SavePackageUtilities.h
@@ -386,7 +386,9 @@ bool IsEditorOnlyObjectInternal(const UObject* InObject, EEditorOnlyObjectFlags
 bool IsEditorOnlyObjectInternal(const UObject* InObject, EEditorOnlyObjectFlags Flags);
 #endif
 
-}
+void HarvestCookRuntimeDependencies(FObjectSaveContextData& Data, UObject* HarvestReferencesFrom);
+
+} // namespace UE::SavePackageUtilities
 
 #if ENABLE_COOK_STATS
 struct FSavePackageStats
diff --git a/Engine/Source/Runtime/CoreUObject/Private/UObject/SavePackage2.cpp b/Engine/Source/Runtime/CoreUObject/Private/UObject/SavePackage2.cpp
--- a/Engine/Source/Runtime/CoreUObject/Private/UObject/SavePackage2.cpp
+++ b/Engine/Source/Runtime/CoreUObject/Private/UObject/SavePackage2.cpp
@@ -9,6 +9,7 @@
 #include "Async/ParallelFor.h"
 #include "Blueprint/BlueprintSupport.h"
 #include "Cooker/CookDependency.h"
+#include "Cooker/CookEvents.h"
 #include "HAL/FileManager.h"
 #include "Internationalization/TextPackageNamespaceUtil.h"
 #include "Internationalization/PackageLocalizationManager.h"
@@ -273,13 +274,14 @@ ESavePackageResult RoutePresave(FSaveContext& SaveContext)
 	// objects likely to not be export will probably not care about PreSave and should be mainly noop
 	TArray<UObject*> ObjectsInPackage;
 	GetObjectsWithPackage(SaveContext.GetPackage(), ObjectsInPackage);
+
+	FObjectSaveContextData& ObjectSaveContext = SaveContext.GetObjectSaveContext();
 	for (UObject* Object : ObjectsInPackage)
 	{
 		// Do not emit unsaveable warning while routing presave. 
 		// this is to prevent warning on objects which won't be harvested later since they are unreferenced
 		if (!SaveContext.IsUnsaveable(Object, SaveContext.GetCachedObjectStatus(Object), false/*bEmitWarning*/))
 		{
-			FObjectSaveContextData& ObjectSaveContext = SaveContext.GetObjectSaveContext();
 			ObjectSaveContext.Object = Object;
 			if (SaveContext.IsCooking() && Object->HasAnyFlags(RF_ClassDefaultObject | RF_ArchetypeObject))
 			{
@@ -288,6 +290,9 @@ ESavePackageResult RoutePresave(FSaveContext& SaveContext)
 				FString PathNameBefore = Object->GetPathName();
 				int32 Before = CrcArchive.Crc32(Object);
 				UE::SavePackageUtilities::CallPreSave(Object, ObjectSaveContext);
+#if WITH_EDITOR
+				UE::SavePackageUtilities::CallCookEventPlatformCookDependencies(Object, ObjectSaveContext);
+#endif
 				int32 After = CrcArchive.Crc32(Object);
 
 				if (Before != After)
@@ -322,6 +327,9 @@ ESavePackageResult RoutePresave(FSaveContext& SaveContext)
 			else
 			{
 				UE::SavePackageUtilities::CallPreSave(Object, ObjectSaveContext);
+#if WITH_EDITOR
+				UE::SavePackageUtilities::CallCookEventPlatformCookDependencies(Object, ObjectSaveContext);
+#endif
 			}
 		}
 	}
@@ -471,9 +479,9 @@ ESavePackageResult HarvestPackage(FSaveContext& SaveContext)
 		SaveContext.AddExportedClassesToDependencies();
 	}
 
-	// Contractually all CookBuildDependencies (which come from PreSave or from Serialize) must be declared by this point.
-	// Copy PackageBuildDependencies from CookBuildDependencies and implement our contract that PackageBuildDependencies
-	// are in the name map.
+	// Contractually all BuildResultDependencies (which come from PreSave or from Serialize) must be
+	// declared by this point. Copy PackageBuildDependencies from those lists and implement our contract that
+	// PackageBuildDependencies are in the name map.
 	SaveContext.UpdateEditorRealmPackageBuildDependencies();
 
 	return ReturnSuccessOrCancel();
@@ -1167,7 +1175,7 @@ ESavePackageResult BuildLinker(FSaveContext& SaveContext)
 		// Use the custom versions we harvested from the dependency harvesting pass
 		SaveContext.GetLinker()->Summary.SetCustomVersionContainer(SaveContext.GetCustomVersions());
 
-		SaveContext.GetLinker()->SetPortFlags(SaveContext.GetPortFlags());
+		SaveContext.GetLinker()->SetPortFlags(UE::SavePackageUtilities::Private::GetSavePackagePortFlags());
 		if (SaveContext.IsSaveAutoOptional() && SaveContext.GetCurrentHarvestingRealm() == ESaveRealm::Optional)
 		{
 			// Do not filter editor only data when automatically creating optional data using full uncooked objects
@@ -2778,11 +2786,12 @@ ESavePackageResult BeginCachePlatformCookedData(FSaveContext& SaveContext)
  * Archive used during post save serialization.
  * It ignores serialization of primitive types. It is just used to provide a callback hook for UStruct::Serialize and UObject::Serialize to execute custom code.
  */
-class FPostSaveArchive : public FArchiveSavePackageCollector
+class FPostSaveArchive : public UE::SavePackageUtilities::Private::FArchiveSavePackageCollector
 {
 public:
 	FPostSaveArchive(FSaveContext& InContext)
-		: FArchiveSavePackageCollector(InContext)
+		: FArchiveSavePackageCollector(InContext.GetArchiveSavePackageData(), InContext.IsFilterEditorOnly(),
+			InContext.IsSaveUnversionedProperties(), InContext.IsCooking())
 	{
 		//disable the delta serialization to make sure the objects are serialized no matter what.
 		ArNoDelta = true;
diff --git a/Engine/Source/Runtime/CoreUObject/Public/Cooker/BuildResultDependenciesMap.h b/Engine/Source/Runtime/CoreUObject/Public/Cooker/BuildResultDependenciesMap.h
new file mode 100644
--- /dev/null
+++ b/Engine/Source/Runtime/CoreUObject/Public/Cooker/BuildResultDependenciesMap.h
@@ -0,0 +1,36 @@
+// Copyright Epic Games, Inc. All Rights Reserved.
+
+#pragma once
+
+#if WITH_EDITOR
+#include "Containers/Array.h"
+#include "Containers/Map.h"
+#include "UObject/NameTypes.h"
+
+namespace UE::Cook { class FCookDependency; }
+#endif
+
+#if WITH_EDITOR
+
+namespace UE::Cook
+{
+
+/**
+ * A container for BuildResult names and an array of CookDependencies for each BuildResult. The cooker collects
+ * these from UClasses during OnCookEvent, and it collects them from UStructs during FArchive::Serialize through
+ * FObjectSavePackageSerializeContext. @see namespace UE::Cook::BuildResult in CookDependency.h.
+ */
+struct FBuildResultDependenciesMap : public TMap<FName, TArray<UE::Cook::FCookDependency>>
+{
+public:
+	using TMap<FName, TArray<UE::Cook::FCookDependency>>::TMap;
+
+	COREUOBJECT_API void Add(FName Name, UE::Cook::FCookDependency CookDependency);
+	COREUOBJECT_API void Append(FName Name, TArray<UE::Cook::FCookDependency> CookDependencies);
+	COREUOBJECT_API void Append(const FBuildResultDependenciesMap& Other);
+	COREUOBJECT_API void Append(FBuildResultDependenciesMap&& Other);
+};
+
+} // namespace UE::Cook
+
+#endif // WITH_EDITOR
diff --git a/Engine/Source/Runtime/CoreUObject/Public/Cooker/CookDependency.h b/Engine/Source/Runtime/CoreUObject/Public/Cooker/CookDependency.h
--- a/Engine/Source/Runtime/CoreUObject/Public/Cooker/CookDependency.h
+++ b/Engine/Source/Runtime/CoreUObject/Public/Cooker/CookDependency.h
@@ -150,6 +150,32 @@ extern COREUOBJECT_API FName None;
 
 }	
 
+/**
+ * BuildResults are groups of data collected during BuildOperations during incremental cooks. Each BuildResult is owned
+ * by a BuildOperation and stored in the cook's oplog, and it has a name, a set of dependencies used to create it, and
+ * an explicit or implicit data object that was created by the build operation. In Unreal, most examples of
+ * BuildOperations are package BuildOperations: the load of a package, the transform of its data into the format that
+ * will be used at runtime, and the save of the package.
+ * 
+ * The generic BuildResults are @see UE::Cook::BuildResult::NAME_Load and UE::Cook::BuildResult::NAME_Save. Other
+ * BuildResults are system-specific. The declaration of system-specific build results, and the reporting of
+ * dependencies for the generic and system-specific BuildResults, are done during native UClass's overrides of
+ * UObject::OnCookEvent(UE::Cook::ECookEvent::PlatformCookDependencies,...).
+ */
+namespace BuildResult
+{
+
+/** The BuildResult that is the bytes of the cooked package. */
+extern COREUOBJECT_API FName NAME_Save;
+
+/**
+ * The BuildResult that is the loaded bytes of editor package in memory. This is the default BuildResult used
+ * in transitive build dependencies.
+ */
+extern COREUOBJECT_API FName NAME_Load;
+
+}
+
 /**
  * TargetDomain dependencies that can be reported from the class instances in a package. These dependencies are
  * stored in the cook oplog and are evaluated during incremental cook. If any of them changes, the package is
diff --git a/Engine/Source/Runtime/CoreUObject/Public/Cooker/CookEvents.h b/Engine/Source/Runtime/CoreUObject/Public/Cooker/CookEvents.h
new file mode 100644
--- /dev/null
+++ b/Engine/Source/Runtime/CoreUObject/Public/Cooker/CookEvents.h
@@ -0,0 +1,136 @@
+// Copyright Epic Games, Inc. All Rights Reserved.
+
+#pragma once
+
+#if WITH_EDITOR
+#include "Containers/Array.h"
+#include "Cooker/CookDependency.h"
+#include "UObject/ObjectSaveContext.h"
+
+class ITargetPlatform;
+namespace UE::Cook { class ICookInfo; }
+#endif // WITH_EDITOR
+
+#if WITH_EDITOR
+
+namespace UE::Cook
+{
+
+/**
+ * The list of possible events that the cooker can call on UObjects during load/transform/save. Native UObject
+ * classes can respond to these events by overriding UObject::OnCookEvent. All UObject classes should call their Super
+ * class's version of the function during their call.
+ */
+enum class ECookEvent : uint32
+{
+	/**
+	 * Called when saving the package to fetch the list of build and runtime dependencies required by the object that
+	 * are not automaticallyd detected and need to be manually declared. For example, a dependency on an AssetRegistry
+	 * query, or a runtime dependency for platforms that enabled Nanite materials on the nanite material that is
+	 * declared as editoronly for non-Nanite-enabled platforms. When called for objects in packages being saved for th
+	 * cook, it is calledfor each object immediately after PreSave on that object.
+	 * 
+	 * This event is also called on objects in packages that are transitive build dependencies of other packages but
+	 * are not cooked themselves. For example, a DataTable can include another DataTable, and needs to depend on all of the
+	 * dependencies of the DataTable it includes, but the DataTable it includes is not needed at runtime. When called for
+	 * objects in these build dependency packages, it is called by itself, with no BeginCacheForCookedPlatformData,
+	 * PreSave, or save-Serialize functions being called on the Object.
+	 * 
+	 * This event is also be called when NOT cooking. In that case, it provides the UObject a way to declare build
+	 * dependencies that are registered with the AssetRegistry and cause propagation of AssetManager chunk assignments
+	 * without causing the target to be cooked. @see UE::Cook::FCookEventContext::AddLoadBuildDependency.
+	 * 
+	 * UObjects that need to declare dependencies should call the Add*Dependency functions on the
+	 * UE::Cook::FCookEventContext passed into UObject::OnCookEvent.
+	 */
+	PlatformCookDependencies,
+};
+
+/**
+ * Context structure to provide information about the cook and the event being called on an object during
+ * UObject::OnCookEvent, and to receive output from those events. The available input and output depend on the event
+ * being called, @see UE::Cook::ECookEvent.
+ */
+class FCookEventContext
+{
+public:
+	explicit FCookEventContext(FObjectSaveContextData& InData)
+		: Data(InData)
+	{
+	}
+	FCookEventContext(const UE::Cook::FCookEventContext& Other)
+		: Data(Other.Data)
+	{
+	}
+
+	/**
+	 * Report whether this OnCookEvent was called during a cook. Some cook events (such as
+	 * UE::Cook::PlatformCookDependencies) can also be called during editor saves of a package.
+	 */
+	bool IsCooking() const;
+	/**
+	 * Return the CookInfo that provides information about the overall cook. Returns null if not cooking.
+	 * Can sometimes be null even if IsCooking is true. (TODO: Fix these cases, make it guaranteed not null.)
+	 */
+	ICookInfo* GetCookInfo() const;
+	/**
+	 * Return the TargetPlatform on behalf of which the event was called on the UObject. In multiprocess cooks,
+	 * this will be the specific platform being saved; each platform will get its own call. Non-null if and only if
+	 * IsCooking().
+	 */
+	const ITargetPlatform* GetTargetPlatform() const;
+
+	/**
+	 * Available during the event UE::Cook::PlatformCookDependencies, ignored in other events.
+	 * 
+	 * Add the given FCookDependency to the build dependencies for the package being cook-saved. Iterative cooks will
+	 * invalidate the package and recook it if the CookDependency changes. Other packages that incorporate the loaded
+	 * data of the package into their own cooked results will also be recooked if this build dependency changes.
+	 * Calling this function during editor save (rather than cook save) has another meaning. It is ignored for
+	 * FCookDependency of type other than FCookDependency::Package, but for FCookDependency::Package it identifies
+	 * a reference that propagates chunk management in the AssetRegistry, but does not cause its target to be cooked.
+	 */
+	COREUOBJECT_API void AddLoadBuildDependency(FCookDependency CookDependency);
+	/**
+	 * Available during the event UE::Cook::PlatformCookDependencies, ignored in other events.
+	 * 
+	 * Add the given FCookDependency to the build dependencies for the package being cook-saved. Iterative cooks will
+	 * invalidate the package and recook it if the CookDependency changes.
+	 */
+	COREUOBJECT_API void AddSaveBuildDependency(FCookDependency CookDependency);
+	/**
+	 * Available during the event UE::Cook::PlatformCookDependencies, ignored in other events.
+	 * 
+	 * Report that the given PackageName is a runtime dependency of the current package and needs to be cooked.
+	 */
+	COREUOBJECT_API void AddRuntimeDependency(FSoftObjectPath PackageName);
+	/**
+	 * Available during the event UE::Cook::PlatformCookDependencies, ignored in other events.
+	 * 
+	 * Serialize an object to find all packages that it references, and AddCookRuntimeDependency for each one.
+	 */
+	COREUOBJECT_API void HarvestCookRuntimeDependencies(UObject* HarvestReferencesFrom);
+
+private:
+	FObjectSaveContextData& Data;
+};
+
+
+inline bool FCookEventContext::IsCooking() const
+{
+	return Data.TargetPlatform != nullptr;
+}
+
+inline ICookInfo* FCookEventContext::GetCookInfo() const
+{
+	return Data.CookInfo;
+}
+
+inline const ITargetPlatform* FCookEventContext::GetTargetPlatform() const
+{
+	return Data.TargetPlatform;
+}
+
+} // namespace UE::Cook
+
+#endif // WITH_EDITOR
\ No newline at end of file
diff --git a/Engine/Source/Runtime/CoreUObject/Public/Serialization/ArchiveSavePackageDataBuffer.h b/Engine/Source/Runtime/CoreUObject/Public/Serialization/ArchiveSavePackageDataBuffer.h
--- a/Engine/Source/Runtime/CoreUObject/Public/Serialization/ArchiveSavePackageDataBuffer.h
+++ b/Engine/Source/Runtime/CoreUObject/Public/Serialization/ArchiveSavePackageDataBuffer.h
@@ -7,6 +7,8 @@
 #include "UObject/ArchiveCookContext.h"
 #include "UObject/ObjectSaveContext.h"
 
+namespace UE::Cook { class ICookInfo; }
+
 /**
  * This is the structure that should be used by most callers of an archive (other than SavePackage which handles
  * the full complexity) that need to provide SavePackage or Cook information to the archive.
@@ -20,14 +22,15 @@
 struct FArchiveSavePackageDataBuffer : public FArchiveSavePackageData
 {
 public:
-	FArchiveSavePackageDataBuffer(const ITargetPlatform* InTargetPlatform = nullptr, UPackage* Package = nullptr)
+	FArchiveSavePackageDataBuffer(const ITargetPlatform* InTargetPlatform = nullptr, UPackage* Package = nullptr,
+		UE::Cook::ICookInfo* CookInfo = nullptr)
 		// warning V1050 : The uninitialized class member 'ObjectSavePackageSerializeContextBuffer' is used when initializing the base class 'FArchiveSavePackageData'.
 		// This warning can be ignored because the base class only records the pointer and does not dereference it
 		: FArchiveSavePackageData(ObjectSavePackageSerializeContextBuffer, nullptr, nullptr) // -V1050
 		, ObjectSaveContextData(FObjectSaveContextData())
 		, ObjectSavePackageSerializeContextBuffer(*ObjectSaveContextData)
 		, CookContextBuffer(Package, UE::Cook::ECookType::Unknown, UE::Cook::ECookingDLC::Unknown,
-			InTargetPlatform)
+			InTargetPlatform, CookInfo)
 	{
 		SetConstructorTargetPlatform(InTargetPlatform);
 	}
@@ -44,7 +47,7 @@ public:
 	FArchiveSavePackageDataBuffer(FObjectSaveContextData& InData, UPackage* InPackage = nullptr)
 		: FArchiveSavePackageData(ObjectSavePackageSerializeContextBuffer, nullptr, nullptr)
 		, ObjectSavePackageSerializeContextBuffer(InData)
-		, CookContextBuffer(InPackage, InData.CookType, InData.CookingDLC, InData.TargetPlatform)
+		, CookContextBuffer(InPackage, InData.CookType, InData.CookingDLC, InData.TargetPlatform, InData.CookInfo)
 	{
 		SetConstructorTargetPlatform(InData.TargetPlatform);
 	}
diff --git a/Engine/Source/Runtime/CoreUObject/Public/UObject/ArchiveCookContext.h b/Engine/Source/Runtime/CoreUObject/Public/UObject/ArchiveCookContext.h
--- a/Engine/Source/Runtime/CoreUObject/Public/UObject/ArchiveCookContext.h
+++ b/Engine/Source/Runtime/CoreUObject/Public/UObject/ArchiveCookContext.h
@@ -5,6 +5,8 @@
 #include "AssetRegistry/CookTagList.h"
 #include "UObject/CookEnums.h"
 
+namespace UE::Cook { class ICookInfo; }
+
 /**
 *	Accessor for data about the package being cooked during UObject::Serialize calls.
 */
@@ -36,15 +38,19 @@ private:
 	bool bCookTagListEnabled = false;
 	UE::Cook::ECookType CookType = UE::Cook::ECookType::Unknown;
 	UE::Cook::ECookingDLC CookingDLC = UE::Cook::ECookingDLC::Unknown;
+	UE::Cook::ICookInfo* CookInfo = nullptr;
 
 public:
 
-	UE_DEPRECATED(5.4, "Call version that takes the TargetPlatform")
+	UE_DEPRECATED(5.4, "Call version that takes TargetPlatform and CookInfo")
 	FArchiveCookContext(UPackage* InPackage, UE::Cook::ECookType InCookType, UE::Cook::ECookingDLC InCookingDLC);
-
+	UE_DEPRECATED(5.6, "Call version that takes CookInfo")
 	FArchiveCookContext(UPackage* InPackage, UE::Cook::ECookType InCookType, UE::Cook::ECookingDLC InCookingDLC,
 		const ITargetPlatform* InTargetPlatform);
 
+	FArchiveCookContext(UPackage* InPackage, UE::Cook::ECookType InCookType, UE::Cook::ECookingDLC InCookingDLC,
+		const ITargetPlatform* InTargetPlatform, UE::Cook::ICookInfo* InCookInfo);
+
 	void Reset();
 
 	FCookTagList* GetCookTagList();
@@ -55,6 +61,7 @@ public:
 	bool IsCookTypeUnknown() const;
 	UE::Cook::ECookType GetCookType() const;
 	UE::Cook::ECookingDLC GetCookingDLC() const;
+	UE::Cook::ICookInfo* GetCookInfo();
 };
 
 
@@ -65,17 +72,24 @@ public:
 
 inline FArchiveCookContext::FArchiveCookContext(UPackage* InPackage, UE::Cook::ECookType InCookType, 
 	UE::Cook::ECookingDLC InCookingDLC)
-	: FArchiveCookContext(InPackage, InCookType, InCookingDLC, nullptr)
+	: FArchiveCookContext(InPackage, InCookType, InCookingDLC, nullptr, nullptr)
 {
 }
 
 inline FArchiveCookContext::FArchiveCookContext(UPackage* InPackage, UE::Cook::ECookType InCookType,
 	UE::Cook::ECookingDLC InCookingDLC, const ITargetPlatform* InTargetPlatform)
+	: FArchiveCookContext(InPackage, InCookType, InCookingDLC, InTargetPlatform, nullptr)
+{
+}
+
+inline FArchiveCookContext::FArchiveCookContext(UPackage* InPackage, UE::Cook::ECookType InCookType,
+	UE::Cook::ECookingDLC InCookingDLC, const ITargetPlatform* InTargetPlatform, UE::Cook::ICookInfo* InCookInfo)
 	: CookTagList(InPackage)
 	, TargetPlatform(InTargetPlatform)
-	, bCookTagListEnabled(InPackage&& InCookType == UE::Cook::ECookType::ByTheBook)
+	, bCookTagListEnabled(InPackage && InCookType == UE::Cook::ECookType::ByTheBook)
 	, CookType(InCookType)
 	, CookingDLC(InCookingDLC)
+	, CookInfo(InCookInfo)
 {
 }
 
@@ -117,4 +131,9 @@ inline UE::Cook::ECookType FArchiveCookContext::GetCookType() const
 inline UE::Cook::ECookingDLC FArchiveCookContext::GetCookingDLC() const
 {
 	return CookingDLC;
-}
\ No newline at end of file
+}
+
+inline UE::Cook::ICookInfo* FArchiveCookContext::GetCookInfo()
+{
+	return CookInfo;
+}
diff --git a/Engine/Source/Runtime/CoreUObject/Public/UObject/ICookInfo.h b/Engine/Source/Runtime/CoreUObject/Public/UObject/ICookInfo.h
--- a/Engine/Source/Runtime/CoreUObject/Public/UObject/ICookInfo.h
+++ b/Engine/Source/Runtime/CoreUObject/Public/UObject/ICookInfo.h
@@ -357,4 +357,7 @@ struct FCookLoadScope
 	}
 };
 
+// Provide a forward declare but no definition for ICookInfo, to remove the need to have WITH_EDITOR around some types that include a pointer to it
+namespace UE::Cook { class ICookInfo; }
+
 #endif // WITH_EDITOR
\ No newline at end of file
diff --git a/Engine/Source/Runtime/CoreUObject/Public/UObject/Object.h b/Engine/Source/Runtime/CoreUObject/Public/UObject/Object.h
--- a/Engine/Source/Runtime/CoreUObject/Public/UObject/Object.h
+++ b/Engine/Source/Runtime/CoreUObject/Public/UObject/Object.h
@@ -18,6 +18,8 @@ struct FAssetData;
 class FAssetRegistryTagsContext;
 namespace UE::AssetRegistry { class FAssetRegistryImpl; }
 namespace UE::ConfigAccessTracking { struct FConfigAccessData; }
+namespace UE::Cook { class FCookEventContext; }
+namespace UE::Cook { enum class ECookEvent : uint32; }
 class FConfigCacheIni;
 class FCustomPropertyConditionState;
 class FEditPropertyChain;
@@ -289,6 +291,14 @@ public:
 	COREUOBJECT_API virtual void CollectSaveOverrides(FObjectCollectSaveOverridesContext SaveContext);
 
 #if WITH_EDITOR
+	/**
+	 * Called during cooking for various events that happen on each UObject in a package during load/transform/save.
+	 * @see UE::Cook::ECookEvent for the possible events.
+	 */
+	virtual void OnCookEvent(UE::Cook::ECookEvent CookEvent, UE::Cook::FCookEventContext& Context)
+	{
+	}
+
 	/**
 	 * Note that the object will be modified.  If we are currently recording into the 
 	 * transaction buffer (undo/redo), save a copy of this object into the buffer and 
diff --git a/Engine/Source/Runtime/CoreUObject/Public/UObject/ObjectSaveContext.h b/Engine/Source/Runtime/CoreUObject/Public/UObject/ObjectSaveContext.h
--- a/Engine/Source/Runtime/CoreUObject/Public/UObject/ObjectSaveContext.h
+++ b/Engine/Source/Runtime/CoreUObject/Public/UObject/ObjectSaveContext.h
@@ -2,13 +2,17 @@
 
 #pragma once
 
+#include "Containers/Array.h"
+#include "Containers/Map.h"
 #include "Containers/UnrealString.h"
 #include "HAL/Platform.h"
 #include "Serialization/ArchiveSavePackageData.h"
 #include "Templates/RefCounting.h"
 #include "UObject/CookEnums.h"
+#include "UObject/NameTypes.h"
 #include "UObject/ObjectSaveOverride.h"
 #if WITH_EDITOR
+#include "Cooker/BuildResultDependenciesMap.h"
 #include "Cooker/CookDependency.h"
 #endif
 
@@ -17,6 +21,7 @@ class IPackageWriter;
 class ITargetPlatform;
 class UPackage;
 struct FSoftObjectPath;
+namespace UE::Cook { class ICookInfo; }
 #if WITH_EDITOR
 namespace UE::Cook { class IDeterminismHelper; }
 #endif
@@ -40,6 +45,14 @@ enum class EObjectSaveContextPhase : uint8
 	 * Build dependencies are valid for writing during this phase.
 	 */
 	Harvest,
+	/**
+	 * The Archive is coming from the cooker, outside of a SavePackage. It is discarding all results other than
+	 * AddCook*Dependency results. No other save functions (BeginCacheForCookedPlatformData, PreSave, PostSave) have
+	 * been called or will be called. Objects and structs serializing with this archive should serialize the same
+	 * structures and UObject pointers that they serialize when serializing Harvest.
+	 * This phase is executed during cook for packages that are build dependencies but are not cooked themselves.
+	 */
+	CookDependencyHarvest,
 	/**
 	 * SavePackage is calling Serialize(FArchive&) on objects in the package. The FArchive is recording the blobs of
 	 * data from each exported UObject that will be written to disk.
@@ -57,6 +70,9 @@ struct FObjectSaveContextData
 {
 	COREUOBJECT_API FObjectSaveContextData();
 	COREUOBJECT_API ~FObjectSaveContextData();
+	COREUOBJECT_API FObjectSaveContextData(const FObjectSaveContextData& Other);
+	COREUOBJECT_API FObjectSaveContextData(FObjectSaveContextData&& Other);
+
 	/** Standard constructor; calculates derived fields from the given externally-specified fields. */
 	COREUOBJECT_API FObjectSaveContextData(UPackage* Package, const ITargetPlatform* InTargetPlatform, const TCHAR* InTargetFilename, uint32 InSaveFlags);
 	COREUOBJECT_API FObjectSaveContextData(UPackage* Package, const ITargetPlatform* InTargetPlatform, const FPackagePath& TargetPath, uint32 InSaveFlags);
@@ -92,6 +108,9 @@ struct FObjectSaveContextData
 	/** The target platform of the save, if cooking. Null if not cooking. */
 	const ITargetPlatform* TargetPlatform = nullptr;
 
+	/** The CookInfo providing extended information about the current cook. Null if not cooking. */
+	UE::Cook::ICookInfo* CookInfo = nullptr;
+
 	/** The PackageWriter passed to SavePackage, may be null. */
 	IPackageWriter* PackageWriter = nullptr;
 
@@ -141,7 +160,11 @@ struct FObjectSaveContextData
 
 	// Collection variables that are written but not read during the PreSave/PostSave functions
 #if WITH_EDITOR
-	TArray<UE::Cook::FCookDependency> CookBuildDependencies;
+	UE_DEPRECATED(5.6, "Use BuildResultDependencies instead.")
+	TArray<UE::Cook::FCookDependency> CookDependencies;
+	/** Output variable for BuildResult names and their dependencies, reported by AddBuildDependency functions. */
+	UE::Cook::FBuildResultDependenciesMap BuildResultDependencies;
+	/** Output variable for runtime dependencies reported by AddRuntimeDependency functions. */
 	TArray<FSoftObjectPath> CookRuntimeDependencies;
 #endif
 
@@ -255,21 +278,11 @@ public:
 	EObjectSaveContextPhase GetPhase() const { return Data.ObjectSaveContextPhase; }
 
 #if WITH_EDITOR
-	/**
-	 * Add the given FCookDependency to the build dependencies for the package being cook-saved. Iterative cooks will
-	 * invalidate the package and recook it if the CookDependency changes.
-	 * CookBuildDependencies of type FCookDependency::CreatePackage are also used when saving an editor package rather
-	 * than a cooked package. They are saved into the AssetRegistry dependency data for the package, as
-	 * EDependencyProperty::Build properties.
-	 */
+	UE_DEPRECATED(5.6, "These dependencies should instead be reported in OnCookEvent(UE::Cook::ECookEvent::PlatformCookDependencies, ...) by calling FCookEventContext.AddLoadBuildDependency (@see CookEvents.h).")
 	COREUOBJECT_API void AddCookBuildDependency(UE::Cook::FCookDependency BuildDependency);
-	/**
-	 * Add the given UObject's package as a runtime dependency for the package being cook-saved. It will
-	 * be cooked for the current platform even if it is not otherwise referenced, and even if the object
-	 * being PreSaved does not end up saving into the package's exports for the current platform.
-	 */
+	UE_DEPRECATED(5.6, "These dependencies should instead be reported in OnCookEvent(UE::Cook::ECookEvent::PlatformCookDependencies, ...) by calling FCookEventContext.AddRuntimeDependency (@see CookEvents.h).")
 	COREUOBJECT_API void AddCookRuntimeDependency(FSoftObjectPath Dependency);
-	/** Serialize an object to find all packages that it references, and AddCookRuntimeDependency for each one. */
+	UE_DEPRECATED(5.6, "These dependencies should instead be reported in OnCookEvent(UE::Cook::ECookEvent::PlatformCookDependencies, ...) by calling FCookEventContext.HarvestCookRuntimeDependencies (@see CookEvents.h).")
 	COREUOBJECT_API void HarvestCookRuntimeDependencies(UObject* HarvestReferencesFrom);
 
 	/**
@@ -357,14 +370,31 @@ public:
 	EObjectSaveContextPhase GetPhase() const { return Data.ObjectSaveContextPhase; }
 
 #if WITH_EDITOR
+	/**
+	 * Return true if the archive is listening to AddCookBuildDependency and AddCookRuntimeDependency
+	 * calls. If it returns true, objects must add their dependencies, failing to do so will cause
+	 * false positive incremental skips.
+	 */
+	COREUOBJECT_API bool IsHarvestingCookDependencies() const;
+
+	UE_DEPRECATED(5.6, "Use AddCookLoadDependency.")
+	COREUOBJECT_API void AddCookBuildDependency(UE::Cook::FCookDependency BuildDependency);
+
+	/**
+	 * Add the given FCookDependency to the build dependencies for the package being cook-saved. Iterative cooks will
+	 * invalidate the package and recook it if the CookDependency changes. Other packages that incorporate the loaded
+	 * data of the package into their own cooked results will also be recooked if this build dependency changes.
+	 * Calling this function during editor save (rather than cook save) has another meaning. It is ignored for
+	 * FCookDependency of type other than FCookDependency::Package, but for FCookDependency::Package it identifies
+	 * a reference that propagates chunk management in the AssetRegistry, but does not cause its target to be cooked.
+	 */
+	COREUOBJECT_API void AddCookLoadDependency(UE::Cook::FCookDependency BuildDependency);
+
 	/**
 	 * Add the given FCookDependency to the build dependencies for the package being cook-saved. Iterative cooks will
 	 * invalidate the package and recook it if the CookDependency changes.
-	 * CookBuildDependencies of type FCookDependency::CreatePackage are also used when saving an editor package rather
-	 * than a cooked package. They are saved into the AssetRegistry dependency data for the package, as
-	 * EDependencyProperty::Build properties.
 	 */
-	COREUOBJECT_API void AddCookBuildDependency(UE::Cook::FCookDependency BuildDependency);
+	COREUOBJECT_API void AddCookSaveDependency(UE::Cook::FCookDependency BuildDependency);
 	/**
 	 * Add the given UObject's package as a runtime dependency for the package being cook-saved. It will
 	 * be cooked for the current platform even if it is not otherwise referenced, and even if the object
diff --git a/Engine/Source/Runtime/CoreUObject/Public/UObject/Package.h b/Engine/Source/Runtime/CoreUObject/Public/UObject/Package.h
--- a/Engine/Source/Runtime/CoreUObject/Public/UObject/Package.h
+++ b/Engine/Source/Runtime/CoreUObject/Public/UObject/Package.h
@@ -36,6 +36,10 @@
 #include "IO/IoHash.h"
 #endif
 
+#if WITH_EDITOR
+#include "Cooker/BuildResultDependenciesMap.h"
+#endif
+
 class Error;
 class FArchive;
 class FLinkerLoad;
@@ -125,7 +129,9 @@ struct FSavePackageResultStruct
 	TArray<FName> UntrackedSoftPackageReferences;
 
 #if WITH_EDITOR
+	UE_DEPRECATED(5.6, "Use BuildResultDependencies instead..")
 	TArray<UE::Cook::FCookDependency> CookDependencies;
+	UE::Cook::FBuildResultDependenciesMap BuildResultDependencies;
 #endif
 
 	/** Constructors, it will implicitly construct from the result enum */
diff --git a/Engine/Source/Runtime/CoreUObject/Public/UObject/SavePackage.h b/Engine/Source/Runtime/CoreUObject/Public/UObject/SavePackage.h
--- a/Engine/Source/Runtime/CoreUObject/Public/UObject/SavePackage.h
+++ b/Engine/Source/Runtime/CoreUObject/Public/UObject/SavePackage.h
@@ -16,6 +16,7 @@
 #include "Misc/OutputDeviceError.h"
 #include "ObjectMacros.h"
 #include "Serialization/ArchiveCookData.h"
+#include "Serialization/ArchiveUObject.h"
 #include "Serialization/FileRegions.h"
 #include "Serialization/PackageWriter.h"
 #include "Templates/Function.h"
@@ -41,6 +42,7 @@ class FPackagePath;
 class FSavePackageContext;
 class IPackageWriter;
 namespace UE { class FLogRecord; }
+struct FArchiveSavePackageData;
 struct FObjectSaveContextData;
 
 /**
@@ -304,6 +306,11 @@ namespace UE::SavePackageUtilities
 	/** Call the PreSave function on the given object and log a warning if there is an incorrect override. */
 	COREUOBJECT_API void CallPreSave(UObject* Object, FObjectSaveContextData& ObjectSaveContext);
 
+#if WITH_EDITOR
+	/** Call the CookEvent UE::Cook::ECookEvent::PlatformCookDependencies on the given object. */
+	COREUOBJECT_API void CallCookEventPlatformCookDependencies(UObject* Object, FObjectSaveContextData& ObjectSaveContext);
+#endif
+
 	/** Call the PreSaveRoot function on the given object. */
 	COREUOBJECT_API void CallPreSaveRoot(UObject* Object, FObjectSaveContextData& ObjectSaveContext);
 
@@ -343,6 +350,34 @@ namespace UE::SavePackageUtilities
 	 */
 	extern COREUOBJECT_API FAddResaveOnDemandPackage OnAddResaveOnDemandPackage;
 #endif
+
+}
+
+namespace UE::SavePackageUtilities::Private
+{
+
+/**
+ * Base archive class for Archives used during SavePackage to implement the flags used for SavePackage Serialization
+ * used to setup all the common flags. This type is an implementation detail of SavePackage (placed in this public
+ * header so it can be shared with some SavePackage implementation present in the cooker) and might change without deprecation.
+ */
+class FArchiveSavePackageCollector : public FArchiveUObject
+{
+public:
+	/** The default constructor sets no flags and relies on subclass calling SetArchiveFlags later instead. */
+	COREUOBJECT_API FArchiveSavePackageCollector() = default;
+	/** Convenience constructor that calls SetArchiveFlags. */
+	COREUOBJECT_API FArchiveSavePackageCollector(FArchiveSavePackageData& InSavePackageData,
+		bool bFilterEditorOnly, bool bSaveUnversioned, bool bCooking);
+
+	/** Set up the Archive with flags and properties needed to act as a collector of references during SavePackage. */
+	COREUOBJECT_API void SetArchiveFlags(FArchiveSavePackageData& InSavePackageData,
+		bool bFilterEditorOnly, bool bSaveUnversioned, bool bCooking);
+};
+
+/** The portflags used for SavePackage archives. */
+COREUOBJECT_API uint32 GetSavePackagePortFlags();
+
 }
 
 COREUOBJECT_API DECLARE_LOG_CATEGORY_EXTERN(LogSavePackage, Log, All);
diff --git a/Engine/Source/Runtime/Engine/Classes/Engine/Blueprint.h b/Engine/Source/Runtime/Engine/Classes/Engine/Blueprint.h
--- a/Engine/Source/Runtime/Engine/Classes/Engine/Blueprint.h
+++ b/Engine/Source/Runtime/Engine/Classes/Engine/Blueprint.h
@@ -896,9 +896,6 @@ private:
 	/** Sets the current object being debugged */
 	ENGINE_API void DebuggingWorldRegistrationHelper(UObject* ObjectProvidingWorld, UObject* ValueToRegister);
 	
-#if WITH_EDITORONLY_DATA
-	void PreSaveCook(FObjectPreSaveContext ObjectSaveContext);
-#endif
 public:
 
 	/** @return the current object being debugged, which can be nullptr */
@@ -1004,6 +1001,9 @@ public:
 	ENGINE_API virtual void PreSave(FObjectPreSaveContext ObjectSaveContext) override;
 	ENGINE_API virtual void GetPreloadDependencies(TArray<UObject*>& OutDeps) override;
 #endif // WITH_EDITORONLY_DATA
+#if WITH_EDITOR
+	ENGINE_API virtual void OnCookEvent(UE::Cook::ECookEvent CookEvent, UE::Cook::FCookEventContext& CookContext) override;
+#endif
 	ENGINE_API virtual void Serialize(FArchive& Ar) override;
 	ENGINE_API virtual FString GetDesc(void) override;
 	ENGINE_API virtual void TagSubobjects(EObjectFlags NewFlags) override;
diff --git a/Engine/Source/Runtime/Engine/Classes/Engine/Texture2D.h b/Engine/Source/Runtime/Engine/Classes/Engine/Texture2D.h
--- a/Engine/Source/Runtime/Engine/Classes/Engine/Texture2D.h
+++ b/Engine/Source/Runtime/Engine/Classes/Engine/Texture2D.h
@@ -117,6 +117,9 @@ public:
 	virtual bool IsReadyForAsyncPostLoad() const override;
 	virtual void PostLoad() override;
 	virtual void PreSave(FObjectPreSaveContext ObjectSaveContext) override;
+#if WITH_EDITOR
+	virtual void OnCookEvent(UE::Cook::ECookEvent CookEvent, UE::Cook::FCookEventContext& CookContext) override;
+#endif
 	virtual void GetAssetRegistryTags(FAssetRegistryTagsContext Context) const override;
 	UE_DEPRECATED(5.4, "Implement the version that takes FAssetRegistryTagsContext instead.")
 	virtual void GetAssetRegistryTags(TArray<FAssetRegistryTag>& OutTags) const override;
diff --git a/Engine/Source/Runtime/Engine/Classes/Sound/SoundWave.h b/Engine/Source/Runtime/Engine/Classes/Sound/SoundWave.h
--- a/Engine/Source/Runtime/Engine/Classes/Sound/SoundWave.h
+++ b/Engine/Source/Runtime/Engine/Classes/Sound/SoundWave.h
@@ -1089,7 +1089,8 @@ public:
 
 	ENGINE_API void PostImport();
 	
-	ENGINE_API virtual void PreSave(FObjectPreSaveContext SaveContext);
+	ENGINE_API virtual void PreSave(FObjectPreSaveContext SaveContext) override;
+	ENGINE_API virtual void OnCookEvent(UE::Cook::ECookEvent CookEvent, UE::Cook::FCookEventContext& CookContext) override;
 
 private:
 	friend class FSoundWaveCompilingManager;
diff --git a/Engine/Source/Runtime/Engine/Classes/Sound/SoundWaveLoadingBehavior.cpp b/Engine/Source/Runtime/Engine/Classes/Sound/SoundWaveLoadingBehavior.cpp
--- a/Engine/Source/Runtime/Engine/Classes/Sound/SoundWaveLoadingBehavior.cpp
+++ b/Engine/Source/Runtime/Engine/Classes/Sound/SoundWaveLoadingBehavior.cpp
@@ -16,9 +16,9 @@
 #if WITH_EDITOR
 #include "Algo/Sort.h"
 #include "Cooker/CookDependency.h"
+#include "Cooker/CookEvents.h"
 #include "Serialization/CompactBinarySerialization.h"
 #include "Serialization/CompactBinaryWriter.h"
-#include "UObject/ObjectSaveContext.h"
 #endif
 
 #ifndef CASE_ENUM_TO_TEXT
@@ -292,7 +292,7 @@ private:
 	mutable TMap<FName, FAnnotatedClassData> CacheClassLoadingBehaviors;
 	mutable FCriticalSection CacheCS;
 	friend void UE::SoundWaveLoadingUtil::Private::HashSoundWaveLoadingBehaviorDependenciesForCook(FCbFieldViewIterator Args, UE::Cook::FCookDependencyContext& Context);
-	friend void UE::SoundWaveLoadingUtil::Private::RecordSoundWaveLoadingBehaviorDependenciesForCook(FObjectPreSaveContext ObjectSaveContext,	const USoundWave* SoundWave);
+	friend void UE::SoundWaveLoadingUtil::Private::RecordSoundWaveLoadingBehaviorDependenciesForCook(UE::Cook::FCookEventContext& CookContext, const USoundWave* SoundWave);
 };
 
 namespace UE { namespace SoundWaveLoadingUtil { namespace Private {
@@ -334,7 +334,7 @@ void HashSoundWaveLoadingBehaviorDependenciesForCook(FCbFieldViewIterator Args,
 
 UE_COOK_DEPENDENCY_FUNCTION(HashSoundWaveLoadingBehaviorDependenciesForCook, HashSoundWaveLoadingBehaviorDependenciesForCook);
 
-void RecordSoundWaveLoadingBehaviorDependenciesForCook(FObjectPreSaveContext ObjectSaveContext,
+void RecordSoundWaveLoadingBehaviorDependenciesForCook(UE::Cook::FCookEventContext& CookContext,
 	const USoundWave* SoundWave)
 {
 	// Cooking of USoundWaves depends on a bunch of indirect state that we want to cooker to know about for incremental cook invalidation
@@ -349,7 +349,7 @@ void RecordSoundWaveLoadingBehaviorDependenciesForCook(FObjectPreSaveContext Obj
 	// HashSoundWaveLoadingBehaviorDependenciesForCook is used to package up that query and make it deterministic for the cooker
 	// dependency checking
 	
-	ObjectSaveContext.AddCookBuildDependency(
+	CookContext.AddLoadBuildDependency(
 		UE::Cook::FCookDependency::Function(
 			UE_COOK_DEPENDENCY_FUNCTION_CALL(HashSoundWaveLoadingBehaviorDependenciesForCook), Writer.Save()));
 
@@ -367,12 +367,12 @@ void RecordSoundWaveLoadingBehaviorDependenciesForCook(FObjectPreSaveContext Obj
 
 		for (FName SoundCue : SoundCueHierarchyDependencies)
 		{
-			ObjectSaveContext.AddCookBuildDependency(UE::Cook::FCookDependency::Package(SoundCue));
+			CookContext.AddLoadBuildDependency(UE::Cook::FCookDependency::Package(SoundCue));
 		}
 	}
 
 	// 3. Finally, we depend on the value of the default loading behavior cvar
-	ObjectSaveContext.AddCookBuildDependency(UE::Cook::FCookDependency::ConsoleVariable(USoundWave::GetDefaultLoadingBehaviorCVarName()));
+	CookContext.AddLoadBuildDependency(UE::Cook::FCookDependency::ConsoleVariable(USoundWave::GetDefaultLoadingBehaviorCVarName()));
 }
 }}}
 
diff --git a/Engine/Source/Runtime/Engine/Classes/Sound/SoundWaveLoadingBehavior.h b/Engine/Source/Runtime/Engine/Classes/Sound/SoundWaveLoadingBehavior.h
--- a/Engine/Source/Runtime/Engine/Classes/Sound/SoundWaveLoadingBehavior.h
+++ b/Engine/Source/Runtime/Engine/Classes/Sound/SoundWaveLoadingBehavior.h
@@ -43,6 +43,7 @@ ENGINE_API const TCHAR* EnumToString(ESoundWaveLoadingBehavior InCurrentState);
 class USoundWave;
 class USoundClass;
 class ITargetPlatform;
+namespace UE::Cook { class FCookEventContext; }
 
 class ISoundWaveLoadingBehaviorUtil
 {
@@ -67,7 +68,7 @@ public:
 
 namespace UE::SoundWaveLoadingUtil::Private
 {
-	void RecordSoundWaveLoadingBehaviorDependenciesForCook(FObjectPreSaveContext ObjectSaveContext, const USoundWave* SoundWave);
+	void RecordSoundWaveLoadingBehaviorDependenciesForCook(UE::Cook::FCookEventContext& CookContext, const USoundWave* SoundWave);
 };
 
 #endif //WITH_EDITOR
diff --git a/Engine/Source/Runtime/Engine/Private/Blueprint.cpp b/Engine/Source/Runtime/Engine/Private/Blueprint.cpp
--- a/Engine/Source/Runtime/Engine/Private/Blueprint.cpp
+++ b/Engine/Source/Runtime/Engine/Private/Blueprint.cpp
@@ -24,6 +24,7 @@
 #include "Kismet2/KismetEditorUtilities.h"
 #include "Kismet2/StructureEditorUtils.h"
 #include "FindInBlueprintManager.h"
+#include "Cooker/CookEvents.h"
 #include "CookerSettings.h"
 #include "Editor.h"
 #include "Logging/MessageLog.h"
@@ -399,13 +400,6 @@ void UBlueprint::PreSave(FObjectPreSaveContext ObjectSaveContext)
 		// Cache the BP for use (immediate, since we're about to save)
 		FFindInBlueprintSearchManager::Get().AddOrUpdateBlueprintSearchMetadata(this, Flags, OverrideVersion);
 	}
-	
-#if WITH_EDITORONLY_DATA
-	if (ObjectSaveContext.IsCooking())
-	{
-		PreSaveCook(ObjectSaveContext);
-	}
-#endif
 }
 
 void UBlueprint::GetPreloadDependencies(TArray<UObject*>& OutDeps)
@@ -871,18 +865,23 @@ void UBlueprint::DebuggingWorldRegistrationHelper(UObject* ObjectProvidingWorld,
 	}
 }
 
-#if WITH_EDITORONLY_DATA
-void UBlueprint::PreSaveCook(FObjectPreSaveContext ObjectSaveContext)
+#if WITH_EDITOR
+void UBlueprint::OnCookEvent(UE::Cook::ECookEvent CookEvent, UE::Cook::FCookEventContext& CookContext)
 {
 	using namespace UE::Cook;
-	static bool bIgnoreBlueprintCookDependencies = FParse::Param(FCommandLine::Get(), TEXT("IgnoreBlueprintCookDependencies"));
-	if (bIgnoreBlueprintCookDependencies)
+	Super::OnCookEvent(CookEvent, CookContext);
+
+	if (CookEvent == UE::Cook::ECookEvent::PlatformCookDependencies && CookContext.IsCooking())
 	{
-		return;
+		static bool bIgnoreBlueprintCookDependencies = FParse::Param(FCommandLine::Get(), TEXT("IgnoreBlueprintCookDependencies"));
+		if (bIgnoreBlueprintCookDependencies)
+		{
+			return;
+		}
+
+		CookContext.AddLoadBuildDependency(
+			UE::Private::BlueprintDependencies::RecordCookDependencies(this));
 	}
-	
-	ObjectSaveContext.AddCookBuildDependency(
-		UE::Private::BlueprintDependencies::RecordCookDependencies(this));
 }
 #endif
 
diff --git a/Engine/Source/Runtime/Engine/Private/LevelInstance/LevelInstanceActor.cpp b/Engine/Source/Runtime/Engine/Private/LevelInstance/LevelInstanceActor.cpp
--- a/Engine/Source/Runtime/Engine/Private/LevelInstance/LevelInstanceActor.cpp
+++ b/Engine/Source/Runtime/Engine/Private/LevelInstance/LevelInstanceActor.cpp
@@ -13,6 +13,7 @@
 
 #if WITH_EDITOR
 #include "Cooker/CookDependency.h"
+#include "Cooker/CookEvents.h"
 #include "UObject/ObjectSaveContext.h"
 #include "WorldPartition/LevelInstance/LevelInstanceActorDesc.h"
 #include "LevelInstance/LevelInstanceEditorPivotActor.h"
@@ -205,14 +206,17 @@ void ALevelInstance::PostLoad()
 #endif
 }
 
-void ALevelInstance::PreSave(FObjectPreSaveContext SaveContext)
+void ALevelInstance::OnCookEvent(UE::Cook::ECookEvent CookEvent, UE::Cook::FCookEventContext& CookContext)
 {
-	Super::PreSave(SaveContext);
+	Super::OnCookEvent(CookEvent, CookContext);
 
-	FName WorldPackageName = FName(GetWorldAsset().GetLongPackageName());
-	if (!WorldPackageName.IsNone())
+	if (CookEvent == UE::Cook::ECookEvent::PlatformCookDependencies)
 	{
-		SaveContext.AddCookBuildDependency(UE::Cook::FCookDependency::Package(WorldPackageName));
+		FName WorldPackageName = FName(GetWorldAsset().GetLongPackageName());
+		if (!WorldPackageName.IsNone())
+		{
+			CookContext.AddLoadBuildDependency(UE::Cook::FCookDependency::Package(WorldPackageName));
+		}
 	}
 }
 
diff --git a/Engine/Source/Runtime/Engine/Private/Materials/Material.cpp b/Engine/Source/Runtime/Engine/Private/Materials/Material.cpp
--- a/Engine/Source/Runtime/Engine/Private/Materials/Material.cpp
+++ b/Engine/Source/Runtime/Engine/Private/Materials/Material.cpp
@@ -97,6 +97,7 @@
 #endif
 
 #if WITH_EDITOR
+#include "Cooker/CookEvents.h"
 #include "Framework/Notifications/NotificationManager.h"
 #include "Logging/MessageLog.h"
 #include "Misc/UObjectToken.h"
@@ -1136,17 +1137,24 @@ void UMaterial::PreSave(FObjectPreSaveContext ObjectSaveContext)
 	Super::PreSave(ObjectSaveContext);
 #if WITH_EDITOR
 	GMaterialsWithDirtyUsageFlags.RemoveAnnotation(this);
-	if (ObjectSaveContext.IsCooking())
+#endif
+}
+
+#if WITH_EDITOR
+void UMaterial::OnCookEvent(UE::Cook::ECookEvent CookEvent, UE::Cook::FCookEventContext& CookContext)
+{
+	Super::OnCookEvent(CookEvent, CookContext);
+	if (CookEvent == UE::Cook::ECookEvent::PlatformCookDependencies && CookContext.IsCooking())
 	{
-		const ITargetPlatform* TargetPlatform = ObjectSaveContext.GetTargetPlatform();
+		const ITargetPlatform* TargetPlatform = CookContext.GetTargetPlatform();
 		check(TargetPlatform);
 
 		TArray<FMaterialResourceForCooking>* Resources = CachedMaterialResourcesForCooking.Find(TargetPlatform);
-		UE::MaterialInterface::Private::RecordMaterialDependenciesForCook(ObjectSaveContext,
+		UE::MaterialInterface::Private::RecordMaterialDependenciesForCook(CookContext,
 			Resources ? *Resources : TArray<FMaterialResourceForCooking>());
 	}
-#endif
 }
+#endif
 
 void UMaterial::PostInitProperties()
 {
diff --git a/Engine/Source/Runtime/Engine/Private/Materials/MaterialInstance.cpp b/Engine/Source/Runtime/Engine/Private/Materials/MaterialInstance.cpp
--- a/Engine/Source/Runtime/Engine/Private/Materials/MaterialInstance.cpp
+++ b/Engine/Source/Runtime/Engine/Private/Materials/MaterialInstance.cpp
@@ -65,6 +65,7 @@
 
 #if WITH_EDITOR
 #include "Cooker/CookDependency.h"
+#include "Cooker/CookEvents.h"
 #include "String/ParseTokens.h"
 #endif
 
@@ -5141,28 +5142,29 @@ void UMaterialInstance::GetLightingGuidChain(bool bIncludeTextures, TArray<FGuid
 #endif
 }
 
-void UMaterialInstance::PreSave(FObjectPreSaveContext ObjectSaveContext)
-{
-	// @TODO : Remove any duplicate data from parent? Aims at improving change propagation (if controlled by parent)
-	Super::PreSave(ObjectSaveContext);
 #if WITH_EDITOR
-	if (ObjectSaveContext.IsCooking())
+void UMaterialInstance::OnCookEvent(UE::Cook::ECookEvent CookEvent, UE::Cook::FCookEventContext& CookContext)
+{
+	Super::OnCookEvent(CookEvent, CookContext);
+	if (CookEvent == UE::Cook::ECookEvent::PlatformCookDependencies && CookContext.IsCooking())
 	{
-		const ITargetPlatform * TargetPlatform = ObjectSaveContext.GetTargetPlatform();
+		// @TODO : Remove any duplicate data from parent? Aims at improving change propagation (if controlled by parent)
+		const ITargetPlatform* TargetPlatform = CookContext.GetTargetPlatform();
 		check(TargetPlatform);
 		TArray<FMaterialResourceForCooking>* Resources = CachedMaterialResourcesForCooking.Find(TargetPlatform);
-		UE::MaterialInterface::Private::RecordMaterialDependenciesForCook(ObjectSaveContext,
+		UE::MaterialInterface::Private::RecordMaterialDependenciesForCook(CookContext,
 			Resources ? *Resources : TArray<FMaterialResourceForCooking>());
 
 		UMaterialInterface* EffectiveParent = Parent ? Parent : UMaterial::GetDefaultMaterial(MD_Surface);
 		if (EffectiveParent && EffectiveParent->GetPackage() != GetPackage())
 		{
-			ObjectSaveContext.AddCookBuildDependency(UE::Cook::FCookDependency::TransitiveBuild(
+			CookContext.AddLoadBuildDependency(UE::Cook::FCookDependency::TransitiveBuild(
 				EffectiveParent->GetPackage()->GetFName()));
 		}
 	}
-#endif
 }
+#endif
+
 
 #if !WITH_EDITOR
 FTextureSamplingInfo UMaterialInstance::CalculateTexturesSamplingInfo(UTexture* Texture)
diff --git a/Engine/Source/Runtime/Engine/Private/Materials/MaterialShader.cpp b/Engine/Source/Runtime/Engine/Private/Materials/MaterialShader.cpp
--- a/Engine/Source/Runtime/Engine/Private/Materials/MaterialShader.cpp
+++ b/Engine/Source/Runtime/Engine/Private/Materials/MaterialShader.cpp
@@ -44,6 +44,7 @@
 #include "Algo/Sort.h"
 #include "Containers/AnsiString.h"
 #include "Cooker/CookDependency.h"
+#include "Cooker/CookEvents.h"
 #include "Interfaces/ITargetPlatform.h"
 #include "MaterialCachedData.h"
 #include "MaterialKeyGeneratorContext.h"
@@ -486,10 +487,10 @@ void HashMaterialDependenciesForCook(FCbFieldViewIterator Args, UE::Cook::FCookD
 
 UE_COOK_DEPENDENCY_FUNCTION(HashMaterialDependenciesForCook, UE::MaterialInterface::Private::HashMaterialDependenciesForCook);
 
-void RecordMaterialDependenciesForCook(FObjectPreSaveContext ObjectSaveContext,
+void RecordMaterialDependenciesForCook(UE::Cook::FCookEventContext CookContext,
 	TConstArrayView<FMaterialResourceForCooking> Resources)
 {
-	const ITargetPlatform* TargetPlatform = ObjectSaveContext.GetTargetPlatform();
+	const ITargetPlatform* TargetPlatform = CookContext.GetTargetPlatform();
 	check(TargetPlatform); // Caller should call only when cooking
 
 	TArray<FRecordedMaterialResourceForCooking> RecordedResources;
@@ -523,12 +524,12 @@ void RecordMaterialDependenciesForCook(FObjectPreSaveContext ObjectSaveContext,
 	}
 
 	// Store the FRecordedMaterialResourceForCooking in compact binary and pass them to the cooker to save in 
-	// ObjectSaveContext's list of dependencies for our package.
+	// CookContext's list of dependencies for our package.
 	FCbWriter Writer;
 	Writer << HashMaterialDependenciesForCookArgsVersion;
 	Writer << RecordedResources;
 
-	ObjectSaveContext.AddCookBuildDependency(
+	CookContext.AddLoadBuildDependency(
 		UE::Cook::FCookDependency::Function(
 			UE_COOK_DEPENDENCY_FUNCTION_CALL(HashMaterialDependenciesForCook), Writer.Save()));
 }
diff --git a/Engine/Source/Runtime/Engine/Private/Materials/MaterialSharedPrivate.h b/Engine/Source/Runtime/Engine/Private/Materials/MaterialSharedPrivate.h
--- a/Engine/Source/Runtime/Engine/Private/Materials/MaterialSharedPrivate.h
+++ b/Engine/Source/Runtime/Engine/Private/Materials/MaterialSharedPrivate.h
@@ -13,7 +13,7 @@
 #include "Serialization/CompactBinary.h"
 #include "Serialization/CompactBinaryWriter.h"
 
-class FObjectPreSaveContext;
+namespace UE::Cook { class FCookEventContext; }
 struct FAppendToClassSchemaContext;
 struct FMaterialResourceForCooking;
 
@@ -21,7 +21,7 @@ namespace UE::MaterialInterface::Private
 {
 
 /** Record the instance-specific cook dependencies (config, hlsl) from the resources for a UMaterialInterface. */
-void RecordMaterialDependenciesForCook(FObjectPreSaveContext ObjectSaveContext,
+void RecordMaterialDependenciesForCook(UE::Cook::FCookEventContext CookContext,
 	TConstArrayView<FMaterialResourceForCooking> Resources);
 
 /** Record the cook dependencies that apply to every UMaterialInterface. */
diff --git a/Engine/Source/Runtime/Engine/Private/SoundWave.cpp b/Engine/Source/Runtime/Engine/Private/SoundWave.cpp
--- a/Engine/Source/Runtime/Engine/Private/SoundWave.cpp
+++ b/Engine/Source/Runtime/Engine/Private/SoundWave.cpp
@@ -42,6 +42,7 @@
 #endif
 
 #if WITH_EDITOR
+#include "Cooker/CookEvents.h"
 #include "UObject/ArchiveCookContext.h"
 #endif
 
@@ -2456,11 +2457,21 @@ void USoundWave::PreSave(FObjectPreSaveContext InSaveContext)
 	{
 		// Populate our cache ahead of Serialize, where we can't do any AssetRegistry queries safely.
 		GetOwnerLoadingBehavior(InSaveContext.GetTargetPlatform());
-		UE::SoundWaveLoadingUtil::Private::RecordSoundWaveLoadingBehaviorDependenciesForCook(InSaveContext, this);
 	}
 	Super::PreSave(InSaveContext);
 }
 
+void USoundWave::OnCookEvent(UE::Cook::ECookEvent CookEvent, UE::Cook::FCookEventContext& CookContext)
+{
+	Super::OnCookEvent(CookEvent, CookContext);
+	if (CookEvent == UE::Cook::ECookEvent::PlatformCookDependencies && CookContext.IsCooking())
+	{
+		// We might need GetOwnerLoadingBehavior cached for RecordSoundWaveLoadingBehaviorDependenciesForCook
+		GetOwnerLoadingBehavior(CookContext.GetTargetPlatform());
+		UE::SoundWaveLoadingUtil::Private::RecordSoundWaveLoadingBehaviorDependenciesForCook(CookContext, this);
+	}
+}
+
 ISoundWaveLoadingBehaviorUtil::FClassData USoundWave::GetOwnerLoadingBehavior(const ITargetPlatform* InTargetPlatform) const
 {
 	// Use {} or Name_None in the case of a missing platform.
diff --git a/Engine/Source/Runtime/Engine/Private/Texture2D.cpp b/Engine/Source/Runtime/Engine/Private/Texture2D.cpp
--- a/Engine/Source/Runtime/Engine/Private/Texture2D.cpp
+++ b/Engine/Source/Runtime/Engine/Private/Texture2D.cpp
@@ -51,6 +51,7 @@
 
 #if WITH_EDITOR
 #include "Cooker/CookDependency.h"
+#include "Cooker/CookEvents.h"
 #include "Interfaces/ITargetPlatform.h"
 #endif
 
@@ -681,14 +682,6 @@ void UTexture2D::PreSave(FObjectPreSaveContext ObjectSaveContext)
 		UpdateResource();
 	}
 
-	if (ObjectSaveContext.IsCooking())
-	{
-		const ITargetPlatform* TargetPlatform = ObjectSaveContext.GetTargetPlatform();
-		check(TargetPlatform);
-		ObjectSaveContext.AddCookBuildDependency(UE::Cook::FCookDependency::SettingsObject(
-			&TargetPlatform->GetTextureLODSettings()));
-	}
-
 	// #TODO DC This is redundant code coming from UTexture::Presave that can be removed once we remove the above streaming code.
 
 	// Ensure that compilation has finished before saving the package
@@ -701,6 +694,20 @@ void UTexture2D::PreSave(FObjectPreSaveContext ObjectSaveContext)
 #endif
 }
 
+#if WITH_EDITOR
+void UTexture2D::OnCookEvent(UE::Cook::ECookEvent CookEvent, UE::Cook::FCookEventContext& CookContext)
+{
+	Super::OnCookEvent(CookEvent, CookContext);
+	if (CookEvent == UE::Cook::ECookEvent::PlatformCookDependencies && CookContext.IsCooking())
+	{
+		const ITargetPlatform* TargetPlatform = CookContext.GetTargetPlatform();
+		check(TargetPlatform);
+		CookContext.AddLoadBuildDependency(UE::Cook::FCookDependency::SettingsObject(
+			&TargetPlatform->GetTextureLODSettings()));
+	}
+}
+#endif
+
 void UTexture2D::GetAssetRegistryTags(TArray<FAssetRegistryTag>& OutTags) const
 {
 	PRAGMA_DISABLE_DEPRECATION_WARNINGS;
diff --git a/Engine/Source/Runtime/Engine/Private/WorldPartition/WorldPartition.cpp b/Engine/Source/Runtime/Engine/Private/WorldPartition/WorldPartition.cpp
--- a/Engine/Source/Runtime/Engine/Private/WorldPartition/WorldPartition.cpp
+++ b/Engine/Source/Runtime/Engine/Private/WorldPartition/WorldPartition.cpp
@@ -31,6 +31,7 @@
 #include UE_INLINE_GENERATED_CPP_BY_NAME(WorldPartition)
 
 #if WITH_EDITOR
+#include "Cooker/CookEvents.h"
 #include "LevelUtils.h"
 #include "Selection.h"
 #include "HAL/FileManager.h"
@@ -588,13 +589,13 @@ void UWorldPartition::PostDuplicate(EDuplicateMode::Type DuplicateMode)
 	}
 }
 
-void UWorldPartition::PreSave(FObjectPreSaveContext ObjectContext)
+void UWorldPartition::OnCookEvent(UE::Cook::ECookEvent CookEvent, UE::Cook::FCookEventContext& CookContext)
 {
-	Super::PreSave(ObjectContext);
-	if (ObjectContext.IsCooking())
+	Super::OnCookEvent(CookEvent, CookContext);
+	if (CookEvent == UE::Cook::ECookEvent::PlatformCookDependencies && CookContext.IsCooking())
 	{
 		FARFilter Filter = FExternalObjectAndActorDependencyGatherer::GetQueryFilter(GetPackage()->GetFName());
-		ObjectContext.AddCookBuildDependency(UE::Cook::FCookDependency::AssetRegistryQuery(MoveTemp(Filter)));
+		CookContext.AddLoadBuildDependency(UE::Cook::FCookDependency::AssetRegistryQuery(MoveTemp(Filter)));
 	}
 }
 
diff --git a/Engine/Source/Runtime/Engine/Public/LevelInstance/LevelInstanceActor.h b/Engine/Source/Runtime/Engine/Public/LevelInstance/LevelInstanceActor.h
--- a/Engine/Source/Runtime/Engine/Public/LevelInstance/LevelInstanceActor.h
+++ b/Engine/Source/Runtime/Engine/Public/LevelInstance/LevelInstanceActor.h
@@ -84,7 +84,7 @@ public:
 	ENGINE_API virtual void PostEditUndo() override;
 	ENGINE_API virtual void PostEditUndo(TSharedPtr<ITransactionObjectAnnotation> TransactionAnnotation) override;
 	ENGINE_API virtual void PostLoad() override;
-	ENGINE_API virtual void PreSave(FObjectPreSaveContext SaveContext) override;
+	ENGINE_API virtual void OnCookEvent(UE::Cook::ECookEvent CookEvent, UE::Cook::FCookEventContext& Context) override;
 	ENGINE_API virtual bool CanEditChange(const FProperty* InProperty) const override;
 	ENGINE_API virtual void PreEditChange(FProperty* PropertyThatWillChange) override;
 	ENGINE_API virtual void PostEditChangeProperty(FPropertyChangedEvent& PropertyChangedEvent) override;
diff --git a/Engine/Source/Runtime/Engine/Public/Materials/Material.h b/Engine/Source/Runtime/Engine/Public/Materials/Material.h
--- a/Engine/Source/Runtime/Engine/Public/Materials/Material.h
+++ b/Engine/Source/Runtime/Engine/Public/Materials/Material.h
@@ -1301,6 +1301,9 @@ public:
 
 	//~ Begin UObject Interface
 	ENGINE_API virtual void PreSave(FObjectPreSaveContext ObjectSaveContext) override;
+#if WITH_EDITOR
+	ENGINE_API virtual void OnCookEvent(UE::Cook::ECookEvent CookEvent, UE::Cook::FCookEventContext& CookContext) override;
+#endif
 	ENGINE_API virtual void PostInitProperties() override;
 	ENGINE_API virtual void Serialize(FArchive& Ar) override;
 	ENGINE_API virtual void PostDuplicate(bool bDuplicateForPIE) override;
diff --git a/Engine/Source/Runtime/Engine/Public/Materials/MaterialInstance.h b/Engine/Source/Runtime/Engine/Public/Materials/MaterialInstance.h
--- a/Engine/Source/Runtime/Engine/Public/Materials/MaterialInstance.h
+++ b/Engine/Source/Runtime/Engine/Public/Materials/MaterialInstance.h
@@ -762,7 +762,9 @@ public:
 	 */
 	class FMaterialInstanceResource* Resource;
 
-	ENGINE_API virtual void PreSave(FObjectPreSaveContext ObjectSaveContext) override;
+#if WITH_EDITOR
+	ENGINE_API virtual void OnCookEvent(UE::Cook::ECookEvent CookEvent, UE::Cook::FCookEventContext& CookContext) override;
+#endif
 	ENGINE_API virtual float GetTextureDensity(FName TextureName, const struct FMeshUVChannelInfo& UVChannelData) const override;
 #if !WITH_EDITOR
 	ENGINE_API virtual FTextureSamplingInfo CalculateTexturesSamplingInfo(UTexture* Texture) override;
diff --git a/Engine/Source/Runtime/Engine/Public/WorldPartition/WorldPartition.h b/Engine/Source/Runtime/Engine/Public/WorldPartition/WorldPartition.h
--- a/Engine/Source/Runtime/Engine/Public/WorldPartition/WorldPartition.h
+++ b/Engine/Source/Runtime/Engine/Public/WorldPartition/WorldPartition.h
@@ -209,7 +209,7 @@ public:
 	ENGINE_API virtual bool CanEditChange(const FProperty* InProperty) const override;
 	ENGINE_API virtual void PostEditChangeChainProperty(FPropertyChangedChainEvent& PropertyChangedEvent) override;
 	ENGINE_API virtual void PostDuplicate(EDuplicateMode::Type DuplicateMode) override;
-	ENGINE_API virtual void PreSave(FObjectPreSaveContext ObjectContext) override;
+	ENGINE_API virtual void OnCookEvent(UE::Cook::ECookEvent CookEvent, UE::Cook::FCookEventContext& CookContext) override;
 #endif //WITH_EDITOR
 	ENGINE_API virtual void Serialize(FArchive& Ar) override;
 	ENGINE_API virtual UWorld* GetWorld() const override;
