IncludeWhatYouUse: CookPackageSplitter.h: Remove old-style CoreMinimal.h, we include all the headers that we need individually.
#rnx

diff --git a/Engine/Source/Editor/UnrealEd/Public/CookPackageSplitter.h b/Engine/Source/Editor/UnrealEd/Public/CookPackageSplitter.h
--- a/Engine/Source/Editor/UnrealEd/Public/CookPackageSplitter.h
+++ b/Engine/Source/Editor/UnrealEd/Public/CookPackageSplitter.h
@@ -10,7 +10,9 @@
 #include "Containers/List.h"
 #include "Containers/StringView.h"
 #include "Containers/UnrealString.h"
+#if UE_ENABLE_INCLUDE_ORDER_DEPRECATED_IN_5_6
 #include "CoreMinimal.h"
+#endif
 #include "Hash/Blake3.h"
 #include "HAL/PlatformCrt.h"
 #include "HAL/PreprocessorHelpers.h"
