IncrementalCook: CookPackageSplitter: Fix typo in calls to PreSaveGeneratorPackage and PreSaveGeneratedPackage that was causing the KeepReferencedPackages to be ignored.
#rnx
#incrementalcook

diff --git a/Engine/Source/Editor/UnrealEd/Private/CookOnTheFlyServer.cpp b/Engine/Source/Editor/UnrealEd/Private/CookOnTheFlyServer.cpp
--- a/Engine/Source/Editor/UnrealEd/Private/CookOnTheFlyServer.cpp
+++ b/Engine/Source/Editor/UnrealEd/Private/CookOnTheFlyServer.cpp
@@ -3778,7 +3778,6 @@ UE::Cook::EPollStatus UCookOnTheFlyServer::PreSaveGeneratorPackage(UE::Cook::FPa
 		return EPollStatus::Error;
 	}
 
-	TArray<UPackage*> KeepReferencedPackages;
 	if (!TryConstructGeneratedPackagesForPopulate(PackageData, GenerationHelper, GeneratedPackagesForPopulate))
 	{
 		UE_LOG(LogCook, Error, TEXT("PackageSplitter unexpected failure: could not Construct GeneratedPackagesForPopulate. Splitter=%s"),
@@ -3786,6 +3785,7 @@ UE::Cook::EPollStatus UCookOnTheFlyServer::PreSaveGeneratorPackage(UE::Cook::FPa
 		return EPollStatus::Error;
 	}
 
+	UE::Cook::CookPackageSplitter::FPopulateContextData PopulateData;
 	{
 		FScopedActivePackage ScopedActivePackage(*this, GenerationHelper.GetOwner().GetPackageName(),
 #if UE_WITH_OBJECT_HANDLE_TRACKING
@@ -3794,7 +3794,6 @@ UE::Cook::EPollStatus UCookOnTheFlyServer::PreSaveGeneratorPackage(UE::Cook::FPa
 			FName()
 #endif
 		);
-		UE::Cook::CookPackageSplitter::FPopulateContextData PopulateData;
 		ICookPackageSplitter::FPopulateContext PopulateContext(PopulateData);
 		PopulateData.OwnerPackage = Package;
 		PopulateData.OwnerObject = SplitDataObject;
@@ -3813,7 +3812,7 @@ UE::Cook::EPollStatus UCookOnTheFlyServer::PreSaveGeneratorPackage(UE::Cook::FPa
 			return EPollStatus::Error;
 		}
 	}
-	Info.AddKeepReferencedPackages(GenerationHelper, KeepReferencedPackages);
+	Info.AddKeepReferencedPackages(GenerationHelper, PopulateData.KeepReferencedPackages);
 
 	return EPollStatus::Success;
 }
@@ -3952,7 +3951,7 @@ UE::Cook::EPollStatus UCookOnTheFlyServer::TryPopulateGeneratedPackage(UE::Cook:
 	GeneratedPackagePopulateData.GeneratedRootPath = GeneratedInfo.GeneratedRootPath;
 	GeneratedPackagePopulateData.Package = GeneratedPackage;
 	GeneratedPackagePopulateData.bCreatedAsMap = GeneratedInfo.IsCreateAsMap();
-	TArray<UPackage*> KeepReferencedPackages;
+	UE::Cook::CookPackageSplitter::FPopulateContextData PopulateData;
 	{
 		FScopedActivePackage ScopedActivePackage(*this, GenerationHelper.GetOwner().GetPackageName(),
 #if UE_WITH_OBJECT_HANDLE_TRACKING
@@ -3961,7 +3960,6 @@ UE::Cook::EPollStatus UCookOnTheFlyServer::TryPopulateGeneratedPackage(UE::Cook:
 			FName()
 #endif
 		);
-		UE::Cook::CookPackageSplitter::FPopulateContextData PopulateData;
 		ICookPackageSplitter::FPopulateContext PopulateContext(PopulateData);
 		PopulateData.OwnerPackage = OwnerPackage;
 		PopulateData.OwnerObject = OwnerObject;
@@ -3980,7 +3978,8 @@ UE::Cook::EPollStatus UCookOnTheFlyServer::TryPopulateGeneratedPackage(UE::Cook:
 			return EPollStatus::Error;
 		}
 	}
-	GeneratedInfo.AddKeepReferencedPackages(GenerationHelper, KeepReferencedPackages);
+	GeneratedInfo.AddKeepReferencedPackages(GenerationHelper, PopulateData.KeepReferencedPackages);
+
 	bool bPackageIsMap = GeneratedPackage->ContainsMap();
 	if (bPackageIsMap != GeneratedInfo.IsCreateAsMap())
 	{
