Remove placement positioning overrides from common Editor Viewport Toolbar functions. Update downstream toolbars to account for this.

End result is that creating a custom viewport toolbar from common functions can simply rely on function call order for item ordering rather than fiddling with positions.

diff --git a/Engine/Plugins/Runtime/GeometryCache/Source/GeometryCacheEd/Private/Widgets/SGeometryCacheEditorViewport.cpp b/Engine/Plugins/Runtime/GeometryCache/Source/GeometryCacheEd/Private/Widgets/SGeometryCacheEditorViewport.cpp
--- a/Engine/Plugins/Runtime/GeometryCache/Source/GeometryCacheEd/Private/Widgets/SGeometryCacheEditorViewport.cpp
+++ b/Engine/Plugins/Runtime/GeometryCache/Source/GeometryCacheEd/Private/Widgets/SGeometryCacheEditorViewport.cpp
@@ -293,16 +293,8 @@ TSharedPtr<SWidget> SGeometryCacheEditorViewport::BuildViewportToolbar()
 			
 			// Performance and Scalability Submenu
 			{
-				FToolMenuEntry Entry = UE::UnrealEd::CreatePerformanceAndScalabilitySubmenu();
-				{
-					Entry.Icon = FSlateIcon(FAppStyle::GetAppStyleSetName(), "Icons.Scalability");
-					Entry.InsertPosition.Position = EToolMenuInsertType::Last;
-					Entry.ToolBarData.LabelOverride = FText();
-					Entry.ToolBarData.ResizeParams.ClippingPriority = 800;
-				}
-
 				// Add Performance and Scalability Menu
-				RightSection.AddEntry(Entry);
+				RightSection.AddEntry(UE::UnrealEd::CreatePerformanceAndScalabilitySubmenu());
 
 				const FName SubmenuName = UToolMenus::JoinMenuPaths(ViewportToolbarName, "PerformanceAndScalability");
 				UToolMenu* Submenu = UToolMenus::Get()->ExtendMenu(SubmenuName);
diff --git a/Engine/Plugins/Runtime/HairStrands/Source/HairStrandsEditor/Private/Widgets/SGroomEditorViewport.cpp b/Engine/Plugins/Runtime/HairStrands/Source/HairStrandsEditor/Private/Widgets/SGroomEditorViewport.cpp
--- a/Engine/Plugins/Runtime/HairStrands/Source/HairStrandsEditor/Private/Widgets/SGroomEditorViewport.cpp
+++ b/Engine/Plugins/Runtime/HairStrands/Source/HairStrandsEditor/Private/Widgets/SGroomEditorViewport.cpp
@@ -410,7 +410,7 @@ TSharedPtr<SWidget> SGroomEditorViewport::BuildViewportToolbar()
 		
 		{
 			// LOD Menu
-			FToolMenuEntry& Entry = RightSection.AddDynamicEntry("DynamicLOD", FNewToolMenuSectionDelegate::CreateLambda([](FToolMenuSection& Section)
+			RightSection.AddDynamicEntry("DynamicLOD", FNewToolMenuSectionDelegate::CreateLambda([](FToolMenuSection& Section)
 			{
 				UUnrealEdViewportToolbarContext* Context = Section.FindContext<UUnrealEdViewportToolbarContext>();
 				if (!Context)
@@ -472,8 +472,6 @@ TSharedPtr<SWidget> SGroomEditorViewport::BuildViewportToolbar()
 					})
 				);
 			}));
-			
-			Entry.InsertPosition.Position = EToolMenuInsertType::Last;
 		}
 		
 		RightSection.AddEntry(UE::UnrealEd::CreatePerformanceAndScalabilitySubmenu());
diff --git a/Engine/Source/Editor/Kismet/Private/SSCSEditorViewport.cpp b/Engine/Source/Editor/Kismet/Private/SSCSEditorViewport.cpp
--- a/Engine/Source/Editor/Kismet/Private/SSCSEditorViewport.cpp
+++ b/Engine/Source/Editor/Kismet/Private/SSCSEditorViewport.cpp
@@ -466,7 +466,6 @@ TSharedPtr<SWidget> SSCSEditorViewport::BuildViewportToolbar()
 				);
 				ShowSubmenu.ToolBarData.LabelOverride = FText();
 				ShowSubmenu.Icon = FSlateIcon(FAppStyle::GetAppStyleSetName(), "Level.VisibleHighlightIcon16x");
-				ShowSubmenu.InsertPosition.Position = EToolMenuInsertType::Last;
 				ShowSubmenu.ToolBarData.ResizeParams.ClippingPriority = 800;
 				RightSection.AddEntry(ShowSubmenu);
 			}
@@ -493,7 +492,6 @@ TSharedPtr<SWidget> SSCSEditorViewport::BuildViewportToolbar()
 				);
 				PerfSubmenu.Icon = FSlateIcon(FAppStyle::GetAppStyleSetName(), "Icons.Scalability");
 				PerfSubmenu.ToolBarData.LabelOverride = FText();
-				PerfSubmenu.InsertPosition.Position = EToolMenuInsertType::Last;
 				PerfSubmenu.ToolBarData.ResizeParams.ClippingPriority = 800;
 				RightSection.AddEntry(PerfSubmenu);
 			}
diff --git a/Engine/Source/Editor/LevelEditor/Private/ViewportToolbar/LevelEditorSubmenus.cpp b/Engine/Source/Editor/LevelEditor/Private/ViewportToolbar/LevelEditorSubmenus.cpp
--- a/Engine/Source/Editor/LevelEditor/Private/ViewportToolbar/LevelEditorSubmenus.cpp
+++ b/Engine/Source/Editor/LevelEditor/Private/ViewportToolbar/LevelEditorSubmenus.cpp
@@ -1965,7 +1965,6 @@ FToolMenuEntry CreateShowSubmenu()
 		)
 	);
 	Entry.Icon = FSlateIcon(FAppStyle::GetAppStyleSetName(), "Level.VisibleHighlightIcon16x");
-	Entry.InsertPosition.Position = EToolMenuInsertType::Last;
 	Entry.ToolBarData.LabelOverride = FText();
 	Entry.ToolBarData.ResizeParams.ClippingPriority = 800;
 	return Entry;
@@ -2139,7 +2138,6 @@ FToolMenuEntry CreatePerformanceAndScalabilitySubmenu()
 		)
 	);
 	Entry.Icon = FSlateIcon(FAppStyle::GetAppStyleSetName(), "Icons.Scalability");
-	Entry.InsertPosition.Position = EToolMenuInsertType::Last;
 	Entry.ToolBarData.LabelOverride = FText();
 	Entry.ToolBarData.ResizeParams.ClippingPriority = 800;
 	return Entry;
@@ -2418,7 +2416,6 @@ FToolMenuEntry CreateSettingsSubmenu()
 	);
 
 	Entry.Icon = FSlateIcon(FAppStyle::GetAppStyleSetName(), "LevelEditor.GameSettings");
-	Entry.InsertPosition.Position = EToolMenuInsertType::Last;
 	Entry.ToolBarData.LabelOverride = FText();
 	Entry.ToolBarData.ResizeParams.ClippingPriority = 1000;
 	return Entry;
@@ -3009,7 +3006,6 @@ FToolMenuEntry CreateToolbarCameraSubmenu()
 						false,
 						Icon
 					);
-					Entry.InsertPosition.Position = EToolMenuInsertType::Last;
 					Entry.ToolBarData.ResizeParams.ClippingPriority = 800;
 				}
 			}
diff --git a/Engine/Source/Editor/StaticMeshEditor/Private/StaticMeshEditorViewportToolbarSections.cpp b/Engine/Source/Editor/StaticMeshEditor/Private/StaticMeshEditorViewportToolbarSections.cpp
--- a/Engine/Source/Editor/StaticMeshEditor/Private/StaticMeshEditorViewportToolbarSections.cpp
+++ b/Engine/Source/Editor/StaticMeshEditor/Private/StaticMeshEditorViewportToolbarSections.cpp
@@ -125,7 +125,6 @@ FToolMenuEntry UE::StaticMeshEditor::CreateLODSubmenu()
 							}
 						)
 					);
-					Entry.InsertPosition.Position = EToolMenuInsertType::Last;
 					Entry.ToolBarData.ResizeParams.ClippingPriority = 800;
 				}
 			}
@@ -148,7 +147,6 @@ FToolMenuEntry UE::StaticMeshEditor::CreateShowSubmenu()
 		)
 	);
 	Entry.Icon = FSlateIcon(FAppStyle::GetAppStyleSetName(), "Level.VisibleHighlightIcon16x");
-	Entry.InsertPosition.Position = EToolMenuInsertType::Last;
 	Entry.ToolBarData.LabelOverride = FText();
 	Entry.ToolBarData.ResizeParams.ClippingPriority = 800;
 
diff --git a/Engine/Source/Editor/UnrealEd/Private/SCommonEditorViewportToolbarBase.cpp b/Engine/Source/Editor/UnrealEd/Private/SCommonEditorViewportToolbarBase.cpp
--- a/Engine/Source/Editor/UnrealEd/Private/SCommonEditorViewportToolbarBase.cpp
+++ b/Engine/Source/Editor/UnrealEd/Private/SCommonEditorViewportToolbarBase.cpp
@@ -462,7 +462,6 @@ void SCommonEditorViewportToolbarBase::Construct(const FArguments& InArgs, TShar
 						})
 					);
 					ShowMenu.Icon = FSlateIcon(FAppStyle::GetAppStyleSetName(), "Level.VisibleHighlightIcon16x");
-					ShowMenu.InsertPosition.Position = EToolMenuInsertType::Last;
 					ShowMenu.ToolBarData.LabelOverride = FText();
 					ShowMenu.ToolBarData.ResizeParams.ClippingPriority = 800;
 				}
diff --git a/Engine/Source/Editor/UnrealEd/Private/ViewportToolbar/UnrealEdViewportToolbar.cpp b/Engine/Source/Editor/UnrealEd/Private/ViewportToolbar/UnrealEdViewportToolbar.cpp
--- a/Engine/Source/Editor/UnrealEd/Private/ViewportToolbar/UnrealEdViewportToolbar.cpp
+++ b/Engine/Source/Editor/UnrealEd/Private/ViewportToolbar/UnrealEdViewportToolbar.cpp
@@ -965,7 +965,6 @@ FToolMenuEntry CreateTransformsSubmenu()
 		)
 	);
 	Entry.Icon = FSlateIcon(FAppStyle::GetAppStyleSetName(), "LevelEditor.SelectMode");
-	Entry.InsertPosition.Position = EToolMenuInsertType::First;
 	Entry.ToolBarData.LabelOverride = FText();
 	Entry.ToolBarData.ResizeParams.ClippingPriority = 1000;
 
@@ -1287,7 +1286,6 @@ FToolMenuEntry CreateSnappingSubmenu()
 	);
 
 	Entry.Icon = FSlateIcon(FAppStyle::GetAppStyleSetName(), "ViewportToolbar.Snap");
-	Entry.InsertPosition.Position = EToolMenuInsertType::First;
 	Entry.ToolBarData.LabelOverride = FText();
 	Entry.ToolBarData.ResizeParams.ClippingPriority = 900;
 	return Entry;
@@ -1845,7 +1843,6 @@ FToolMenuEntry CreateViewModesSubmenu()
 					false,
 					IconAttribute
 				);
-				Entry.InsertPosition.Position = EToolMenuInsertType::Last;
 				Entry.ToolBarData.ResizeParams.ClippingPriority = 800;
 			}
 		)
@@ -2691,7 +2688,6 @@ FToolMenuEntry CreateCameraSubmenu()
 					false,
 					Icon
 				);
-				Entry.InsertPosition.Position = EToolMenuInsertType::Last;
 				Entry.ToolBarData.ResizeParams.ClippingPriority = 800;
 			}
 		)
@@ -2806,7 +2802,6 @@ FToolMenuEntry CreateAssetViewerProfileSubmenu()
 					false,
 					FSlateIcon(FAppStyle::GetAppStyleSetName(), "ViewportToolbar.PreviewSceneSettings")
 				);
-				Entry.InsertPosition.Position = EToolMenuInsertType::Last;
 				Entry.ToolBarData.ResizeParams.ClippingPriority = 1100;
 			}
 		)
@@ -3497,7 +3492,6 @@ FToolMenuEntry CreatePerformanceAndScalabilitySubmenu()
 		)
 	);
 	Entry.Icon = FSlateIcon(FAppStyle::GetAppStyleSetName(), "Icons.Scalability");
-	Entry.InsertPosition.Position = EToolMenuInsertType::Last;
 	Entry.ToolBarData.LabelOverride = FText();
 	Entry.ToolBarData.ResizeParams.ClippingPriority = 800;
 
@@ -3550,7 +3544,6 @@ FToolMenuEntry CreateDefaultShowSubmenu()
 	);
 
 	Entry.Icon = FSlateIcon(FAppStyle::GetAppStyleSetName(), "Level.VisibleHighlightIcon16x");
-	Entry.InsertPosition.Position = EToolMenuInsertType::Last;
 	Entry.ToolBarData.LabelOverride = FText();
 	Entry.ToolBarData.ResizeParams.ClippingPriority = 800;
 
