Fix regression in GatherObjectReferencersForDeletion from 40093672, where we wouldn't report referenced objects properly.

diff --git a/Engine/Source/Editor/UnrealEd/Private/ObjectTools.cpp b/Engine/Source/Editor/UnrealEd/Private/ObjectTools.cpp
--- a/Engine/Source/Editor/UnrealEd/Private/ObjectTools.cpp
+++ b/Engine/Source/Editor/UnrealEd/Private/ObjectTools.cpp
@@ -331,11 +331,9 @@ namespace ObjectTools
 			TSet<UObject*> InternalReferences;
 			// The old behavior of GatherObjectReferencersForDeletion will find anything that prevents 
 			// InObject from being garbage collected, including internal sub objects.
-			// it does make an exception with very specific package metadata case.
 			for (UObject* ObjectToDelete : ObjectsToDelete)
 			{
-				if ((ObjectToDelete->HasAnyFlags(GARBAGE_COLLECTION_KEEPFLAGS) || ObjectToDelete->HasAnyInternalFlags(EInternalObjectFlags_GarbageCollectionKeepFlags)) &&
-					(!bIsGatheringPackageRef))
+				if ((ObjectToDelete->HasAnyFlags(GARBAGE_COLLECTION_KEEPFLAGS) || ObjectToDelete->HasAnyInternalFlags(EInternalObjectFlags_GarbageCollectionKeepFlags)))
 				{
 					InternalReferences.Add(ObjectToDelete);
 					bOutIsReferenced = true;
