[CompileShadersTestBed] Log list of shader formats that are being used to compile global shaders.

This adds the list of target shader formats the logging output when global shaders are compiled for each target platform.
Example for Android: "Compile global shaders for SF_VULKAN_ES31_ANDROID, GLSL_ES3_1_ANDROID"

#rnx

diff --git a/Engine/Source/Editor/UnrealEd/Private/Commandlets/CompileShadersTestBedCommandlet.cpp b/Engine/Source/Editor/UnrealEd/Private/Commandlets/CompileShadersTestBedCommandlet.cpp
--- a/Engine/Source/Editor/UnrealEd/Private/Commandlets/CompileShadersTestBedCommandlet.cpp
+++ b/Engine/Source/Editor/UnrealEd/Private/Commandlets/CompileShadersTestBedCommandlet.cpp
@@ -163,14 +163,30 @@ int32 UCompileShadersTestBedCommandlet::Main(const FString& Params)
 			TArray<FName> DesiredShaderFormats;
 			Platform->GetAllTargetedShaderFormats(DesiredShaderFormats);
 
-			UE_LOG(LogCompileShadersTestBedCommandlet, Display, TEXT("Compile global shaders"));
+			if (!DesiredShaderFormats.IsEmpty())
+			{
+				TStringBuilder<4086> ShaderFormatListString;
+				for (const FName& ShaderFormat : DesiredShaderFormats)
+				{
+					if (ShaderFormatListString.Len() > 0)
+					{
+						ShaderFormatListString.Append(TEXT(", "));
+					}
+					ShaderFormatListString.Append(*ShaderFormat.ToString());
+				}
+				UE_LOG(LogCompileShadersTestBedCommandlet, Display, TEXT("Compile global shaders for %s"), ShaderFormatListString.GetData());
 
-			for (int32 FormatIndex = 0; FormatIndex < DesiredShaderFormats.Num(); FormatIndex++)
+				for (int32 FormatIndex = 0; FormatIndex < DesiredShaderFormats.Num(); FormatIndex++)
+				{
+					const EShaderPlatform ShaderPlatform = ShaderFormatToLegacyShaderPlatform(DesiredShaderFormats[FormatIndex]);
+					CompileGlobalShaderMap(ShaderPlatform, Platform, false);
+					GShaderCompilingManager->ProcessAsyncResults(bLimitExecutationTime, true /* bBlockOnGlobalShaderCompilation */);
+					check(GGlobalShaderMap[ShaderPlatform]->IsComplete(Platform));
+				}
+			}
+			else
 			{
-				const EShaderPlatform ShaderPlatform = ShaderFormatToLegacyShaderPlatform(DesiredShaderFormats[FormatIndex]);
-				CompileGlobalShaderMap(ShaderPlatform, Platform, false);
-				GShaderCompilingManager->ProcessAsyncResults(bLimitExecutationTime, true /* bBlockOnGlobalShaderCompilation */);
-				check(GGlobalShaderMap[ShaderPlatform]->IsComplete(Platform));
+				UE_LOG(LogCompileShadersTestBedCommandlet, Display, TEXT("No target shader formats available to compile global shaders"));
 			}
 		}
 
