IncrementalCook: CookPackageSplitter: Add FPopulateContext API function to report FCookDependencies for generated packages during Populate and PreSave.
#rnx

diff --git a/Engine/Source/Editor/UnrealEd/Private/CookOnTheFlyServer.cpp b/Engine/Source/Editor/UnrealEd/Private/CookOnTheFlyServer.cpp
--- a/Engine/Source/Editor/UnrealEd/Private/CookOnTheFlyServer.cpp
+++ b/Engine/Source/Editor/UnrealEd/Private/CookOnTheFlyServer.cpp
@@ -3979,6 +3979,7 @@ UE::Cook::EPollStatus UCookOnTheFlyServer::TryPopulateGeneratedPackage(UE::Cook:
 		}
 	}
 	GeneratedInfo.AddKeepReferencedPackages(GenerationHelper, PopulateData.KeepReferencedPackages);
+	GeneratedInfo.BuildResultDependencies.Append(MoveTemp(PopulateData.BuildResultDependencies));
 
 	bool bPackageIsMap = GeneratedPackage->ContainsMap();
 	if (bPackageIsMap != GeneratedInfo.IsCreateAsMap())
diff --git a/Engine/Source/Editor/UnrealEd/Private/Cooker/CookGenerationHelper.cpp b/Engine/Source/Editor/UnrealEd/Private/Cooker/CookGenerationHelper.cpp
--- a/Engine/Source/Editor/UnrealEd/Private/Cooker/CookGenerationHelper.cpp
+++ b/Engine/Source/Editor/UnrealEd/Private/Cooker/CookGenerationHelper.cpp
@@ -78,6 +78,11 @@ void ICookPackageSplitter::FPopulateContext::ReportKeepReferencedPackages(TConst
 	Data.KeepReferencedPackages.Append(Packages);
 }
 
+void ICookPackageSplitter::FPopulateContext::ReportSaveDependency(UE::Cook::FCookDependency CookDependency)
+{
+	Data.BuildResultDependencies.Add(UE::Cook::BuildResult::NAME_Save, MoveTemp(CookDependency));
+}
+
 bool ICookPackageSplitter::PopulateGeneratorPackage(UPackage* OwnerPackage, UObject* OwnerObject,
 	const TArray<ICookPackageSplitter::FGeneratedPackageForPopulate>& GeneratedPackages, TArray<UObject*>& OutObjectsToMove,
 	TArray<UPackage*>& OutKeepReferencedPackages)
@@ -1172,6 +1177,7 @@ bool FGenerationHelper::TryCallPopulateGeneratedPackage(UE::Cook::FCookGeneratio
 
 	OutObjectsToMove.Append(PopulateData.ObjectsToMove);
 	Info.AddKeepReferencedPackages(*this, PopulateData.KeepReferencedPackages);
+	Info.BuildResultDependencies.Append(MoveTemp(PopulateData.BuildResultDependencies));
 
 	// Contract Point 1: We will not call populate again until the splitter has been destroyed
 	// Contract Point 2: We will not call populate again without garbage collecting the generator package
@@ -1447,7 +1453,7 @@ void FGenerationHelper::FinishGeneratorPlatformSave(FPackageData& PackageData, b
 }
 
 void FGenerationHelper::FinishGeneratedPlatformSave(FPackageData& PackageData, FAssetPackageData& OutAssetPackageData,
-	TArray<FAssetDependency>& OutDependencies)
+	TArray<FAssetDependency>& OutDependencies, FBuildResultDependenciesMap& OutBuildResultDependencies)
 {
 	ConditionalInitialize();
 
@@ -1461,9 +1467,13 @@ void FGenerationHelper::FinishGeneratedPlatformSave(FPackageData& PackageData, F
 	UCookOnTheFlyServer& COTFS = Info->PackageData->GetPackageDatas().GetCookOnTheFlyServer();
 
 	// There should be no package dependencies present for the package from the global assetregistry
-	// because it is newly created. Add on the dependencies declared for it from the CookPackageSplitter.
+	// because it is newly created. Add on the package dependencies declared for it from the CookPackageSplitter.
+	// These dependencies are used in the AssetRegistry entry and runtime dependencies for the generated package
 	OutDependencies = Info->PackageDependencies;
 
+	// Add on the custom cook BuildResultDependencies declared for the generated package by the CookPackageSplitter.
+	OutBuildResultDependencies = Info->BuildResultDependencies;
+
 	// Update the AssetPackageData for each requested platform with Guid and ImportedClasses
 	TSet<UClass*> PackageClasses;
 	UPackage* Package = PackageData.GetPackage();
@@ -2239,6 +2249,7 @@ void FCookGenerationInfo::Uninitialize()
 	GeneratedRootPath.Empty();
 	GenerationHash.Reset();
 	PackageDependencies.Empty();
+	BuildResultDependencies.Empty();
 	// Keep PackageData; it is allowed in the uninitialized state
 	KeepReferencedPackages.Empty();
 	// CachedObjectsInOuterInfo should have been cleared before uninitialize, either by leaving the save state, or
diff --git a/Engine/Source/Editor/UnrealEd/Private/Cooker/CookGenerationHelper.h b/Engine/Source/Editor/UnrealEd/Private/Cooker/CookGenerationHelper.h
--- a/Engine/Source/Editor/UnrealEd/Private/Cooker/CookGenerationHelper.h
+++ b/Engine/Source/Editor/UnrealEd/Private/Cooker/CookGenerationHelper.h
@@ -8,6 +8,7 @@
 #include "Containers/ArrayView.h"
 #include "Containers/Map.h"
 #include "Containers/UnrealString.h"
+#include "Cooker/BuildResultDependenciesMap.h"
 #include "Cooker/CookPackageData.h"
 #include "Cooker/CookTypes.h"
 #include "Cooker/MPCollector.h"
@@ -36,6 +37,7 @@ struct FPopulateContextData
 {
 	TArray<UPackage*> KeepReferencedPackages;
 	TArray<UObject*> ObjectsToMove;
+	UE::Cook::FBuildResultDependenciesMap BuildResultDependencies;
 	TConstArrayView<ICookPackageSplitter::FGeneratedPackageForPopulate> GeneratedPackages;
 	UPackage* OwnerPackage = nullptr;
 	UObject* OwnerObject = nullptr;
@@ -144,6 +146,7 @@ public:
 	FString GeneratedRootPath;
 	FBlake3Hash GenerationHash;
 	TArray<FAssetDependency> PackageDependencies;
+	FBuildResultDependenciesMap BuildResultDependencies;
 	/** Cannot be null, set in constructor */
 	FPackageData* PackageData;
 	TArray<TWeakObjectPtr<UPackage>> KeepReferencedPackages;
@@ -318,7 +321,7 @@ public:
 	 * Calculates AssetRegistryData.
 	 */
 	void FinishGeneratedPlatformSave(FPackageData& PackageData, FAssetPackageData& OutAssetPackageData,
-		TArray< FAssetDependency>& OutDependencies);
+		TArray< FAssetDependency>& OutARDependencies, FBuildResultDependenciesMap& OutBuildResultDependencies);
 	/**
 	 * Return the AssetPackageData that was created for the given PackageData, or loaded for it from a previous cook.
 	 * Will return nullptr if not loaded from a previous incremental cook and it has not yet been saved.
diff --git a/Engine/Source/Editor/UnrealEd/Private/Cooker/CookSavePackage.cpp b/Engine/Source/Editor/UnrealEd/Private/Cooker/CookSavePackage.cpp
--- a/Engine/Source/Editor/UnrealEd/Private/Cooker/CookSavePackage.cpp
+++ b/Engine/Source/Editor/UnrealEd/Private/Cooker/CookSavePackage.cpp
@@ -421,6 +421,7 @@ void FSaveCookedPackageContext::FinishPlatform()
 
 	// Calculate up-to-date assetregistry data for Generator and Generated packages
 	TOptional<TArray<FAssetDependency>> OverridePackageDependencies;
+	TOptional<FBuildResultDependenciesMap> GeneratedPackageBuildResultDependencies;
 	TOptional<FAssetPackageData> AssetPackageDataBuffer;
 	TOptional<FAssetPackageData> OverrideAssetPackageData;
 	const FAssetPackageData* AssetPackageData = nullptr;
@@ -440,7 +441,9 @@ void FSaveCookedPackageContext::FinishPlatform()
 			bGenerated = true;
 			OverrideAssetPackageData.Emplace();
 			OverridePackageDependencies.Emplace();
-			GenerationHelper->FinishGeneratedPlatformSave(PackageData, *OverrideAssetPackageData, *OverridePackageDependencies);
+			GeneratedPackageBuildResultDependencies.Emplace();
+			GenerationHelper->FinishGeneratedPlatformSave(PackageData, *OverrideAssetPackageData,
+				 *OverridePackageDependencies, *GeneratedPackageBuildResultDependencies);
 			AssetPackageData = OverrideAssetPackageData.GetPtrOrNull();
 		}
 		if (!AssetPackageData)
@@ -460,7 +463,9 @@ void FSaveCookedPackageContext::FinishPlatform()
 		CalculatePlatformAgnosticRuntimeDependencies();
 		CalculatePlatformRuntimeDependencies();
 		// Note CalculateCookDependencies mutates this->SaveResult; it moves BuildResultDependencies out of it
-		CalculateCookDependencies(GenerationHelper.GetReference(), bGenerated, OverridePackageDependencies.GetPtrOrNull());
+		CalculateCookDependencies(GenerationHelper.GetReference(), bGenerated,
+			OverridePackageDependencies.GetPtrOrNull() /** const, not modified */,
+			GeneratedPackageBuildResultDependencies.GetPtrOrNull() /** non-const, function sets it to empty */);
 		RecordPlatformBuildDependencies();
 
 		ICookedPackageWriter::FCommitPackageInfo Info;
@@ -767,7 +772,7 @@ private:
 };
 
 void FSaveCookedPackageContext::CalculateCookDependencies(FGenerationHelper* GenerationHelper, bool bGenerated,
-	TArray<FAssetDependency>* OverrideDependencies)
+	const TArray<FAssetDependency>* ExtraArDependencies, FBuildResultDependenciesMap* ExtraBuildResultDependencies)
 {
 	if (COTFS.bHybridIterativeEnabled)
 	{
@@ -777,14 +782,18 @@ void FSaveCookedPackageContext::CalculateCookDependencies(FGenerationHelper* Gen
 		TConstArrayView<FName> UntrackedSoftPackageReferences;
 		
 		PlatformRuntimeDependencies = GetPlatformRuntimeDependencies();
-		if (OverrideDependencies)
+		if (ExtraArDependencies)
 		{
-			PlatformRuntimeDependencies.Reserve(PlatformRuntimeDependencies.Num() + OverrideDependencies->Num());
-			for (const FAssetDependency& Dependency : *OverrideDependencies)
+			PlatformRuntimeDependencies.Reserve(PlatformRuntimeDependencies.Num() + ExtraArDependencies->Num());
+			for (const FAssetDependency& Dependency : *ExtraArDependencies)
 			{
 				PlatformRuntimeDependencies.Add(Dependency.AssetId.PackageName);
 			}
 		}
+		if (ExtraBuildResultDependencies)
+		{
+			BuildResultDependencies.Append(MoveTemp(*ExtraBuildResultDependencies));
+		}
 		bool bHasSavePackageResult = HasSavePackageResult();
 		FSavePackageResultStruct* LocalSavePackageResult = bHasSavePackageResult ? &SavePackageResult : nullptr;
 
diff --git a/Engine/Source/Editor/UnrealEd/Private/Cooker/CookSavePackage.h b/Engine/Source/Editor/UnrealEd/Private/Cooker/CookSavePackage.h
--- a/Engine/Source/Editor/UnrealEd/Private/Cooker/CookSavePackage.h
+++ b/Engine/Source/Editor/UnrealEd/Private/Cooker/CookSavePackage.h
@@ -51,7 +51,7 @@ private: // Used only by UCookOnTheFlyServer, which has private access
 	void CalculatePlatformRuntimeDependencies();
 	TArray<FName> GetPlatformRuntimeDependencies() const;
 	void CalculateCookDependencies(FGenerationHelper* GenerationHelper, bool bGenerated,
-		TArray<FAssetDependency>* OverrideDependencies);
+		const TArray<FAssetDependency>* ExtraARDependencies, FBuildResultDependenciesMap* ExtraBuildResultDependencies);
 	void RecordPlatformBuildDependencies();
 	bool HasSavePackageResult() const;
 
diff --git a/Engine/Source/Editor/UnrealEd/Public/CookPackageSplitter.h b/Engine/Source/Editor/UnrealEd/Public/CookPackageSplitter.h
--- a/Engine/Source/Editor/UnrealEd/Public/CookPackageSplitter.h
+++ b/Engine/Source/Editor/UnrealEd/Public/CookPackageSplitter.h
@@ -27,6 +27,7 @@ class UPackage;
 struct FAssetDependency;
 template <typename FuncType> class TFunctionRef;
 
+namespace UE::Cook { class FCookDependency; }
 namespace UE::Cook::CookPackageSplitter { struct FPopulateContextData; }
 
 /**
@@ -145,8 +146,16 @@ public:
 		UE_DEPRECATED(5.3, "Write to PackageDependencies instead")
 		TArray<FName> Dependencies;
 		/**
-		 * AssetRegistry dependencies for the generated package. AR dependencies cause packages to be added to the
-		 * current cook and cause invalidation of this package in iterative cooks if any dependencies change.
+		 * Source packages outside of the generator package that will be incorporated into the generated
+		 * package (e.g. ExternalActor packages). These are used to construct the PackageSavedHash for the generated
+		 * package. Some objects use the PackageSavedHash during derived data construction as a change marker, adding
+		 * the source packages here is important for those types to work.
+		 *
+		 * These packages are also recorded as dependencies in the AssetRegistry generated by the cook.
+		 *
+		 * During incrementalcook, changes to these packages cause the recook of the generated package, but that
+		 * invalidation can also be accomplished without the other effects, and with more types of dependencies, using
+		 * PopulateContext.ReportSaveDependency during PopulateGeneratedPackage and PreSaveGeneratedPackage.
 		 */
 		TArray<FAssetDependency> PackageDependencies;
 		/**
@@ -247,6 +256,15 @@ public:
 		 */
 		UNREALED_API void ReportKeepReferencedPackage(UPackage* Package);
 		UNREALED_API void ReportKeepReferencedPackages(TConstArrayView<UPackage*> Packages);
+
+		/**
+		 * Add the given FCookDependency to the build dependencies for the TargetGeneratedPackage. Incremental cooks
+		 * will invalidate the package and recook it if the CookDependency changes.
+		 * 
+		 * This callback is only valid during Populate and PreSave functions. It is ignored during PostSave functions.
+		 * This callback is only valid in calls on generated packages. It is ignored during calls to the generator.
+		 */
+		UNREALED_API void ReportSaveDependency(UE::Cook::FCookDependency CookDependency);
 	private:
 		UE::Cook::CookPackageSplitter::FPopulateContextData& Data;
 	};
