AxisDisplayInfo : Add GetAxisDisplayCoordinateSystem() to query the editor's axis display coordinate system.

EAxisList adds values for Left, Up, Forward and LeftUpForward.
The cvar Editor.UseForwardRightUpDisplayNames is replaced by [REDACTED_PERSON] and defaults to XYZ.
AxisDisplayInfo methods are updated to handle EAxisList::Left, EAxisList::Up, and EAxisList::Forward.

This will allow for a gradual transition to the LeftUpForward display coordinate system.

diff --git a/Engine/Source/Editor/UnrealEd/Private/EditorAxisDisplayInfo.cpp b/Engine/Source/Editor/UnrealEd/Private/EditorAxisDisplayInfo.cpp
--- a/Engine/Source/Editor/UnrealEd/Private/EditorAxisDisplayInfo.cpp
+++ b/Engine/Source/Editor/UnrealEd/Private/EditorAxisDisplayInfo.cpp
@@ -10,10 +10,10 @@
 
 namespace UE::Editor::Private
 {
-	static TAutoConsoleVariable<bool> EditorUseForwardRightUpDisplayNames(
-		TEXT("Editor.UseForwardRightUpDisplayNames"),
-		false,
-		TEXT("When true, the editor uses (Forward, Right, Up) nomenclature for coordinate axes"),
+	static TAutoConsoleVariable<int> EditorAxisDisplayCoordinateSystem(
+		TEXT("Editor.AxisDisplayCoordinateSystem"),
+		0, // XYZ
+		TEXT("Sets the editor's axis display coordinate system { 0 = XYZ (default), 1 = FRU (deprecated), 2 = LUF }"),
 		ECVF_ReadOnly
 	);
 }
@@ -22,53 +22,102 @@ FEditorAxisDisplayInfo::FEditorAxisDisplayInfo()
 {
 	if (FProperty* XAxisColorProperty = UEditorStyleSettings::StaticClass()->FindPropertyByName(GET_MEMBER_NAME_CHECKED(UEditorStyleSettings, XAxisColor)))
 	{
-		XAxisColorProperty->SetMetaData(TEXT("DisplayName"), *FText::Format(LOCTEXT("XAxisColorDisplayName", "{0} Axis Color"), GetAxisDisplayName(EAxisList::X)).ToString());
+		XAxisColorProperty->SetMetaData(TEXT("DisplayName"), *FText::Format(LOCTEXT("XAxisColorDisplayName", "{0} Axis Color"), GetAxisDisplayName(EAxisList::Forward)).ToString());
+		XAxisColorProperty->SetMetaData(TEXT("ToolTip"), *FText::Format(LOCTEXT("XAxisColorToolTip", "The color used for the {0} axis"), GetAxisDisplayName(EAxisList::Forward)).ToString());
 	}
 	if (FProperty* YAxisColorProperty = UEditorStyleSettings::StaticClass()->FindPropertyByName(GET_MEMBER_NAME_CHECKED(UEditorStyleSettings, YAxisColor)))
 	{
-		YAxisColorProperty->SetMetaData(TEXT("DisplayName"), *FText::Format(LOCTEXT("YAxisColorDisplayName", "{0} Axis Color"), GetAxisDisplayName(EAxisList::Y)).ToString());
+		YAxisColorProperty->SetMetaData(TEXT("DisplayName"), *FText::Format(LOCTEXT("YAxisColorDisplayName", "{0} Axis Color"), GetAxisDisplayName(EAxisList::Left)).ToString());
+		YAxisColorProperty->SetMetaData(TEXT("ToolTip"), *FText::Format(LOCTEXT("YAxisColorToolTip", "The color used for the {0} axis"), GetAxisDisplayName(EAxisList::Left)).ToString());
 	}
 	if (FProperty* ZAxisColorProperty = UEditorStyleSettings::StaticClass()->FindPropertyByName(GET_MEMBER_NAME_CHECKED(UEditorStyleSettings, ZAxisColor)))
 	{
-		ZAxisColorProperty->SetMetaData(TEXT("DisplayName"), *FText::Format(LOCTEXT("ZAxisColorDisplayName", "{0} Axis Color"), GetAxisDisplayName(EAxisList::Z)).ToString());
+		ZAxisColorProperty->SetMetaData(TEXT("DisplayName"), *FText::Format(LOCTEXT("ZAxisColorDisplayName", "{0} Axis Color"), GetAxisDisplayName(EAxisList::Up)).ToString());
+		ZAxisColorProperty->SetMetaData(TEXT("ToolTip"), *FText::Format(LOCTEXT("ZAxisColorToolTip", "The color used for the {0} axis"), GetAxisDisplayName(EAxisList::Up)).ToString());
 	}
 }
 
+EAxisList::Type FEditorAxisDisplayInfo::GetAxisDisplayCoordinateSystem()
+{
+	if (!AxisDisplayCoordinateSystem.IsSet())
+	{
+		const int EditorAxisDisplayCoordinateSystem = UE::Editor::Private::EditorAxisDisplayCoordinateSystem.GetValueOnAnyThread();
+		
+		if (EditorAxisDisplayCoordinateSystem == 2)
+		{
+			AxisDisplayCoordinateSystem = EAxisList::LeftUpForward;
+		}
+		else
+		{
+			// There is no enumeration for the deprecated (Forward,Right,Up) (EditorAxisDisplayCoordinateSystem == 1), it uses EAxisList::XYZ
+
+			ensureMsgf(EditorAxisDisplayCoordinateSystem == 0 || EditorAxisDisplayCoordinateSystem == 1,
+				TEXT("Unsupported Editor.AxisDisplayCoordinateSystem: %d"), EditorAxisDisplayCoordinateSystem);
+
+			AxisDisplayCoordinateSystem = EAxisList::XYZ;
+		}
+	}
+
+	return AxisDisplayCoordinateSystem.Get(EAxisList::XYZ/*DefaultValue*/);
+}
+
 FText FEditorAxisDisplayInfo::GetAxisDisplayName(EAxisList::Type Axis)
 {
+	static const FText XDisplayName = LOCTEXT("XDisplayName", "X");
+	static const FText YDisplayName = LOCTEXT("YDisplayName", "Y");
+	static const FText ZDisplayName = LOCTEXT("ZDisplayName", "Z");
+	static const FText RightDisplayName = LOCTEXT("RightDisplayName", "Right (was Y)");
+	static const FText LeftDisplayName = LOCTEXT("LeftDisplayName", "Left (was -Y)");
+	static const FText UpDisplayName = LOCTEXT("UpDisplayName", "Up (was Z)");
+	static const FText ForwardDisplayName = LOCTEXT("ForwardDisplayName", "Forward (was X)");
+	
+	Axis = MapAxis(Axis);
+
 	if (Axis == EAxisList::X)
 	{
 		if (UseForwardRightUpDisplayNames())
 		{
-			return LOCTEXT("ForwardDisplayName", "Forward (was X)");
+			return ForwardDisplayName;
 		}
 		else
 		{
-			return LOCTEXT("XDisplayName", "X");
+			return XDisplayName;
 		}
 	}
 	else if (Axis == EAxisList::Y)
 	{
 		if (UseForwardRightUpDisplayNames())
 		{
-			return LOCTEXT("RightDisplayName", "Right (was Y)");
+			return RightDisplayName;
 		}
 		else
 		{
-			return LOCTEXT("YDisplayName", "Y");
+			return YDisplayName;
 		}
 	}
 	else if (Axis == EAxisList::Z)
 	{
 		if (UseForwardRightUpDisplayNames())
 		{
-			return LOCTEXT("UpDisplayName", "Up (was Z)");
+			return UpDisplayName;
 		}
 		else
 		{
-			return LOCTEXT("ZDisplayName", "Z");
+			return ZDisplayName;
 		}
 	}
+	else if (Axis == EAxisList::Left)
+	{
+		return LeftDisplayName;
+	}
+	else if (Axis == EAxisList::Up)
+	{
+		return UpDisplayName;
+	}
+	else if (Axis == EAxisList::Forward)
+	{
+		return ForwardDisplayName;
+	}
 
 	ensureMsgf(false, TEXT("Unsupported Axis: %d"), Axis);
 
@@ -77,39 +126,61 @@ FText FEditorAxisDisplayInfo::GetAxisDisplayName(EAxisList::Type Axis)
 
 FText FEditorAxisDisplayInfo::GetAxisDisplayNameShort(EAxisList::Type Axis)
 {
+	static const FText XDisplayNameShort = LOCTEXT("XDisplayNameShort", "X");
+	static const FText YDisplayNameShort = LOCTEXT("YDisplayNameShort", "Y");
+	static const FText ZDisplayNameShort = LOCTEXT("ZDisplayNameShort", "Z");
+	static const FText RightDisplayNameShort = LOCTEXT("RightDisplayNameShort", "R");
+	static const FText LeftDisplayNameShort = LOCTEXT("LeftDisplayNameShort", "L");
+	static const FText UpDisplayNameShort = LOCTEXT("UpDisplayNameShort", "U");
+	static const FText ForwardDisplayNameShort = LOCTEXT("ForwardDisplayNameShort", "F");
+
+	Axis = MapAxis(Axis);
+
 	if (Axis == EAxisList::X)
 	{
 		if (UseForwardRightUpDisplayNames())
 		{
-			return LOCTEXT("ForwardDisplayNameShort", "F");
+			return ForwardDisplayNameShort;
 		}
 		else
 		{
-			return LOCTEXT("XDisplayNameShort", "X");
+			return XDisplayNameShort;
 		}
 	}
 	else if (Axis == EAxisList::Y)
 	{
 		if (UseForwardRightUpDisplayNames())
 		{
-			return LOCTEXT("RightDisplayNameShort", "R");
+			return RightDisplayNameShort;
 		}
 		else
 		{
-			return LOCTEXT("YDisplayNameShort", "Y");
+			return YDisplayNameShort;
 		}
 	}
 	else if (Axis == EAxisList::Z)
 	{
 		if (UseForwardRightUpDisplayNames())
 		{
-			return LOCTEXT("UpDisplayNameShort", "U");
+			return UpDisplayNameShort;
 		}
 		else
 		{
-			return LOCTEXT("ZDisplayNameShort", "Z");
+			return ZDisplayNameShort;
 		}
 	}
+	else if (Axis == EAxisList::Left)
+	{
+		return LeftDisplayNameShort;
+	}
+	else if (Axis == EAxisList::Up)
+	{
+		return UpDisplayNameShort;
+	}
+	else if (Axis == EAxisList::Forward)
+	{
+		return ForwardDisplayNameShort;
+	}
 
 	ensureMsgf(false, TEXT("Unsupported Axis: %d"), Axis);
 
@@ -118,15 +189,17 @@ FText FEditorAxisDisplayInfo::GetAxisDisplayNameShort(EAxisList::Type Axis)
 
 FLinearColor FEditorAxisDisplayInfo::GetAxisColor(EAxisList::Type Axis)
 {
-	if (Axis == EAxisList::X)
+	Axis = MapAxis(Axis);
+
+	if (Axis == EAxisList::X || Axis == EAxisList::Forward)
 	{
 		return GetDefault<UEditorStyleSettings>()->XAxisColor;
 	}
-	else if (Axis == EAxisList::Y)
+	else if (Axis == EAxisList::Y || Axis == EAxisList::Left)
 	{
 		return GetDefault<UEditorStyleSettings>()->YAxisColor;
 	}
-	else if (Axis == EAxisList::Z)
+	else if (Axis == EAxisList::Z || Axis == EAxisList::Up)
 	{
 		return GetDefault<UEditorStyleSettings>()->ZAxisColor;
 	}
@@ -140,10 +213,32 @@ bool FEditorAxisDisplayInfo::UseForwardRightUpDisplayNames()
 {
 	if (!bUseForwardRightUpDisplayNames.IsSet())
 	{
-		bUseForwardRightUpDisplayNames = UE::Editor::Private::EditorUseForwardRightUpDisplayNames.GetValueOnAnyThread();
+		const int EditorAxisDisplayCoordinateSystem = UE::Editor::Private::EditorAxisDisplayCoordinateSystem.GetValueOnAnyThread();
+		bUseForwardRightUpDisplayNames = (EditorAxisDisplayCoordinateSystem == 1);
 	}
 
 	return bUseForwardRightUpDisplayNames.Get(false/*DefaultValue*/);
 }
 
+EAxisList::Type FEditorAxisDisplayInfo::MapAxis(EAxisList::Type Axis)
+{
+	if (GetAxisDisplayCoordinateSystem() == EAxisList::XYZ)
+	{
+		if (Axis == EAxisList::Forward)
+		{
+			return EAxisList::X;
+		}
+		else if (Axis == EAxisList::Left)
+		{
+			return EAxisList::Y;
+		}
+		else if (Axis == EAxisList::Up)
+		{
+			return EAxisList::Z;
+		}
+	}
+
+	return Axis;
+}
+
 #undef LOCTEXT_NAMESPACE
diff --git a/Engine/Source/Editor/UnrealEd/Public/EditorAxisDisplayInfo.h b/Engine/Source/Editor/UnrealEd/Public/EditorAxisDisplayInfo.h
--- a/Engine/Source/Editor/UnrealEd/Public/EditorAxisDisplayInfo.h
+++ b/Engine/Source/Editor/UnrealEd/Public/EditorAxisDisplayInfo.h
@@ -11,6 +11,8 @@ public:
 	FEditorAxisDisplayInfo();
 	virtual ~FEditorAxisDisplayInfo() = default;
 
+	virtual EAxisList::Type GetAxisDisplayCoordinateSystem() override;
+
 	virtual FText GetAxisDisplayName(EAxisList::Type Axis) override;
 
 	virtual FText GetAxisDisplayNameShort(EAxisList::Type Axis) override;
@@ -21,6 +23,10 @@ public:
 
 private:
 
+	// Maps the given axis from FLU -> XYZ if the AxisDisplayCoordinateSystem is XYZ
+	EAxisList::Type MapAxis(EAxisList::Type Axis);
+
 	TOptional<bool> bUseForwardRightUpDisplayNames;
+	TOptional<EAxisList::Type> AxisDisplayCoordinateSystem;
 
 };
diff --git a/Engine/Source/Runtime/Core/Private/Misc/AxisDisplayInfo.cpp b/Engine/Source/Runtime/Core/Private/Misc/AxisDisplayInfo.cpp
--- a/Engine/Source/Runtime/Core/Private/Misc/AxisDisplayInfo.cpp
+++ b/Engine/Source/Runtime/Core/Private/Misc/AxisDisplayInfo.cpp
@@ -7,6 +7,17 @@
 
 #define LOCTEXT_NAMESPACE "AxisDisplayInfo"
 
+EAxisList::Type AxisDisplayInfo::GetAxisDisplayCoordinateSystem()
+{
+	if (IModularFeatures::Get().IsModularFeatureAvailable(IAxisDisplayInfo::GetModularFeatureName()))
+	{
+		IAxisDisplayInfo& DisplayInfo = IModularFeatures::Get().GetModularFeature<IAxisDisplayInfo>(IAxisDisplayInfo::GetModularFeatureName());
+		return DisplayInfo.GetAxisDisplayCoordinateSystem();
+	}
+
+	return EAxisList::XYZ;
+}
+
 FText AxisDisplayInfo::GetAxisDisplayName(EAxisList::Type Axis)
 {
 	if (IModularFeatures::Get().IsModularFeatureAvailable(IAxisDisplayInfo::GetModularFeatureName()))
@@ -15,15 +26,15 @@ FText AxisDisplayInfo::GetAxisDisplayName(EAxisList::Type Axis)
 		return DisplayInfo.GetAxisDisplayName(Axis);
 	}
 
-	if (Axis == EAxisList::X)
+	if (Axis == EAxisList::X || Axis == EAxisList::Forward)
 	{
 		return LOCTEXT("XDisplayName", "X");
 	}
-	else if (Axis == EAxisList::Y)
+	else if (Axis == EAxisList::Y || Axis == EAxisList::Left)
 	{
 		return LOCTEXT("YDisplayName", "Y");
 	}
-	else if (Axis == EAxisList::Z)
+	else if (Axis == EAxisList::Z || Axis == EAxisList::Up)
 	{
 		return LOCTEXT("ZDisplayName", "Z");
 	}
@@ -41,15 +52,15 @@ FText AxisDisplayInfo::GetAxisDisplayNameShort(EAxisList::Type Axis)
 		return DisplayInfo.GetAxisDisplayNameShort(Axis);
 	}
 
-	if (Axis == EAxisList::X)
+	if (Axis == EAxisList::X || Axis == EAxisList::Forward)
 	{
 		return LOCTEXT("XDisplayNameShort", "X");
 	}
-	else if (Axis == EAxisList::Y)
+	else if (Axis == EAxisList::Y || Axis == EAxisList::Left)
 	{
 		return LOCTEXT("YDisplayNameShort", "Y");
 	}
-	else if (Axis == EAxisList::Z)
+	else if (Axis == EAxisList::Z || Axis == EAxisList::Up)
 	{
 		return LOCTEXT("ZDisplayNameShort", "Z");
 	}
@@ -67,15 +78,15 @@ FLinearColor AxisDisplayInfo::GetAxisColor(EAxisList::Type Axis)
 		return DisplayInfo.GetAxisColor(Axis);
 	}
 
-	if (Axis == EAxisList::X)
+	if (Axis == EAxisList::X || Axis == EAxisList::Forward)
 	{
 		return FLinearColor(0.594f, 0.0197f, 0.0f);
 	}
-	else if (Axis == EAxisList::Y)
+	else if (Axis == EAxisList::Y || Axis == EAxisList::Left)
 	{
 		return FLinearColor(0.1349f, 0.3959f, 0.0f);
 	}
-	else if (Axis == EAxisList::Z)
+	else if (Axis == EAxisList::Z || Axis == EAxisList::Up)
 	{
 		return FLinearColor(0.0251f, 0.207f, 0.85f);
 	}
diff --git a/Engine/Source/Runtime/Core/Public/Math/Axis.h b/Engine/Source/Runtime/Core/Public/Math/Axis.h
--- a/Engine/Source/Runtime/Core/Public/Math/Axis.h
+++ b/Engine/Source/Runtime/Core/Public/Math/Axis.h
@@ -23,11 +23,11 @@ namespace EAxisList
 	enum Type
 	{
 		None		= 0,
-		X			= 1,
-		Y			= 2,
-		Z			= 4,
+		X			= 1 << 0,
+		Y			= 1 << 1,
+		Z			= 1 << 2,
 
-		Screen		= 8,
+		Screen		= 1 << 3,
 		XY			= X | Y,
 		XZ			= X | Z,
 		YZ			= Y | Z,
@@ -39,5 +39,10 @@ namespace EAxisList
 		
 		// alias over Screen since it isn't used when the 2d translate rotate widget is being used
 		Rotate2D	= Screen,
+
+		Left			= 1 << 4,
+		Up				= 1 << 5,
+		Forward			= 1 << 6,
+		LeftUpForward	= Left | Up | Forward,
 	};
 }
diff --git a/Engine/Source/Runtime/Core/Public/Misc/AxisDisplayInfo.h b/Engine/Source/Runtime/Core/Public/Misc/AxisDisplayInfo.h
--- a/Engine/Source/Runtime/Core/Public/Misc/AxisDisplayInfo.h
+++ b/Engine/Source/Runtime/Core/Public/Misc/AxisDisplayInfo.h
@@ -9,6 +9,16 @@
 
 namespace AxisDisplayInfo
 {
+	/**
+	 * Gets the axis display coordinate system to use.
+	 *
+	 * The return value may be driven by the config system and therefore
+	 * it should not be cached until the editor has finished loading.
+	 *
+	 * @return The display coordinate system, either EAxisList::XYZ or EAxisList::LeftUpForward
+	 */
+	CORE_API EAxisList::Type GetAxisDisplayCoordinateSystem();
+
 	/**
 	 * Gets the display name to use for the given axis.
 	 * 
@@ -62,6 +72,16 @@ public:
 		return Name;
 	}
 
+	/**
+	 * Gets the axis display coordinate system to use.
+	 *
+	 * The return value may be driven by the config system and therefore
+	 * it should not be cached until the editor has finished loading.
+	 *
+	 * @return The display coordinate system, either EAxisList::XYZ or EAxisList::LeftUpForward
+	 */
+	virtual EAxisList::Type GetAxisDisplayCoordinateSystem() = 0;
+
 	/**
 	 * Gets the display name to use for the given axis.
 	 *
