Cooker: -diffonly cooks: Fix copy-paste typo in FDiffWriterZenHeader::IsDependencyBundleIdentical.
#rnx

diff --git a/Engine/Source/Editor/UnrealEd/Private/Cooker/DiffWriterZenHeader.cpp b/Engine/Source/Editor/UnrealEd/Private/Cooker/DiffWriterZenHeader.cpp
--- a/Engine/Source/Editor/UnrealEd/Private/Cooker/DiffWriterZenHeader.cpp
+++ b/Engine/Source/Editor/UnrealEd/Private/Cooker/DiffWriterZenHeader.cpp
@@ -373,7 +373,7 @@ bool FDiffWriterZenHeader::IsDependencyBundleIdentical(FDiffWriterZenHeader& Des
 				{
 					return false;
 				}
-				if (SrcEntries[SrcIndex].LocalImportOrExportIndex != DestEntries[SrcIndex].LocalImportOrExportIndex)
+				if (SrcEntries[SrcIndex].LocalImportOrExportIndex != DestEntries[DestIndex].LocalImportOrExportIndex)
 				{
 					return false;
 				}
