[DataTable Import Options] Adding a reimport options popup to when the Curve Table reimport button is pressed so an interpolation mode can be selected

diff --git a/Engine/Source/Editor/UnrealEd/Private/Factories/CSVImportFactory.cpp b/Engine/Source/Editor/UnrealEd/Private/Factories/CSVImportFactory.cpp
--- a/Engine/Source/Editor/UnrealEd/Private/Factories/CSVImportFactory.cpp
+++ b/Engine/Source/Editor/UnrealEd/Private/Factories/CSVImportFactory.cpp
@@ -155,6 +155,7 @@ UObject* UCSVImportFactory::FactoryCreateText(UClass* InClass, UObject* InParent
 	FCSVImportSettings ImportSettings;
 
 	bool bHaveInfo = false;
+	bool bReimportingCurveTable = false;
 
 	if (IsAutomatedImport())
 	{
@@ -173,7 +174,7 @@ UObject* UCSVImportFactory::FactoryCreateText(UClass* InClass, UObject* InParent
 	else if (ExistingCurveTable != nullptr)
 	{
 		ImportSettings.ImportType = ECSVImportType::ECSV_CurveTable;
-		bHaveInfo = true;
+		bReimportingCurveTable = true;
 	}
 	else if (ExistingCurve != nullptr)
 	{
@@ -232,6 +233,7 @@ UObject* UCSVImportFactory::FactoryCreateText(UClass* InClass, UObject* InParent
 				.WidgetWindow(Window)
 				.FullPath(FText::FromString(ParentFullPath))
 				.TempImportDataTable(DataTableImportOptions)
+				.bReimportingCurveTable(bReimportingCurveTable)
 		);
 
 		FSlateApplication::Get().AddModalWindow(Window, ParentWindow, false);
diff --git a/Engine/Source/Editor/UnrealEd/Private/SCSVImportOptions.cpp b/Engine/Source/Editor/UnrealEd/Private/SCSVImportOptions.cpp
--- a/Engine/Source/Editor/UnrealEd/Private/SCSVImportOptions.cpp
+++ b/Engine/Source/Editor/UnrealEd/Private/SCSVImportOptions.cpp
@@ -38,11 +38,13 @@ void SCSVImportOptions::Construct(const FArguments& InArgs)
 {
 	WidgetWindow = InArgs._WidgetWindow;
 	TempImportDataTable = InArgs._TempImportDataTable;
+	bReimportingCurveTable = InArgs._bReimportingCurveTable;
 
 	// Make array of enum pointers
 	TSharedPtr<ECSVImportType> DataTableTypePtr = MakeShareable(new ECSVImportType(ECSVImportType::ECSV_DataTable));
 	ImportTypes.Add( DataTableTypePtr );
-	ImportTypes.Add(MakeShareable(new ECSVImportType(ECSVImportType::ECSV_CurveTable)));
+	TSharedPtr<ECSVImportType> CurveTableTypePtr = MakeShareable(new ECSVImportType(ECSVImportType::ECSV_CurveTable));
+	ImportTypes.Add(CurveTableTypePtr);
 	ImportTypes.Add(MakeShareable(new ECSVImportType(ECSVImportType::ECSV_CurveFloat)));
 	ImportTypes.Add(MakeShareable(new ECSVImportType(ECSVImportType::ECSV_CurveVector)));
 	ImportTypes.Add(MakeShareable(new ECSVImportType(ECSVImportType::ECSV_CurveLinearColor)));
@@ -119,6 +121,8 @@ void SCSVImportOptions::Construct(const FArguments& InArgs)
 			[
 				SAssignNew(ImportTypeCombo, SComboBox< TSharedPtr<ECSVImportType> >)
 				.OptionsSource( &ImportTypes )
+				// If we are reimporting a curve table, we do not want to let the user pick a different import type
+				.IsEnabled(!bReimportingCurveTable)
 				.OnGenerateWidget( this, &SCSVImportOptions::MakeImportTypeItemWidget )
 				.OnSelectionChanged( this, &SCSVImportOptions::OnImportTypeSelected)
 				[
@@ -209,7 +213,7 @@ void SCSVImportOptions::Construct(const FArguments& InArgs)
 	];
 
 	// set-up selection
-	ImportTypeCombo->SetSelectedItem(DataTableTypePtr);
+	ImportTypeCombo->SetSelectedItem(bReimportingCurveTable ? CurveTableTypePtr : DataTableTypePtr);
 	PropertyView->SetObject(TempImportDataTable.Get());
 
 	// Populate the valid interpolation modes
diff --git a/Engine/Source/Editor/UnrealEd/Public/SCSVImportOptions.h b/Engine/Source/Editor/UnrealEd/Public/SCSVImportOptions.h
--- a/Engine/Source/Editor/UnrealEd/Public/SCSVImportOptions.h
+++ b/Engine/Source/Editor/UnrealEd/Public/SCSVImportOptions.h
@@ -56,11 +56,13 @@ public:
 		: _WidgetWindow()
 		, _FullPath()
 		, _TempImportDataTable()
+		, _bReimportingCurveTable()
 		{}
 
 		SLATE_ARGUMENT(TSharedPtr<SWindow>, WidgetWindow)
 		SLATE_ARGUMENT(FText, FullPath)
 		SLATE_ARGUMENT(UDataTable*, TempImportDataTable)
+		SLATE_ARGUMENT(bool, bReimportingCurveTable)
 	SLATE_END_ARGS()
 
 	SCSVImportOptions()
@@ -135,6 +137,9 @@ private:
 	/** Window that owns us */
 	TWeakPtr< SWindow >							WidgetWindow;
 
+	/** Whether we are reimporting a curve table */
+	bool										bReimportingCurveTable;
+
 	// Import type
 
 	/** List of import types to pick from, drives combo box */
