Stop forcing preservation of root component before cpfuo when native class has not changed

[RN] The blueprint reinstancer is now less aggressive about disallowing RootComponent overwrites, use bp.UseLegacyRootComponentPreservation to return to the baseline behavior if the RootComponent is incorrect after blueprint compilation

diff --git a/Engine/Source/Editor/UnrealEd/Private/Kismet2/KismetReinstanceUtilities.cpp b/Engine/Source/Editor/UnrealEd/Private/Kismet2/KismetReinstanceUtilities.cpp
--- a/Engine/Source/Editor/UnrealEd/Private/Kismet2/KismetReinstanceUtilities.cpp
+++ b/Engine/Source/Editor/UnrealEd/Private/Kismet2/KismetReinstanceUtilities.cpp
@@ -67,6 +67,13 @@ static FAutoConsoleVariableRef CVarUseLegacyAnimInstanceReinstancingBehavior(
 	TEXT("Use the legacy re-instancing behavior for anim instances where the instance is destroyed and re-created.")
 );
 
+bool GUseLegacyRootComponentPreservation = false;
+static FAutoConsoleVariableRef CVarUseLegacyRootComponentPreservation(
+	TEXT("bp.UseLegacyRootComponentPreservation"),
+	GUseLegacyRootComponentPreservation,
+	TEXT("Use the legacy re-instancing behavior to preserve root components, ignoring data in previous instances")
+);
+
 namespace UE::ReinstanceUtils
 {
 	const EObjectFlags FlagMask = RF_Public | RF_ArchetypeObject | RF_Transactional | RF_Transient | RF_TextExportTransient | RF_InheritableComponentTemplate | RF_Standalone; //TODO: what about RF_RootSet?
@@ -3223,6 +3230,12 @@ void FBlueprintCompileReinstancer::CopyPropertiesForUnrelatedObjects(UObject* Ol
 		Params.bReplaceInternalReferenceUponRead = true;
 		Params.OptionalOldToNewClassMappings = OldToNewClassMap;
 	}
+	Params.bPreserveRootComponent = GUseLegacyRootComponentPreservation || 
+		// If the native parent class has changed, we may need to conform the root component
+		// this behavior may not be exactly correct but it's better than preserving 
+		// the root component that was likely set by our old native parent:
+		(FBlueprintEditorUtils::FindFirstNativeClass(OldObject->GetClass()) !=
+		FBlueprintEditorUtils::FindFirstNativeClass(NewObject->GetClass()));
 
 	UEngine::CopyPropertiesForUnrelatedObjects(OldObject, NewObject, Params);
 }
