Fixed crash due to change in signature and behavior of IMeshBuilderModule::BuildSkeletalMesh
- Solution: Allocate an uninitialized rendering resource then release it after the call to IMeshBuilderModule::BuildSkeletalMesh
Although the legacy importer does not support Nanite for skeletal mesh, the call succeeds even if Nanite is manually set to true
Hopefully, this solution will last until the legacy FBX importer is deprecated.

#rnx

diff --git a/Engine/Source/Editor/UnrealEd/Private/Fbx/FbxSkeletalMeshImport.cpp b/Engine/Source/Editor/UnrealEd/Private/Fbx/FbxSkeletalMeshImport.cpp
--- a/Engine/Source/Editor/UnrealEd/Private/Fbx/FbxSkeletalMeshImport.cpp
+++ b/Engine/Source/Editor/UnrealEd/Private/Fbx/FbxSkeletalMeshImport.cpp
@@ -1976,9 +1976,13 @@ USkeletalMesh* UnFbx::FFbxImporter::ImportSkeletalMesh(FImportSkeletalMeshArgs &
 			//We must build the LODModel so we can restore properly the mesh, but we do not have to regenerate LODs
 			const bool bRegenDepLODs = false;
 			FSkeletalMeshBuildParameters SkeletalMeshBuildParameters(SkeletalMesh, GetTargetPlatformManagerRef().GetRunningTargetPlatform(), ImportLODModelIndex, bRegenDepLODs);
-			FSkeletalMeshRenderData* RenderData = SkeletalMesh->GetResourceForRendering();
-			check(RenderData != nullptr);
-			bBuildSuccess = MeshBuilderModule.BuildSkeletalMesh(*RenderData, SkeletalMeshBuildParameters);
+			// JIRA UE-250536: Change of API to IMeshBuilderModule::BuildSkeletalMesh now requires a rendering resource
+			// Solution: Allocate an uninitialized rendering resource then release it after the call to IMeshBuilderModule::BuildSkeletalMesh
+			//			 Although the legacy importer does not support Nanite for skeletal mesh, the call succeeds even if Nanite is manually set to true
+			//			 Hopefully, this solution will last until the legacy FBX importer is deprecated.
+			SkeletalMesh->AllocateResourceForRendering();
+			bBuildSuccess = MeshBuilderModule.BuildSkeletalMesh(*SkeletalMesh->GetResourceForRendering(), SkeletalMeshBuildParameters);
+			SkeletalMesh->ReleaseResources();
 		}
 
 		if( !bBuildSuccess )
