Update ToolMenus.Edit console command description.
#rnx

diff --git a/Engine/Source/Editor/UnrealEd/Private/EditorEngine.cpp b/Engine/Source/Editor/UnrealEd/Private/EditorEngine.cpp
--- a/Engine/Source/Editor/UnrealEd/Private/EditorEngine.cpp
+++ b/Engine/Source/Editor/UnrealEd/Private/EditorEngine.cpp
@@ -1122,7 +1122,7 @@ void UEditorEngine::InitEditor(IEngineLoop* InEngineLoop)
 
 		static FAutoConsoleCommand ToolMenusEditMenusModeCVar = FAutoConsoleCommand(
 			TEXT("ToolMenus.Edit"),
-			TEXT("Experimental: Enable edit menus mode toggle in level editor's windows menu"),
+			TEXT("Show edit menu buttons"),
 			FConsoleCommandWithArgsDelegate::CreateLambda([](const TArray<FString>& Args)
 		{
 			IToolMenusEditorModule::Get().RegisterShowEditMenusModeCheckbox();
