StructUtils: Support instanced struct when copying an actor.

diff --git a/Engine/Source/Editor/UnrealEd/Private/Editor.cpp b/Engine/Source/Editor/UnrealEd/Private/Editor.cpp
--- a/Engine/Source/Editor/UnrealEd/Private/Editor.cpp
+++ b/Engine/Source/Editor/UnrealEd/Private/Editor.cpp
@@ -1351,7 +1351,7 @@ namespace EditorUtilities
 								else
 								{
 									TargetComponent = nullptr;
-								}								
+								}
 							}
 						}
 					}
@@ -1363,7 +1363,7 @@ namespace EditorUtilities
 	}
 
 
-	UActorComponent* FindMatchingComponentInstance( UActorComponent* SourceComponent, AActor* TargetActor )
+	UActorComponent* FindMatchingComponentInstance(UActorComponent* SourceComponent, AActor* TargetActor)
 	{
 		UActorComponent* MatchingComponent = NULL;
 		int32 StartIndex = 0;
@@ -1372,119 +1372,155 @@ namespace EditorUtilities
 		{
 			TInlineComponentArray<UActorComponent*> TargetComponents;
 			TargetActor->GetComponents(TargetComponents);
-			MatchingComponent = FindMatchingComponentInstance( SourceComponent, TargetActor, TargetComponents, StartIndex );
+			MatchingComponent = FindMatchingComponentInstance(SourceComponent, TargetActor, TargetComponents, StartIndex);
 		}
 
 		return MatchingComponent;
 	}
-
-
-	void CopySinglePropertyRecursive(const void* const InSourcePtr, void* const InTargetPtr, UObject* const InTargetObject, FProperty* const InProperty)
+	
+	void CopySingleProperty(const UObject* const InSourceObject, UObject* const InTargetObject, FProperty* const InProperty)
 	{
-		// Properties that are *object* properties are tricky
-		// Sometimes the object will be a reference to a PIE-world object, and copying that reference back to an actor CDO asset is not a good idea
-		// If the property is referencing an actor or actor component in the PIE world, then we can try and fix that reference up to the equivalent
-		// from the editor world; otherwise we have to skip it
-		bool bNeedsGenericCopy = true;
-		if( FObjectPropertyBase* const ObjectProperty = CastField<FObjectPropertyBase>(InProperty) )
-		{
-			const int32 PropertyArrayDim = InProperty->ArrayDim;
-			for (int32 ArrayIndex = 0; ArrayIndex < PropertyArrayDim; ArrayIndex++)
+		// Copy the property. Go through all child properties to find UObject.
+		//If the UObject is a PIE object, then replace it with a valid object.
+		//If the object is replaced, then rehash the set/map that contained that object.
+
+		TArray<FPropertyVisitorPath> PropertiesToRehash;
+		auto VisitPred = [InTargetObject, &PropertiesToRehash](const FPropertyVisitorPath& Path, const FPropertyVisitorData& Data) -> EPropertyVisitorControlFlow
 			{
-				UObject* const SourceObjectPropertyValue = ObjectProperty->GetObjectPropertyValue_InContainer(InSourcePtr, ArrayIndex);
-				if (SourceObjectPropertyValue && SourceObjectPropertyValue->GetOutermost()->HasAnyPackageFlags(PKG_PlayInEditor))
+				const FProperty* Property = Path.Top().Property;
+				const void* SourcePtr = Data.PropertyData;
+				void* TargetPtr = Path.GetPropertyDataPtr(InTargetObject);
+				check(TargetPtr);
+
+				EPropertyVisitorControlFlow Result = EPropertyVisitorControlFlow::StepInto;
+				bool bCopyPropertyValue = Path.Num() == 1;
+				// If the property is referencing an actor or actor component in the PIE world, then we can try and fix that reference up to the equivalent
+				// from the editor world; otherwise we have to skip it
+				if (const FObjectPropertyBase* const ObjectProperty = CastField<FObjectPropertyBase>(Property))
 				{
-					// Not all the code paths below actually copy the object, but even if they don't we need to claim that they
-					// did, as copying a reference to an object in a PIE world leads to crashes
-					bNeedsGenericCopy = false;
+					bool bObjectValueChanged = false;
 
-					// REFERENCE an existing actor in the editor world from a REFERENCE in the PIE world
-					if (SourceObjectPropertyValue->IsA(AActor::StaticClass()))
+					UObject* const SourceObjectPropertyValue = ObjectProperty->GetObjectPropertyValue(SourcePtr);
+					if (SourceObjectPropertyValue && SourceObjectPropertyValue->GetOutermost()->HasAnyPackageFlags(PKG_PlayInEditor))
 					{
-						// We can try and fix-up an actor reference from the PIE world to instead be the version from the persistent world
-						AActor* const EditorWorldActor = GetEditorWorldCounterpartActor(Cast<AActor>(SourceObjectPropertyValue));
-						if (EditorWorldActor)
+						// Don't copy any PIE object.
+
+						// REFERENCE an existing actor in the editor world from a REFERENCE in the PIE world
+						if (SourceObjectPropertyValue->IsA(AActor::StaticClass()))
 						{
-							ObjectProperty->SetObjectPropertyValue_InContainer(InTargetPtr, EditorWorldActor, ArrayIndex);
+							// We can try and fix-up an actor reference from the PIE world to instead be the version from the persistent world
+							AActor* const EditorWorldActor = GetEditorWorldCounterpartActor(Cast<AActor>(SourceObjectPropertyValue));
+							if (EditorWorldActor)
+							{
+								ObjectProperty->SetObjectPropertyValue(TargetPtr, EditorWorldActor);
+								bObjectValueChanged = true;
+							}
+						}
+						// REFERENCE an existing actor component in the editor world from a REFERENCE in the PIE world
+						else if (SourceObjectPropertyValue->IsA(UActorComponent::StaticClass()) && InTargetObject->IsA(AActor::StaticClass()))
+						{
+							AActor* const TargetActor = Cast<AActor>(InTargetObject);
+							TInlineComponentArray<UActorComponent*> TargetComponents;
+							TargetActor->GetComponents(TargetComponents);
+
+							// We can try and fix-up an actor component reference from the PIE world to instead be the version from the persistent world
+							int32 TargetComponentIndex = 0;
+							UActorComponent* const EditorWorldComponent = FindMatchingComponentInstance(Cast<UActorComponent>(SourceObjectPropertyValue), TargetActor, TargetComponents, TargetComponentIndex);
+							if (EditorWorldComponent)
+							{
+								ObjectProperty->SetObjectPropertyValue(TargetPtr, EditorWorldComponent);
+								bObjectValueChanged = true;
+							}
 						}
 					}
-					// REFERENCE an existing actor component in the editor world from a REFERENCE in the PIE world
-					else if (SourceObjectPropertyValue->IsA(UActorComponent::StaticClass()) && InTargetObject->IsA(AActor::StaticClass()))
+
+					// TSet and TMap needs to rehash if modified.
+					// TObjectPtr should not be used for keys (use FObjectKey instead).
+					// TSoftObjectPtr are valid keys and could be modified above.
+					if (bObjectValueChanged)
 					{
-						AActor* const TargetActor = Cast<AActor>(InTargetObject);
-						TInlineComponentArray<UActorComponent*> TargetComponents;
-						TargetActor->GetComponents(TargetComponents);
-
-						// We can try and fix-up an actor component reference from the PIE world to instead be the version from the persistent world
-						int32 TargetComponentIndex = 0;
-						UActorComponent* const EditorWorldComponent = FindMatchingComponentInstance(Cast<UActorComponent>(SourceObjectPropertyValue), TargetActor, TargetComponents, TargetComponentIndex);
-						if (EditorWorldComponent)
+						const TArray<FPropertyVisitorInfo>& Infos = Path.GetPath();
+						for (int32 Index = Infos.Num() - 1; Index >= 0; --Index)
 						{
-							ObjectProperty->SetObjectPropertyValue_InContainer(InTargetPtr, EditorWorldComponent, ArrayIndex);
+							if (CastField<FSetProperty>(Infos[Index].Property)
+								|| CastField<FMapProperty>(Infos[Index].Property))
+							{
+								PropertiesToRehash.AddUnique(FPropertyVisitorPath(MakeArrayView(Infos.GetData(), Index+1)));
+							}
 						}
 					}
-				}
-			}
-		}
-		else if (FStructProperty* const StructProperty = CastField<FStructProperty>(InProperty))
-		{
-			// Ensure that the target struct is initialized before copying fields from the source.
-			StructProperty->InitializeValue_InContainer(InTargetPtr);
 
-			const int32 PropertyArrayDim = InProperty->ArrayDim;
-			for (int32 ArrayIndex = 0; ArrayIndex < PropertyArrayDim; ArrayIndex++)
-			{
-				const void* const SourcePtr = StructProperty->ContainerPtrToValuePtr<void>(InSourcePtr, ArrayIndex);
-				void* const TargetPtr = StructProperty->ContainerPtrToValuePtr<void>(InTargetPtr, ArrayIndex);
+					Result = EPropertyVisitorControlFlow::StepOver;
+				}
+				
+				if (bCopyPropertyValue)
+				{
+					// For all the other property types. Only copy if first level.
+					Property->CopyCompleteValue(TargetPtr, SourcePtr);
+				}
 
-				for (TFieldIterator<FProperty> It(StructProperty->Struct); It; ++It)
+				// As an optimization skip over container's of int/double/... only step into property that can contains UObject.
+				auto CanPropertyContainsObject = [](const FProperty* Property)
+					{
+						return CastField<FStructProperty>(Property)
+							|| CastField<FObjectPropertyBase>(Property)
+							|| CastField<FArrayProperty>(Property)
+							|| CastField<FMapProperty>(Property)
+							|| CastField<FSetProperty>(Property);
+					};
+				if (const FSetProperty* const SetProperty = CastField<FSetProperty>(Property))
 				{
-					FProperty* const InnerProperty = *It;
-					CopySinglePropertyRecursive(SourcePtr, TargetPtr, InTargetObject, InnerProperty);
+					if (!CanPropertyContainsObject(SetProperty->ElementProp))
+					{
+						Result = EPropertyVisitorControlFlow::StepOver;
+					}
+				}
+				else if (const FMapProperty* const MapProperty = CastField<FMapProperty>(Property))
+				{
+					if (!CanPropertyContainsObject(MapProperty->KeyProp)
+						&& !CanPropertyContainsObject(MapProperty->ValueProp))
+					{
+						Result = EPropertyVisitorControlFlow::StepOver;
+					}
+				}
+				else if (const FArrayProperty* const ArrayProperty = CastField<FArrayProperty>(Property))
+				{
+					if (!CanPropertyContainsObject(ArrayProperty->Inner))
+					{
+						Result = EPropertyVisitorControlFlow::StepOver;
+					}
 				}
-			}
 
-			bNeedsGenericCopy = false;
-		}
-		else if (FArrayProperty* const ArrayProperty = CastField<FArrayProperty>(InProperty))
-		{
-			check(InProperty->ArrayDim == 1);
-			FScriptArrayHelper SourceArrayHelper(ArrayProperty, ArrayProperty->ContainerPtrToValuePtr<void>(InSourcePtr));
-			FScriptArrayHelper TargetArrayHelper(ArrayProperty, ArrayProperty->ContainerPtrToValuePtr<void>(InTargetPtr));
+				return Result;
+			};
 
-			FProperty* InnerProperty = ArrayProperty->Inner;
-			int32 Num = SourceArrayHelper.Num();
+		UStruct* PropertyOwner = InProperty->GetOwnerStruct();
+		check(PropertyOwner);
+		FPropertyVisitorPath Path;
+		void* SourceObjectPtr = const_cast<UObject*>(InSourceObject);
+		FPropertyVisitorData Data(SourceObjectPtr, /*ParentStructData*/nullptr);
+		PropertyVisitorHelpers::VisitProperty(PropertyOwner, InProperty, Path, Data, VisitPred);
 
-			// here we emulate FArrayProperty::CopyValuesInternal()
-			if (!(InnerProperty->PropertyFlags & CPF_IsPlainOldData))
+		// Rehash TSet and TMap in reverse order in case one contains the other.
+		for (int32 Index = PropertiesToRehash.Num() - 1; Index >= 0; --Index)
+		{
+			const FPropertyVisitorPath& PathToRehash = PropertiesToRehash[Index];
+			const FProperty* Property = PathToRehash.Top().Property;
+			void* TargetPtr = PathToRehash.GetPropertyDataPtr(InTargetObject);
+			if (const FSetProperty* const SetProperty = CastField<FSetProperty>(Property))
 			{
-				TargetArrayHelper.EmptyAndAddValues(Num);
+				FScriptSetHelper SetHelper(SetProperty, TargetPtr);
+				SetHelper.Rehash();
 			}
 			else
 			{
-				TargetArrayHelper.EmptyAndAddUninitializedValues(Num);
-			}
-			
-			for (int32 Index = 0; Index < Num; Index++)
-			{
-				CopySinglePropertyRecursive(SourceArrayHelper.GetRawPtr(Index), TargetArrayHelper.GetRawPtr(Index), InTargetObject, InnerProperty);
+				const FMapProperty* const MapProperty = CastFieldChecked<FMapProperty>(Property);
+				FScriptMapHelper MapHelper(MapProperty, TargetPtr);
+				MapHelper.Rehash();
 			}
-
-			bNeedsGenericCopy = false;
-		}
-		
-		// Handle copying properties that either aren't an object, or aren't part of the PIE world
-		if( bNeedsGenericCopy )
-		{
-			InProperty->CopyCompleteValue_InContainer(InTargetPtr, InSourcePtr);
 		}
 	}
-
-	void CopySingleProperty(const UObject* const InSourceObject, UObject* const InTargetObject, FProperty* const InProperty)
-	{
-		CopySinglePropertyRecursive(InSourceObject, InTargetObject, InTargetObject, InProperty);
-	}
-
+	
 	int32 CopyActorProperties( AActor* SourceActor, AActor* TargetActor, const FCopyOptions& Options )
 	{
 		check( SourceActor != nullptr && TargetActor != nullptr );
diff --git a/Engine/Source/Runtime/CoreUObject/Private/UObject/Class.cpp b/Engine/Source/Runtime/CoreUObject/Private/UObject/Class.cpp
--- a/Engine/Source/Runtime/CoreUObject/Private/UObject/Class.cpp
+++ b/Engine/Source/Runtime/CoreUObject/Private/UObject/Class.cpp
@@ -805,44 +805,14 @@ EPropertyVisitorControlFlow UStruct::Visit(FPropertyVisitorPath& Path, const FPr
 	EPropertyVisitorControlFlow RetVal = EPropertyVisitorControlFlow::StepOver;
 	for (const FProperty* Property = PropertyLink; Property; Property = Property->PropertyLinkNext)
 	{
-		FPropertyVisitorScope Scope(Path, FPropertyVisitorInfo(Property));
-		Path.Top().ParentStructType = this;
-
-		// ArrayDim of one means it is just a single property, not a static array
-		if(Property->ArrayDim == 1)
+		RetVal = PropertyVisitorHelpers::VisitProperty(this, Property, Path, InData, InFunc);
+		if (RetVal == EPropertyVisitorControlFlow::Stop)
 		{
-			FPropertyVisitorData Data(Property->ContainerPtrToValuePtr<void>(InData.PropertyData), /*ParentStructData*/InData.PropertyData);
-
-			RetVal = Property->Visit(Path, Data, InFunc);
-			if (RetVal == EPropertyVisitorControlFlow::Stop)
-			{
-				return EPropertyVisitorControlFlow::Stop;
-			}
-			if (RetVal == EPropertyVisitorControlFlow::StepOut)
-			{
-				return EPropertyVisitorControlFlow::StepOver;
-			}
+			return EPropertyVisitorControlFlow::Stop;
 		}
-		else
+		if (RetVal == EPropertyVisitorControlFlow::StepOut)
 		{
-			// In the case of a static array, we need to setup the EPropertyVisitorInfoType
-			for (int32 StaticArrayIndex = 0; StaticArrayIndex < Property->ArrayDim; ++StaticArrayIndex)
-			{
-				Path.Top().SetIndex(StaticArrayIndex, EPropertyVisitorInfoType::StaticArrayIndex);
-
-				FPropertyVisitorData Data(Property->ContainerPtrToValuePtr<void>(InData.PropertyData, StaticArrayIndex), /*ParentStructData*/InData.PropertyData);
-
-				RetVal = Property->Visit(Path, Data, InFunc);
-
-				if (RetVal == EPropertyVisitorControlFlow::Stop)
-				{
-					return EPropertyVisitorControlFlow::Stop;
-				}
-				if (RetVal == EPropertyVisitorControlFlow::StepOut)
-				{
-					return EPropertyVisitorControlFlow::StepOver;
-				}
-			}
+			return EPropertyVisitorControlFlow::StepOver;
 		}
 	}
 	return RetVal;
diff --git a/Engine/Source/Runtime/CoreUObject/Private/UObject/PropertyVisitor.cpp b/Engine/Source/Runtime/CoreUObject/Private/UObject/PropertyVisitor.cpp
--- a/Engine/Source/Runtime/CoreUObject/Private/UObject/PropertyVisitor.cpp
+++ b/Engine/Source/Runtime/CoreUObject/Private/UObject/PropertyVisitor.cpp
@@ -98,6 +98,41 @@ FArchiveSerializedPropertyChain FPropertyVisitorPath::ToSerializedPropertyChain(
 	return PropertyVisitorHelpers::PathToSerializedPropertyChain(Path);
 }
 
+EPropertyVisitorControlFlow PropertyVisitorHelpers::VisitProperty(const UStruct* PropertyOwner, const FProperty* Property, FPropertyVisitorPath& Path, const FPropertyVisitorData& InData, const TFunctionRef<EPropertyVisitorControlFlow(const FPropertyVisitorPath& /*Path*/, const FPropertyVisitorData& /*Data*/)> InFunc)
+{
+	EPropertyVisitorControlFlow RetVal = EPropertyVisitorControlFlow::StepOver;
+
+	FPropertyVisitorScope Scope(Path, FPropertyVisitorInfo(Property));
+	Path.Top().ParentStructType = PropertyOwner;
+
+	// ArrayDim of one means it is just a single property, not a static array
+	if (Property->ArrayDim == 1)
+	{
+		FPropertyVisitorData Data(Property->ContainerPtrToValuePtr<void>(InData.PropertyData), /*ParentStructData*/InData.PropertyData);
+
+		RetVal = Property->Visit(Path, Data, InFunc);
+	}
+	else
+	{
+		// In the case of a static array, we need to setup the EPropertyVisitorInfoType
+		for (int32 StaticArrayIndex = 0; StaticArrayIndex < Property->ArrayDim; ++StaticArrayIndex)
+		{
+			Path.Top().SetIndex(StaticArrayIndex, EPropertyVisitorInfoType::StaticArrayIndex);
+
+			FPropertyVisitorData Data(Property->ContainerPtrToValuePtr<void>(InData.PropertyData, StaticArrayIndex), /*ParentStructData*/InData.PropertyData);
+
+			RetVal = Property->Visit(Path, Data, InFunc);
+
+			if (RetVal == EPropertyVisitorControlFlow::Stop
+				|| RetVal == EPropertyVisitorControlFlow::StepOut)
+			{
+				break;
+			}
+		}
+	}
+	return RetVal;
+}
+
 FString PropertyVisitorHelpers::PathToString(TArrayView<const FPropertyVisitorInfo> Path, const TCHAR* Separator)
 {
 	TStringBuilder<FName::StringBufferSize> PropertyPath;
diff --git a/Engine/Source/Runtime/CoreUObject/Public/UObject/PropertyVisitor.h b/Engine/Source/Runtime/CoreUObject/Public/UObject/PropertyVisitor.h
--- a/Engine/Source/Runtime/CoreUObject/Public/UObject/PropertyVisitor.h
+++ b/Engine/Source/Runtime/CoreUObject/Public/UObject/PropertyVisitor.h
@@ -4,6 +4,7 @@
 
 #include "Containers/Array.h"
 #include "Containers/ArrayView.h"
+#include "Templates/FunctionFwd.h"
 
 class FProperty;
 class FEditPropertyChain;
@@ -313,6 +314,9 @@ void* ResolveVisitedPathInfo_Generic(Type* This, void* Data, const FPropertyVisi
 	return ResolveVisitedPathInfo_Generic(This, Path, Data, Info);
 }
 
+/** Visit the property from an instance. */
+COREUOBJECT_API EPropertyVisitorControlFlow VisitProperty(const UStruct* PropertyOwner, const FProperty* Property, FPropertyVisitorPath& Path, const FPropertyVisitorData& InData, const TFunctionRef<EPropertyVisitorControlFlow(const FPropertyVisitorPath& /*Path*/, const FPropertyVisitorData& /*Data*/)> InFunc);
+
 /** Convert the given path to a string */
 COREUOBJECT_API FString PathToString(TArrayView<const FPropertyVisitorInfo> Path, const TCHAR* Separator = TEXT("."));
 COREUOBJECT_API void PathToString(TArrayView<const FPropertyVisitorInfo> Path, FStringBuilderBase& Out, const TCHAR* Separator = TEXT("."));
