Apply feedback from review CL 40314606
#rnx

diff --git a/Engine/Source/Editor/UnrealEd/Private/Editor.cpp b/Engine/Source/Editor/UnrealEd/Private/Editor.cpp
--- a/Engine/Source/Editor/UnrealEd/Private/Editor.cpp
+++ b/Engine/Source/Editor/UnrealEd/Private/Editor.cpp
@@ -20,6 +20,8 @@
 #include "Factories/PhysicalMaterialMaskFactory.h"
 #include "Engine/BlueprintGeneratedClass.h"
 #include "Misc/ScopedSlowTask.h"
+#include "UObject/EnumProperty.h"
+#include "UObject/TextProperty.h"
 #include "UObject/UObjectIterator.h"
 #include "EngineUtils.h"
 #include "Misc/MessageDialog.h"
@@ -1439,6 +1441,7 @@ namespace EditorUtilities
 					// TSoftObjectPtr are valid keys and could be modified above.
 					if (bObjectValueChanged)
 					{
+						bCopyPropertyValue = false;
 						const TArray<FPropertyVisitorInfo>& Infos = Path.GetPath();
 						for (int32 Index = Infos.Num() - 1; Index >= 0; --Index)
 						{
@@ -1462,11 +1465,15 @@ namespace EditorUtilities
 				// As an optimization skip over container's of int/double/... only step into property that can contains UObject.
 				auto CanPropertyContainsObject = [](const FProperty* Property)
 					{
-						return CastField<FStructProperty>(Property)
-							|| CastField<FObjectPropertyBase>(Property)
-							|| CastField<FArrayProperty>(Property)
-							|| CastField<FMapProperty>(Property)
-							|| CastField<FSetProperty>(Property);
+						const bool bPropertyCannotContainsObject = CastField<FNumericProperty>(Property)
+							|| CastField<FBoolProperty>(Property)
+							|| CastField<FEnumProperty>(Property)
+							|| CastField<FNameProperty>(Property)
+							|| CastField<FStrProperty>(Property)
+							|| CastField<FTextProperty>(Property)
+							|| CastField<FDelegateProperty>(Property)
+							|| CastField<FMulticastDelegateProperty>(Property);
+						return !bPropertyCannotContainsObject;
 					};
 				if (const FSetProperty* const SetProperty = CastField<FSetProperty>(Property))
 				{
