Revert to old root component preservation behavior as it is being depended upon for asset duplication

#rnx

diff --git a/Engine/Source/Editor/UnrealEd/Private/Kismet2/KismetReinstanceUtilities.cpp b/Engine/Source/Editor/UnrealEd/Private/Kismet2/KismetReinstanceUtilities.cpp
--- a/Engine/Source/Editor/UnrealEd/Private/Kismet2/KismetReinstanceUtilities.cpp
+++ b/Engine/Source/Editor/UnrealEd/Private/Kismet2/KismetReinstanceUtilities.cpp
@@ -67,7 +67,7 @@ static FAutoConsoleVariableRef CVarUseLegacyAnimInstanceReinstancingBehavior(
 	TEXT("Use the legacy re-instancing behavior for anim instances where the instance is destroyed and re-created.")
 );
 
-bool GUseLegacyRootComponentPreservation = false;
+bool GUseLegacyRootComponentPreservation = true;
 static FAutoConsoleVariableRef CVarUseLegacyRootComponentPreservation(
 	TEXT("bp.UseLegacyRootComponentPreservation"),
 	GUseLegacyRootComponentPreservation,
