IncrementalCook: Suppress warning "Package %s is a development-only package but was not marked with DiskSize explaining why..." that can occur during incrementalcook. This is a workaround until UE-252126 is fixed.
#rnx
#incrementalcook

diff --git a/Engine/Source/Editor/UnrealEd/Private/Commandlets/AssetRegistryGenerator.cpp b/Engine/Source/Editor/UnrealEd/Private/Commandlets/AssetRegistryGenerator.cpp
--- a/Engine/Source/Editor/UnrealEd/Private/Commandlets/AssetRegistryGenerator.cpp
+++ b/Engine/Source/Editor/UnrealEd/Private/Commandlets/AssetRegistryGenerator.cpp
@@ -1716,7 +1716,9 @@ void FAssetRegistryGenerator::FinalizeChunkIDs(const TSet<FName>& InCookedPackag
 		FAssetPackageData* PackageData = State.GetAssetPackageData(DevelopmentOnlyPackage);
 		if (PackageData && PackageData->DiskSize >= 0)
 		{
-			UE_LOG(LogAssetRegistryGenerator, Warning,
+			// TODO UE-252126: Raise this logseverity to Warning once UE-252126 is fixed;
+			// UE-252126 causes some low-priority occurrences of this warning.
+			UE_LOG(LogAssetRegistryGenerator, Display,
 				TEXT("Package %s is a development-only package but was not marked with DiskSize explaining why it was not cooked. Marking it as a generic skip."),
 				*DevelopmentOnlyPackage.ToString());
 			PackageData->DiskSize = -1;
