Flip back to allowing root component to be overridden by a blueprint, but  don't copy root components if they are invalid (e.g. an old actor has a root component owned by some other actor, this happens when an asset is duplicated)

#rnx

diff --git a/Engine/Source/Editor/UnrealEd/Private/Kismet2/KismetReinstanceUtilities.cpp b/Engine/Source/Editor/UnrealEd/Private/Kismet2/KismetReinstanceUtilities.cpp
--- a/Engine/Source/Editor/UnrealEd/Private/Kismet2/KismetReinstanceUtilities.cpp
+++ b/Engine/Source/Editor/UnrealEd/Private/Kismet2/KismetReinstanceUtilities.cpp
@@ -67,7 +67,7 @@ static FAutoConsoleVariableRef CVarUseLegacyAnimInstanceReinstancingBehavior(
 	TEXT("Use the legacy re-instancing behavior for anim instances where the instance is destroyed and re-created.")
 );
 
-bool GUseLegacyRootComponentPreservation = true;
+bool GUseLegacyRootComponentPreservation = false;
 static FAutoConsoleVariableRef CVarUseLegacyRootComponentPreservation(
 	TEXT("bp.UseLegacyRootComponentPreservation"),
 	GUseLegacyRootComponentPreservation,
@@ -3236,6 +3236,21 @@ void FBlueprintCompileReinstancer::CopyPropertiesForUnrelatedObjects(UObject* Ol
 		// the root component that was likely set by our old native parent:
 		(FBlueprintEditorUtils::FindFirstNativeClass(OldObject->GetClass()) !=
 		FBlueprintEditorUtils::FindFirstNativeClass(NewObject->GetClass()));
+	if(!Params.bPreserveRootComponent)
+	{
+		// If the old actor has a bad root component then don't copy it
+		// and just preserve the one created by the native constructor:
+		if(const AActor* OldActor = Cast<AActor>(OldObject))
+		{
+			if(const USceneComponent* OldComponent = OldActor->GetRootComponent())
+			{
+				// this flag's name is somewhat confusing in this context, setting it to true
+				// will preserve the root component created by the NewActor's constructor,
+				// rather than the one on the old actor:
+				Params.bPreserveRootComponent = (OldComponent->GetOwner() != OldActor);
+			}
+		}
+	}
 
 	UEngine::CopyPropertiesForUnrelatedObjects(OldObject, NewObject, Params);
 }
