Keep Simulation Changes not keeping all the changes

After some investigation, the root cause of the issue is that `AActor::PostEditChangeProperty` will destroy and rebuild the ActorComponent array when simulating.
The code in `EditorUtilities::CopyActorProperties` (used by [REDACTED_PERSON]) generates an array of components and iterate over it to apply the changed properties. After each property is changed though, `PostEditChangeProperty` is called which invalidates and recreates all the components. So when we iterate over the next component from our array, we are actually iterating over an array containing invalid components that are not linked to the actor anymore.
This mean all the changes from Keep Simulation Changes are actually applied but only the first change will be applied to a valid component. The next changes will be applied to the destroyed component, explaining the issue described.

The current flow is:
```
FLevelEditorActionCallbacks::OnKeepSimulationChanges
=> EditorUtilities::CopyActorProperties
=> BuildComponentPairs
=> for (... ComponentPair : SourceTargetComponentPairs)
=> for( FProperty* Property ... )
=> AActor::PostEditChangeProperty
=> if (... GEditor->bIsSimulatingInEditor ...)
=> RerunConstructionScripts();
=> DestroyConstructedComponents(); => this will make the components saved in SourceTargetComponentPairs invalid, but we still continue iterating over them
```

The proposed solution is to add some validity tests at strategic locations and refetch the latest version of the component if the validity test fails.

diff --git a/Engine/Source/Editor/UnrealEd/Private/Editor.cpp b/Engine/Source/Editor/UnrealEd/Private/Editor.cpp
--- a/Engine/Source/Editor/UnrealEd/Private/Editor.cpp
+++ b/Engine/Source/Editor/UnrealEd/Private/Editor.cpp
@@ -1688,6 +1688,10 @@ namespace EditorUtilities
 		{
 			UActorComponent* SourceComponent = ComponentPair.Key;
 			UActorComponent* TargetComponent = ComponentPair.Value;
+			if (!IsValid(TargetComponent)) // The Target Component might not be valid if the component list has been regenerated (see AActor::PostEditChangeProperty)
+			{
+				TargetComponent = FindMatchingComponentInstance(SourceComponent, TargetActor);
+			}
 
 			UClass* ComponentClass = SourceComponent->GetClass();
 			check( ComponentClass == TargetComponent->GetClass() );
@@ -1810,6 +1814,12 @@ namespace EditorUtilities
 							{
 								FPropertyChangedEvent PropertyChangedEvent( Property );
 								TargetActor->PostEditChangeProperty( PropertyChangedEvent );
+								// When simulating, a call to AActor::PostEditChangeProperty will reset and reconstruct the Actor's component array
+								// which will make the TargetComponent invalid, so we need to update the variable to point to the newly created one
+								if (!IsValid(TargetComponent))
+								{
+									TargetComponent = FindMatchingComponentInstance(SourceComponent, TargetActor);
+								}
 							}
 
 							if( Options.Flags & ECopyOptions::PropagateChangesToArchetypeInstances )
