IncrementalCook: Add logging of lifetime events for CookGenerationHelper to diagnose state errors that can occur when incrementally cooking in MPCook:
"LogCook: Warning: RecordResults received for generated package %s, but its ParentGenerationHelper has already been destructed..."
#rnx
#incrementalcook

diff --git a/Engine/Source/Editor/UnrealEd/Private/Cooker/CookGenerationHelper.cpp b/Engine/Source/Editor/UnrealEd/Private/Cooker/CookGenerationHelper.cpp
--- a/Engine/Source/Editor/UnrealEd/Private/Cooker/CookGenerationHelper.cpp
+++ b/Engine/Source/Editor/UnrealEd/Private/Cooker/CookGenerationHelper.cpp
@@ -162,6 +162,9 @@ FGenerationHelper::FGenerationHelper(FPackageData& InOwner)
 
 FGenerationHelper::~FGenerationHelper()
 {
+	UE_LOG(LogCookGenerationHelper, Verbose, TEXT("%s: Destructed."),
+		*WriteToString<256>(GetOwner().GetPackageName()));
+
 	NotifyCompletion(ICookPackageSplitter::ETeardown::Complete);
 	GetOwner().OnGenerationHelperDestroyed(*this);
 }
@@ -224,6 +227,8 @@ void FGenerationHelper::Initialize(const UObject* InSplitDataObject,
 	check(InSplitDataObject);
 	if (InitializeStatus != EInitializeStatus::Uninitialized)
 	{
+		UE_LOG(LogCookGenerationHelper, Verbose, TEXT("%s: Redundant initialize ignored."),
+			*WriteToString<256>(GetOwner().GetPackageName()));
 		// If we already have a splitter, keep the old and throw out the new. The old one
 		// still contains some state.
 		return;
@@ -252,6 +257,8 @@ void FGenerationHelper::Initialize(const UObject* InSplitDataObject,
 	{
 		DoesGeneratedRequireGeneratorValue = ICookPackageSplitter::EGeneratedRequiresGenerator::Save;
 	}
+
+	UE_LOG(LogCookGenerationHelper, Verbose, TEXT("%s: Initialized."), *WriteToString<256>(GetOwner().GetPackageName()));
 }
 
 void FGenerationHelper::InitializeAsInvalid()
@@ -269,6 +276,7 @@ void FGenerationHelper::Uninitialize()
 	{
 		return;
 	}
+	UE_LOG(LogCookGenerationHelper, Verbose, TEXT("%s: Uninitialized."), *WriteToString<256>(GetOwner().GetPackageName()));
 
 	// Demote stalled packages; we will be garbage collecting so they no longer need to be preserved.
 	// And we need to demote them so that they drop their references to the generation helper and allow it to be
@@ -325,6 +333,9 @@ void FGenerationHelper::ModifyNumSaved(int32 Delta)
 	check(0 <= NumSaved && NumSaved <= NumAllSaved);
 	if (NumSaved == NumAllSaved)
 	{
+		UE_LOG(LogCookGenerationHelper, Verbose, TEXT("%s: All saves completed (%d/%d)."),
+			*WriteToString<256>(GetOwner().GetPackageName()), NumSaved, NumAllSaved);
+
 		UCookOnTheFlyServer& COTFS = GetOwner().GetPackageDatas().GetCookOnTheFlyServer();
 		// Only send OnAllSavesCompleted from director; clients have incomplete information and
 		// could send it spuriously.
@@ -405,11 +416,19 @@ void FGenerationHelper::DemoteStalledPackages(UCookOnTheFlyServer& COTFS, bool b
 					{
 						// Otherwise, when called on the CookDirector, demote them out of save but keep them in the assigned to worker state.
 						Info.PackageData->SendToState(EPackageState::AssignedToWorker, ESendFlags::QueueAddAndRemove, EStateChangeReason::GarbageCollected);
+						UE_LOG(LogCookGenerationHelper, Verbose,
+							TEXT("%s generated package %s: DemoteStalledPackage from SaveStalledAssignedToWorker to AssignedToWorker."),
+							*WriteToString<256>(GetOwner().GetPackageName()),
+							*WriteToString<256>(Info.PackageData->GetPackageName()));
 					}
 				}
 				else
 				{
 					// Demoting stalled packages on a Client just returns them to Idle.
+					UE_LOG(LogCookGenerationHelper, Verbose, TEXT("%s generated package %s: DemoteStalledPackage from %s to Idle."),
+						*WriteToString<256>(GetOwner().GetPackageName()),
+						*WriteToString<256>(Info.PackageData->GetPackageName()),
+						LexToString(Info.PackageData->GetState()));
 					COTFS.DemoteToIdle(*Info.PackageData, ESendFlags::QueueAddAndRemove, ESuppressCookReason::RetractedByCookDirector);
 				}
 			}
@@ -912,6 +931,7 @@ bool FGenerationHelper::TryGenerateList()
 	}
 
 	TMap<FPackageData*, int32> AlreadyExistingInfoPackageToIndex;
+	int32 NumAlreadyExisting = PackagesToGenerate.Num();
 	for (int32 ExistingIndex = 0; ExistingIndex < PackagesToGenerate.Num(); ++ExistingIndex)
 	{
 		AlreadyExistingInfoPackageToIndex.Add(PackagesToGenerate[ExistingIndex].PackageData, ExistingIndex);
@@ -1023,6 +1043,11 @@ bool FGenerationHelper::TryGenerateList()
 			PackagesToGenerate.RemoveAt(UnusedIndex);
 		}
 	}
+	UE_LOG(LogCookGenerationHelper, Verbose,
+		TEXT("%s: TryGenerateList: %d packages to generate, %d previously recorded, %d/%d already saved."),
+		*WriteToString<256>(GetOwner().GetPackageName()), PackagesToGenerate.Num(),
+		NumAlreadyExisting, NumSaved, PackagesToGenerate.Num() + 1);
+
 	ModifyNumSaved(0);
 
 	bGeneratedList = true;
@@ -1101,6 +1126,9 @@ bool FGenerationHelper::TryCallPopulateGeneratorPackage(
 		}
 	}
 
+	UE_LOG(LogCookGenerationHelper, Verbose, TEXT("%s: PopulateGeneratorPackage."),
+		*WriteToString<256>(GetOwner().GetPackageName()));
+
 	// Contract Point 1: We will not call populate again until the splitter has been destroyed
 	// Contract Point 2: We will not call populate again without garbage collecting the generator package
 	OwnerInfo.SetHasCalledPopulate(true);
@@ -1179,6 +1207,10 @@ bool FGenerationHelper::TryCallPopulateGeneratedPackage(UE::Cook::FCookGeneratio
 	Info.AddKeepReferencedPackages(*this, PopulateData.KeepReferencedPackages);
 	Info.BuildResultDependencies.Append(MoveTemp(PopulateData.BuildResultDependencies));
 
+	UE_LOG(LogCookGenerationHelper, Verbose, TEXT("%s generated package %s: PopulateGeneratedPackage."),
+		*WriteToString<256>(GetOwner().GetPackageName()),
+		*WriteToString<256>(Info.PackageData->GetPackageName()));
+
 	// Contract Point 1: We will not call populate again until the splitter has been destroyed
 	// Contract Point 2: We will not call populate again without garbage collecting the generator package
 	Info.SetHasCalledPopulate(true);
@@ -1205,6 +1237,9 @@ void FGenerationHelper::StartQueueGeneratedPackages(UCookOnTheFlyServer& COTFS)
 	{
 		return;
 	}
+	UE_LOG(LogCookGenerationHelper, Verbose, TEXT("%s: StartQueueGeneratedPackages."),
+		*WriteToString<256>(GetOwner().GetPackageName()));
+
 	NotifyStartQueueGeneratedPackages(COTFS, FWorkerId::Local());
 
 	bool bHybridIterativeEnabled = COTFS.bHybridIterativeEnabled;
@@ -1291,6 +1326,9 @@ bool FGenerationHelper::FDirectorAPI::HasStartedQueueGeneratedPackages() const
 
 void FGenerationHelper::EndQueueGeneratedPackages(UCookOnTheFlyServer& COTFS)
 {
+	UE_LOG(LogCookGenerationHelper, Verbose, TEXT("%s: EndQueueGeneratedPackages."),
+		*WriteToString<256>(GetOwner().GetPackageName()));
+
 	bHasFinishedQueueGeneratedPackages = true;
 	SetKeepForQueueResults();
 	COTFS.WorkerRequests->EndQueueGeneratedPackages(COTFS, *this);
@@ -1298,6 +1336,9 @@ void FGenerationHelper::EndQueueGeneratedPackages(UCookOnTheFlyServer& COTFS)
 
 void FGenerationHelper::EndQueueGeneratedPackagesOnDirector(UCookOnTheFlyServer& COTFS, FWorkerId SourceWorkerId)
 {
+	UE_LOG(LogCookGenerationHelper, Verbose, TEXT("%s: EndQueueGeneratedPackagesOnDirector."),
+		*WriteToString<256>(GetOwner().GetPackageName()));
+
 	// Note this function can be called on an uninitialized Generator; the generator is only needed
 	// on the director so it can serve as the passer of messages.
 	bHasFinishedQueueGeneratedPackages = true;
@@ -1334,6 +1375,8 @@ void FGenerationHelper::OnRequestFencePassed(UCookOnTheFlyServer& COTFS)
 	{
 		// PumpRequests has completed and we marked ourselves and all generated packages as iteratively skipped,
 		// so we no longer need the PreviouslyCookedData or this entire GenerationHelper
+		UE_LOG(LogCookGenerationHelper, Verbose, TEXT("%s: RequestFencePassed and found to be entirely iteratively skipped."),
+			*WriteToString<256>(GetOwner().GetPackageName()));
 		ClearKeepForIterative();
 		PreviousGeneratedPackages.Empty();
 	}
@@ -1357,10 +1400,18 @@ void FGenerationHelper::OnRequestFencePassed(UCookOnTheFlyServer& COTFS)
 		}
 		OnQueuedGeneratedPackagesFencePassed(COTFS);
 	}
+	else
+	{
+		UE_LOG(LogCookGenerationHelper, Verbose, TEXT("%s: RequestFencePassed, but not yet HasFinishedQueueGeneratedPackages."),
+			*WriteToString<256>(GetOwner().GetPackageName()));
+	}
 }
 
 void FGenerationHelper::OnQueuedGeneratedPackagesFencePassed(UCookOnTheFlyServer& COTFS)
 {
+	UE_LOG(LogCookGenerationHelper, Verbose, TEXT("%s: RequestFencePassed, and HasFinishedQueueGeneratedPackages."),
+		*WriteToString<256>(GetOwner().GetPackageName()));
+
 	ClearKeepForQueueResults();
 	// We no longer need PreviousGeneratedPackages or KeepForIterative, because they are used only in 
 	// StartQueueGeneratedPackages or the request cluster that they end up in in PumpRequests, both of which
@@ -1450,6 +1501,9 @@ void FGenerationHelper::FinishGeneratorPlatformSave(FPackageData& PackageData, b
 		FetchExternalActorDependencies();
 		COTFS.RecordExternalActorDependencies(GetExternalActorDependencies());
 	}
+
+	UE_LOG(LogCookGenerationHelper, Verbose, TEXT("%s: Generator package saved."),
+		*WriteToString<256>(GetOwner().GetPackageName()));
 }
 
 void FGenerationHelper::FinishGeneratedPlatformSave(FPackageData& PackageData, FAssetPackageData& OutAssetPackageData,
@@ -1509,6 +1563,10 @@ void FGenerationHelper::FinishGeneratedPlatformSave(FPackageData& PackageData, F
 	AssetPackageData.ImportedClasses = MoveTemp(ImportedClasses);
 
 	OutAssetPackageData = AssetPackageData;
+
+	UE_LOG(LogCookGenerationHelper, Verbose, TEXT("%s generated package %s: Generated package saved."),
+		*WriteToString<256>(GetOwner().GetPackageName()),
+		*WriteToString<256>(Info->PackageData->GetPackageName()));
 }
 
 const FAssetPackageData* FGenerationHelper::GetAssetPackageData(FName PackageName)
@@ -1746,6 +1804,9 @@ void FGenerationHelper::FetchExternalActorDependencies()
 
 void FGenerationHelper::SetPreviousGeneratedPackages(TMap<FName, FAssetPackageData>&& Packages)
 {
+	UE_LOG(LogCookGenerationHelper, Verbose,
+		TEXT("%s: SetPreviousGeneratedPackages with %d generated packages."),
+		*WriteToString<256>(GetOwner().GetPackageName()), Packages.Num());
 	SetKeepForIterative();
 	PreviousGeneratedPackages = MoveTemp(Packages);
 }
@@ -1977,21 +2038,44 @@ void FGenerationHelper::MarkPackageSavedRemotely(UCookOnTheFlyServer& COTFS, FPa
 	}
 }
 
-void FGenerationHelper::MarkPackageIterativelySkipped(FPackageData& PackageData)
+void FGenerationHelper::MarkPackageIterativelySkipped(FPackageData& PackageData, bool bIncrementallySkipped)
 {
 	FCookGenerationInfo* Info = FindInfoNoInitialize(PackageData);
-	if (Info)
+	if (!Info)
 	{
-		Info->SetHasSaved(*this, true, FWorkerId::Local());
-		Info->SetIterativelySkipped(true);
+		check(&PackageData != &GetOwner()); // FindInfoNoInitialize should have succeeded on the owner.
+		return;
 	}
 
-	if (&PackageData == &GetOwner())
+	bool bGenerator = &PackageData == &GetOwner();
+	if (bGenerator)
+	{
+		UE_LOG(LogCookGenerationHelper, Verbose,
+			TEXT("%s: %s"),
+			*WriteToString<256>(GetOwner().GetPackageName()),
+			bIncrementallySkipped ? TEXT("generator found to be entirely incrementally skippable.")
+			: TEXT("generator found to be NOT incrementally skippable. Generation will be reexecuted."));
+	}
+	else
+	{
+		UE_LOG(LogCookGenerationHelper, Verbose,
+			TEXT("%s generated package %s: %s"),
+			*WriteToString<256>(GetOwner().GetPackageName()),
+			*WriteToString<256>(PackageData.GetPackageName()),
+			bIncrementallySkipped ? TEXT("generated package found to be incrementally skippable.")
+			: TEXT("generated package found to be NOT incrementally skippable."));
+	}
+	if (bIncrementallySkipped)
 	{
-		// The entire generator package has been skipped. Wait for the current cluster to complete
-		// so we can mark all of our generated packages as skipped, but then clear the iterative data;
-		// it will no longer be needed.
-		GetOwner().GetPackageDatas().GetRequestQueue().AddRequestFenceListener(GetOwner().GetPackageName());
+		Info->SetHasSaved(*this, true, FWorkerId::Local());
+		Info->SetIterativelySkipped(true);
+		if (bGenerator)
+		{
+			// The entire generator package has been skipped. Wait for the current cluster to complete
+			// so we can mark all of our generated packages as skipped, but then clear the iterative data;
+			// it will no longer be needed.
+			GetOwner().GetPackageDatas().GetRequestQueue().AddRequestFenceListener(GetOwner().GetPackageName());
+		}
 	}
 }
 
diff --git a/Engine/Source/Editor/UnrealEd/Private/Cooker/CookGenerationHelper.h b/Engine/Source/Editor/UnrealEd/Private/Cooker/CookGenerationHelper.h
--- a/Engine/Source/Editor/UnrealEd/Private/Cooker/CookGenerationHelper.h
+++ b/Engine/Source/Editor/UnrealEd/Private/Cooker/CookGenerationHelper.h
@@ -424,7 +424,7 @@ public:
 	 * Called on the director when the generator or one of the generated packages was found to be iteratively skipped
 	 * and marked already cooked in an incremental cook. Does not call Initialize.
 	 */
-	void MarkPackageIterativelySkipped(FPackageData& PackageData);
+	void MarkPackageIterativelySkipped(FPackageData& PackageData, bool bIncrementallySkipped);
 	/**
 	 * Called on the director and every CookWorker when all saves have been completed; this indicates that
 	 * some of our contract points are complete and the splitter can be destroyed. Does not call Initialize.
@@ -742,21 +742,33 @@ inline const TMap<FName, FAssetPackageData>& FGenerationHelper::GetPreviousGener
 
 inline void FGenerationHelper::SetKeepForIterative()
 {
+	UE_LOG(LogCookGenerationHelper, Verbose, TEXT("%s: SetKeepForIterative."),
+		*WriteToString<256>(GetOwner().GetPackageName()));
+
 	ReferenceFromKeepForIterative = this;
 }
 
 inline void FGenerationHelper::ClearKeepForIterative()
 {
+	UE_LOG(LogCookGenerationHelper, Verbose, TEXT("%s: ClearKeepForIterative."),
+		*WriteToString<256>(GetOwner().GetPackageName()));
+
 	ReferenceFromKeepForIterative.SafeRelease();
 }
 
 inline void FGenerationHelper::SetKeepForQueueResults()
 {
+	UE_LOG(LogCookGenerationHelper, Verbose, TEXT("%s: SetKeepForQueueResults."),
+		*WriteToString<256>(GetOwner().GetPackageName()));
+
 	ReferenceFromKeepForQueueResults = this;
 }
 
 inline void FGenerationHelper::ClearKeepForQueueResults()
 {
+	UE_LOG(LogCookGenerationHelper, Verbose, TEXT("%s: ClearKeepForQueueResults."),
+		*WriteToString<256>(GetOwner().GetPackageName()));
+
 	ReferenceFromKeepForQueueResults.SafeRelease();
 }
 
@@ -767,22 +779,34 @@ inline bool FGenerationHelper::IsWaitingForQueueResults() const
 
 inline void FGenerationHelper::SetKeepForGeneratorSave()
 {
+	UE_LOG(LogCookGenerationHelper, Verbose, TEXT("%s: SetKeepForGeneratorSave."),
+		*WriteToString<256>(GetOwner().GetPackageName()));
+
 	ReferenceFromKeepForGeneratorSave = this;
 }
 
 inline void FGenerationHelper::ClearKeepForGeneratorSave()
 {
+	UE_LOG(LogCookGenerationHelper, Verbose, TEXT("%s: ClearKeepForGeneratorSave."),
+		*WriteToString<256>(GetOwner().GetPackageName()));
+
 	ReferenceFromKeepForGeneratorSave.SafeRelease();
 }
 
 inline void FGenerationHelper::SetKeepForAllSavedOrGC()
 {
+	UE_LOG(LogCookGenerationHelper, Verbose, TEXT("%s: SetKeepForAllSavedOrGC."),
+		*WriteToString<256>(GetOwner().GetPackageName()));
+
 	ReferenceFromKeepForAllSavedOrGC = this;
 	bKeepForAllSavedOrGC = true;
 }
 
 inline void FGenerationHelper::ClearKeepForAllSavedOrGC()
 {
+	UE_LOG(LogCookGenerationHelper, Verbose, TEXT("%s: ClearKeepForAllSavedOrGC."),
+		*WriteToString<256>(GetOwner().GetPackageName()));
+
 	bKeepForAllSavedOrGC = false;
 	if (!bKeepForAllSavedOrGC && !bKeepForCompletedAllSavesMessage)
 	{
@@ -792,12 +816,18 @@ inline void FGenerationHelper::ClearKeepForAllSavedOrGC()
 
 inline void FGenerationHelper::SetKeepForCompletedAllSavesMessage()
 {
+	UE_LOG(LogCookGenerationHelper, Verbose, TEXT("%s: SetKeepForCompletedAllSavesMessage."),
+		*WriteToString<256>(GetOwner().GetPackageName()));
+
 	ReferenceFromKeepForAllSavedOrGC = this;
 	bKeepForCompletedAllSavesMessage = true;
 }
 
 inline void FGenerationHelper::ClearKeepForCompletedAllSavesMessage()
 {
+	UE_LOG(LogCookGenerationHelper, Verbose, TEXT("%s: ClearKeepForCompletedAllSavesMessage."),
+		*WriteToString<256>(GetOwner().GetPackageName()));
+
 	bKeepForCompletedAllSavesMessage = false;
 	if (!bKeepForAllSavedOrGC && !bKeepForCompletedAllSavesMessage)
 	{
diff --git a/Engine/Source/Editor/UnrealEd/Private/Cooker/CookRequestCluster.cpp b/Engine/Source/Editor/UnrealEd/Private/Cooker/CookRequestCluster.cpp
--- a/Engine/Source/Editor/UnrealEd/Private/Cooker/CookRequestCluster.cpp
+++ b/Engine/Source/Editor/UnrealEd/Private/Cooker/CookRequestCluster.cpp
@@ -2100,11 +2100,11 @@ void FRequestCluster::FGraphSearch::FExploreEdgesContext::ProcessPlatformAttachm
 				PackagePlatformData.SetIterativelySkipped(true);
 				if (TRefCountPtr<FGenerationHelper> GenerationHelper = PackageData->GetGenerationHelper())
 				{
-					GenerationHelper->MarkPackageIterativelySkipped(*PackageData);
+					GenerationHelper->MarkPackageIterativelySkipped(*PackageData, true /* bIncrementallySkipped */);
 				}
 				if (ParentGenerationHelper)
 				{
-					ParentGenerationHelper->MarkPackageIterativelySkipped(*PackageData);
+					ParentGenerationHelper->MarkPackageIterativelySkipped(*PackageData, true /* bIncrementallySkipped */);
 				}
 				if (CookResult == ECookResult::Succeeded && PlatformIndex == FirstSessionPlatformIndex)
 				{
@@ -2121,6 +2121,20 @@ void FRequestCluster::FGraphSearch::FExploreEdgesContext::ProcessPlatformAttachm
 
 			Cluster.SetWasMarkedSkipped(*Vertex, true);
 		}
+		else
+		{
+			if (Cluster.TraversalMarkCookable())
+			{
+				if (TRefCountPtr<FGenerationHelper> GenerationHelper = PackageData->GetGenerationHelper())
+				{
+					GenerationHelper->MarkPackageIterativelySkipped(*PackageData, false /* bIncrementallySkipped */);
+				}
+				if (ParentGenerationHelper)
+				{
+					ParentGenerationHelper->MarkPackageIterativelySkipped(*PackageData, false /* bIncrementallySkipped */);
+				}
+			}
+		}
 	}
 
 	if (Cluster.bPreQueueBuildDefinitions && !bFoundBuildDefinitions)
diff --git a/Engine/Source/Editor/UnrealEd/Private/Cooker/CookTypes.cpp b/Engine/Source/Editor/UnrealEd/Private/Cooker/CookTypes.cpp
--- a/Engine/Source/Editor/UnrealEd/Private/Cooker/CookTypes.cpp
+++ b/Engine/Source/Editor/UnrealEd/Private/Cooker/CookTypes.cpp
@@ -21,6 +21,7 @@
 #include "Serialization/PackageWriterToSharedBuffer.h"
 
 LLM_DEFINE_TAG(Cooker_CachedPlatformData);
+DEFINE_LOG_CATEGORY(LogCookGenerationHelper);
 
 namespace UE::Cook
 {
diff --git a/Engine/Source/Editor/UnrealEd/Private/Cooker/CookTypes.h b/Engine/Source/Editor/UnrealEd/Private/Cooker/CookTypes.h
--- a/Engine/Source/Editor/UnrealEd/Private/Cooker/CookTypes.h
+++ b/Engine/Source/Editor/UnrealEd/Private/Cooker/CookTypes.h
@@ -14,6 +14,7 @@
 #include "HAL/PlatformMath.h"
 #include "HAL/Platform.h"
 #include "Interfaces/IPluginManager.h"
+#include "Logging/LogMacros.h"
 #include "Logging/TokenizedMessage.h"
 #include "Misc/AssertionMacros.h"
 #include "Serialization/PackageWriter.h"
@@ -50,6 +51,7 @@ bool LoadFromCompactBinary(FCbFieldView Field, UE::Cook::FInitializeConfigSettin
 #define DEBUG_COOKONTHEFLY 0
 
 LLM_DECLARE_TAG(Cooker_CachedPlatformData);
+DECLARE_LOG_CATEGORY_EXTERN(LogCookGenerationHelper, Verbose, All);
 
 inline constexpr uint32 ExpectedMaxNumPlatforms = 32;
 
