Add an important slow task around startup map loading, as it can be quite slow and the user can take action if they wish to avoid it

#rnx

diff --git a/Engine/Source/Editor/UnrealEd/Private/FileHelpers.cpp b/Engine/Source/Editor/UnrealEd/Private/FileHelpers.cpp
--- a/Engine/Source/Editor/UnrealEd/Private/FileHelpers.cpp
+++ b/Engine/Source/Editor/UnrealEd/Private/FileHelpers.cpp
@@ -5035,7 +5035,12 @@ void FEditorFileUtils::LoadDefaultMapAtStartup()
 	if ( FPackageName::IsValidLongPackageName(EditorStartupMap, bIncludeReadOnlyRoots) )
 	{
 		FString MapFilenameToLoad = FPackageName::LongPackageNameToFilename( EditorStartupMap );
-
+		
+		FScopedSlowTask SlowTask(1.f /* Number of steps */,
+			FText::Format(
+				LOCTEXT("LoadDefaultMapAtStartup", "Loading Startup Map: {0}"),
+				FText::FromString(EditorStartupMap)));
+		SlowTask.Visibility = ESlowTaskVisibility::Important; // this function can be very slow, users will benefit from our messages
 		bIsLoadingDefaultStartupMap = true;
 		FEditorFileUtils::LoadMap( MapFilenameToLoad + FPackageName::GetMapPackageExtension(), GUnrealEd && GUnrealEd->IsTemplateMap(EditorStartupMap), true );
 		bIsLoadingDefaultStartupMap = false;
