Add a const version of FWorkflowAllowedTabSet::CreateIterator
Allow accessing a const version of ApplicationModeList from derived classes through GetApplicationModeList()

diff --git a/Engine/Source/Editor/UnrealEd/Private/WorkflowCentricApplication.cpp b/Engine/Source/Editor/UnrealEd/Private/WorkflowCentricApplication.cpp
--- a/Engine/Source/Editor/UnrealEd/Private/WorkflowCentricApplication.cpp
+++ b/Engine/Source/Editor/UnrealEd/Private/WorkflowCentricApplication.cpp
@@ -113,10 +113,10 @@ void FWorkflowCentricApplication::SetCurrentMode(FName NewMode)
 	}
 }
 
-void FWorkflowCentricApplication::PushTabFactories(FWorkflowAllowedTabSet& FactorySetToPush)
+void FWorkflowCentricApplication::PushTabFactories(FWorkflowAllowedTabSet& FactorySetToPush) const
 {
 	check(TabManager.IsValid());
-	for (auto FactoryIt = FactorySetToPush.CreateIterator(); FactoryIt; ++FactoryIt)
+	for (auto FactoryIt = FactorySetToPush.CreateConstIterator(); FactoryIt; ++FactoryIt)
 	{
 		FactoryIt.Value()->RegisterTabSpawner(TabManager.ToSharedRef(), CurrentAppModePtr.Get());
 	}
@@ -159,6 +159,11 @@ void FWorkflowCentricApplication::RemoveAllApplicationModes()
 	ApplicationModeList.Reset();
 }
 
+const TMap<FName, TSharedPtr<FApplicationMode>>& FWorkflowCentricApplication::GetApplicationModeList() const
+{
+	return ApplicationModeList;
+}
+
 /////////////////////////////////////////////////////
 
 #undef LOCTEXT_NAMESPACE
diff --git a/Engine/Source/Editor/UnrealEd/Private/WorkflowOrientedApp/WorkflowTabManager.cpp b/Engine/Source/Editor/UnrealEd/Private/WorkflowOrientedApp/WorkflowTabManager.cpp
--- a/Engine/Source/Editor/UnrealEd/Private/WorkflowOrientedApp/WorkflowTabManager.cpp
+++ b/Engine/Source/Editor/UnrealEd/Private/WorkflowOrientedApp/WorkflowTabManager.cpp
@@ -402,6 +402,11 @@ TMap< FName, TSharedPtr<class FWorkflowTabFactory> >::TIterator FWorkflowAllowed
 	return Factories.CreateIterator();
 }
 
+TMap< FName, TSharedPtr<class FWorkflowTabFactory> >::TConstIterator FWorkflowAllowedTabSet::CreateConstIterator() const
+{
+	return Factories.CreateConstIterator();
+}
+
 /////////////////////////////////////////////////////
 // FDocumentTracker
 
diff --git a/Engine/Source/Editor/UnrealEd/Public/WorkflowOrientedApp/WorkflowCentricApplication.h b/Engine/Source/Editor/UnrealEd/Public/WorkflowOrientedApp/WorkflowCentricApplication.h
--- a/Engine/Source/Editor/UnrealEd/Public/WorkflowOrientedApp/WorkflowCentricApplication.h
+++ b/Engine/Source/Editor/UnrealEd/Public/WorkflowOrientedApp/WorkflowCentricApplication.h
@@ -41,7 +41,7 @@ public:
 	// Attempt to set the current mode.  If this mode is illegal or unknown, the mode will remain unchanged.
 	UNREALED_API virtual void SetCurrentMode(FName NewMode);
 
-	UNREALED_API void PushTabFactories(FWorkflowAllowedTabSet& FactorySetToPush);
+	UNREALED_API void PushTabFactories(FWorkflowAllowedTabSet& FactorySetToPush) const;
 
 	// Gets the mode extender list for all workflow applications (append to customize a specific mode)
 	static TArray<FWorkflowApplicationModeExtender>& GetModeExtenderList() { return ModeExtenderList; }
@@ -59,6 +59,9 @@ protected:
 
 	// Remove all application modes that were added through AddApplicationMode.
 	UNREALED_API virtual void RemoveAllApplicationModes();
+
+	// Get the list of registered application modes. Use AddApplicationMode to add modes.
+	UNREALED_API virtual const TMap< FName, TSharedPtr<FApplicationMode> >& GetApplicationModeList() const;
 protected:
 	TSharedPtr<FApplicationMode> CurrentAppModePtr;
 
diff --git a/Engine/Source/Editor/UnrealEd/Public/WorkflowOrientedApp/WorkflowTabManager.h b/Engine/Source/Editor/UnrealEd/Public/WorkflowOrientedApp/WorkflowTabManager.h
--- a/Engine/Source/Editor/UnrealEd/Public/WorkflowOrientedApp/WorkflowTabManager.h
+++ b/Engine/Source/Editor/UnrealEd/Public/WorkflowOrientedApp/WorkflowTabManager.h
@@ -46,6 +46,7 @@ public:
 
 	// Allows iteration over Factories
 	UNREALED_API TMap< FName, TSharedPtr<class FWorkflowTabFactory> >::TIterator CreateIterator();
+	UNREALED_API TMap< FName, TSharedPtr<class FWorkflowTabFactory> >::TConstIterator CreateConstIterator() const;
 
 private:
 	// Map of available factories in this mode
