Removing missing packages from SoftObjectPaths collected from startup packages.

diff --git a/Engine/Source/Editor/UnrealEd/Private/CookOnTheFlyServer.cpp b/Engine/Source/Editor/UnrealEd/Private/CookOnTheFlyServer.cpp
--- a/Engine/Source/Editor/UnrealEd/Private/CookOnTheFlyServer.cpp
+++ b/Engine/Source/Editor/UnrealEd/Private/CookOnTheFlyServer.cpp
@@ -11314,13 +11314,22 @@ void UCookOnTheFlyServer::GenerateInitialRequests(FBeginCookContext& BeginContex
 	GetGameDefaultObjects(TargetPlatforms, GameDefaultObjects);
 
 	// Strip out the default maps from SoftObjectPaths collected from startup packages. They will be added to the cook if necessary by CollectFilesToCook.
-	for (const auto& GameDefaultSet : GameDefaultObjects)
+	for (const TTuple<FName, TArray<FName>>& GameDefaultSet : GameDefaultObjects)
 	{
 		for (FName AssetName : GameDefaultSet.Value)
 		{
 			StartupSoftObjectPackages.Remove(AssetName);
 		}
 	}
+	// Strip out missing packages from SoftObjectPaths collected from startup packages.
+	for (TSet<FName>::TIterator It = StartupSoftObjectPackages.CreateIterator(); It; ++It)
+	{
+		const FName PackageName = *It;
+		if (PackageName.IsNone() || PackageDatas->GetFileNameByPackageName(PackageName).IsNone())
+		{
+			It.RemoveCurrent();
+		}
+	}
 
 	TArray<FString> CookMaps = BeginContext.StartupOptions->CookMaps;
 	TArray<FString> CookFirstPackages;
