Move UEditorInteractiveToolsContext classes from EditorInteractiveToolsFramework module to UnrealEd

Unshelved from pending changelist '40462909'

diff --git a/Engine/Source/Editor/Experimental/EditorInteractiveToolsFramework/Private/EdModeInteractiveToolsContext.cpp b/Engine/Source/Editor/Experimental/EditorInteractiveToolsFramework/Private/EdModeInteractiveToolsContext.cpp
deleted file mode 100644
diff --git a/Engine/Source/Editor/Experimental/EditorInteractiveToolsFramework/Public/EdModeInteractiveToolsContext.h b/Engine/Source/Editor/Experimental/EditorInteractiveToolsFramework/Public/EdModeInteractiveToolsContext.h
--- a/Engine/Source/Editor/Experimental/EditorInteractiveToolsFramework/Public/EdModeInteractiveToolsContext.h
+++ b/Engine/Source/Editor/Experimental/EditorInteractiveToolsFramework/Public/EdModeInteractiveToolsContext.h
@@ -1,388 +1,7 @@
 // Copyright Epic Games, Inc. All Rights Reserved.
-
 #pragma once
+#include "Misc/CoreMiscDefines.h"
 
-#include "Containers/Array.h"
-#include "Containers/Map.h"
-#include "Containers/UnrealString.h"
-#include "Delegates/Delegate.h"
-#include "Engine/EngineBaseTypes.h"
-#include "Engine/World.h"
-#include "HAL/Platform.h"
-#include "InputCoreTypes.h"
-#include "InputState.h"
-#include "InteractiveTool.h"
-#include "InteractiveToolManager.h"
-#include "InteractiveToolsContext.h"
-#include "Math/Ray.h"
-#include "Misc/AssertionMacros.h"
-#include "Misc/Optional.h"
-#include "Templates/SharedPointer.h"
-#include "UObject/ObjectMacros.h"
-#include "UObject/ObjectPtr.h"
-#include "UObject/UObjectGlobals.h"
-
-#include "EdModeInteractiveToolsContext.generated.h"
-
-class FCanvas;
-class FEdMode;
-class FEditorModeTools;
-class FEditorViewportClient;
-class FPrimitiveDrawInterface;
-class FSceneView;
-class FViewport;
-class FViewportClient;
-class ILevelEditor;
-class IToolsContextQueriesAPI;
-class IToolsContextRenderAPI;
-class IToolsContextTransactionsAPI;
-class UEdModeInteractiveToolsContext;
-class UGizmoViewContext;
-class UInputRouter;
-class UInteractiveToolBuilder;
-class UMaterialInterface;
-class UObject;
-class USelection;
-class UTypedElementSelectionSet;
-struct FToolBuilderState;
-
-/**
- * UEditorInteractiveToolsContext is an extension/adapter of an InteractiveToolsContext designed 
- * for use in the UE Editor. Currently this implementation assumes that it is created by a
- * Mode Manager (FEditorModeTools), and that the Mode Manager will call various API functions
- * like Render() and Tick() when necessary. 
- * 
- * 
- * allows it to be easily embedded inside an FEdMode. A set of functions are provided which can be
- * called from the FEdMode functions of the same name. These will handle the data type
- * conversions and forwarding calls necessary to operate the ToolsContext
- */
-UCLASS(Transient)
-class EDITORINTERACTIVETOOLSFRAMEWORK_API UEditorInteractiveToolsContext : public UInteractiveToolsContext
-{
-	GENERATED_BODY()
-
-public:
-	UEditorInteractiveToolsContext();
-
-	/**
-	 * Initialize a new ToolsContext for an EdMode owned by the given InEditorModeManager
-	 * @param 
-	 */
-	void InitializeContextWithEditorModeManager(FEditorModeTools* InEditorModeManager, UInputRouter* UseInputRouter = nullptr);
-
-	/** Shutdown ToolsContext and clean up any connections/etc */
-	virtual void ShutdownContext();
-
-	// default behavior is to accept active tool
-	virtual void TerminateActiveToolsOnPIEStart();
-
-	// default behavior is to accept active tool
-	virtual void TerminateActiveToolsOnSaveWorld();
-
-	// default behavior is to cancel active tool
-	virtual void TerminateActiveToolsOnWorldTearDown();
-
-	// default behavior is to cancel active tool
-	virtual void TerminateActiveToolsOnLevelChange();
-
-	FEditorModeTools* GetParentEditorModeManager() const { return EditorModeManager; }
-
-	IToolsContextQueriesAPI* GetQueriesAPI() const { return QueriesAPI; }
-	IToolsContextTransactionsAPI* GetTransactionAPI() const { return TransactionAPI; }
-
-	/** Call this to notify the Editor that the viewports this ToolsContext is related to may need a repaint, ie during interactive tool usage */
-	virtual void PostInvalidation();
-
-	// UObject Interface
-	virtual UWorld* GetWorld() const override;
-
-	// call functions of the same name on the ToolManager and GizmoManager
-	virtual void Tick(FEditorViewportClient* ViewportClient, float DeltaTime);
-	virtual void Render(const FSceneView* View, FViewport* Viewport, FPrimitiveDrawInterface* PDI);
-	virtual void DrawHUD(FViewportClient* ViewportClient,FViewport* Viewport,const FSceneView* View, FCanvas* Canvas);
-
-	// These delegates can be used to hook into the Render() / DrawHUD() / Tick() calls above. In particular, non-legacy UEdMode's
-	// don't normally receive Render() and DrawHUD() calls from the mode manager, but can attach to these.
-	DECLARE_MULTICAST_DELEGATE_OneParam(FOnRender, IToolsContextRenderAPI* RenderAPI);
-	FOnRender OnRender;
-	DECLARE_MULTICAST_DELEGATE_TwoParams(FOnDrawHUD, FCanvas* Canvas, IToolsContextRenderAPI* RenderAPI);
-	FOnDrawHUD OnDrawHUD;
-	DECLARE_MULTICAST_DELEGATE_OneParam(FOnTick, float DeltaTime);
-	FOnTick OnTick;
-
-	/** @return true if selected actors/components can be deleted */
-	virtual bool ProcessEditDelete();
-
-	//
-	// Utility functions useful for hooking up to UICommand/etc
-	//
-
-	virtual bool CanStartTool(const FString ToolTypeIdentifier) const;
-	virtual bool HasActiveTool() const;
-	virtual FString GetActiveToolName() const;
-	virtual bool ActiveToolHasAccept() const;
-	virtual bool CanAcceptActiveTool() const;
-	virtual bool CanCancelActiveTool() const;
-	virtual bool CanCompleteActiveTool() const;
-	virtual void StartTool(const FString ToolTypeIdentifier);
-	virtual void EndTool(EToolShutdownType ShutdownType);
-	void Activate();
-	void Deactivate();
-
-
-	/** @return Ray into 3D scene at last mouse event */
-	virtual FRay GetLastWorldRay() const
-	{
-		check(false);
-		return FRay();
-	}
-
-	//
-	// Configuration functions
-	//
-
-	/*
-	 * Configure whether ::Render() should early-out for HitProxy rendering passes.
-	 * If the Mode does not use HitProxy, and the Tools/Gizmos have expensive Render() calls, this can help with interactive performance.
-	 */
-	void SetEnableRenderingDuringHitProxyPass(bool bEnabled);
-
-	/** @return true if HitProxy rendering will be allowed in ::Render() */
-	bool GetEnableRenderingDuringHitProxyPass() const { return bEnableRenderingDuringHitProxyPass; }
-
-
-	/**
-	 * Configure whether Transform Gizmos created by the ITF (eg CombinedTransformGizmo) should prefer to show in 'Combined' mode.
-	 * If this is disabled, the Gizmo should respect the active Editor Gizmo setting (eg in the Level Viewport)
-	 */
-	void SetForceCombinedGizmoMode(bool bEnabled);
-
-	/** @return true if Force Combined Gizmo mode is Enabled */
-	bool GetForceCombinedGizmoModeEnabled() const { return bForceCombinedGizmoMode; }
-
-
-	/**
-	 * Configure whether Transform Gizmos created by the ITF (eg CombinedTransformGizmo) should, when in World coordinate system,
-	 * snap to an Absolute world-aligned grid, or snap Relative to the initial position of any particular gizmo transform.
-	 * Relative is the default and is also the behavior of the standard UE Gizmo.
-	 */
-	void SetAbsoluteWorldSnappingEnabled(bool bEnabled);
-
-	/** @return true if Absolute World Snapping mode is Enabled */
-	bool GetAbsoluteWorldSnappingEnabled() const { return bEnableAbsoluteWorldSnapping; }
-
-	/*
-	 * Configure whether tools should shutdown when entering PIE.
-	 * By default, the context will shut down any active tools on PIE start.
-	 * Setting this to true will prevent this, but the system makes no promises about whether tools work properly across PIE start/run/shutdown.
-	 * Therefore, do not use this if you do not know for certain that it is safe to persist tools across PIE start in your situation.
-	 */
-	void SetDeactivateToolsOnPIEStart(bool bDeactivateTools);
-
-	/** @return true if tools should be deactivated when PIE is started */
-	bool GetDeactivateToolsOnPIEStart() const { return bDeactivateOnPIEStart; }
-
-protected:
-
-	// we hide these 
-	virtual void Initialize(IToolsContextQueriesAPI* QueriesAPI, IToolsContextTransactionsAPI* TransactionsAPI) override;
-	virtual void Shutdown() override;
-
-	virtual void DeactivateActiveTool(EToolSide WhichSide, EToolShutdownType ShutdownType);
-	virtual void DeactivateAllActiveTools(EToolShutdownType ShutdownType);
-
-public:
-	UPROPERTY()
-	TObjectPtr<UMaterialInterface> StandardVertexColorMaterial;
-
-protected:
-	// EdMode implementation of InteractiveToolFramework APIs - see ToolContextInterfaces.h
-	IToolsContextQueriesAPI* QueriesAPI;
-	IToolsContextTransactionsAPI* TransactionAPI;
-
-	// Tools need to be able to Invalidate the view, in case it is not Realtime.
-	// Currently we do this very aggressively, and also force Realtime to be on, but in general we should be able to rely on Invalidation.
-	// However there are multiple Views and we do not want to Invalidate immediately, so we store a timestamp for each
-	// ViewportClient, and invalidate it when we see it if it's timestamp is out-of-date.
-	// (In theory this map will continually grow as new Viewports are created...)
-	TMap<FViewportClient*, int32> InvalidationMap;
-	// current invalidation timestamp, incremented by invalidation calls
-	int32 InvalidationTimestamp = 0;
-
-	// An object in which we save the current scene view information that gizmos can use on the game thread
-	// to figure out how big the gizmo is for hit testing. Lives in the context store, but we keep a pointer here
-	// to avoid having to look for it.
-	UGizmoViewContext* GizmoViewContext = nullptr;
-
-	// Utility function to convert viewport x/y from mouse events (and others?) into scene ray.
-	// Copy-pasted from other Editor code, seems kind of expensive?
-	static FRay GetRayFromMousePos(FEditorViewportClient* ViewportClient, FViewport* Viewport, int MouseX, int MouseY);
-
-	// editor UI state that we set before starting tool and when exiting tool
-	// Currently disabling anti-aliasing during active Tools because it causes PDI flickering
-	void SetEditorStateForTool();
-	void RestoreEditorState();
-
-	void OnToolEnded(UInteractiveToolManager* InToolManager, UInteractiveTool* InEndedTool);
-	void OnToolPostBuild(UInteractiveToolManager* InToolManager, EToolSide InSide, UInteractiveTool* InBuiltTool, UInteractiveToolBuilder* InToolBuilder, const FToolBuilderState& ToolState);
-
-	TOptional<FString> PendingToolToStart = {};
-	TOptional<EToolShutdownType> PendingToolShutdownType = {};
-
-private:
-	FEditorModeTools* EditorModeManager = nullptr;
-
-	// currently defaulting to enabled as FEdModes generally assume this, and in most cases hitproxy pass is not expensive.
-	bool bEnableRenderingDuringHitProxyPass = true;
-
-	bool bForceCombinedGizmoMode = false;
-	bool bEnableAbsoluteWorldSnapping = false;
-
-	bool bDeactivateOnPIEStart = true;
-
-	bool bIsActive = false;
-};
-
-
-/**
- * UModeManagerInteractiveToolsContext extends UEditorInteractiveToolsContext with various functions for handling 
- * device (mouse) input. These functions are currently called by the EdMode Manager (FEditorModeTools).
- */
-UCLASS(Transient)
-class EDITORINTERACTIVETOOLSFRAMEWORK_API UModeManagerInteractiveToolsContext : public UEditorInteractiveToolsContext
-{
-	GENERATED_BODY()
-
-	//
-	// UEditorInteractiveToolsContext API implementations that also forward calls to any child EdMode ToolsContexts
-	//
-public:
-	virtual void Tick(FEditorViewportClient* ViewportClient, float DeltaTime) override;
-	virtual void Render(const FSceneView* View, FViewport* Viewport, FPrimitiveDrawInterface* PDI) override;
-	virtual void DrawHUD(FViewportClient* ViewportClient,FViewport* Viewport,const FSceneView* View, FCanvas* Canvas) override;
-
-	virtual bool ProcessEditDelete() override;
-
-	//
-	// Input handling, these functions forward ViewportClient events to the UInputRouter
-	//
-public:
-	bool InputKey(FEditorViewportClient* ViewportClient, FViewport* Viewport, FKey Key, EInputEvent Event);
-	/**
-	 * This updates internal state like InputKey, but doesn't route the results to the input router. 
-	 * Use this if the input is captured by some higher system, to avoid this class from having an
-	 * incorrect view of e.g. the mouse state because it did not receive a mouse release event.
-	 */
-	void UpdateStateWithoutRoutingInputKey(FEditorViewportClient* ViewportClient, FViewport* Viewport, FKey Key, EInputEvent Event);
-
-	bool MouseEnter(FEditorViewportClient* ViewportClient, FViewport* Viewport, int32 x, int32 y);
-	bool MouseLeave(FEditorViewportClient* ViewportClient, FViewport* Viewport);
-	bool MouseMove(FEditorViewportClient* ViewportClient, FViewport* Viewport, int32 x, int32 y);
-
-	bool StartTracking(FEditorViewportClient* InViewportClient, FViewport* InViewport);
-	bool CapturedMouseMove(FEditorViewportClient* InViewportClient, FViewport* InViewport, int32 InMouseX, int32 InMouseY);
-	bool EndTracking(FEditorViewportClient* InViewportClient, FViewport* InViewport);
-
-	/** @return True if the context has overriden the cursor style, false if not. */
-	bool GetCursor(EMouseCursor::Type& OutCursor) const;
-
-	/** @return Ray into 3D scene at last mouse event */
-	virtual FRay GetLastWorldRay() const override;
-
-protected:
-	virtual void DeactivateAllActiveTools(EToolShutdownType ShutdownType) override;
-	virtual void Initialize(IToolsContextQueriesAPI* QueriesAPIIn, IToolsContextTransactionsAPI* TransactionsAPIIn) override;
-	virtual void Shutdown() override;
-
-	/** Input event instance used to keep track of various button states, etc, that we cannot directly query on-demand */
-	FInputDeviceState CurrentMouseState;
-	// called when PIE is about to start, shuts down active tools
-	FDelegateHandle BeginPIEDelegateHandle;
-	// called before a Save starts. This currently shuts down active tools.
-	FDelegateHandle PreSaveWorldDelegateHandle;
-	// called when a map is changed
-	FDelegateHandle WorldTearDownDelegateHandle;
-	// called when viewport clients change
-	FDelegateHandle ViewportClientListChangedHandle;
-
-private:
-	bool bIsTrackingMouse;
-
-
-
-
-protected:
-	UPROPERTY()
-	TArray<TObjectPtr<UEdModeInteractiveToolsContext>> EdModeToolsContexts;
-
-public:
-	/** 
-	* Create and initialize a new EdMode-level ToolsContext derived from the ModeManager ToolsContext.
-	* The EdMode ToolsContext does not have it's own InputRouter, it shares the InputRouter with the ModeManager ToolsContext.
-	* The ModeManager ToolsContext keeps track of these derived ToolsContext's and automatically Tick()'s them/etc.
-	* When the child ToolsContext is shut down, OnChildEdModeToolsContextShutdown() must be called to clean up
-	* @return new ToolsContext
-	*/
-	UEdModeInteractiveToolsContext* CreateNewChildEdModeToolsContext();
-
-	/**
-	 * Call to add a child EdMode ToolsContext created using the above function
-	 * @return true if child was added
-	 */
-	bool OnChildEdModeActivated(UEdModeInteractiveToolsContext* ChildToolsContext);
-
-	/**
-	 * Call to release a child EdMode ToolsContext created using the above function
-	 * @return true if child was found and removed
-	 */
-	bool OnChildEdModeDeactivated(UEdModeInteractiveToolsContext* ChildToolsContext);
-
-protected:
-	//~Begin ITF Alt processing Bypass
-
-	/**
-	 * These functions are declared protected to avoid having to deal with deprecation later on.
-	 * The Alt modifier bypass will be removed once a complete switch to ITF happens
-	 */
-
-	// Needs to use Alt modifier for Drag Tools to work
-	friend class FEditorViewportClient;
-
-	/** Enable/Disable Alt modifier bypass */
-	void SetBypassAltModifier(bool bInBypassAltModifier);
-
-	/** Is this Interactive Tools Context skipping Alt processing? */
-	bool BypassAltModifier() const { return bBypassAltModifier; }
-
-	/** Should we skip processing Alt modifier inputs? See UModeManagerInteractiveToolsContext::InputKey */
-	bool bBypassAltModifier = true;
-
-	//~End ITF Alt processing Bypass
-};
-
-
-/**
- * UEdModeInteractiveToolsContext is an UEditorInteractiveToolsContext intended for use/lifetime in the context of a UEdMode.
- * This ITC subclass is dependent on a UModeManagerInteractiveToolsContext to provide an InputRouter.
- */
-UCLASS(Transient)
-class EDITORINTERACTIVETOOLSFRAMEWORK_API UEdModeInteractiveToolsContext : public UEditorInteractiveToolsContext
-{
-	friend class UModeManagerInteractiveToolsContext;
-
-	GENERATED_BODY()
-public:
-	/**
-	 * Initialize a new EdModeToolsContext that is derived from a ModeManagerToolsContext.
-	 * This new ToolsContext will not have it's own InputRouter, it will share the InputRouter with the ModeManagerToolsContext
-	 */
-	void InitializeContextFromModeManagerContext(UModeManagerInteractiveToolsContext* ModeManagerToolsContext);
-
-	/** @return Ray into 3D scene at last mouse event */
-	virtual FRay GetLastWorldRay() const override;
-
-protected:
-	UPROPERTY()
-	TObjectPtr<UModeManagerInteractiveToolsContext> ParentModeManagerToolsContext = nullptr;
-};
+//todo: have a dedicated CL to uncomment this deprecation notice and update all files using this header
+//UE_DEPRECATED_HEADER(5.6, "Use Tools/EditorInteractiveToolsContext.h instead of EdModeInteractiveToolsContext.h.")
+#include "Tools/EdModeInteractiveToolsContext.h"
\ No newline at end of file
diff --git a/Engine/Source/Editor/UnrealEd/Private/Tools/EdModeInteractiveToolsContext.cpp b/Engine/Source/Editor/UnrealEd/Private/Tools/EdModeInteractiveToolsContext.cpp
new file mode 100644
--- /dev/null
+++ b/Engine/Source/Editor/UnrealEd/Private/Tools/EdModeInteractiveToolsContext.cpp
@@ -0,0 +1,1317 @@
+// Copyright Epic Games, Inc. All Rights Reserved.
+
+
+#include "Tools/EdModeInteractiveToolsContext.h"
+
+#include "BaseGizmos/GizmoViewContext.h"
+#include "ContextObjectStore.h"
+#include "Editor.h"
+#include "EditorViewportClient.h"
+#include "EditorModeManager.h"
+#include "Framework/Application/SlateApplication.h"
+#include "LevelEditor.h"
+#include "LevelEditorViewport.h"
+#include "IAssetViewport.h"
+#include "Math/Rotator.h"
+#include "Misc/AssertionMacros.h"
+#include "SLevelViewport.h"
+
+#include "Modules/ModuleManager.h"
+#include "ShowFlags.h"				// for EngineShowFlags
+#include "Engine/Engine.h"
+#include "Engine/Selection.h"
+#include "Engine/World.h"
+#include "Misc/ITransaction.h"
+#include "ScopedTransaction.h"
+#include "Materials/Material.h"
+#include "Engine/StaticMesh.h"
+#include "Components/StaticMeshComponent.h"
+
+#include "ToolContextInterfaces.h"
+#include "InteractiveToolObjects.h"
+#include "EditorInteractiveGizmoManager.h"
+#include "BaseBehaviors/ClickDragBehavior.h"
+#include "EditorModeManager.h"
+#include "EdMode.h"
+
+#include "BaseGizmos/GizmoRenderingUtil.h"
+#include "UnrealClient.h"
+#include "UObject/ObjectSaveContext.h"
+
+//#define ENABLE_DEBUG_PRINTING
+
+DEFINE_LOG_CATEGORY_STATIC(LogInteractiveToolsContext, Log, All);
+
+class FEdModeToolsContextQueriesImpl : public IToolsContextQueriesAPI
+{
+public:
+	UEditorInteractiveToolsContext* ToolsContext;
+	FEditorModeTools* EditorModeManager;
+
+	FViewCameraState CachedViewState;
+	FEditorViewportClient* CachedViewportClient;
+
+	FEdModeToolsContextQueriesImpl(UEditorInteractiveToolsContext* Context, FEditorModeTools* InEditorModeManager)
+	{
+		ToolsContext = Context;
+		EditorModeManager = InEditorModeManager;
+	}
+
+	void CacheCurrentViewState(FEditorViewportClient* ViewportClient)
+	{
+		CachedViewportClient = ViewportClient;
+		FViewportCameraTransform ViewTransform = ViewportClient->GetViewTransform();
+		CachedViewState.bIsOrthographic = ViewportClient->IsOrtho();
+		CachedViewState.Position = ViewTransform.GetLocation();
+		CachedViewState.HorizontalFOVDegrees = ViewportClient->ViewFOV;
+		CachedViewState.AspectRatio = ViewportClient->AspectRatio;
+
+		// ViewTransform rotation is only initialized for perspective!
+		if (CachedViewState.bIsOrthographic == false)
+		{
+			// if using Orbit camera, the rotation in the ViewTransform is not the current camera rotation, it
+			// is set to a different rotation based on the Orbit. So we have to convert back to camera rotation.
+			FRotator ViewRotation = (ViewportClient->bUsingOrbitCamera) ? 
+				ViewTransform.ComputeOrbitMatrix().InverseFast().Rotator()  :   ViewTransform.GetRotation();
+
+			CachedViewState.Orientation = ViewRotation.Quaternion();
+		}
+		else
+		{
+			// These rotations are based on hardcoded values in EditorViewportClient.cpp, see switches in FEditorViewportClient::CalcSceneView and FEditorViewportClient::Draw
+			switch (ViewportClient->ViewportType)
+			{
+			case LVT_OrthoXY:
+				CachedViewState.Orientation = FQuat(FRotator(-90.0f, -90.0f, 0.0f));
+				break;
+			case LVT_OrthoNegativeXY:
+				CachedViewState.Orientation = FQuat(FRotator(90.0f, 90.0f, 0.0f));
+				break;
+			case LVT_OrthoXZ:
+				CachedViewState.Orientation = FQuat(FRotator(0.0f, -90.0f, 0.0f));
+				break;
+			case LVT_OrthoNegativeXZ:
+				CachedViewState.Orientation = FQuat(FRotator(0.0f, 90.0f, 0.0f));
+				break;
+			case LVT_OrthoYZ:
+				CachedViewState.Orientation = FQuat(FRotator(0.0f, 0.0f, 0.0f));
+				break;
+			case LVT_OrthoNegativeYZ:
+				CachedViewState.Orientation = FQuat(FRotator(0.0f, 180.0f, 0.0f));
+				break;
+			default:
+				CachedViewState.Orientation = FQuat::Identity;
+			}
+
+			CachedViewState.OrthoWorldCoordinateWidth = ViewportClient->GetOrthoUnitsPerPixel(ViewportClient->Viewport) * static_cast<float>(ViewportClient->Viewport->GetSizeXY().X);
+		}
+
+		CachedViewState.bIsVR = false;
+	}
+
+	virtual UWorld* GetCurrentEditingWorld() const override
+	{
+		return EditorModeManager->GetWorld();
+	}
+
+	virtual void GetCurrentSelectionState(FToolBuilderState& StateOut) const override
+	{
+		StateOut.ToolManager = ToolsContext->ToolManager;
+		StateOut.TargetManager = ToolsContext->TargetManager;
+		StateOut.GizmoManager = ToolsContext->GizmoManager;
+		StateOut.World = EditorModeManager->GetWorld();
+		EditorModeManager->GetSelectedActors()->GetSelectedObjects(StateOut.SelectedActors);
+		EditorModeManager->GetSelectedComponents()->GetSelectedObjects(StateOut.SelectedComponents);
+
+		PRAGMA_DISABLE_DEPRECATION_WARNINGS
+		StateOut.TypedElementSelectionSet = EditorModeManager->GetEditorSelectionSet();
+		PRAGMA_ENABLE_DEPRECATION_WARNINGS
+	}
+
+	virtual void GetCurrentViewState(FViewCameraState& StateOut) const override
+	{
+		StateOut = CachedViewState;
+	}
+
+	virtual EToolContextCoordinateSystem GetCurrentCoordinateSystem() const override
+	{
+		ECoordSystem CoordSys = EditorModeManager->GetCoordSystem();
+		return (CoordSys == COORD_World) ? EToolContextCoordinateSystem::World : EToolContextCoordinateSystem::Local;
+	}
+
+	virtual EToolContextTransformGizmoMode GetCurrentTransformGizmoMode() const override
+	{
+		if (ToolsContext->GetForceCombinedGizmoModeEnabled() == false)
+		{
+			UE::Widget::EWidgetMode WidgetMode = EditorModeManager->GetWidgetMode();
+			switch (WidgetMode)
+			{
+			case UE::Widget::EWidgetMode::WM_None:
+				return EToolContextTransformGizmoMode::NoGizmo;
+			case UE::Widget::EWidgetMode::WM_Translate:
+				return EToolContextTransformGizmoMode::Translation;
+			case UE::Widget::EWidgetMode::WM_Rotate:
+				return EToolContextTransformGizmoMode::Rotation;
+			case UE::Widget::EWidgetMode::WM_Scale:
+				return EToolContextTransformGizmoMode::Scale;
+			}
+		}
+		
+		return EToolContextTransformGizmoMode::Combined;
+	}
+
+	virtual FToolContextSnappingConfiguration GetCurrentSnappingSettings() const override
+	{
+		FToolContextSnappingConfiguration Config;
+		Config.bEnablePositionGridSnapping = (GetDefault<ULevelEditorViewportSettings>()->GridEnabled != 0);
+		float EditorGridSize = GEditor->GetGridSize();
+		Config.PositionGridDimensions = FVector(EditorGridSize, EditorGridSize, EditorGridSize);
+		Config.bEnableRotationGridSnapping = (GetDefault<ULevelEditorViewportSettings>()->RotGridEnabled != 0);
+		Config.RotationGridAngles = GEditor->GetRotGridSize();
+		Config.bEnableScaleGridSnapping = (GetDefault<ULevelEditorViewportSettings>()->SnapScaleEnabled != 0);
+		Config.ScaleGridSize = GEditor->GetScaleGridSize();
+		Config.bEnableAbsoluteWorldSnapping = ToolsContext->GetAbsoluteWorldSnappingEnabled();
+		return Config;
+	}
+
+	virtual UMaterialInterface* GetStandardMaterial(EStandardToolContextMaterials MaterialType) const
+	{
+		if (MaterialType == EStandardToolContextMaterials::VertexColorMaterial)
+		{
+			return ToolsContext->StandardVertexColorMaterial;
+		}
+		check(false);
+		return nullptr;
+	}
+
+	virtual FViewport* GetHoveredViewport() const override
+	{
+		if (FEditorViewportClient* HoveredClient = EditorModeManager->GetHoveredViewportClient())
+		{
+			return HoveredClient->Viewport;
+		}
+
+		return nullptr;
+	}
+
+	virtual FViewport* GetFocusedViewport() const override
+	{
+		if (FEditorViewportClient* FocusedClient = EditorModeManager->GetFocusedViewportClient())
+		{
+			return FocusedClient->Viewport;
+		}
+
+		return nullptr;
+	}
+};
+
+class FEdModeToolsContextTransactionImpl : public IToolsContextTransactionsAPI
+{
+public:
+	FEdModeToolsContextTransactionImpl(UEditorInteractiveToolsContext* Context, FEditorModeTools* InEditorModeManager)
+		: ToolsContext(Context)
+		, EditorModeManager(InEditorModeManager)
+	{
+		check(EditorModeManager);
+		check(ToolsContext);
+	}
+
+
+	virtual void DisplayMessage(const FText& Message, EToolMessageLevel Level) override
+	{
+		if (Level == EToolMessageLevel::UserNotification || Level == EToolMessageLevel::UserMessage)
+		{
+			UE_LOG(LogInteractiveToolsContext, Display, TEXT("%s"), *Message.ToString());
+			ToolsContext->PostToolNotificationMessage(Message);
+		}
+		else if (Level == EToolMessageLevel::UserWarning || Level == EToolMessageLevel::UserError)
+		{
+			UE_LOG(LogInteractiveToolsContext, Warning, TEXT("%s"), *Message.ToString());
+			ToolsContext->PostToolWarningMessage(Message);
+		}
+		else
+		{
+			UE_LOG(LogInteractiveToolsContext, Warning, TEXT("%s"), *Message.ToString());
+		}
+	}
+
+
+	virtual void PostInvalidation() override
+	{
+		ToolsContext->PostInvalidation();
+	}
+
+	virtual void BeginUndoTransaction(const FText& Description) override
+	{
+		GEditor->BeginTransaction(Description);
+	}
+
+	virtual void EndUndoTransaction() override
+	{
+		GEditor->EndTransaction();
+	}
+
+	virtual void AppendChange(UObject* TargetObject, TUniquePtr<FToolCommandChange> Change, const FText& Description) override
+	{
+		FScopedTransaction Transaction(Description);
+		//if (ensure(GUndo != nullptr))		// ideally we would ensure here, but currently this can be hit on world teardown, resolution TBD
+		if (GUndo != nullptr)
+		{
+			GUndo->StoreUndo(TargetObject, MoveTemp(Change));
+		}
+	}
+
+
+	virtual bool RequestSelectionChange(const FSelectedObjectsChangeList& SelectionChange) override
+	{
+		checkf(SelectionChange.Components.Num() == 0, TEXT("FEdModeToolsContextTransactionImpl::RequestSelectionChange - Component selection not supported yet"));
+
+		if (SelectionChange.ModificationType == ESelectedObjectsModificationType::Clear)
+		{
+			GEditor->SelectNone(true, true, false);
+			return true;
+		}
+
+		if (SelectionChange.ModificationType == ESelectedObjectsModificationType::Replace )
+		{
+			GEditor->SelectNone(false, true, false);
+		}
+
+		bool bAdd = (SelectionChange.ModificationType != ESelectedObjectsModificationType::Remove);
+		int NumActors = SelectionChange.Actors.Num();
+		for (int k = 0; k < NumActors; ++k)
+		{
+			// Calling GEditor->NoteSelectionChange(true) will not send out change notifications on the TypedElementSelectionSet.
+			// The selection change will work but any Editor stuff listening for changes will not be notified.
+			// This may be a bug, needs further investigation.
+			// In the meantime, just send out the notification on the last SelectActor() call
+			bool bNotify = (k == NumActors-1);
+			GEditor->SelectActor(SelectionChange.Actors[k], bAdd, bNotify, true, false);
+		}
+
+		GEditor->NoteSelectionChange(true);
+
+		return true;
+	}
+
+protected:
+	UEditorInteractiveToolsContext* ToolsContext;
+	FEditorModeTools* EditorModeManager;
+};
+
+
+
+
+UEditorInteractiveToolsContext::UEditorInteractiveToolsContext()
+{
+	QueriesAPI = nullptr;
+	TransactionAPI = nullptr;
+}
+
+
+
+void UEditorInteractiveToolsContext::Initialize(IToolsContextQueriesAPI* QueriesAPIIn, IToolsContextTransactionsAPI* TransactionsAPIIn)
+{
+	UInteractiveToolsContext::Initialize(QueriesAPIIn, TransactionsAPIIn);
+
+	InvalidationTimestamp = 0;
+
+	// This gets set up in UInteractiveToolsContext::Initialize;
+	GizmoViewContext = ToolManager->GetContextObjectStore()->FindContext<UGizmoViewContext>();
+
+	UToolsContextCursorAPI* ToolsContextCursorAPI = ToolManager->GetContextObjectStore()->FindContext<UToolsContextCursorAPI>();
+	if (!ToolsContextCursorAPI)
+	{
+		ToolsContextCursorAPI = NewObject< UToolsContextCursorAPI >();
+		ToolManager->GetContextObjectStore()->AddContextObject(ToolsContextCursorAPI);
+	}
+}
+
+void UEditorInteractiveToolsContext::Shutdown()
+{
+	bIsActive = false;
+
+	// auto-accept any in-progress tools
+	DeactivateAllActiveTools(EToolShutdownType::Accept);
+
+	UInteractiveToolsContext::Shutdown();
+}
+
+
+void UEditorInteractiveToolsContext::InitializeContextWithEditorModeManager(FEditorModeTools* InEditorModeManager, UInputRouter* UseInputRouter)
+{
+	check(InEditorModeManager);
+	EditorModeManager = InEditorModeManager;
+
+	this->TransactionAPI = new FEdModeToolsContextTransactionImpl(this, InEditorModeManager);
+	this->QueriesAPI = new FEdModeToolsContextQueriesImpl(this, InEditorModeManager);
+
+	SetCreateGizmoManagerFunc([this](const FContextInitInfo& ContextInfo)
+	{
+		UEditorInteractiveGizmoManager* NewGizmoManager = NewObject<UEditorInteractiveGizmoManager>(ContextInfo.ToolsContext);
+		NewGizmoManager->InitializeWithEditorModeManager(ContextInfo.QueriesAPI, ContextInfo.TransactionsAPI, ContextInfo.InputRouter, EditorModeManager);
+		NewGizmoManager->RegisterDefaultGizmos();
+		return NewGizmoManager;
+	});
+
+	if (UseInputRouter != nullptr)
+	{
+		SetCreateInputRouterFunc([this, UseInputRouter](const FContextInitInfo& ContextInfo)
+		{
+			return UseInputRouter;
+		});
+		SetShutdownInputRouterFunc([this](UInputRouter*) {});
+	}
+
+	Initialize(QueriesAPI, TransactionAPI);
+
+	if (UseInputRouter != nullptr)
+	{
+		// enable auto invalidation in Editor, because invalidating for all hover and capture events is unpleasant
+		this->InputRouter->bAutoInvalidateOnHover = true;
+		this->InputRouter->bAutoInvalidateOnCapture = true;
+	}
+
+	// set up standard materials
+	StandardVertexColorMaterial = GEngine->VertexColorMaterial;
+}
+
+
+void UEditorInteractiveToolsContext::ShutdownContext()
+{
+	Shutdown();
+
+	OnToolNotificationMessage.Clear();
+	OnToolWarningMessage.Clear();
+
+	if (QueriesAPI != nullptr)
+	{
+		delete QueriesAPI;
+		QueriesAPI = nullptr;
+	}
+
+	if (TransactionAPI != nullptr)
+	{
+		delete TransactionAPI;
+		TransactionAPI = nullptr;
+	}
+}
+
+
+void UEditorInteractiveToolsContext::TerminateActiveToolsOnPIEStart()
+{
+	if (bDeactivateOnPIEStart)
+	{
+		DeactivateAllActiveTools(EToolShutdownType::Accept);	
+	}
+}
+void UEditorInteractiveToolsContext::TerminateActiveToolsOnSaveWorld()
+{
+	DeactivateAllActiveTools(EToolShutdownType::Accept);
+}
+void UEditorInteractiveToolsContext::TerminateActiveToolsOnWorldTearDown()
+{
+	DeactivateAllActiveTools(EToolShutdownType::Cancel);
+}
+void UEditorInteractiveToolsContext::TerminateActiveToolsOnLevelChange()
+{
+	DeactivateAllActiveTools(EToolShutdownType::Cancel);
+}
+
+void UEditorInteractiveToolsContext::PostInvalidation()
+{
+	InvalidationTimestamp++;
+}
+
+UWorld* UEditorInteractiveToolsContext::GetWorld() const
+{
+	if (bIsActive && EditorModeManager)
+	{
+		return EditorModeManager->GetWorld();
+	}
+
+	return nullptr;
+}
+
+void UEditorInteractiveToolsContext::Tick(FEditorViewportClient* ViewportClient, float DeltaTime)
+{
+	// invalidate this viewport if it's timestamp is not current
+	const int32* FoundTimestamp = InvalidationMap.Find(ViewportClient);
+	if (ViewportClient && FoundTimestamp == nullptr)
+	{
+		ViewportClient->Invalidate(false, false);
+		InvalidationMap.Add(ViewportClient, InvalidationTimestamp);
+	}
+	if (ViewportClient && FoundTimestamp != nullptr && *FoundTimestamp < InvalidationTimestamp)
+	{
+		ViewportClient->Invalidate(false, false);
+		InvalidationMap[ViewportClient] = InvalidationTimestamp;
+	}
+
+	if ( PendingToolShutdownType )
+	{
+		UInteractiveToolsContext::EndTool(EToolSide::Mouse, *PendingToolShutdownType);
+		PendingToolShutdownType.Reset();
+	}
+	if ( PendingToolToStart )
+	{
+		if (UInteractiveToolsContext::StartTool(EToolSide::Mouse, *PendingToolToStart))
+		{
+			SetEditorStateForTool();
+		}
+		PendingToolToStart.Reset();
+	}
+	
+	// This Tick() is called for every ViewportClient, however we only want to Tick the ToolManager and GizmoManager
+	// once, for the 'Active'/Focused Viewport, so early-out here
+	if (ViewportClient != EditorModeManager->GetFocusedViewportClient())
+	{
+		return;
+	}
+
+	// Cache current camera state from this Viewport in the ContextQueries, which we will use for things like snapping/etc that
+	// is computed by the Tool and Gizmo Tick()s
+	// (This is not necessarily correct for Hover, because we might be Hovering over a different Viewport than the Active one...)
+	if (ViewportClient)
+	{
+		((FEdModeToolsContextQueriesImpl*)this->QueriesAPI)->CacheCurrentViewState(ViewportClient);
+	}
+
+	// tick our stuff
+	ToolManager->Tick(DeltaTime);
+	GizmoManager->Tick(DeltaTime);
+	OnTick.Broadcast(DeltaTime);
+}
+
+
+
+class FEdModeTempRenderContext : public IToolsContextRenderAPI
+{
+public:
+	FPrimitiveDrawInterface* PDI;
+	const FSceneView* SceneView;
+	FViewCameraState ViewCameraState;
+	EViewInteractionState ViewInteractionState;
+
+	FEdModeTempRenderContext(const FSceneView* View, FViewport* Viewport, FEditorViewportClient* ViewportClient, FPrimitiveDrawInterface* DrawInterface, EViewInteractionState ViewInteractionState)
+		:PDI(DrawInterface), SceneView(View), ViewInteractionState(ViewInteractionState)
+	{
+		CacheCurrentViewState(Viewport, ViewportClient);
+	}
+
+	virtual FPrimitiveDrawInterface* GetPrimitiveDrawInterface() override
+	{
+		return PDI;
+	}
+
+	virtual const FSceneView* GetSceneView() override
+	{
+		return SceneView;
+	}
+
+	virtual FViewCameraState GetCameraState() override
+	{
+		return ViewCameraState;
+	}
+
+	virtual EViewInteractionState GetViewInteractionState() override
+	{
+		return ViewInteractionState;
+	}
+
+	void CacheCurrentViewState(FViewport* Viewport, FEditorViewportClient* ViewportClient)
+	{
+		if (!ViewportClient)
+		{
+			return;
+		}
+
+		FViewportCameraTransform ViewTransform = ViewportClient->GetViewTransform();
+		ViewCameraState.bIsOrthographic = ViewportClient->IsOrtho();
+		ViewCameraState.Position = ViewTransform.GetLocation();
+		ViewCameraState.HorizontalFOVDegrees = ViewportClient->ViewFOV;
+		ViewCameraState.AspectRatio = ViewportClient->AspectRatio;
+
+		// ViewTransform rotation is only initialized for perspective!
+		if (ViewCameraState.bIsOrthographic == false)
+		{
+			// if using Orbit camera, the rotation in the ViewTransform is not the current camera rotation, it
+			// is set to a different rotation based on the Orbit. So we have to convert back to camera rotation.
+			FRotator ViewRotation = (ViewportClient->bUsingOrbitCamera) ?
+				ViewTransform.ComputeOrbitMatrix().InverseFast().Rotator() : ViewTransform.GetRotation();
+
+			ViewCameraState.Orientation = ViewRotation.Quaternion();
+		}
+		else
+		{
+			// These rotations are based on hardcoded values in EditorViewportClient.cpp, see switches in FEditorViewportClient::CalcSceneView and FEditorViewportClient::Draw
+			switch (ViewportClient->ViewportType)
+			{
+			case LVT_OrthoXY:
+				ViewCameraState.Orientation = FQuat(FRotator(-90.0f, -90.0f, 0.0f));
+				break;
+			case LVT_OrthoNegativeXY:
+				ViewCameraState.Orientation = FQuat(FRotator(90.0f, 90.0f, 0.0f));
+				break;
+			case LVT_OrthoXZ:
+				ViewCameraState.Orientation = FQuat(FRotator(0.0f, -90.0f, 0.0f));
+				break;
+			case LVT_OrthoNegativeXZ:
+				ViewCameraState.Orientation = FQuat(FRotator(0.0f, 90.0f, 0.0f));
+				break;
+			case LVT_OrthoYZ:
+				ViewCameraState.Orientation = FQuat(FRotator(0.0f, 0.0f, 0.0f));
+				break;
+			case LVT_OrthoNegativeYZ:
+				ViewCameraState.Orientation = FQuat(FRotator(0.0f, 180.0f, 0.0f));
+				break;
+			default:
+				ViewCameraState.Orientation = FQuat::Identity;
+			}
+
+			ViewCameraState.OrthoWorldCoordinateWidth = ViewportClient->GetOrthoUnitsPerPixel(ViewportClient->Viewport) * static_cast<float>(ViewportClient->Viewport->GetSizeXY().X);
+		}
+
+		ViewCameraState.bIsVR = false;
+	}
+
+};
+
+
+void UEditorInteractiveToolsContext::Render(const FSceneView* View, FViewport* Viewport, FPrimitiveDrawInterface* PDI)
+{
+	// skip HitProxy rendering passes if desired
+	if (PDI->IsHitTesting() && bEnableRenderingDuringHitProxyPass == false)
+	{
+		return;
+	}
+
+	// THIS IS NOT SAFE!! However it appears that (1) it is only possible to get certain info from the EditorViewportClient,
+	// but (2) there is no way to know if a FViewportClient is an FEditorViewportClient. Currently this ::Render() function
+	// is only intended to be called by FEdMode/UEdMode::Render(), and their ::Render() calls are only called by the
+	// FEditorViewportClient, which passes it's own Viewport down. So, this cast should be valid (for now)
+	FEditorViewportClient* ViewportClient = static_cast<FEditorViewportClient*>(Viewport->GetClient());
+
+	// Update the currently-hovered scene view information, which GizmoArrowComponent and friends will
+	// use to recalculate their size/visibility/etc.
+	if (GizmoViewContext && ViewportClient == EditorModeManager->GetHoveredViewportClient())
+	{
+		GizmoViewContext->ResetFromSceneView(*View);
+	}
+
+	// Render Tool and Gizmos
+	const FEditorViewportClient* Focused = EditorModeManager->GetFocusedViewportClient();
+	const FEditorViewportClient* Hovered = EditorModeManager->GetHoveredViewportClient();
+	EViewInteractionState InteractionState = EViewInteractionState::None;
+	if (ViewportClient == Focused )
+	{
+		InteractionState |= EViewInteractionState::Focused;
+	}
+	if (ViewportClient == Hovered )
+	{
+		InteractionState |= EViewInteractionState::Hovered;
+	}
+	FEdModeTempRenderContext RenderContext(View, Viewport, ViewportClient, PDI, InteractionState);
+	ToolManager->Render(&RenderContext);
+	GizmoManager->Render(&RenderContext);
+	OnRender.Broadcast(&RenderContext);
+}
+
+void UEditorInteractiveToolsContext::DrawHUD(FViewportClient* ViewportClient,FViewport* Viewport,const FSceneView* View, FCanvas* Canvas)
+{
+	FEditorViewportClient* EditorViewportClient = static_cast<FEditorViewportClient*>(Viewport->GetClient());
+	const FViewportClient* Focused = EditorModeManager->GetFocusedViewportClient();
+	const FViewportClient* Hovered = EditorModeManager->GetHoveredViewportClient();
+	EViewInteractionState InteractionState = EViewInteractionState::None;
+	if (ViewportClient == Focused )
+	{
+		InteractionState |= EViewInteractionState::Focused;
+	}
+	if (ViewportClient == Hovered )
+	{
+		InteractionState |= EViewInteractionState::Hovered;
+	}
+	FEdModeTempRenderContext RenderContext(View, Viewport, EditorViewportClient, nullptr /*PDI*/, InteractionState);
+	ToolManager->DrawHUD(Canvas, &RenderContext);
+	GizmoManager->DrawHUD(Canvas, &RenderContext);
+	OnDrawHUD.Broadcast(Canvas, &RenderContext);
+}
+
+
+bool UEditorInteractiveToolsContext::ProcessEditDelete()
+{
+	if (ToolManager->HasAnyActiveTool() == false)
+	{
+		return false;
+	}
+
+	bool bSkipDelete = false;
+
+	USelection* SelectedActors = GEditor->GetSelectedActors();
+	for (int i = 0; i < SelectedActors->Num() && bSkipDelete == false; ++i)
+	{
+		UObject* SelectedActor = SelectedActors->GetSelectedObject(i);
+
+		// If any of the selected actors are AInternalToolFrameworkActor, we do not want to allow them to be deleted,
+		// as generally this will cause problems for the Tool.
+		if ( Cast<AInternalToolFrameworkActor>(SelectedActor) != nullptr)
+		{
+			bSkipDelete = true;
+		}
+
+		// If any Components of selected Actors implement UToolFrameworkComponent, we disable delete (for now).
+		// (Currently Sculpt and a few other Modeling Tools attach their preview mesh components to the selected Actor)
+		AActor* Actor = Cast<AActor>(SelectedActor);
+		if (Actor != nullptr)
+		{
+			const TSet<UActorComponent*>& Components = Actor->GetComponents();
+			for (const UActorComponent* Component : Components)
+			{
+				if ( Component->Implements<UToolFrameworkComponent>() )
+				{
+					bSkipDelete = true;
+				}
+			}
+		}
+
+	}
+
+	return bSkipDelete;
+}
+
+
+
+FRay UEditorInteractiveToolsContext::GetRayFromMousePos(FEditorViewportClient* ViewportClient, FViewport* Viewport, int MouseX, int MouseY)
+{
+	if (!ViewportClient)
+	{
+		return FRay();
+	}
+
+	FSceneViewFamilyContext ViewFamily(FSceneViewFamily::ConstructionValues(
+		ViewportClient->Viewport,
+		ViewportClient->GetScene(),
+		ViewportClient->EngineShowFlags).SetRealtimeUpdate(ViewportClient->IsRealtime()));		// why SetRealtimeUpdate here??
+	// this View is deleted by the FSceneViewFamilyContext destructor
+	FSceneView* View = ViewportClient->CalcSceneView(&ViewFamily);
+	FViewportCursorLocation MouseViewportRay(View, ViewportClient, MouseX, MouseY);
+
+	FVector RayOrigin = MouseViewportRay.GetOrigin();
+	FVector RayDirection = MouseViewportRay.GetDirection();
+
+	// in Ortho views, the RayOrigin appears to be completely arbitrary, in some views it is on the view plane,
+	// others it moves back/forth with the OrthoZoom. Translate by a large amount here in hopes of getting
+	// ray origin "outside" the scene (which is a disaster for numerical precision !! ... )
+	if (ViewportClient->IsOrtho())
+	{
+		RayOrigin -= 0.1 * HALF_WORLD_MAX * RayDirection;
+	}
+
+	return FRay(RayOrigin, RayDirection, true);
+}
+
+
+bool UEditorInteractiveToolsContext::CanStartTool(const FString ToolTypeIdentifier) const
+{
+	return UInteractiveToolsContext::CanStartTool(EToolSide::Mouse, ToolTypeIdentifier);
+}
+
+bool UEditorInteractiveToolsContext::HasActiveTool() const
+{
+	return UInteractiveToolsContext::HasActiveTool(EToolSide::Mouse);
+}
+
+FString UEditorInteractiveToolsContext::GetActiveToolName() const
+{
+	return UInteractiveToolsContext::GetActiveToolName(EToolSide::Mouse);
+}
+
+bool UEditorInteractiveToolsContext::ActiveToolHasAccept() const
+{
+	return  UInteractiveToolsContext::ActiveToolHasAccept(EToolSide::Mouse);
+}
+
+bool UEditorInteractiveToolsContext::CanAcceptActiveTool() const
+{
+	return UInteractiveToolsContext::CanAcceptActiveTool(EToolSide::Mouse);
+}
+
+bool UEditorInteractiveToolsContext::CanCancelActiveTool() const
+{
+	return UInteractiveToolsContext::CanCancelActiveTool(EToolSide::Mouse);
+}
+
+bool UEditorInteractiveToolsContext::CanCompleteActiveTool() const
+{
+	return UInteractiveToolsContext::CanCompleteActiveTool(EToolSide::Mouse);
+}
+
+void UEditorInteractiveToolsContext::StartTool(const FString ToolTypeIdentifier)
+{
+	FString LocalIdentifier(ToolTypeIdentifier);
+	PendingToolToStart = LocalIdentifier;
+	PostInvalidation();
+}
+
+void UEditorInteractiveToolsContext::EndTool(EToolShutdownType ShutdownType)
+{
+	PendingToolShutdownType = ShutdownType;
+	PostInvalidation();
+}
+
+void UEditorInteractiveToolsContext::Activate()
+{
+	bIsActive = true;
+}
+
+void UEditorInteractiveToolsContext::Deactivate()
+{
+	bIsActive = false;
+}
+
+
+void UEditorInteractiveToolsContext::DeactivateActiveTool(EToolSide WhichSide, EToolShutdownType ShutdownType)
+{
+	UInteractiveToolsContext::DeactivateActiveTool(WhichSide, ShutdownType);
+	RestoreEditorState();
+}
+
+void UEditorInteractiveToolsContext::DeactivateAllActiveTools(EToolShutdownType ShutdownType)
+{
+	UInteractiveToolsContext::DeactivateAllActiveTools(ShutdownType);
+	RestoreEditorState();
+}
+
+void UEditorInteractiveToolsContext::SetEditorStateForTool()
+{
+	FLevelEditorModule& LevelEditorModule = FModuleManager::GetModuleChecked<FLevelEditorModule>("LevelEditor");
+	TSharedPtr<ILevelEditor> LevelEditor = LevelEditorModule.GetFirstLevelEditor();
+	if (LevelEditor.IsValid())
+	{
+		TArray<TSharedPtr<SLevelViewport>> Viewports = LevelEditor->GetViewports();
+		for (const TSharedPtr<SLevelViewport>& ViewportWindow : Viewports)
+		{
+			if (ViewportWindow.IsValid())
+			{
+				FEditorViewportClient& Viewport = ViewportWindow->GetAssetViewportClient();
+				Viewport.EnableOverrideEngineShowFlags([](FEngineShowFlags& Flags)
+				{
+					Flags.SetTemporalAA(false);
+					Flags.SetMotionBlur(false);
+					// disable this as depending on fixed exposure settings the entire scene may turn black
+					//Flags.SetEyeAdaptation(false);
+				});
+			}
+		}
+	}
+}
+
+void UEditorInteractiveToolsContext::RestoreEditorState()
+{
+	FLevelEditorModule& LevelEditorModule = FModuleManager::GetModuleChecked<FLevelEditorModule>("LevelEditor");
+	TSharedPtr<ILevelEditor> LevelEditor  = LevelEditorModule.GetFirstLevelEditor();
+	if (LevelEditor.IsValid())
+	{
+		TArray<TSharedPtr<SLevelViewport>> Viewports = LevelEditor->GetViewports();
+		for (const TSharedPtr<SLevelViewport>& ViewportWindow : Viewports)
+		{
+			if (ViewportWindow.IsValid())
+			{
+				FEditorViewportClient& ViewportClient = ViewportWindow->GetAssetViewportClient();
+				ViewportClient.DisableOverrideEngineShowFlags();
+
+				// Rebuild the cached hit proxy. The tool may have disabled some viewport items while active (e.g. transform gizmo) and
+				// so the hit proxy cache may be out of date. In the past this has led to, for example, the gizmo being visible but not
+				// clickable when returning from a tool (UE-116888).
+				// TODO: Figure out why the hit proxy is not being rebuilt elsewhere
+				ViewportClient.RequestInvalidateHitProxy(ViewportClient.Viewport);
+			}
+		}
+	}
+}
+
+void UEditorInteractiveToolsContext::OnToolEnded(UInteractiveToolManager* InToolManager, UInteractiveTool* InEndedTool)
+{
+	RestoreEditorState();
+}
+
+void UEditorInteractiveToolsContext::OnToolPostBuild(UInteractiveToolManager* InToolManager, EToolSide InSide, UInteractiveTool* InBuiltTool, UInteractiveToolBuilder* InToolBuilder, const FToolBuilderState& ToolState)
+{
+	// todo: Add any shared tool targets for the mode toolkit
+}
+
+void UEditorInteractiveToolsContext::SetEnableRenderingDuringHitProxyPass(bool bEnabled)
+{
+	bEnableRenderingDuringHitProxyPass = bEnabled;
+}
+
+void UEditorInteractiveToolsContext::SetForceCombinedGizmoMode(bool bEnabled)
+{
+	bForceCombinedGizmoMode = bEnabled;
+}
+
+void UEditorInteractiveToolsContext::SetAbsoluteWorldSnappingEnabled(bool bEnabled)
+{
+	bEnableAbsoluteWorldSnapping = bEnabled;
+}
+
+void UEditorInteractiveToolsContext::SetDeactivateToolsOnPIEStart(bool bDeactivateTools)
+{
+	bDeactivateOnPIEStart = bDeactivateTools;
+}
+
+void UModeManagerInteractiveToolsContext::Tick(FEditorViewportClient* ViewportClient, float DeltaTime)
+{
+	UEditorInteractiveToolsContext::Tick(ViewportClient, DeltaTime);
+	for (const TObjectPtr<UEdModeInteractiveToolsContext>& EdModeContext : EdModeToolsContexts)
+	{
+		EdModeContext->Tick(ViewportClient, DeltaTime);
+	}
+}
+
+void UModeManagerInteractiveToolsContext::Render(const FSceneView* View, FViewport* Viewport, FPrimitiveDrawInterface* PDI)
+{
+	UEditorInteractiveToolsContext::Render(View, Viewport, PDI);
+	for (const TObjectPtr<UEdModeInteractiveToolsContext>& EdModeContext : EdModeToolsContexts)
+	{
+		EdModeContext->Render(View, Viewport, PDI);
+	}
+}
+
+void UModeManagerInteractiveToolsContext::DrawHUD(FViewportClient* ViewportClient, FViewport* Viewport, const FSceneView* View, FCanvas* Canvas)
+{
+	UEditorInteractiveToolsContext::DrawHUD(ViewportClient, Viewport, View, Canvas);
+	for (const TObjectPtr<UEdModeInteractiveToolsContext>& EdModeContext : EdModeToolsContexts)
+	{
+		EdModeContext->DrawHUD(ViewportClient, Viewport, View, Canvas);
+	}
+}
+
+
+bool UModeManagerInteractiveToolsContext::ProcessEditDelete()
+{
+	bool bHandled = UEditorInteractiveToolsContext::ProcessEditDelete();
+	for (const TObjectPtr<UEdModeInteractiveToolsContext>& EdModeContext : EdModeToolsContexts)
+	{
+		bHandled |= EdModeContext->ProcessEditDelete();
+	}
+	return bHandled;
+}
+
+void UModeManagerInteractiveToolsContext::DeactivateAllActiveTools(EToolShutdownType ShutdownType)
+{
+	for (UEdModeInteractiveToolsContext* EdModeContext : EdModeToolsContexts)
+	{
+		EdModeContext->DeactivateAllActiveTools(ShutdownType);
+	}
+
+	Super::DeactivateAllActiveTools(ShutdownType);
+}
+
+void UModeManagerInteractiveToolsContext::UpdateStateWithoutRoutingInputKey(FEditorViewportClient* ViewportClient, FViewport* Viewport, FKey Key, EInputEvent Event)
+{
+	// Currently, the only internal state we keep is the state of various mouse keys being down. Note
+	// that we don't want to save bPressed or bReleased for them as those are one-time events issued
+	// from InputKey, and shouldn't show up on mouse moves or on other keys being pressed/released.
+	if ((Event == IE_Pressed || Event == IE_Released)
+		&& Key.IsMouseButton())
+	{
+		if (Key.IsMouseButton())
+		{
+			if (Key == EKeys::LeftMouseButton)
+			{
+				CurrentMouseState.Mouse.Left.bDown = (Event == IE_Pressed);
+			}
+			else if (Key == EKeys::MiddleMouseButton)
+			{
+				CurrentMouseState.Mouse.Middle.bDown = (Event == IE_Pressed);
+			}
+			else if (Key == EKeys::RightMouseButton)
+			{
+				CurrentMouseState.Mouse.Right.bDown = (Event == IE_Pressed);
+			}
+		}
+	}
+}
+
+bool UModeManagerInteractiveToolsContext::InputKey(FEditorViewportClient* ViewportClient, FViewport* Viewport, FKey Key, EInputEvent Event)
+{
+#ifdef ENABLE_DEBUG_PRINTING
+	if (Event == IE_Pressed) { UE_LOG(LogInteractiveToolsContext, Warning, TEXT("PRESSED EVENT")); }
+	else if (Event == IE_Released) { UE_LOG(LogInteractiveToolsContext, Warning, TEXT("RELEASED EVENT")); }
+	else if (Event == IE_Repeat) { UE_LOG(LogInteractiveToolsContext, Warning, TEXT("REPEAT EVENT")); }
+	else if (Event == IE_Axis) { UE_LOG(LogInteractiveToolsContext, Warning, TEXT("AXIS EVENT")); }
+	else if (Event == IE_DoubleClick) { UE_LOG(LogInteractiveToolsContext, Warning, TEXT("DOUBLECLICK EVENT")); }
+#endif
+
+	// Update the current state, then route result
+	UpdateStateWithoutRoutingInputKey(ViewportClient, Viewport, Key, Event);
+
+	if (Event == IE_Pressed || Event == IE_Released || Event == IE_DoubleClick)
+	{
+		if (Key.IsMouseButton())
+		{
+			bool bIsLeftMouse = (Key == EKeys::LeftMouseButton);
+			bool bIsMiddleMouse = (Key == EKeys::MiddleMouseButton);
+			bool bIsRightMouse = (Key == EKeys::RightMouseButton);
+			if (bIsLeftMouse || bIsMiddleMouse || bIsRightMouse)
+			{
+				// Currently, we don't capture mouse clicks that start with Alt being down because we want
+				// Alt camera manipulation to have priority over tools. So, we let those inputs pass on up
+				// to wherever they get handled.
+				// Someday these kinds of prioritizations will be handled by having camera manipulation be
+				// in a common input router so that behavior priorities can determine the ordering.
+				if (BypassAltModifier())
+				{
+					if (ViewportClient && ViewportClient->IsAltPressed() && InputRouter->HasActiveMouseCapture() == false)
+					{
+						return false;
+					}
+				}
+
+				FInputDeviceState InputState = CurrentMouseState;
+				InputState.InputDevice = EInputDevices::Mouse;
+
+				FModifierKeysState ModifierKeys = FSlateApplication::Get().GetModifierKeys();
+				InputState.SetModifierKeyStates(
+					ModifierKeys.IsShiftDown(), ModifierKeys.IsAltDown(),
+					ModifierKeys.IsControlDown(), ModifierKeys.IsCommandDown());
+
+				// DoubleClick replaces a Press event. Include DoubleClick in Pressed/Down events so
+				// that it can still be processed by single click behaviors. For explicit double click
+				// handling, clients can listen for the DoubleClick event with a higher capture priority.
+				const bool bIsPressed = (Event == IE_Pressed || Event == IE_DoubleClick);
+				const bool bIsReleased = (Event == IE_Released);
+				const bool bIsDoubleClick = (Event == IE_DoubleClick);
+				if (bIsLeftMouse)
+				{
+					InputState.Mouse.Left.SetStates(bIsPressed, bIsPressed, bIsReleased, bIsDoubleClick);
+				}
+				else if (bIsMiddleMouse)
+				{
+					InputState.Mouse.Middle.SetStates(bIsPressed, bIsPressed, bIsReleased, bIsDoubleClick);
+				}
+				else
+				{
+					InputState.Mouse.Right.SetStates(bIsPressed, bIsPressed, bIsReleased, bIsDoubleClick);
+				}
+				if (InputRouter->PostInputEvent(InputState))
+				{
+					return true;
+				}
+			}
+			else if (Key == EKeys::MouseScrollUp || Key == EKeys::MouseScrollDown)
+			{
+				// Note that we get two events for each scroll- an IE_Pressed, and IE_Released.
+				// We pass both of these in, though only the first one will have WheelDelta set.
+				// If a behavior captures the mouse wheel interaction, the second event will give
+				// it the opportunity to immediately release capture. Not passing in the second
+				// event would probably be ok too- capture would get released on next hover event.
+
+				FInputDeviceState InputState = CurrentMouseState;
+				InputState.InputDevice = EInputDevices::Mouse;
+
+				FModifierKeysState ModifierKeys = FSlateApplication::Get().GetModifierKeys();
+				InputState.SetModifierKeyStates(
+					ModifierKeys.IsShiftDown(), ModifierKeys.IsAltDown(),
+					ModifierKeys.IsControlDown(), ModifierKeys.IsCommandDown());
+
+				InputState.Mouse.WheelDelta = (Event != IE_Pressed) ? 0.f
+					: (Key == EKeys::MouseScrollUp) ? 1.f 
+					: -1.f;
+
+				return InputRouter->PostInputEvent(InputState);
+			}
+		}
+		else if (Key.IsGamepadKey())
+		{
+			// not supported yet
+		}
+		else if (Key.IsTouch())
+		{
+			// not supported yet
+		}
+		else if (Key.IsAnalog())
+		{
+			// not supported yet
+		}
+		else    // is this definitely a keyboard key?
+		{
+			FInputDeviceState InputState;
+			InputState.InputDevice = EInputDevices::Keyboard;
+
+			FModifierKeysState ModifierKeys = FSlateApplication::Get().GetModifierKeys();
+			InputState.SetModifierKeyStates(
+				ModifierKeys.IsShiftDown(), ModifierKeys.IsAltDown(),
+				ModifierKeys.IsControlDown(), ModifierKeys.IsCommandDown());
+
+			InputState.Keyboard.ActiveKey.Button = Key;
+			bool bPressed = (Event == IE_Pressed);
+			InputState.Keyboard.ActiveKey.SetStates(bPressed, bPressed, !bPressed);
+			return InputRouter->PostInputEvent(InputState);
+		}
+	}
+
+	return false;
+}
+
+bool UModeManagerInteractiveToolsContext::MouseEnter(FEditorViewportClient* ViewportClient, FViewport* Viewport, int32 x, int32 y)
+{
+#ifdef ENABLE_DEBUG_PRINTING
+	UE_LOG(LogInteractiveToolsContext, Warning, TEXT("MOUSE ENTER"));
+#endif
+
+	CurrentMouseState.Mouse.Position2D = FVector2D(x, y);
+	CurrentMouseState.Mouse.WorldRay = GetRayFromMousePos(ViewportClient, Viewport, x, y);
+
+	return false;
+}
+
+
+bool UModeManagerInteractiveToolsContext::MouseMove(FEditorViewportClient* ViewportClient, FViewport* Viewport, int32 x, int32 y)
+{
+#ifdef ENABLE_DEBUG_PRINTING
+	UE_LOG(LogInteractiveToolsContext, Warning, TEXT("HOVER %p"), ViewportClient);
+#endif
+
+	CurrentMouseState.Mouse.Position2D = FVector2D(x, y);
+	CurrentMouseState.Mouse.WorldRay = GetRayFromMousePos(ViewportClient, Viewport, x, y);
+	FInputDeviceState InputState = CurrentMouseState;
+	InputState.InputDevice = EInputDevices::Mouse;
+
+	FModifierKeysState ModifierKeys = FSlateApplication::Get().GetModifierKeys();
+	InputState.SetModifierKeyStates(
+		ModifierKeys.IsShiftDown(), ModifierKeys.IsAltDown(),
+		ModifierKeys.IsControlDown(), ModifierKeys.IsCommandDown());
+
+	if (InputRouter->HasActiveMouseCapture())
+	{
+		// TODO: This should no longer be necessary: test and remove.
+		// This state occurs if InputBehavior did not release capture on mouse release.
+		// UMultiClickSequenceInputBehavior does this, eg for multi-click draw-polygon sequences.
+		// It's not ideal though and maybe would be better done via multiple captures + hover...?
+		InputRouter->PostInputEvent(InputState);
+	}
+	else
+	{
+		InputRouter->PostHoverInputEvent(InputState);
+	}
+
+	return false;
+}
+
+
+bool UModeManagerInteractiveToolsContext::MouseLeave(FEditorViewportClient* ViewportClient, FViewport* Viewport)
+{
+#ifdef ENABLE_DEBUG_PRINTING
+	UE_LOG(LogInteractiveToolsContext, Warning, TEXT("MOUSE LEAVE"));
+#endif
+
+	return false;
+}
+
+
+
+bool UModeManagerInteractiveToolsContext::StartTracking(FEditorViewportClient* InViewportClient, FViewport* InViewport)
+{
+	bIsTrackingMouse = InputRouter->HasActiveMouseCapture();
+	return bIsTrackingMouse;
+}
+
+bool UModeManagerInteractiveToolsContext::CapturedMouseMove(FEditorViewportClient* InViewportClient, FViewport* InViewport, int32 InMouseX, int32 InMouseY)
+{
+	FVector2D OldPosition = CurrentMouseState.Mouse.Position2D;
+	CurrentMouseState.Mouse.Position2D = FVector2D(InMouseX, InMouseY);
+	CurrentMouseState.Mouse.WorldRay = GetRayFromMousePos(InViewportClient, InViewport, InMouseX, InMouseY);
+
+	if (InputRouter->HasActiveMouseCapture())
+	{
+#ifdef ENABLE_DEBUG_PRINTING
+		UE_LOG(LogInteractiveToolsContext, Warning, TEXT("CAPTURED MOUSE MOVE"));
+#endif
+
+		FInputDeviceState InputState = CurrentMouseState;
+		InputState.InputDevice = EInputDevices::Mouse;
+		
+		FModifierKeysState ModifierKeys = FSlateApplication::Get().GetModifierKeys();
+		InputState.SetModifierKeyStates(
+			ModifierKeys.IsShiftDown(), ModifierKeys.IsAltDown(),
+			ModifierKeys.IsControlDown(), ModifierKeys.IsCommandDown());
+
+		InputState.Mouse.Delta2D = CurrentMouseState.Mouse.Position2D - OldPosition;
+		InputRouter->PostInputEvent(InputState);
+		return true;
+	}
+
+	return false;
+}
+
+bool UModeManagerInteractiveToolsContext::EndTracking(FEditorViewportClient* InViewportClient, FViewport* InViewport)
+{
+	if (bIsTrackingMouse)
+	{
+		// If the input router captured the mouse input, we need to invalidate the viewport client here, since the mouse delta tracker's end tracking will not be called.
+		constexpr bool bForceChildViewportRedraw = true;
+		constexpr bool bInvalidateHitProxies = true;
+		if (InViewportClient)
+		{
+			InViewportClient->Invalidate(bForceChildViewportRedraw, bInvalidateHitProxies);
+		}
+		bIsTrackingMouse = false;
+		return true;
+	}
+
+	return false;
+}
+
+bool UModeManagerInteractiveToolsContext::GetCursor(EMouseCursor::Type& OutCursor) const
+{	
+	UToolsContextCursorAPI* ToolsContextCursorAPI = ToolManager->GetContextObjectStore()->FindContext<UToolsContextCursorAPI>();
+
+	if (ToolsContextCursorAPI && ToolsContextCursorAPI->IsCursorOverridden())
+	{
+		OutCursor = ToolsContextCursorAPI->GetCurrentCursorOverride();
+		return true;
+	}
+
+	if (bIsTrackingMouse)
+	{
+		OutCursor = EMouseCursor::SlashedCircle;				
+		return true;
+	}
+
+	return false;
+}
+
+FRay UModeManagerInteractiveToolsContext::GetLastWorldRay() const
+{
+	return CurrentMouseState.Mouse.WorldRay;
+}
+
+void UModeManagerInteractiveToolsContext::Initialize(IToolsContextQueriesAPI* QueriesAPIIn, IToolsContextTransactionsAPI* TransactionsAPIIn)
+{
+	Super::Initialize(QueriesAPIIn, TransactionsAPIIn);
+
+	BeginPIEDelegateHandle = FEditorDelegates::BeginPIE.AddLambda([this](bool bSimulating)
+	{
+		TerminateActiveToolsOnPIEStart();
+	});
+	PreSaveWorldDelegateHandle = FEditorDelegates::PreSaveWorldWithContext.AddLambda([this](UWorld* World, FObjectPreSaveContext ObjectSaveContext)
+	{
+		TerminateActiveToolsOnSaveWorld();
+	});
+
+	FLevelEditorModule& LevelEditor = FModuleManager::LoadModuleChecked<FLevelEditorModule>("LevelEditor");
+	WorldTearDownDelegateHandle = LevelEditor.OnMapChanged().AddLambda([this](UWorld* World, EMapChangeType ChangeType)
+	{
+		if (ChangeType == EMapChangeType::TearDownWorld)
+		{
+			TerminateActiveToolsOnWorldTearDown();
+		}
+	});
+
+	// Tools frequently spawn temporary objects in the level for visualization, gizmos, etc, so having a level be
+	// removed from the world puts the tools at risk of letting those objects be garbage collected and cause crashes
+	FWorldDelegates::PreLevelRemovedFromWorld.AddWeakLambda(this, [this](ULevel*, UWorld*) {
+		TerminateActiveToolsOnLevelChange();
+	});
+
+	ToolManager->OnToolEnded.AddUObject(this, &UModeManagerInteractiveToolsContext::OnToolEnded);
+	ToolManager->OnToolPostBuild.AddUObject(this, &UModeManagerInteractiveToolsContext::OnToolPostBuild);
+
+	// if viewport clients change we will discard our overrides as we aren't sure what happened
+	ViewportClientListChangedHandle = GEditor->OnViewportClientListChanged().AddLambda([this]()
+	{
+		RestoreEditorState();
+	});
+}
+
+void UModeManagerInteractiveToolsContext::Shutdown()
+{
+	if (ToolManager)
+	{
+		ToolManager->OnToolPostBuild.RemoveAll(this);
+		ToolManager->OnToolEnded.RemoveAll(this);
+	}
+
+	if (FLevelEditorModule* LevelEditor = FModuleManager::GetModulePtr<FLevelEditorModule>("LevelEditor"))
+	{
+		LevelEditor->OnMapChanged().Remove(WorldTearDownDelegateHandle);
+		FEditorDelegates::BeginPIE.Remove(BeginPIEDelegateHandle);
+		FEditorDelegates::PreSaveWorldWithContext.Remove(PreSaveWorldDelegateHandle);
+		GEditor->OnViewportClientListChanged().Remove(ViewportClientListChangedHandle);
+	}
+
+	FWorldDelegates::PreLevelRemovedFromWorld.RemoveAll(this);
+
+	Super::Shutdown();
+}
+
+
+UEdModeInteractiveToolsContext* UModeManagerInteractiveToolsContext::CreateNewChildEdModeToolsContext()
+{
+	UEdModeInteractiveToolsContext* NewModeToolsContext = 
+		NewObject<UEdModeInteractiveToolsContext>(GetTransientPackage(), UEdModeInteractiveToolsContext::StaticClass(), NAME_None, RF_Transient);
+	NewModeToolsContext->InitializeContextFromModeManagerContext(this);
+
+	return NewModeToolsContext;
+}
+
+bool UModeManagerInteractiveToolsContext::OnChildEdModeActivated(UEdModeInteractiveToolsContext* ChildToolsContext)
+{
+	if (!ensureMsgf(EdModeToolsContexts.Find(ChildToolsContext), TEXT("Child ToolsContext was already found!")))
+	{
+		return false;
+	}
+
+	EdModeToolsContexts.Add(ChildToolsContext);
+	ChildToolsContext->Activate();
+	return true;
+}
+
+
+bool UModeManagerInteractiveToolsContext::OnChildEdModeDeactivated(UEdModeInteractiveToolsContext* ChildToolsContext)
+{
+	for (int32 k = 0; k < EdModeToolsContexts.Num(); ++k)
+	{
+		if (EdModeToolsContexts[k] == ChildToolsContext)
+		{
+			ChildToolsContext->DeactivateAllActiveTools(EToolShutdownType::Cancel);
+			ChildToolsContext->Deactivate();
+			EdModeToolsContexts.RemoveAt(k);
+			return true;
+		}
+	}
+	ensureMsgf(false, TEXT("Child ToolsContext was not found! It may have already been removed"));
+	return false;
+}
+
+void UModeManagerInteractiveToolsContext::SetBypassAltModifier(bool bInBypassAltModifier)
+{
+	bBypassAltModifier = bInBypassAltModifier;
+}
+
+void UEdModeInteractiveToolsContext::InitializeContextFromModeManagerContext(UModeManagerInteractiveToolsContext* ModeManagerToolsContext)
+{
+	check(ModeManagerToolsContext != nullptr);
+	check(ModeManagerToolsContext->InputRouter != nullptr);
+	FEditorModeTools* ModeManager = ModeManagerToolsContext->GetParentEditorModeManager();
+	check(ModeManager != nullptr);
+
+	ParentModeManagerToolsContext = ModeManagerToolsContext;
+
+
+	SetCreateContextStoreFunc([this, ModeManagerToolsContext](const FContextInitInfo& ContextInfo)
+	{
+		UContextObjectStore* NewContextStore = NewObject<UContextObjectStore>(ModeManagerToolsContext->ContextObjectStore);
+		return NewContextStore;
+	});
+
+	InitializeContextWithEditorModeManager(ModeManager, ModeManagerToolsContext->InputRouter);
+}
+
+
+FRay UEdModeInteractiveToolsContext::GetLastWorldRay() const
+{
+	return ParentModeManagerToolsContext->GetLastWorldRay();
+}
diff --git a/Engine/Source/Editor/UnrealEd/Public/Tools/EdModeInteractiveToolsContext.h b/Engine/Source/Editor/UnrealEd/Public/Tools/EdModeInteractiveToolsContext.h
new file mode 100644
--- /dev/null
+++ b/Engine/Source/Editor/UnrealEd/Public/Tools/EdModeInteractiveToolsContext.h
@@ -0,0 +1,388 @@
+// Copyright Epic Games, Inc. All Rights Reserved.
+
+#pragma once
+
+#include "Containers/Array.h"
+#include "Containers/Map.h"
+#include "Containers/UnrealString.h"
+#include "Delegates/Delegate.h"
+#include "Engine/EngineBaseTypes.h"
+#include "Engine/World.h"
+#include "HAL/Platform.h"
+#include "InputCoreTypes.h"
+#include "InputState.h"
+#include "InteractiveTool.h"
+#include "InteractiveToolManager.h"
+#include "InteractiveToolsContext.h"
+#include "Math/Ray.h"
+#include "Misc/AssertionMacros.h"
+#include "Misc/Optional.h"
+#include "Templates/SharedPointer.h"
+#include "UObject/ObjectMacros.h"
+#include "UObject/ObjectPtr.h"
+#include "UObject/UObjectGlobals.h"
+
+#include "EdModeInteractiveToolsContext.generated.h"
+
+class FCanvas;
+class FEdMode;
+class FEditorModeTools;
+class FEditorViewportClient;
+class FPrimitiveDrawInterface;
+class FSceneView;
+class FViewport;
+class FViewportClient;
+class ILevelEditor;
+class IToolsContextQueriesAPI;
+class IToolsContextRenderAPI;
+class IToolsContextTransactionsAPI;
+class UEdModeInteractiveToolsContext;
+class UGizmoViewContext;
+class UInputRouter;
+class UInteractiveToolBuilder;
+class UMaterialInterface;
+class UObject;
+class USelection;
+class UTypedElementSelectionSet;
+struct FToolBuilderState;
+
+/**
+ * UEditorInteractiveToolsContext is an extension/adapter of an InteractiveToolsContext designed 
+ * for use in the UE Editor. Currently this implementation assumes that it is created by a
+ * Mode Manager (FEditorModeTools), and that the Mode Manager will call various API functions
+ * like Render() and Tick() when necessary. 
+ * 
+ * 
+ * allows it to be easily embedded inside an FEdMode. A set of functions are provided which can be
+ * called from the FEdMode functions of the same name. These will handle the data type
+ * conversions and forwarding calls necessary to operate the ToolsContext
+ */
+UCLASS(Transient)
+class UNREALED_API UEditorInteractiveToolsContext : public UInteractiveToolsContext
+{
+	GENERATED_BODY()
+
+public:
+	UEditorInteractiveToolsContext();
+
+	/**
+	 * Initialize a new ToolsContext for an EdMode owned by the given InEditorModeManager
+	 * @param 
+	 */
+	void InitializeContextWithEditorModeManager(FEditorModeTools* InEditorModeManager, UInputRouter* UseInputRouter = nullptr);
+
+	/** Shutdown ToolsContext and clean up any connections/etc */
+	virtual void ShutdownContext();
+
+	// default behavior is to accept active tool
+	virtual void TerminateActiveToolsOnPIEStart();
+
+	// default behavior is to accept active tool
+	virtual void TerminateActiveToolsOnSaveWorld();
+
+	// default behavior is to cancel active tool
+	virtual void TerminateActiveToolsOnWorldTearDown();
+
+	// default behavior is to cancel active tool
+	virtual void TerminateActiveToolsOnLevelChange();
+
+	FEditorModeTools* GetParentEditorModeManager() const { return EditorModeManager; }
+
+	IToolsContextQueriesAPI* GetQueriesAPI() const { return QueriesAPI; }
+	IToolsContextTransactionsAPI* GetTransactionAPI() const { return TransactionAPI; }
+
+	/** Call this to notify the Editor that the viewports this ToolsContext is related to may need a repaint, ie during interactive tool usage */
+	virtual void PostInvalidation();
+
+	// UObject Interface
+	virtual UWorld* GetWorld() const override;
+
+	// call functions of the same name on the ToolManager and GizmoManager
+	virtual void Tick(FEditorViewportClient* ViewportClient, float DeltaTime);
+	virtual void Render(const FSceneView* View, FViewport* Viewport, FPrimitiveDrawInterface* PDI);
+	virtual void DrawHUD(FViewportClient* ViewportClient,FViewport* Viewport,const FSceneView* View, FCanvas* Canvas);
+
+	// These delegates can be used to hook into the Render() / DrawHUD() / Tick() calls above. In particular, non-legacy UEdMode's
+	// don't normally receive Render() and DrawHUD() calls from the mode manager, but can attach to these.
+	DECLARE_MULTICAST_DELEGATE_OneParam(FOnRender, IToolsContextRenderAPI* RenderAPI);
+	FOnRender OnRender;
+	DECLARE_MULTICAST_DELEGATE_TwoParams(FOnDrawHUD, FCanvas* Canvas, IToolsContextRenderAPI* RenderAPI);
+	FOnDrawHUD OnDrawHUD;
+	DECLARE_MULTICAST_DELEGATE_OneParam(FOnTick, float DeltaTime);
+	FOnTick OnTick;
+
+	/** @return true if selected actors/components can be deleted */
+	virtual bool ProcessEditDelete();
+
+	//
+	// Utility functions useful for hooking up to UICommand/etc
+	//
+
+	virtual bool CanStartTool(const FString ToolTypeIdentifier) const;
+	virtual bool HasActiveTool() const;
+	virtual FString GetActiveToolName() const;
+	virtual bool ActiveToolHasAccept() const;
+	virtual bool CanAcceptActiveTool() const;
+	virtual bool CanCancelActiveTool() const;
+	virtual bool CanCompleteActiveTool() const;
+	virtual void StartTool(const FString ToolTypeIdentifier);
+	virtual void EndTool(EToolShutdownType ShutdownType);
+	void Activate();
+	void Deactivate();
+
+
+	/** @return Ray into 3D scene at last mouse event */
+	virtual FRay GetLastWorldRay() const
+	{
+		check(false);
+		return FRay();
+	}
+
+	//
+	// Configuration functions
+	//
+
+	/*
+	 * Configure whether ::Render() should early-out for HitProxy rendering passes.
+	 * If the Mode does not use HitProxy, and the Tools/Gizmos have expensive Render() calls, this can help with interactive performance.
+	 */
+	void SetEnableRenderingDuringHitProxyPass(bool bEnabled);
+
+	/** @return true if HitProxy rendering will be allowed in ::Render() */
+	bool GetEnableRenderingDuringHitProxyPass() const { return bEnableRenderingDuringHitProxyPass; }
+
+
+	/**
+	 * Configure whether Transform Gizmos created by the ITF (eg CombinedTransformGizmo) should prefer to show in 'Combined' mode.
+	 * If this is disabled, the Gizmo should respect the active Editor Gizmo setting (eg in the Level Viewport)
+	 */
+	void SetForceCombinedGizmoMode(bool bEnabled);
+
+	/** @return true if Force Combined Gizmo mode is Enabled */
+	bool GetForceCombinedGizmoModeEnabled() const { return bForceCombinedGizmoMode; }
+
+
+	/**
+	 * Configure whether Transform Gizmos created by the ITF (eg CombinedTransformGizmo) should, when in World coordinate system,
+	 * snap to an Absolute world-aligned grid, or snap Relative to the initial position of any particular gizmo transform.
+	 * Relative is the default and is also the behavior of the standard UE Gizmo.
+	 */
+	void SetAbsoluteWorldSnappingEnabled(bool bEnabled);
+
+	/** @return true if Absolute World Snapping mode is Enabled */
+	bool GetAbsoluteWorldSnappingEnabled() const { return bEnableAbsoluteWorldSnapping; }
+
+	/*
+	 * Configure whether tools should shutdown when entering PIE.
+	 * By default, the context will shut down any active tools on PIE start.
+	 * Setting this to true will prevent this, but the system makes no promises about whether tools work properly across PIE start/run/shutdown.
+	 * Therefore, do not use this if you do not know for certain that it is safe to persist tools across PIE start in your situation.
+	 */
+	void SetDeactivateToolsOnPIEStart(bool bDeactivateTools);
+
+	/** @return true if tools should be deactivated when PIE is started */
+	bool GetDeactivateToolsOnPIEStart() const { return bDeactivateOnPIEStart; }
+
+protected:
+
+	// we hide these 
+	virtual void Initialize(IToolsContextQueriesAPI* QueriesAPI, IToolsContextTransactionsAPI* TransactionsAPI) override;
+	virtual void Shutdown() override;
+
+	virtual void DeactivateActiveTool(EToolSide WhichSide, EToolShutdownType ShutdownType);
+	virtual void DeactivateAllActiveTools(EToolShutdownType ShutdownType);
+
+public:
+	UPROPERTY()
+	TObjectPtr<UMaterialInterface> StandardVertexColorMaterial;
+
+protected:
+	// EdMode implementation of InteractiveToolFramework APIs - see ToolContextInterfaces.h
+	IToolsContextQueriesAPI* QueriesAPI;
+	IToolsContextTransactionsAPI* TransactionAPI;
+
+	// Tools need to be able to Invalidate the view, in case it is not Realtime.
+	// Currently we do this very aggressively, and also force Realtime to be on, but in general we should be able to rely on Invalidation.
+	// However there are multiple Views and we do not want to Invalidate immediately, so we store a timestamp for each
+	// ViewportClient, and invalidate it when we see it if it's timestamp is out-of-date.
+	// (In theory this map will continually grow as new Viewports are created...)
+	TMap<FViewportClient*, int32> InvalidationMap;
+	// current invalidation timestamp, incremented by invalidation calls
+	int32 InvalidationTimestamp = 0;
+
+	// An object in which we save the current scene view information that gizmos can use on the game thread
+	// to figure out how big the gizmo is for hit testing. Lives in the context store, but we keep a pointer here
+	// to avoid having to look for it.
+	UGizmoViewContext* GizmoViewContext = nullptr;
+
+	// Utility function to convert viewport x/y from mouse events (and others?) into scene ray.
+	// Copy-pasted from other Editor code, seems kind of expensive?
+	static FRay GetRayFromMousePos(FEditorViewportClient* ViewportClient, FViewport* Viewport, int MouseX, int MouseY);
+
+	// editor UI state that we set before starting tool and when exiting tool
+	// Currently disabling anti-aliasing during active Tools because it causes PDI flickering
+	void SetEditorStateForTool();
+	void RestoreEditorState();
+
+	void OnToolEnded(UInteractiveToolManager* InToolManager, UInteractiveTool* InEndedTool);
+	void OnToolPostBuild(UInteractiveToolManager* InToolManager, EToolSide InSide, UInteractiveTool* InBuiltTool, UInteractiveToolBuilder* InToolBuilder, const FToolBuilderState& ToolState);
+
+	TOptional<FString> PendingToolToStart = {};
+	TOptional<EToolShutdownType> PendingToolShutdownType = {};
+
+private:
+	FEditorModeTools* EditorModeManager = nullptr;
+
+	// currently defaulting to enabled as FEdModes generally assume this, and in most cases hitproxy pass is not expensive.
+	bool bEnableRenderingDuringHitProxyPass = true;
+
+	bool bForceCombinedGizmoMode = false;
+	bool bEnableAbsoluteWorldSnapping = false;
+
+	bool bDeactivateOnPIEStart = true;
+
+	bool bIsActive = false;
+};
+
+
+/**
+ * UModeManagerInteractiveToolsContext extends UEditorInteractiveToolsContext with various functions for handling 
+ * device (mouse) input. These functions are currently called by the EdMode Manager (FEditorModeTools).
+ */
+UCLASS(Transient)
+class UNREALED_API UModeManagerInteractiveToolsContext : public UEditorInteractiveToolsContext
+{
+	GENERATED_BODY()
+
+	//
+	// UEditorInteractiveToolsContext API implementations that also forward calls to any child EdMode ToolsContexts
+	//
+public:
+	virtual void Tick(FEditorViewportClient* ViewportClient, float DeltaTime) override;
+	virtual void Render(const FSceneView* View, FViewport* Viewport, FPrimitiveDrawInterface* PDI) override;
+	virtual void DrawHUD(FViewportClient* ViewportClient,FViewport* Viewport,const FSceneView* View, FCanvas* Canvas) override;
+
+	virtual bool ProcessEditDelete() override;
+
+	//
+	// Input handling, these functions forward ViewportClient events to the UInputRouter
+	//
+public:
+	bool InputKey(FEditorViewportClient* ViewportClient, FViewport* Viewport, FKey Key, EInputEvent Event);
+	/**
+	 * This updates internal state like InputKey, but doesn't route the results to the input router. 
+	 * Use this if the input is captured by some higher system, to avoid this class from having an
+	 * incorrect view of e.g. the mouse state because it did not receive a mouse release event.
+	 */
+	void UpdateStateWithoutRoutingInputKey(FEditorViewportClient* ViewportClient, FViewport* Viewport, FKey Key, EInputEvent Event);
+
+	bool MouseEnter(FEditorViewportClient* ViewportClient, FViewport* Viewport, int32 x, int32 y);
+	bool MouseLeave(FEditorViewportClient* ViewportClient, FViewport* Viewport);
+	bool MouseMove(FEditorViewportClient* ViewportClient, FViewport* Viewport, int32 x, int32 y);
+
+	bool StartTracking(FEditorViewportClient* InViewportClient, FViewport* InViewport);
+	bool CapturedMouseMove(FEditorViewportClient* InViewportClient, FViewport* InViewport, int32 InMouseX, int32 InMouseY);
+	bool EndTracking(FEditorViewportClient* InViewportClient, FViewport* InViewport);
+
+	/** @return True if the context has overriden the cursor style, false if not. */
+	bool GetCursor(EMouseCursor::Type& OutCursor) const;
+
+	/** @return Ray into 3D scene at last mouse event */
+	virtual FRay GetLastWorldRay() const override;
+
+protected:
+	virtual void DeactivateAllActiveTools(EToolShutdownType ShutdownType) override;
+	virtual void Initialize(IToolsContextQueriesAPI* QueriesAPIIn, IToolsContextTransactionsAPI* TransactionsAPIIn) override;
+	virtual void Shutdown() override;
+
+	/** Input event instance used to keep track of various button states, etc, that we cannot directly query on-demand */
+	FInputDeviceState CurrentMouseState;
+	// called when PIE is about to start, shuts down active tools
+	FDelegateHandle BeginPIEDelegateHandle;
+	// called before a Save starts. This currently shuts down active tools.
+	FDelegateHandle PreSaveWorldDelegateHandle;
+	// called when a map is changed
+	FDelegateHandle WorldTearDownDelegateHandle;
+	// called when viewport clients change
+	FDelegateHandle ViewportClientListChangedHandle;
+
+private:
+	bool bIsTrackingMouse;
+
+
+
+
+protected:
+	UPROPERTY()
+	TArray<TObjectPtr<UEdModeInteractiveToolsContext>> EdModeToolsContexts;
+
+public:
+	/** 
+	* Create and initialize a new EdMode-level ToolsContext derived from the ModeManager ToolsContext.
+	* The EdMode ToolsContext does not have it's own InputRouter, it shares the InputRouter with the ModeManager ToolsContext.
+	* The ModeManager ToolsContext keeps track of these derived ToolsContext's and automatically Tick()'s them/etc.
+	* When the child ToolsContext is shut down, OnChildEdModeToolsContextShutdown() must be called to clean up
+	* @return new ToolsContext
+	*/
+	UEdModeInteractiveToolsContext* CreateNewChildEdModeToolsContext();
+
+	/**
+	 * Call to add a child EdMode ToolsContext created using the above function
+	 * @return true if child was added
+	 */
+	bool OnChildEdModeActivated(UEdModeInteractiveToolsContext* ChildToolsContext);
+
+	/**
+	 * Call to release a child EdMode ToolsContext created using the above function
+	 * @return true if child was found and removed
+	 */
+	bool OnChildEdModeDeactivated(UEdModeInteractiveToolsContext* ChildToolsContext);
+
+protected:
+	//~Begin ITF Alt processing Bypass
+
+	/**
+	 * These functions are declared protected to avoid having to deal with deprecation later on.
+	 * The Alt modifier bypass will be removed once a complete switch to ITF happens
+	 */
+
+	// Needs to use Alt modifier for Drag Tools to work
+	friend class FEditorViewportClient;
+
+	/** Enable/Disable Alt modifier bypass */
+	void SetBypassAltModifier(bool bInBypassAltModifier);
+
+	/** Is this Interactive Tools Context skipping Alt processing? */
+	bool BypassAltModifier() const { return bBypassAltModifier; }
+
+	/** Should we skip processing Alt modifier inputs? See UModeManagerInteractiveToolsContext::InputKey */
+	bool bBypassAltModifier = true;
+
+	//~End ITF Alt processing Bypass
+};
+
+
+/**
+ * UEdModeInteractiveToolsContext is an UEditorInteractiveToolsContext intended for use/lifetime in the context of a UEdMode.
+ * This ITC subclass is dependent on a UModeManagerInteractiveToolsContext to provide an InputRouter.
+ */
+UCLASS(Transient)
+class UNREALED_API UEdModeInteractiveToolsContext : public UEditorInteractiveToolsContext
+{
+	friend class UModeManagerInteractiveToolsContext;
+
+	GENERATED_BODY()
+public:
+	/**
+	 * Initialize a new EdModeToolsContext that is derived from a ModeManagerToolsContext.
+	 * This new ToolsContext will not have it's own InputRouter, it will share the InputRouter with the ModeManagerToolsContext
+	 */
+	void InitializeContextFromModeManagerContext(UModeManagerInteractiveToolsContext* ModeManagerToolsContext);
+
+	/** @return Ray into 3D scene at last mouse event */
+	virtual FRay GetLastWorldRay() const override;
+
+protected:
+	UPROPERTY()
+	TObjectPtr<UModeManagerInteractiveToolsContext> ParentModeManagerToolsContext = nullptr;
+};
