Make objects in external packages considered not browsable  by default, since they are never shown in the content browser.
This fixes thumbnails creation for DataLayerInstance, ActorFolder, AnimNextRigVMAssetEntry, etc.

diff --git a/Engine/Source/Editor/UnrealEd/Private/ObjectTools.cpp b/Engine/Source/Editor/UnrealEd/Private/ObjectTools.cpp
--- a/Engine/Source/Editor/UnrealEd/Private/ObjectTools.cpp
+++ b/Engine/Source/Editor/UnrealEd/Private/ObjectTools.cpp
@@ -236,7 +236,7 @@ namespace ObjectTools
 				if( ObjectPackage != GetTransientPackage()
 					&& (ObjectPackage->HasAnyPackageFlags(PKG_PlayInEditor) == false)
 					&& IsValidChecked(Obj) 
-					&& !Obj->IsA<AActor>())
+					&& !Obj->IsPackageExternal())
 				{
 					bIsSupported = true;
 				}
