[Viewport Axes Gizmo] Update to add LUF rendering functionality

diff --git a/Engine/Source/Editor/UnrealEd/Private/EditorViewportClient.cpp b/Engine/Source/Editor/UnrealEd/Private/EditorViewportClient.cpp
--- a/Engine/Source/Editor/UnrealEd/Private/EditorViewportClient.cpp
+++ b/Engine/Source/Editor/UnrealEd/Private/EditorViewportClient.cpp
@@ -81,6 +81,7 @@
 #include "GameFramework/ActorPrimitiveColorHandler.h"
 #include "LandscapeRender.h"
 #include "Tools/DragToolsBehaviourSource.h"
+#include "Misc/AxisDisplayInfo.h"
 
 #define LOCTEXT_NAMESPACE "EditorViewportClient"
 
@@ -759,10 +760,12 @@ void FEditorViewportClient::SetInitialViewTransform(ELevelViewportType InViewpor
 {
 	check(InViewportType < LVT_MAX);
 
-	FViewportCameraTransform& ViewTransform = (InViewportType == LVT_Perspective) ? ViewTransformPerspective : ViewTransformOrthographic;
+	const bool bUsingLUFCoordinateSysem = AxisDisplayInfo::GetAxisDisplayCoordinateSystem() == EAxisList::LeftUpForward;
 
-	ViewTransform.SetLocation(ViewLocation);
-	ViewTransform.SetRotation(ViewRotation);
+	FViewportCameraTransform& ViewTransform = (InViewportType == LVT_Perspective) ? ViewTransformPerspective : ViewTransformOrthographic;
+	// If using the LUF coord system, invert Y position and rotate yaw by -90* so the gizmo is centered since the Left axis draw is inverse to the Y
+	ViewTransform.SetLocation(FVector(ViewLocation.X, (bUsingLUFCoordinateSysem ? -ViewLocation.Y : ViewLocation.Y), ViewLocation.Z));
+	ViewTransform.SetRotation(ViewRotation + (bUsingLUFCoordinateSysem ? FRotator(0, -90, 0) : FRotator(0, 0, 0)));
 
 	// Make a look at location in front of the camera
 	const FQuat CameraOrientation = FQuat::MakeFromEuler(ViewRotation.Euler());
@@ -3608,13 +3611,13 @@ void FEditorViewportClient::DrawAxes(FViewport* InViewport, FCanvas* Canvas, con
 	FVector TextVec;
 	FIntPoint AxisEnd;
 	FIntPoint TextCenter;
-	float TextDistance = 12.0f;
+	const float TextDistance = 12.0f;
 	FCanvasLineItem LineItem;
 	LineItem.LineThickness = 1.0f;
 	FCanvasTextItem TextItem( FVector2D::ZeroVector, FText::GetEmpty(), Font, FLinearColor::White );
-	FVector2D InsideCenterSize = FVector2D(2, 2);
+	const FVector2D InsideCenterSize = FVector2D(2, 2);
 	FCanvasTileItem InsideTileItem(FVector2D::ZeroVector, InsideCenterSize, FLinearColor(0.2f, 0.2f, 0.2f));
-	FVector2D OutsideCenterSize = FVector2D(4, 4);
+	const FVector2D OutsideCenterSize = FVector2D(4, 4);
 	FCanvasTileItem OutsideTileItem(FVector2D::ZeroVector, OutsideCenterSize, FLinearColor(0.05f, 0.05f, 0.05f));
 	
 	auto GetShadowColor = [](const FLinearColor& Color) -> FLinearColor
@@ -3624,12 +3627,13 @@ void FEditorViewportClient::DrawAxes(FViewport* InViewport, FCanvas* Canvas, con
 			return ColorHSV.HSVToLinearRGB();
 		};
 
-	FLinearColor XBaseColor = AxisDisplayInfo::GetAxisColor(EAxisList::X);
-	FLinearColor XShadowColor = GetShadowColor(XBaseColor);
-	FLinearColor YBaseColor = AxisDisplayInfo::GetAxisColor(EAxisList::Y);
-	FLinearColor YShadowColor = GetShadowColor(YBaseColor);
-	FLinearColor ZBaseColor = AxisDisplayInfo::GetAxisColor(EAxisList::Z);
-	FLinearColor ZShadowColor = GetShadowColor(ZBaseColor);
+	// If LUF isn't being used, it will automatically translate to its correlating XYZ value
+	const FLinearColor XBaseColor = AxisDisplayInfo::GetAxisColor(EAxisList::Forward);
+	const FLinearColor XShadowColor = GetShadowColor(XBaseColor);
+	const FLinearColor YBaseColor = AxisDisplayInfo::GetAxisColor(EAxisList::Left);
+	const FLinearColor YShadowColor = GetShadowColor(YBaseColor);
+	const FLinearColor ZBaseColor = AxisDisplayInfo::GetAxisColor(EAxisList::Up);
+	const FLinearColor ZShadowColor = GetShadowColor(ZBaseColor);
 
 	auto GetShadowOffset = [this](FIntPoint EndPoint, FIntPoint OriginPoint) -> FIntPoint
 		{
@@ -3637,61 +3641,69 @@ void FEditorViewportClient::DrawAxes(FViewport* InViewport, FCanvas* Canvas, con
 			int32 Y = (FMath::Abs(EndPoint.X - OriginPoint.X) > FMath::Abs(EndPoint.Y - OriginPoint.Y)) ? 1 : 0;
 			return FIntPoint(X, Y);
 		};
-		
-	if( ( InAxis & EAxisList::X ) ==  EAxisList::X)
+
+	if ((InAxis & EAxisList::X) == EAxisList::X || (InAxis & EAxisList::Forward) == EAxisList::Forward)
 	{
-		AxisVec = AxisSize * ViewTM.InverseTransformVector( FVector(1,0,0) );
-		AxisEnd = AxisOrigin + FIntPoint( AxisVec.Y, -AxisVec.Z );
+		AxisVec = AxisSize * ViewTM.InverseTransformVector(FVector(1, 0, 0));
+		AxisEnd = AxisOrigin + FIntPoint(AxisVec.Y, -AxisVec.Z);
 		TextVec = (AxisSize + TextDistance) * ViewTM.InverseTransformVector(FVector(1, 0, 0));
 		TextCenter = AxisOrigin + FIntPoint(TextVec.Y, -TextVec.Z);
 		LineItem.SetColor(XShadowColor);
 		LineItem.Draw(Canvas, AxisOrigin + GetShadowOffset(AxisEnd, AxisOrigin), AxisEnd + GetShadowOffset(AxisEnd, AxisOrigin));
 		LineItem.SetColor(XBaseColor);
-		LineItem.Draw( Canvas, AxisOrigin, AxisEnd );
-		TextItem.Text = AxisDisplayInfo::GetAxisDisplayNameShort(EAxisList::X);
+		LineItem.Draw(Canvas, AxisOrigin, AxisEnd);
+		TextItem.Text = AxisDisplayInfo::GetAxisDisplayNameShort(EAxisList::Forward);
 		TextItem.SetColor(XBaseColor);
 		TextItem.EnableShadow(XShadowColor);
-		TextItem.Draw( Canvas, FVector2D(TextCenter.X - 0.5*XL, TextCenter.Y - 0.5*YL) );
+		TextItem.Draw(Canvas, FVector2D(TextCenter.X - 0.5 * XL, TextCenter.Y - 0.5 * YL));
 	}
 
-	if( ( InAxis & EAxisList::Y ) == EAxisList::Y)
+	const bool bIsLeft = (InAxis & EAxisList::Left) == EAxisList::Left;
+	if ((InAxis & EAxisList::Y) == EAxisList::Y || bIsLeft)
 	{
-		AxisVec = AxisSize * ViewTM.InverseTransformVector( FVector(0,1,0) );
-		AxisEnd = AxisOrigin + FIntPoint( AxisVec.Y, -AxisVec.Z );
-		TextVec = (AxisSize + TextDistance) * ViewTM.InverseTransformVector(FVector(0, 1, 0));
+		if (bIsLeft)
+		{
+			AxisVec = AxisSize * ViewTM.InverseTransformVector(FVector(0, -1, 0));
+			TextVec = (AxisSize + TextDistance) * ViewTM.InverseTransformVector(FVector(0, -1, 0));
+		}
+		else
+		{
+			AxisVec = AxisSize * ViewTM.InverseTransformVector(FVector(0, 1, 0));
+			TextVec = (AxisSize + TextDistance) * ViewTM.InverseTransformVector(FVector(0, 1, 0));
+		}
+
+		AxisEnd = AxisOrigin + FIntPoint(AxisVec.Y, -AxisVec.Z);
 		TextCenter = AxisOrigin + FIntPoint(TextVec.Y, -TextVec.Z);
 		LineItem.SetColor(YShadowColor);
 		LineItem.Draw(Canvas, AxisOrigin + GetShadowOffset(AxisEnd, AxisOrigin), AxisEnd + GetShadowOffset(AxisEnd, AxisOrigin));
 		LineItem.SetColor(YBaseColor);
-		LineItem.Draw( Canvas, AxisOrigin, AxisEnd );
-		TextItem.Text = AxisDisplayInfo::GetAxisDisplayNameShort(EAxisList::Y);
+		LineItem.Draw(Canvas, AxisOrigin, AxisEnd);
+		TextItem.Text = AxisDisplayInfo::GetAxisDisplayNameShort(EAxisList::Left);
 		TextItem.SetColor(YBaseColor);
 		TextItem.EnableShadow(YShadowColor);
-		TextItem.Draw( Canvas, FVector2D(TextCenter.X - 0.5*XL, TextCenter.Y - 0.5*YL) );
-
+		TextItem.Draw(Canvas, FVector2D(TextCenter.X - 0.5 * XL, TextCenter.Y - 0.5 * YL));
 	}
 
-	if( ( InAxis & EAxisList::Z ) == EAxisList::Z )
+	if ((InAxis & EAxisList::Z) == EAxisList::Z || (InAxis & EAxisList::Up) == EAxisList::Up)
 	{
-		AxisVec = AxisSize * ViewTM.InverseTransformVector( FVector(0,0,1) );
-		AxisEnd = AxisOrigin + FIntPoint( AxisVec.Y, -AxisVec.Z );
+		AxisVec = AxisSize * ViewTM.InverseTransformVector(FVector(0, 0, 1));
+		AxisEnd = AxisOrigin + FIntPoint(AxisVec.Y, -AxisVec.Z);
 		TextVec = (AxisSize + TextDistance) * ViewTM.InverseTransformVector(FVector(0, 0, 1));
 		TextCenter = AxisOrigin + FIntPoint(TextVec.Y, -TextVec.Z);
 		LineItem.SetColor(ZShadowColor);
 		LineItem.Draw(Canvas, AxisOrigin + GetShadowOffset(AxisEnd, AxisOrigin), AxisEnd + GetShadowOffset(AxisEnd, AxisOrigin));
 		LineItem.SetColor(ZBaseColor);
-		LineItem.Draw( Canvas, AxisOrigin, AxisEnd );
-		TextItem.Text = AxisDisplayInfo::GetAxisDisplayNameShort(EAxisList::Z);
+		LineItem.Draw(Canvas, AxisOrigin, AxisEnd);
+		TextItem.Text = AxisDisplayInfo::GetAxisDisplayNameShort(EAxisList::Up);
 		TextItem.SetColor(ZBaseColor);
 		TextItem.EnableShadow(ZShadowColor);
-		TextItem.Draw( Canvas, FVector2D(TextCenter.X - 0.5*XL, TextCenter.Y - 0.5*YL) );
+		TextItem.Draw(Canvas, FVector2D(TextCenter.X - 0.5 * XL, TextCenter.Y - 0.5 * YL));
 	}
 
 	OutsideTileItem.Draw(Canvas, FIntPoint(AxisOrigin.X - FMath::Floor(OutsideCenterSize.X * 0.5), AxisOrigin.Y - FMath::Floor(OutsideCenterSize.Y * 0.5)));
 	InsideTileItem.Draw(Canvas, FIntPoint(AxisOrigin.X - FMath::Floor(InsideCenterSize.X * 0.5), AxisOrigin.Y - FMath::Floor(InsideCenterSize.Y * 0.5) ));
 }
 
-
 /** Convert the specified number (in cm or unreal units) into a readable string with relevant si units */
 FString FEditorViewportClient::UnrealUnitsToSiUnits(float UnrealUnits)
 {
@@ -4448,76 +4460,55 @@ void FEditorViewportClient::Draw(FViewport* InViewport, FCanvas* Canvas)
 	// Axes indicators
 	const bool bShouldDrawAxes = (bDrawAxes && !ViewFamily.EngineShowFlags.Game) || (bDrawAxesGame && ViewFamily.EngineShowFlags.Game);
 	const bool bIsViewportUIHidden = ModeTools && ModeTools->IsViewportUIHidden();
+	const EAxisList::Type AxisDisplayCoordinateSystem = AxisDisplayInfo::GetAxisDisplayCoordinateSystem();
+
+	auto DrawAxesInOrthoViewport = [this, AxisDisplayCoordinateSystem, Canvas, View](FRotator Rotator, EAxisList::Type AxisType)
+		{
+			DrawAxes(Viewport, Canvas, &Rotator, AxisType);
+
+			if (View)
+			{
+				DrawScaleUnits(Viewport, Canvas, *View);
+			}
+		};
+
 	if (bShouldDrawAxes && !bIsViewportUIHidden && !IsVisualizeCalibrationMaterialEnabled())
 	{
 		switch (GetViewportType())
 		{
 		case LVT_OrthoXY:
 			{
-				const FRotator XYRot = []()
-				{
-					return FRotator(-90.0f, -90.0f, -90.0f);
-				}();
-				DrawAxes(Viewport, Canvas, &XYRot, EAxisList::XY);
-				if (View)
-				{
-					DrawScaleUnits(Viewport, Canvas, *View);
-				}
+				DrawAxesInOrthoViewport(FRotator(-90.0f, -90.0f, -90.0f), EAxisList::LF);
 				break;
 			}
 		case LVT_OrthoXZ:
 			{
-				const FRotator XZRot(0.0f, -90.0f, 0.0f);
-				DrawAxes(Viewport, Canvas, &XZRot, EAxisList::XZ);
-				if (View)
-				{
-					DrawScaleUnits(Viewport, Canvas, *View);
-				}
+				DrawAxesInOrthoViewport(FRotator(0.0f, -90.0f, 0.0f), EAxisList::UF);
 				break;
 			}
 		case LVT_OrthoYZ:
 			{
-				const FRotator YZRot(0.0f, 0.0f, 0.0f);
-				DrawAxes(Viewport, Canvas, &YZRot, EAxisList::YZ);
-				if (View)
-				{
-					DrawScaleUnits(Viewport, Canvas, *View);
-				}
+				DrawAxesInOrthoViewport(FRotator(0.0f, 0.0f, 0.0f), EAxisList::LU);
 				break;
 			}
 		case LVT_OrthoNegativeXY:
 			{
-				const FRotator XYRot = FRotator(90.0f, 90.0f, 90.0f);
-				DrawAxes(Viewport, Canvas, &XYRot, EAxisList::XY);
-				if (View)
-				{
-					DrawScaleUnits(Viewport, Canvas, *View);
-				}
+				DrawAxesInOrthoViewport(FRotator(90.0f, 90.0f, 90.0f), EAxisList::LF);
 				break;
 			}
 		case LVT_OrthoNegativeXZ:
 			{
-				const FRotator XZRot(0.0f, 90.0f, 0.0f);
-				DrawAxes(Viewport, Canvas, &XZRot, EAxisList::XZ);
-				if (View)
-				{
-					DrawScaleUnits(Viewport, Canvas, *View);
-				}
+				DrawAxesInOrthoViewport(FRotator(0.0f, 90.0f, 0.0f), EAxisList::UF);
 				break;
 			}
 		case LVT_OrthoNegativeYZ:
 			{
-				const FRotator YZRot(0.0f, 180.0f, 0.0f);
-				DrawAxes(Viewport, Canvas, &YZRot, EAxisList::YZ);
-				if (View)
-				{
-					DrawScaleUnits(Viewport, Canvas, *View);
-				}
+				DrawAxesInOrthoViewport(FRotator(0.0f, 180.0f, 0.0f), EAxisList::LU);
 				break;
 			}
 		default:
 			{
-				DrawAxes(Viewport, Canvas);
+				DrawAxes(Viewport, Canvas, nullptr, AxisDisplayCoordinateSystem);
 				break;
 			}
 		}
diff --git a/Engine/Source/Editor/UnrealEd/Public/EditorViewportClient.h b/Engine/Source/Editor/UnrealEd/Public/EditorViewportClient.h
--- a/Engine/Source/Editor/UnrealEd/Public/EditorViewportClient.h
+++ b/Engine/Source/Editor/UnrealEd/Public/EditorViewportClient.h
@@ -845,7 +845,7 @@ public:
 	 *
 	 * Returns true if preview actors were updated
 	 */
-	virtual bool UpdateDropPreviewActors(int32 MouseX, int32 MouseY, const TArray<UObject*>& DroppedObjects, bool& out_bDroppedObjectsVisible, UActorFactory* FactoryToUse = NULL) { return false; }
+	virtual bool UpdateDropPreviewActors(int32 MouseX, int32 MouseY, const TArray<UObject*>& DroppedObjects, bool& out_bDroppedObjectsVisible, UActorFactory* FactoryToUse = nullptr) { return false; }
 
 	//~ TODO: UE_DEPRECATED(5.4,"Use DestroyDropPreviewElements instead.")
 	/**
@@ -920,7 +920,7 @@ public:
 	UE_DEPRECATED(5.4, "Use the overload that takes FDropObjectOptions instead.")
 	UNREALED_API virtual bool DropObjectsAtCoordinates(int32 MouseX, int32 MouseY, const TArray<UObject*>& DroppedObjects, TArray<AActor*>& OutNewActors,
 		bool bOnlyDropOnTarget = false, bool bCreateDropPreview = false, bool bSelectActors = true,
-		UActorFactory* FactoryToUse = NULL) { return false; }
+		UActorFactory* FactoryToUse = nullptr) { return false; }
 
 	/** Returns true if the viewport is allowed to be possessed for previewing cinematic sequences or keyframe animations*/
 	bool AllowsCinematicControl() const { return bAllowCinematicControl; }
@@ -1625,7 +1625,7 @@ protected:
 	 * @param	InRotation		Specifies the rotation to apply to the axes
 	 * @param	InAxis			Specifies which axes to draw
 	 */
-	UNREALED_API void DrawAxes(FViewport* Viewport,FCanvas* Canvas, const FRotator* InRotation = NULL, EAxisList::Type InAxis = EAxisList::XYZ);
+	UNREALED_API void DrawAxes(FViewport* Viewport,FCanvas* Canvas, const FRotator* InRotation = nullptr, EAxisList::Type InAxis = EAxisList::XYZ);
 
 	/**
 	 * Draws viewport scale units
diff --git a/Engine/Source/Runtime/Core/Public/Math/Axis.h b/Engine/Source/Runtime/Core/Public/Math/Axis.h
--- a/Engine/Source/Runtime/Core/Public/Math/Axis.h
+++ b/Engine/Source/Runtime/Core/Public/Math/Axis.h
@@ -43,6 +43,10 @@ namespace EAxisList
 		Left			= 1 << 4,
 		Up				= 1 << 5,
 		Forward			= 1 << 6,
+
+		LU = Left | Up,
+		LF = Left | Forward,
+		UF = Up | Forward,
 		LeftUpForward	= Left | Up | Forward,
 	};
 }
