[Viewport Transform Gizmos] Update to flip the Y axis rendering to reflect the new Left direction for the transform and scale gizmos

diff --git a/Engine/Source/Editor/UnrealEd/Private/UnrealWidget.cpp b/Engine/Source/Editor/UnrealEd/Private/UnrealWidget.cpp
--- a/Engine/Source/Editor/UnrealEd/Private/UnrealWidget.cpp
+++ b/Engine/Source/Editor/UnrealEd/Private/UnrealWidget.cpp
@@ -29,13 +29,13 @@ constexpr float FWidget::AXIS_LENGTH_SCALE_OFFSET;
 
 FWidget::FWidget()
 {
-	EditorModeTools      = NULL;
+	EditorModeTools      = nullptr;
 	TotalDeltaRotation   = 0;
 	CurrentDeltaRotation = 0;
 
-	AxisColorX       = AxisDisplayInfo::GetAxisColor(EAxisList::X);
-	AxisColorY       = AxisDisplayInfo::GetAxisColor(EAxisList::Y);
-	AxisColorZ       = AxisDisplayInfo::GetAxisColor(EAxisList::Z);
+	AxisColorX       = AxisDisplayInfo::GetAxisColor(EAxisList::Forward);
+	AxisColorY       = AxisDisplayInfo::GetAxisColor(EAxisList::Left);
+	AxisColorZ       = AxisDisplayInfo::GetAxisColor(EAxisList::Up);
 	ScreenAxisColor  = FLinearColor(0.76, 0.72, 0.14f);
 	PlaneColorXY     = FColor::Yellow;
 	ArcBallColor     = FColor(128, 128, 128, 6);
@@ -44,29 +44,29 @@ FWidget::FWidget()
 
 	UMaterial* AxisMaterialBase = GEngine->ArrowMaterial;
 
-	AxisMaterialX = UMaterialInstanceDynamic::Create(AxisMaterialBase, NULL);
+	AxisMaterialX = UMaterialInstanceDynamic::Create(AxisMaterialBase, nullptr);
 	AxisMaterialX->SetVectorParameterValue("GizmoColor", AxisColorX);
 
-	AxisMaterialY = UMaterialInstanceDynamic::Create(AxisMaterialBase, NULL);
+	AxisMaterialY = UMaterialInstanceDynamic::Create(AxisMaterialBase, nullptr);
 	AxisMaterialY->SetVectorParameterValue("GizmoColor", AxisColorY);
 
-	AxisMaterialZ = UMaterialInstanceDynamic::Create(AxisMaterialBase, NULL);
+	AxisMaterialZ = UMaterialInstanceDynamic::Create(AxisMaterialBase, nullptr);
 	AxisMaterialZ->SetVectorParameterValue("GizmoColor", AxisColorZ);
 
-	CurrentAxisMaterial = UMaterialInstanceDynamic::Create(AxisMaterialBase, NULL);
+	CurrentAxisMaterial = UMaterialInstanceDynamic::Create(AxisMaterialBase, nullptr);
 	CurrentAxisMaterial->SetVectorParameterValue("GizmoColor", CurrentColor);
 
-	OpaquePlaneMaterialXY = UMaterialInstanceDynamic::Create(AxisMaterialBase, NULL);
+	OpaquePlaneMaterialXY = UMaterialInstanceDynamic::Create(AxisMaterialBase, nullptr);
 	OpaquePlaneMaterialXY->SetVectorParameterValue("GizmoColor", FLinearColor::White);
 
 	TransparentPlaneMaterialXY = (UMaterial*)StaticLoadObject(
-	    UMaterial::StaticClass(), NULL,
-	    TEXT("/Engine/EditorMaterials/WidgetVertexColorMaterial.WidgetVertexColorMaterial"), NULL, LOAD_None, NULL);
+	    UMaterial::StaticClass(), nullptr,
+	    TEXT("/Engine/EditorMaterials/WidgetVertexColorMaterial.WidgetVertexColorMaterial"), nullptr, LOAD_None, nullptr);
 
 	GridMaterial = (UMaterial*)StaticLoadObject(
-	    UMaterial::StaticClass(), NULL,
-	    TEXT("/Engine/EditorMaterials/WidgetGridVertexColorMaterial_Ma.WidgetGridVertexColorMaterial_Ma"), NULL,
-	    LOAD_None, NULL);
+	    UMaterial::StaticClass(), nullptr,
+	    TEXT("/Engine/EditorMaterials/WidgetGridVertexColorMaterial_Ma.WidgetGridVertexColorMaterial_Ma"), nullptr,
+	    LOAD_None, nullptr);
 	if (!GridMaterial)
 	{
 		GridMaterial = TransparentPlaneMaterialXY;
diff --git a/Engine/Source/Editor/UnrealEd/Private/UnrealWidgetRender.cpp b/Engine/Source/Editor/UnrealEd/Private/UnrealWidgetRender.cpp
--- a/Engine/Source/Editor/UnrealEd/Private/UnrealWidgetRender.cpp
+++ b/Engine/Source/Editor/UnrealEd/Private/UnrealWidgetRender.cpp
@@ -14,6 +14,7 @@
 #include "DynamicMeshBuilder.h"
 #include "CanvasTypes.h"
 #include "EditorInteractiveGizmoManager.h"
+#include "Misc/AxisDisplayInfo.h"
 
 /*
  *  Simple struct used to create and group data related to the current window's / viewport's space,
@@ -259,13 +260,14 @@ void FWidget::Render( const FSceneView* View,FPrimitiveDrawInterface* PDI, FEdit
 void FWidget::Render_Axis(const FSceneView* View, FPrimitiveDrawInterface* PDI, EAxisList::Type InAxis, FMatrix& InMatrix, UMaterialInterface* InMaterial, const FLinearColor& InColor, FVector2D& OutAxisDir, const FVector& InScale, bool bDrawWidget, bool bCubeHead, float AxisLengthOffset)
 {
 	FMatrix AxisRotation = FMatrix::Identity;
+	const bool bUsingLUFCoordinateSysem = AxisDisplayInfo::GetAxisDisplayCoordinateSystem() == EAxisList::LeftUpForward;
 	if( InAxis == EAxisList::Y )
 	{
-		AxisRotation = FRotationMatrix::MakeFromXZ(FVector(0, 1, 0), FVector(0, 0, 1));
+		AxisRotation = FRotationMatrix::MakeFromXZ(FVector(0, bUsingLUFCoordinateSysem ? -1 : 1, 0), FVector(0, 0, 1));
 	}
 	else if( InAxis == EAxisList::Z )
 	{
-		AxisRotation = FRotationMatrix::MakeFromXY(FVector(0, 0, 1), FVector(0, 1, 0));
+		AxisRotation = FRotationMatrix::MakeFromXY(FVector(0, 0, 1), FVector(0, bUsingLUFCoordinateSysem ? -1 : 1, 0));
 	}
 
 	FMatrix ArrowToWorld = AxisRotation * InMatrix;
@@ -300,7 +302,7 @@ void FWidget::Render_Axis(const FSceneView* View, FPrimitiveDrawInterface* PDI,
 			}
 			case EAxisList::Y:
 			{
-				DrawCylinder(PDI, (Scale * FRotationMatrix(FRotator(0, 0, 90)) * InMatrix)* FScaleMatrix(FlattenScale), Offset, FVector(1, 0, 0), FVector(0, 1, 0), FVector(0, 0, 1), CylinderRadius, HalfHeight, 16, InMaterial->GetRenderProxy(), SDPG_Foreground );
+				DrawCylinder(PDI, (Scale * FRotationMatrix(FRotator(0, 0, bUsingLUFCoordinateSysem ? -90 : 90)) * InMatrix)* FScaleMatrix(FlattenScale), Offset, FVector(1, 0, 0), FVector(0, 1, 0), FVector(0, 0, 1), CylinderRadius, HalfHeight, 16, InMaterial->GetRenderProxy(), SDPG_Foreground );
 				break;
 			}
 			case EAxisList::Z:
@@ -326,7 +328,7 @@ void FWidget::Render_Axis(const FSceneView* View, FPrimitiveDrawInterface* PDI,
 			DrawCone(PDI, ( FScaleMatrix(-13) * FTranslationMatrix(RootPos) * ArrowToWorld ) * FScaleMatrix(FlattenScale), Angle, Angle, 32, false, FColor::White, InMaterial->GetRenderProxy(), SDPG_Foreground);
 		}
 	
-		PDI->SetHitProxy( NULL );
+		PDI->SetHitProxy( nullptr );
 	}
 
 	FVector2D NewOrigin;
@@ -496,6 +498,8 @@ void FWidget::Render_Translate( const FSceneView* View, FPrimitiveDrawInterface*
 		// but in the ortho viewports it will prevent scaling in the direction of the camera and thus intersecting the near plane.
 		FVector FlattenScale = FVector(Space.Scale.Component(0) == 1.0f ? 1.0f / Space.UniformScale : 1.0f, Space.Scale.Component(1) == 1.0f ? 1.0f / Space.UniformScale : 1.0f, Space.Scale.Component(2) == 1.0f ? 1.0f / Space.UniformScale : 1.0f);
 
+		const bool bUsingLUFCoordinateSysem = AxisDisplayInfo::GetAxisDisplayCoordinateSystem() == EAxisList::LeftUpForward;
+
 		if (Space.bIsPerspective || Space.bIsLocalSpace || Space.bIsOrthoXY)
 		{
 			if( (DrawAxis&EAxisList::XY) == EAxisList::XY )							// Top
@@ -505,9 +509,9 @@ void FWidget::Render_Translate( const FSceneView* View, FPrimitiveDrawInterface*
 
 				PDI->SetHitProxy( new HWidgetAxis(EAxisList::XY, bDisabled) );
 				{
-					DrawDualAxis(PDI, ( FTranslationMatrix(CornerPos) * FRotationMatrix(FRotator(0, 0, 90)) * WidgetMatrix ) * FScaleMatrix(FlattenScale), AxisSize, CornerLength, XMaterial->GetRenderProxy(), YMaterial->GetRenderProxy());
+					DrawDualAxis(PDI, ( FTranslationMatrix(CornerPos) * FRotationMatrix(FRotator(0, 0, bUsingLUFCoordinateSysem ? -90 : 90)) * WidgetMatrix ) * FScaleMatrix(FlattenScale), AxisSize, CornerLength, XMaterial->GetRenderProxy(), YMaterial->GetRenderProxy());
 				}
-				PDI->SetHitProxy( NULL );
+				PDI->SetHitProxy( nullptr );
 			}
 		}
 
@@ -522,7 +526,7 @@ void FWidget::Render_Translate( const FSceneView* View, FPrimitiveDrawInterface*
 				{
 					DrawDualAxis(PDI, (FTranslationMatrix(CornerPos) * WidgetMatrix) * FScaleMatrix(FlattenScale), AxisSize, CornerLength, XMaterial->GetRenderProxy(), ZMaterial->GetRenderProxy() );
 				}
-				PDI->SetHitProxy( NULL );
+				PDI->SetHitProxy( nullptr );
 			}
 		}
 
@@ -535,9 +539,9 @@ void FWidget::Render_Translate( const FSceneView* View, FPrimitiveDrawInterface*
 
 				PDI->SetHitProxy( new HWidgetAxis(EAxisList::YZ, bDisabled) );
 				{
-					DrawDualAxis(PDI, (FTranslationMatrix(CornerPos) * FRotationMatrix(FRotator(0, 90, 0)) * WidgetMatrix) * FScaleMatrix(FlattenScale), AxisSize, CornerLength, YMaterial->GetRenderProxy(), ZMaterial->GetRenderProxy() );
+					DrawDualAxis(PDI, (FTranslationMatrix(CornerPos) * FRotationMatrix(FRotator(0, bUsingLUFCoordinateSysem ? -90 : 90, 0)) * WidgetMatrix) * FScaleMatrix(FlattenScale), AxisSize, CornerLength, YMaterial->GetRenderProxy(), ZMaterial->GetRenderProxy() );
 				}
-				PDI->SetHitProxy( NULL );
+				PDI->SetHitProxy( nullptr );
 			}
 		}
 	}
@@ -553,7 +557,7 @@ void FWidget::Render_Translate( const FSceneView* View, FPrimitiveDrawInterface*
 		UMaterialInstanceDynamic* XYZMaterial = ( CurrentAxis&EAxisList::Screen) ? CurrentAxisMaterial : OpaquePlaneMaterialXY;
 		DrawSphere( PDI, InLocation, FRotator::ZeroRotator, 4.0f * Space.Scale, 10, 5, XYZMaterial->GetRenderProxy(), SDPG_Foreground );
 
-		PDI->SetHitProxy( NULL );
+		PDI->SetHitProxy( nullptr );
 	}
 }
 
@@ -714,7 +718,7 @@ void FWidget::Render_Rotate( const FSceneView* View,FPrimitiveDrawInterface* PDI
 			{
 				DrawColoredSphere(PDI, InLocation, FRotator::ZeroRotator, ArcColor, Radii, 32, 24, TransparentPlaneMaterialXY->GetRenderProxy(), SDPG_Foreground, true);
 			}
-			PDI->SetHitProxy(NULL);
+			PDI->SetHitProxy(nullptr);
 		}
 		
 		if(DrawAxis&EAxisList::Screen && (!bDragging || CurrentAxis == EAxisList::Screen) && GetDefault<ULevelEditorViewportSettings>()->bAllowScreenRotate)
@@ -738,7 +742,7 @@ void FWidget::Render_Rotate( const FSceneView* View,FPrimitiveDrawInterface* PDI
 				//Pass through alpha
 				DrawThickArc(OuterArcParams, Axis0, Axis1, 0.0f, 2.0f * PI, ArcColor, DirectionToWidget, !bIsPerspective);
 			}
-			PDI->SetHitProxy(NULL);
+			PDI->SetHitProxy(nullptr);
 		}
 	}
 }
@@ -789,6 +793,11 @@ void FWidget::Render_Scale( const FSceneView* View,FPrimitiveDrawInterface* PDI,
 	if ( bDrawWidget )
 	{
 		const bool bDisabled = IsWidgetDisabled();
+		const bool bUsingLUFCoordinateSysem = AxisDisplayInfo::GetAxisDisplayCoordinateSystem() == EAxisList::LeftUpForward;
+
+		// Since the Y axis gets flipped in the LUF coord system, flip the Y position of the grabber handles as well
+		const int8_t GrabberHandlesHalfYPos = bUsingLUFCoordinateSysem ? -12 : 12;
+		const int8_t GrabberHandlesYPos = bUsingLUFCoordinateSysem ? -24 : 24;
 
 		// Grabber handles - since orthographic scale widgets are not supported, we should always draw grabber handles if we're drawing
 		// the corresponding axes.
@@ -796,10 +805,10 @@ void FWidget::Render_Scale( const FSceneView* View,FPrimitiveDrawInterface* PDI,
 		{
 			PDI->SetHitProxy( new HWidgetAxis(EAxisList::XY, bDisabled) );
 			{
-				PDI->DrawLine( WidgetMatrix.TransformPosition(FVector(24,0,0) * UniformScale), WidgetMatrix.TransformPosition(FVector(12,12,0) * UniformScale), XColor, SDPG_Foreground );
-				PDI->DrawLine( WidgetMatrix.TransformPosition(FVector(12,12,0) * UniformScale), WidgetMatrix.TransformPosition(FVector(0,24,0) * UniformScale), YColor, SDPG_Foreground );
+				PDI->DrawLine( WidgetMatrix.TransformPosition(FVector(24,0,0) * UniformScale), WidgetMatrix.TransformPosition(FVector(12,GrabberHandlesHalfYPos,0) * UniformScale), XColor, SDPG_Foreground );
+				PDI->DrawLine( WidgetMatrix.TransformPosition(FVector(12,GrabberHandlesHalfYPos,0) * UniformScale), WidgetMatrix.TransformPosition(FVector(0,GrabberHandlesYPos,0) * UniformScale), YColor, SDPG_Foreground );
 			}
-			PDI->SetHitProxy( NULL );
+			PDI->SetHitProxy( nullptr );
 		}
 
 		if( (DrawAxis&(EAxisList::X|EAxisList::Z)) == (EAxisList::X|EAxisList::Z) )
@@ -809,17 +818,17 @@ void FWidget::Render_Scale( const FSceneView* View,FPrimitiveDrawInterface* PDI,
 				PDI->DrawLine( WidgetMatrix.TransformPosition(FVector(24,0,0) * UniformScale), WidgetMatrix.TransformPosition(FVector(12,0,12) * UniformScale), XColor, SDPG_Foreground );
 				PDI->DrawLine( WidgetMatrix.TransformPosition(FVector(12,0,12) * UniformScale), WidgetMatrix.TransformPosition(FVector(0,0,24) * UniformScale), ZColor, SDPG_Foreground );
 			}
-			PDI->SetHitProxy( NULL );
+			PDI->SetHitProxy( nullptr );
 		}
 
 		if( (DrawAxis&(EAxisList::Y|EAxisList::Z)) == (EAxisList::Y|EAxisList::Z) )
 		{
 			PDI->SetHitProxy( new HWidgetAxis(EAxisList::YZ, bDisabled) );
 			{
-				PDI->DrawLine( WidgetMatrix.TransformPosition(FVector(0,24,0) * UniformScale), WidgetMatrix.TransformPosition(FVector(0,12,12) * UniformScale), YColor, SDPG_Foreground );
-				PDI->DrawLine( WidgetMatrix.TransformPosition(FVector(0,12,12) * UniformScale), WidgetMatrix.TransformPosition(FVector(0,0,24) * UniformScale), ZColor, SDPG_Foreground );
+				PDI->DrawLine( WidgetMatrix.TransformPosition(FVector(0,GrabberHandlesYPos,0) * UniformScale), WidgetMatrix.TransformPosition(FVector(0,GrabberHandlesHalfYPos,12) * UniformScale), YColor, SDPG_Foreground );
+				PDI->DrawLine( WidgetMatrix.TransformPosition(FVector(0,GrabberHandlesHalfYPos,12) * UniformScale), WidgetMatrix.TransformPosition(FVector(0,0,24) * UniformScale), ZColor, SDPG_Foreground );
 			}
-			PDI->SetHitProxy( NULL );
+			PDI->SetHitProxy( nullptr );
 		}
 
 		// Center cube
@@ -829,7 +838,7 @@ void FWidget::Render_Scale( const FSceneView* View,FPrimitiveDrawInterface* PDI,
 
 			Render_Cube(PDI, WidgetMatrix, XYZMaterial, UniformScale * 4 );
 
-			PDI->SetHitProxy( NULL );
+			PDI->SetHitProxy( nullptr );
 		}
 	}
 }
@@ -907,7 +916,7 @@ void FWidget::Render_TranslateRotateZ( const FSceneView* View, FPrimitiveDrawInt
 				FVector BaseArrowPoint = InLocation + XAxis * ScaledRadius;
 				DrawFlatArrow(PDI, BaseArrowPoint, XAxis, YAxis, ZRotateColor, ScaledRadius, ScaledRadius*.5f, ZRotateMaterial->GetRenderProxy(), SDPG_Foreground);
 			}
-			PDI->SetHitProxy( NULL );
+			PDI->SetHitProxy( nullptr );
 		}
 
 		//XY Plane
@@ -927,7 +936,7 @@ void FWidget::Render_TranslateRotateZ( const FSceneView* View, FPrimitiveDrawInt
 					XYPlaneColor.A = ((CurrentAxis&EAxisList::XY) == EAxisList::XY) ? 0x3f : 0x0f;	//make the disc transparent
 					DrawDisc  ( PDI, InLocation, CustomCoordSystem.TransformPosition( FVector(1,0,0) ), CustomCoordSystem.TransformPosition( FVector(0,1,0) ), XYPlaneColor, ScaledRadius, CircleSides, TransparentPlaneMaterialXY->GetRenderProxy(), SDPG_Foreground );
 				}
-				PDI->SetHitProxy( NULL );
+				PDI->SetHitProxy( nullptr );
 			}
 		}
 	}
@@ -1051,7 +1060,7 @@ void FWidget::Render_2D(const FSceneView* View, FPrimitiveDrawInterface* PDI, FE
 				Color.A = Alpha;
 				DrawDisc(PDI, InLocation, CustomCoordSystem.TransformPosition(FVector(1, 0, 0)), CustomCoordSystem.TransformPosition(FVector(0, 0, 1)), Color, ScaledRadius, CircleSides, TransparentPlaneMaterialXY->GetRenderProxy(), SDPG_Foreground);
 			}
-			PDI->SetHitProxy(NULL);
+			PDI->SetHitProxy(nullptr);
 
 			PDI->SetHitProxy(new HWidgetAxis(EAxisList::Rotate2D, bDisabled));
 			{
@@ -1063,7 +1072,7 @@ void FWidget::Render_2D(const FSceneView* View, FPrimitiveDrawInterface* PDI, FE
 				FVector BaseArrowPoint = InLocation + XAxis * ArrowStartRadius;
 				DrawFlatArrow(PDI, BaseArrowPoint, XAxis, YAxis, Color, ArrowRadius, ArrowRadius*.5f, TransparentPlaneMaterialXY->GetRenderProxy(), SDPG_Foreground);
 			}
-			PDI->SetHitProxy(NULL);
+			PDI->SetHitProxy(nullptr);
 		}
 	}
 }
@@ -1258,7 +1267,7 @@ void FWidget::DrawPartialRotationArc(const FSceneView* View, FPrimitiveDrawInter
 		OuterColor.A = InColor.A;
 		DrawThickArc(OuterArcParams, Axis0, Axis1, InStartAngle, InEndAngle, OuterColor, InDirectionToWidget, !bIsPerspective );
 	}
-	PDI->SetHitProxy( NULL );
+	PDI->SetHitProxy( nullptr );
 
 	const bool bIsHitProxyView = View->Family->EngineShowFlags.HitProxies;
 	if (bIsPerspective && !bIsHitProxyView && !PDI->IsHitTesting())
