Avoid spurious warnings about non-determinism caused by comparing saved packages with metadata in their header against packages without metadata in their header due to the -DiffOptional argument trying to exclusively control whether the SAVE_Optional flag is present on the save args.  This issue surfaced after the recent UMetaData->FMetaData change which now controls whether metadata is put into the package header based on the presence or absence of the SAVE_Optional flag on the save args.  This is normally controlled with the -editoroptional commandline argument, but the archive diffing seems to want it controlled exclusively with the -DiffOptional argument.  This change implies that -DiffOptional will enable that flag if it isn't already enabled via -editoroptional, but it will not clear that flag.

diff --git a/Engine/Source/Editor/UnrealEd/Private/Cooker/DiffPackageWriter.cpp b/Engine/Source/Editor/UnrealEd/Private/Cooker/DiffPackageWriter.cpp
--- a/Engine/Source/Editor/UnrealEd/Private/Cooker/DiffPackageWriter.cpp
+++ b/Engine/Source/Editor/UnrealEd/Private/Cooker/DiffPackageWriter.cpp
@@ -309,10 +309,7 @@ void FDiffPackageWriter::UpdateSaveArguments(FSavePackageArgs& SaveArgs)
 	{
 		SaveArgs.SaveFlags |= SAVE_Optional;
 	}
-	else
-	{
-		SaveArgs.SaveFlags &= ~SAVE_Optional;
-	}
+
 	Inner->UpdateSaveArguments(SaveArgs);
 }
 
