Added landscape project setting to prevent landscape mode from disabling Temporal Anti Aliasing. Some projects are just so noisy that the lack of TAA is too detrimental

diff --git a/Engine/Source/Editor/UnrealEd/Private/EditorViewportClient.cpp b/Engine/Source/Editor/UnrealEd/Private/EditorViewportClient.cpp
--- a/Engine/Source/Editor/UnrealEd/Private/EditorViewportClient.cpp
+++ b/Engine/Source/Editor/UnrealEd/Private/EditorViewportClient.cpp
@@ -79,6 +79,7 @@
 #include "Settings/EditorStyleSettings.h"
 #include "GameFramework/ActorPrimitiveColorHandler.h"
 #include "LandscapeRender.h"
+#include "LandscapeSettings.h"
 #include "Misc/AxisDisplayInfo.h"
 
 #define LOCTEXT_NAMESPACE "EditorViewportClient"
@@ -3635,21 +3636,21 @@ void FEditorViewportClient::DrawAxes(FViewport* InViewport, FCanvas* Canvas, con
 			int32 Y = (FMath::Abs(EndPoint.X - OriginPoint.X) > FMath::Abs(EndPoint.Y - OriginPoint.Y)) ? 1 : 0;
 			return FIntPoint(X, Y);
 		};
-
+		
 	if ((InAxis & EAxisList::X) == EAxisList::X || (InAxis & EAxisList::Forward) == EAxisList::Forward)
 	{
-		AxisVec = AxisSize * ViewTM.InverseTransformVector(FVector(1, 0, 0));
-		AxisEnd = AxisOrigin + FIntPoint(AxisVec.Y, -AxisVec.Z);
+		AxisVec = AxisSize * ViewTM.InverseTransformVector( FVector(1,0,0) );
+		AxisEnd = AxisOrigin + FIntPoint( AxisVec.Y, -AxisVec.Z );
 		TextVec = (AxisSize + TextDistance) * ViewTM.InverseTransformVector(FVector(1, 0, 0));
 		TextCenter = AxisOrigin + FIntPoint(TextVec.Y, -TextVec.Z);
 		LineItem.SetColor(XShadowColor);
 		LineItem.Draw(Canvas, AxisOrigin + GetShadowOffset(AxisEnd, AxisOrigin), AxisEnd + GetShadowOffset(AxisEnd, AxisOrigin));
 		LineItem.SetColor(XBaseColor);
-		LineItem.Draw(Canvas, AxisOrigin, AxisEnd);
+		LineItem.Draw( Canvas, AxisOrigin, AxisEnd );
 		TextItem.Text = AxisDisplayInfo::GetAxisDisplayNameShort(EAxisList::Forward);
 		TextItem.SetColor(XBaseColor);
 		TextItem.EnableShadow(XShadowColor);
-		TextItem.Draw(Canvas, FVector2D(TextCenter.X - 0.5 * XL, TextCenter.Y - 0.5 * YL));
+		TextItem.Draw( Canvas, FVector2D(TextCenter.X - 0.5*XL, TextCenter.Y - 0.5*YL) );
 	}
 
 	const bool bIsLeft = (InAxis & EAxisList::Left) == EAxisList::Left;
@@ -3661,37 +3662,37 @@ void FEditorViewportClient::DrawAxes(FViewport* InViewport, FCanvas* Canvas, con
 			TextVec = (AxisSize + TextDistance) * ViewTM.InverseTransformVector(FVector(0, -1, 0));
 		}
 		else
-		{
-			AxisVec = AxisSize * ViewTM.InverseTransformVector(FVector(0, 1, 0));
+	{
+		AxisVec = AxisSize * ViewTM.InverseTransformVector( FVector(0,1,0) );
 			TextVec = (AxisSize + TextDistance) * ViewTM.InverseTransformVector(FVector(0, 1, 0));
 		}
 
-		AxisEnd = AxisOrigin + FIntPoint(AxisVec.Y, -AxisVec.Z);
+		AxisEnd = AxisOrigin + FIntPoint( AxisVec.Y, -AxisVec.Z );
 		TextCenter = AxisOrigin + FIntPoint(TextVec.Y, -TextVec.Z);
 		LineItem.SetColor(YShadowColor);
 		LineItem.Draw(Canvas, AxisOrigin + GetShadowOffset(AxisEnd, AxisOrigin), AxisEnd + GetShadowOffset(AxisEnd, AxisOrigin));
 		LineItem.SetColor(YBaseColor);
-		LineItem.Draw(Canvas, AxisOrigin, AxisEnd);
+		LineItem.Draw( Canvas, AxisOrigin, AxisEnd );
 		TextItem.Text = AxisDisplayInfo::GetAxisDisplayNameShort(EAxisList::Left);
 		TextItem.SetColor(YBaseColor);
 		TextItem.EnableShadow(YShadowColor);
-		TextItem.Draw(Canvas, FVector2D(TextCenter.X - 0.5 * XL, TextCenter.Y - 0.5 * YL));
+		TextItem.Draw( Canvas, FVector2D(TextCenter.X - 0.5*XL, TextCenter.Y - 0.5*YL) );
 	}
 
 	if ((InAxis & EAxisList::Z) == EAxisList::Z || (InAxis & EAxisList::Up) == EAxisList::Up)
 	{
-		AxisVec = AxisSize * ViewTM.InverseTransformVector(FVector(0, 0, 1));
-		AxisEnd = AxisOrigin + FIntPoint(AxisVec.Y, -AxisVec.Z);
+		AxisVec = AxisSize * ViewTM.InverseTransformVector( FVector(0,0,1) );
+		AxisEnd = AxisOrigin + FIntPoint( AxisVec.Y, -AxisVec.Z );
 		TextVec = (AxisSize + TextDistance) * ViewTM.InverseTransformVector(FVector(0, 0, 1));
 		TextCenter = AxisOrigin + FIntPoint(TextVec.Y, -TextVec.Z);
 		LineItem.SetColor(ZShadowColor);
 		LineItem.Draw(Canvas, AxisOrigin + GetShadowOffset(AxisEnd, AxisOrigin), AxisEnd + GetShadowOffset(AxisEnd, AxisOrigin));
 		LineItem.SetColor(ZBaseColor);
-		LineItem.Draw(Canvas, AxisOrigin, AxisEnd);
+		LineItem.Draw( Canvas, AxisOrigin, AxisEnd );
 		TextItem.Text = AxisDisplayInfo::GetAxisDisplayNameShort(EAxisList::Up);
 		TextItem.SetColor(ZBaseColor);
 		TextItem.EnableShadow(ZShadowColor);
-		TextItem.Draw(Canvas, FVector2D(TextCenter.X - 0.5 * XL, TextCenter.Y - 0.5 * YL));
+		TextItem.Draw( Canvas, FVector2D(TextCenter.X - 0.5*XL, TextCenter.Y - 0.5*YL) );
 	}
 
 	OutsideTileItem.Draw(Canvas, FIntPoint(AxisOrigin.X - FMath::Floor(OutsideCenterSize.X * 0.5), AxisOrigin.Y - FMath::Floor(OutsideCenterSize.Y * 0.5)));
@@ -4336,7 +4337,7 @@ void FEditorViewportClient::Draw(FViewport* InViewport, FCanvas* Canvas)
 
 	ViewFamily.LandscapeLODOverride = LandscapeLODOverride;
 	// Landscape brushes don't play well with TAA : 
-	if (GLandscapeEditModeActive)
+	if (GLandscapeEditModeActive && GetDefault<ULandscapeSettings>()->ShouldDisableTemporalAntiAliasingInLandscapeMode())
 	{
 		ViewFamily.EngineShowFlags.SetTemporalAA(false);
 	}
diff --git a/Engine/Source/Runtime/Landscape/Public/LandscapeSettings.h b/Engine/Source/Runtime/Landscape/Public/LandscapeSettings.h
--- a/Engine/Source/Runtime/Landscape/Public/LandscapeSettings.h
+++ b/Engine/Source/Runtime/Landscape/Public/LandscapeSettings.h
@@ -65,6 +65,8 @@ public:
 
 	bool ShouldDisplayTargetLayerThumbnails() const { return bDisplayTargetLayerThumbnails; }
 
+	bool ShouldDisableTemporalAntiAliasingInLandscapeMode() const { return bDisableTemporalAntiAliasingInLandscapeMode; }
+
 #if WITH_EDITOR
 	//~ Begin UObject Interface.
 	virtual void PostEditChangeProperty(FPropertyChangedEvent& PropertyChangedEvent) override;
@@ -126,4 +128,8 @@ protected:
 	UPROPERTY(EditAnywhere, config, Category = "Configuration", meta = (ToolTip = 
 		"When true, each target layer will have a representative thumbnail in landscape mode. However, the thumbnail system can be slow. Plus, certain landscape materials (e.g. using RVT) are not able to properly render thumbnails. Thus, setting this to false will skip needlessly rendering landscape layer thumbnails, which can improve the editing experience"))
 	bool bDisplayTargetLayerThumbnails = true;
+
+	UPROPERTY(EditAnywhere, config, Category = "Configuration", meta = (ToolTip = 
+		"When true, temporal anti-aliasing will be inactive while in landscape mode. This avoids the ghosting effect on the landscape brush but can lead to aliasing or shimmering on other parts of the image."))
+	bool bDisableTemporalAntiAliasingInLandscapeMode = true;
 };
