Fix export of theme file from Editor Preferences.

The source path that was being passed to CopyFile did not exist (e.g. Content/Slate/Themes/Dark.json does not exist in the engine directory anymore).

diff --git a/Engine/Source/Editor/UnrealEd/Private/Settings/EditorStyleClasses.cpp b/Engine/Source/Editor/UnrealEd/Private/Settings/EditorStyleClasses.cpp
--- a/Engine/Source/Editor/UnrealEd/Private/Settings/EditorStyleClasses.cpp
+++ b/Engine/Source/Editor/UnrealEd/Private/Settings/EditorStyleClasses.cpp
@@ -143,11 +143,14 @@ bool UEditorStyleSettings::OnExportBegin(const FString& ExportToPath)
 
 	const FString DestinationPath = PathPart / CurrentTheme.DisplayName.ToString() + TEXT(".json");
 
-	// copy the theme file to the same destination
-	if (IPlatformFile::GetPlatformPhysical().CopyFile(*DestinationPath, *CurrentTheme.Filename))
+	// Save the theme file to the same destination, replacing any existing file at that path.
+	if (IPlatformFile::GetPlatformPhysical().FileExists(*DestinationPath))
 	{
-		return true; 
+		IPlatformFile::GetPlatformPhysical().DeleteFile(*DestinationPath);
 	}
-	return false; 
-}
 
+	USlateThemeManager::Get().SaveCurrentThemeAs(DestinationPath);
+
+	return IPlatformFile::GetPlatformPhysical().FileExists(*ExportToPath)
+		&& IPlatformFile::GetPlatformPhysical().FileExists(*DestinationPath);
+}
diff --git a/Engine/Source/Editor/UnrealEd/Private/Settings/EditorStyleSettingsCustomization.cpp b/Engine/Source/Editor/UnrealEd/Private/Settings/EditorStyleSettingsCustomization.cpp
--- a/Engine/Source/Editor/UnrealEd/Private/Settings/EditorStyleSettingsCustomization.cpp
+++ b/Engine/Source/Editor/UnrealEd/Private/Settings/EditorStyleSettingsCustomization.cpp
@@ -572,10 +572,16 @@ FReply FEditorStyleSettingsCustomization::OnExportThemeClicked()
 
 	if (FDesktopPlatformModule::Get()->SaveFileDialog(ParentWindowHandle, LOCTEXT("ExportThemeDialogTitle", "Export current theme...").ToString(), FPaths::GetPath(ExportPath), DefaultFileName, TEXT("JSON files (*.json)|*.json"), EFileDialogFlags::None, OutFiles))
 	{
-		const FString SourcePath = USlateThemeManager::Get().GetCurrentTheme().Filename;
-		const FString DestPath = OutFiles[0]; 
+		const FString DestPath = OutFiles[0];
 
-		if (IPlatformFile::GetPlatformPhysical().CopyFile(*DestPath, *SourcePath)) 
+		if (IPlatformFile::GetPlatformPhysical().FileExists(*DestPath))
+		{
+			IPlatformFile::GetPlatformPhysical().DeleteFile(*DestPath);
+		}
+
+		USlateThemeManager::Get().SaveCurrentThemeAs(DestPath);
+
+		if (IPlatformFile::GetPlatformPhysical().FileExists(*DestPath))
 		{
 			ShowNotification(LOCTEXT("ExportThemeSuccess", "Export theme succeeded"), SNotificationItem::CS_Success);
 		}
