Change the cooker to allow for checking the `PKG_RuntimeGenerated` flag to also mean that it should treat the data being generated as "this was generated during the cook process itself" and thus treat the data as being already-cooked and available for the current platform being cooked. These packages are marked with said flag when the Verse codegen runs the step to generate the `UPackage`s containing the compiled `UVerseClass`es, and this will carry over to the serialized cooked data on disk as well. Since Verse is guaranteed to be available on any target platform we currently support for this case, this is guaranteed to be true. This allows for multi-plugin cooking with Verse to work when the `DLCNAME` argument is used.

diff --git a/Engine/Source/Editor/UnrealEd/Private/Cooker/CookRequestCluster.cpp b/Engine/Source/Editor/UnrealEd/Private/Cooker/CookRequestCluster.cpp
--- a/Engine/Source/Editor/UnrealEd/Private/Cooker/CookRequestCluster.cpp
+++ b/Engine/Source/Editor/UnrealEd/Private/Cooker/CookRequestCluster.cpp
@@ -2588,16 +2588,19 @@ void FRequestCluster::IsRequestCookable(const ITargetPlatform* Platform, FPackag
 
 			if (!bEditorOnly)
 			{
-				if (!PackageData.HasCookedPlatform(Platform, true /* bIncludeFailed */))
+				bool bIsReferenceAnError = true;
+				const UPackage* UnrealPackage = FindPackage(NULL, *PackageData.GetPackageName().ToString());
+				if (UnrealPackage && UnrealPackage->HasAnyPackageFlags(PKG_RuntimeGenerated))
 				{
-					// AllowUncookedAssetReferences should only be used when the DLC plugin to cook is going to be
-					// mounted where uncooked packages are available. This will allow a DLC plugin to be recooked 
-					// continually and mounted in an uncooked editor which is useful for CI.
-					if (!InCOTFS.CookByTheBookOptions->bAllowUncookedAssetReferences)
-					{
-						UE_LOG(LogCook, Error,
-							TEXT("Uncooked Engine or Game content %s is being referenced by DLC!"), *NameBuffer);
-					}
+					bIsReferenceAnError = false;
+				}
+				if (bIsReferenceAnError)
+				{
+					bIsReferenceAnError = !PackageData.HasCookedPlatform(Platform, true /* bIncludeFailed */) && !InCOTFS.CookByTheBookOptions->bAllowUncookedAssetReferences;
+				}
+				if (bIsReferenceAnError)
+				{
+					UE_LOG(LogCook, Error, TEXT("Uncooked Engine or Game content %s is being referenced by DLC!"), *NameBuffer);
 				}
 				OutReason = ESuppressCookReason::NotInCurrentPlugin;
 				bOutCookable = false;
@@ -2778,4 +2781,4 @@ static TArray<T> FindRootsFromLeafToRootOrderList(TConstArrayView<T> LeafToRootO
 	return Roots;
 }
 
-} // namespace UE::Cook
\ No newline at end of file
+} // namespace UE::Cook
