SceneGraph - add collision events to mesh_component
- fix duplicate rigid body problem when using rigid_body_component

rigid_body_component now modifies the rigid body owned by the mesh_component rather than creating its own. Previously you would have to disable collision on the mesh_component to use rigid_body_component.

#okforversepublic

diff --git a/Engine/Source/Editor/UnrealEd/Private/Cooker/CookPackageArtifacts.cpp b/Engine/Source/Editor/UnrealEd/Private/Cooker/CookPackageArtifacts.cpp
--- a/Engine/Source/Editor/UnrealEd/Private/Cooker/CookPackageArtifacts.cpp
+++ b/Engine/Source/Editor/UnrealEd/Private/Cooker/CookPackageArtifacts.cpp
@@ -32,7 +32,7 @@ namespace UE::Cook
 {
 
 // Change TargetDomainVersion to a new guid when all TargetDomain keys in an incremental cook need to be invalidated
-FGuid TargetDomainVersion(0xE780E176, 0xE04F4849, 0xA8E8D63D, 0x67B2EFD2);
+FGuid TargetDomainVersion(0xE0DBB4ED, 0x5C8A4914, 0xB83C77D0, 0x4152131F);
 
 // Bump PackageArtifactsVersion version when the serialization of PackageArtifacts has changed and we want to add
 // backwards compatibility rather than invalidating everything.
