[LUF Axes Colors] Changed the reference in the EditorComponents GetAxisColor to be LUF

diff --git a/Engine/Source/Editor/UnrealEd/Private/EditorComponents.cpp b/Engine/Source/Editor/UnrealEd/Private/EditorComponents.cpp
--- a/Engine/Source/Editor/UnrealEd/Private/EditorComponents.cpp
+++ b/Engine/Source/Editor/UnrealEd/Private/EditorComponents.cpp
@@ -120,9 +120,9 @@ UMaterialInstanceDynamic* FGridWidget::GetActiveLevelGridMID()
 
 static void GetAxisColors(FLinearColor Out[3], bool bIsPerspective)
 {
-	Out[0] = AxisDisplayInfo::GetAxisColor(EAxisList::X);
-	Out[1] = AxisDisplayInfo::GetAxisColor(EAxisList::Y);
-	Out[2] = AxisDisplayInfo::GetAxisColor(EAxisList::Z);
+	Out[0] = AxisDisplayInfo::GetAxisColor(EAxisList::Forward);
+	Out[1] = AxisDisplayInfo::GetAxisColor(EAxisList::Left);
+	Out[2] = AxisDisplayInfo::GetAxisColor(EAxisList::Up);
 
 	// less prominent axis lines for perspective views
 	if (bIsPerspective)
