[Viewport Toolbar] Filter whole Camera submenu based on Perspective / Orthographic

Dynamically hiding:

- Pilot Section
- Camera movement Section
- Camera shake

Also enables Game View in Ortho Views.

diff --git a/Engine/Plugins/Cameras/CameraShakePreviewer/Source/CameraShakePreviewer/Private/CameraShakePreviewerModule.cpp b/Engine/Plugins/Cameras/CameraShakePreviewer/Source/CameraShakePreviewer/Private/CameraShakePreviewerModule.cpp
--- a/Engine/Plugins/Cameras/CameraShakePreviewer/Source/CameraShakePreviewer/Private/CameraShakePreviewerModule.cpp
+++ b/Engine/Plugins/Cameras/CameraShakePreviewer/Source/CameraShakePreviewer/Private/CameraShakePreviewerModule.cpp
@@ -14,6 +14,7 @@
 #include "SLevelViewport.h"
 #include "Framework/MultiBox/MultiBoxBuilder.h"
 #include "LevelEditorMenuContext.h"
+#include "ViewportToolbar/UnrealEdViewportToolbar.h"
 
 #define LOCTEXT_NAMESPACE "CameraShakePreviewer"
 
@@ -24,25 +25,41 @@ namespace UE::CameraShakePreviewer::Private
 static const FName LevelEditorModuleName("LevelEditor");
 static const FName LevelEditorCameraShakePreviewerTab("CameraShakePreviewer");
 
-FLevelEditorViewportClient* GetPerspectiveLevelEditorViewportClient()
+TSharedPtr<SLevelViewport> GetLevelViewport()
 {
 	FLevelEditorModule& LevelEditorModule = FModuleManager::LoadModuleChecked<FLevelEditorModule>(LevelEditorModuleName);
 
 	if (TSharedPtr<ILevelEditor> LevelEditor = LevelEditorModule.GetLevelEditorInstance().Pin())
 	{
-		if (TSharedPtr<SLevelViewport> ViewportInterface = LevelEditor->GetActiveViewportInterface())
+		return LevelEditor->GetActiveViewportInterface();
+	}
+
+	return nullptr;
+}
+
+FLevelEditorViewportClient* GetPerspectiveLevelEditorViewportClient()
+{
+	if (TSharedPtr<SLevelViewport> ViewportInterface = GetLevelViewport())
+	{
+		FLevelEditorViewportClient* ViewportClient = &ViewportInterface->GetLevelViewportClient();
+		if (ViewportClient->ViewportType == ELevelViewportType::LVT_Perspective)
 		{
-			FLevelEditorViewportClient* ViewportClient =
-				static_cast<FLevelEditorViewportClient*>(&ViewportInterface->GetAssetViewportClient());
-			if (ViewportClient->ViewportType == ELevelViewportType::LVT_Perspective)
-			{
-				return ViewportClient;
-			}
+			return ViewportClient;
 		}
 	}
 
 	return nullptr;
 }
+
+TSharedPtr<FEditorViewportClient> GetLevelEditorViewportClient()
+{
+	if (TSharedPtr<SLevelViewport> ViewportInterface = GetLevelViewport())
+	{
+		return ViewportInterface->GetViewportClient();
+	}
+
+	return nullptr;
+}
 } // namespace UE::CameraShakePreviewer::Private
 
 /**
@@ -231,6 +248,9 @@ void FCameraShakePreviewerModule::CreateCameraShakeToggleOption(FMenuBuilder& In
 		}
 	);
 
+	// Show this entry only if viewport is perspective
+	TAttribute<EVisibility> VisibilityOverride = UE::UnrealEd::GetPerspectiveOnlyVisibility(UE::CameraShakePreviewer::Private::GetLevelEditorViewportClient());
+
 	TSharedPtr<FUICommandInfo> ToggleCameraShakesPreview = FCameraShakePreviewerCommands::Get().ToggleCameraShakesPreview;
 	InMenuBuilder.AddMenuEntry(
 		ToggleCameraShakesPreview->GetLabel(),
@@ -238,7 +258,10 @@ void FCameraShakePreviewerModule::CreateCameraShakeToggleOption(FMenuBuilder& In
 		FSlateIcon(FAppStyle::Get().GetStyleSetName(), "LevelViewport.ToggleCameraShakePreview"),
 		ToggleCameraShakeAction,
 		NAME_None,
-		EUserInterfaceActionType::ToggleButton
+		EUserInterfaceActionType::ToggleButton,
+		NAME_None,
+		FText(),
+		VisibilityOverride
 	);
 }
 
diff --git a/Engine/Source/Editor/LevelEditor/Private/SLevelViewport.cpp b/Engine/Source/Editor/LevelEditor/Private/SLevelViewport.cpp
--- a/Engine/Source/Editor/LevelEditor/Private/SLevelViewport.cpp
+++ b/Engine/Source/Editor/LevelEditor/Private/SLevelViewport.cpp
@@ -1555,7 +1555,7 @@ void SLevelViewport::BindOptionCommands( FUICommandList& OutCommandList )
 	OutCommandList.MapAction(
 		ViewportActions.ToggleGameView,
 		FExecuteAction::CreateSP( this, &SLevelViewport::ToggleGameView ),
-		FCanExecuteAction::CreateSP( this, &SLevelViewport::CanToggleGameView ),
+		FCanExecuteAction(),
 		FIsActionChecked::CreateSP( this, &SLevelViewport::IsInGameView ) );
 
 	OutCommandList.MapAction(
diff --git a/Engine/Source/Editor/LevelEditor/Private/ViewportToolbar/LevelEditorSubmenus.cpp b/Engine/Source/Editor/LevelEditor/Private/ViewportToolbar/LevelEditorSubmenus.cpp
--- a/Engine/Source/Editor/LevelEditor/Private/ViewportToolbar/LevelEditorSubmenus.cpp
+++ b/Engine/Source/Editor/LevelEditor/Private/ViewportToolbar/LevelEditorSubmenus.cpp
@@ -520,9 +520,20 @@ FToolMenuEntry CreateActorSnapCheckboxMenu()
 	return ActorSnapping;
 }
 
+// Can be used to show entries only in perspective view - specialized version of UE::UnrealEd method, for LevelViewport argument
+TAttribute<EVisibility> GetPerspectiveOnlyVisibility(const TWeakPtr<::SLevelViewport>& LevelViewportWeak)
+{
+	if (TSharedPtr<::SLevelViewport> LevelViewport = LevelViewportWeak.Pin())
+	{
+		return UE::UnrealEd::GetPerspectiveOnlyVisibility(LevelViewport->GetViewportClient());
+	}
+
+	return EVisibility::Collapsed;
+}
+
 FToolMenuEntry CreatePilotSubmenu(TWeakPtr<::SLevelViewport> LevelViewportWeak)
 {
-	return FToolMenuEntry::InitSubMenu(
+	FToolMenuEntry Entry = FToolMenuEntry::InitSubMenu(
 		"PilotingSubmenu",
 		LOCTEXT("PilotingSubmenu", "Pilot"),
 		LOCTEXT("PilotingSubmenu_ToolTip", "Piloting cameras and actors"),
@@ -586,11 +597,15 @@ FToolMenuEntry CreatePilotSubmenu(TWeakPtr<::SLevelViewport> LevelViewportWeak)
 		false,
 		FSlateIcon(FAppStyle::Get().GetStyleSetName(), "LevelViewport.PilotSelectedActor")
 	);
+
+	Entry.SetVisibility(UE::LevelEditor::Private::GetPerspectiveOnlyVisibility(LevelViewportWeak));
+
+	return Entry;
 }
 
-FToolMenuEntry CreateCameraMovementSubmenu()
+FToolMenuEntry CreateCameraMovementSubmenu(const TWeakPtr<::SLevelViewport>& LevelViewportWeak)
 {
-	return FToolMenuEntry::InitSubMenu(
+	FToolMenuEntry Entry = FToolMenuEntry::InitSubMenu(
 		"CameraMovement",
 		LOCTEXT("CameraMovementSubmenu", "Camera Movement"),
 		LOCTEXT("CameraMovementSubmenu_ToolTip", "Camera movement options"),
@@ -624,6 +639,10 @@ FToolMenuEntry CreateCameraMovementSubmenu()
 		false,
 		FSlateIcon(FSlateIcon(FAppStyle::GetAppStyleSetName(), "ClassIcon.CameraComponent"))
 	);
+
+	Entry.SetVisibility(GetPerspectiveOnlyVisibility(LevelViewportWeak));
+
+	return Entry;
 }
 
 FToolMenuEntry CreatePreviewSelectedCamerasCheckBoxSubmenu()
@@ -2819,7 +2838,7 @@ void ExtendCameraSubmenu(FName InCameraOptionsSubmenuName)
 					FToolMenuSection& MovementSection = InDynamicMenu->FindOrAddSection("Movement");
 					
 					MovementSection.AddEntry(Private::CreatePilotSubmenu(LevelViewportWeak));
-					MovementSection.AddEntry(Private::CreateCameraMovementSubmenu());
+					MovementSection.AddEntry(Private::CreateCameraMovementSubmenu(LevelViewportWeak));
 				}
 				
 				UE::UnrealEd::GenerateViewportTypeMenu(InDynamicMenu);
diff --git a/Engine/Source/Editor/UnrealEd/Private/ViewportToolbar/UnrealEdViewportToolbar.cpp b/Engine/Source/Editor/UnrealEd/Private/ViewportToolbar/UnrealEdViewportToolbar.cpp
--- a/Engine/Source/Editor/UnrealEd/Private/ViewportToolbar/UnrealEdViewportToolbar.cpp
+++ b/Engine/Source/Editor/UnrealEd/Private/ViewportToolbar/UnrealEdViewportToolbar.cpp
@@ -3004,23 +3004,25 @@ void PopulateCameraMenu(UToolMenu* InMenu, const FViewportCameraMenuOptions& InO
 		OrthographicCameraSection.AddMenuEntry(FEditorViewportCommands::Get().Front);
 		OrthographicCameraSection.AddMenuEntry(FEditorViewportCommands::Get().Back);
 	}
-	
+
 	// Movement
 	{
 		FToolMenuSection& MovementSection = InMenu->FindOrAddSection("Movement", LOCTEXT("CameraMovementSectionLabel", "Movement"));
 		if (InOptions.bShowCameraMovement)
 		{
 			MovementSection.AddEntry(CreateCameraSpeedMenu());
-			
+
 			FToolMenuEntry& FrameEntry = MovementSection.AddMenuEntry(FEditorViewportCommands::Get().FocusViewportToSelection);
 			FrameEntry.Icon = FSlateIcon(FAppStyle::GetAppStyleSetName(), "Icons.FrameActor");
 		}
+
+		MovementSection.Visibility = UE::UnrealEd::GetPerspectiveOnlyVisibility(EditorViewport->GetViewportClient());
 	}
-	
+
 	// View
 	{
 		FToolMenuSection& ViewSection = InMenu->FindOrAddSection("View", LOCTEXT("CameraViewSectionLabel", "View"));
-		
+
 		if (InOptions.bShowFieldOfView)
 		{
 			FToolMenuEntry CameraFOV = FToolMenuEntry::InitWidget(
@@ -3032,7 +3034,9 @@ void PopulateCameraMenu(UToolMenu* InMenu, const FViewportCameraMenuOptions& InO
 			CameraFOV.ToolTip = LOCTEXT("CameraMovementTooltip", "Sets the field of view of the viewport's camera.");
 			ViewSection.AddEntry(CameraFOV);
 		}
-		
+
+		ViewSection.Visibility = UE::UnrealEd::GetPerspectiveOnlyVisibility(EditorViewport->GetViewportClient());
+
 		if (InOptions.bShowNearAndFarPlanes)
 		{
 			FToolMenuEntry CameraNearViewPlane = FToolMenuEntry::InitWidget(
@@ -3042,7 +3046,7 @@ void PopulateCameraMenu(UToolMenu* InMenu, const FViewportCameraMenuOptions& InO
 			);
 			CameraNearViewPlane.Icon = FSlateIcon(FAppStyle::GetAppStyleSetName(), "ViewportToolbar.NearViewPlane");
 			ViewSection.AddEntry(CameraNearViewPlane);
-				
+
 			FToolMenuEntry CameraFarViewPlane = FToolMenuEntry::InitWidget(
 				"CameraFarViewPlane",
 				UE::UnrealEd::CreateFarViewPlaneMenuWidget(EditorViewport.ToSharedRef()),
@@ -4033,14 +4037,52 @@ FToolMenuEntry CreateCameraSpeedMenu()
 							return UE::UnrealEd::GetCameraSpeedLabel(WeakViewport);
 						}
 					);
+
+					// Don't show camera speed raised entry if viewport is orthographic (or missing)
+					Entry.SetShowInToolbarTopLevel(
+						TAttribute<bool>::CreateLambda(
+							[WeakViewport = Context->Viewport]() -> bool
+							{
+								if (const TSharedPtr<SEditorViewport>& ViewportPinned = WeakViewport.Pin())
+								{
+									if (const TSharedPtr<FEditorViewportClient>& ViewportClient =
+											ViewportPinned->GetViewportClient())
+									{
+										return ViewportClient->IsPerspective();
+									}
+								}
+
+								return false;
+							}
+						)
+					);
 				}
-				Entry.SetShowInToolbarTopLevel(true);
+
 				Entry.ToolBarData.PlacementOverride = MenuPlacement_BelowRightAnchor;
 			}
 		)
 	);
 }
 
+TAttribute<EVisibility> GetPerspectiveOnlyVisibility(const TSharedPtr<FEditorViewportClient>& InViewportClient)
+{
+	if (!InViewportClient)
+	{
+		return EVisibility::Collapsed;
+	}
+
+	return TAttribute<EVisibility>::CreateLambda(
+		[WeakViewport = InViewportClient.ToWeakPtr()]() -> EVisibility
+		{
+			if (const TSharedPtr<FEditorViewportClient>& ViewportClient = WeakViewport.Pin())
+			{
+				return ViewportClient->IsPerspective() ? EVisibility::Visible : EVisibility::Collapsed;
+			}
+			return EVisibility::Collapsed;
+		}
+	);
+}
+
 } // namespace UE::UnrealEd
 
 #undef LOCTEXT_NAMESPACE
diff --git a/Engine/Source/Editor/UnrealEd/Public/ViewportToolbar/UnrealEdViewportToolbar.h b/Engine/Source/Editor/UnrealEd/Public/ViewportToolbar/UnrealEdViewportToolbar.h
--- a/Engine/Source/Editor/UnrealEd/Public/ViewportToolbar/UnrealEdViewportToolbar.h
+++ b/Engine/Source/Editor/UnrealEd/Public/ViewportToolbar/UnrealEdViewportToolbar.h
@@ -244,6 +244,11 @@ UNREALED_API FToolMenuEntry CreateScalabilitySubmenu();
 
 UNREALED_API FText GetCameraSpeedTooltip();
 
+/**
+ * Returns Visible if the provided Viewport Client is in perspective view, Collapsed if not
+ */
+UNREALED_API TAttribute<EVisibility> GetPerspectiveOnlyVisibility(const TSharedPtr<FEditorViewportClient>& InViewportClient);
+
 // Camera Menu Widgets
 TSharedRef<SWidget> CreateCameraMenuWidget(const TSharedRef<SEditorViewport>& InViewport, bool bInShowExposureSettings = false);
 TSharedRef<SWidget> CreateFOVMenuWidget(const TSharedRef<SEditorViewport>& InViewport);
