[Shaders] Fix NULL-terminated string in CompileShadersTestBed commandlet.

#rnx

diff --git a/Engine/Source/Editor/UnrealEd/Private/Commandlets/CompileShadersTestBedCommandlet.cpp b/Engine/Source/Editor/UnrealEd/Private/Commandlets/CompileShadersTestBedCommandlet.cpp
--- a/Engine/Source/Editor/UnrealEd/Private/Commandlets/CompileShadersTestBedCommandlet.cpp
+++ b/Engine/Source/Editor/UnrealEd/Private/Commandlets/CompileShadersTestBedCommandlet.cpp
@@ -174,7 +174,7 @@ int32 UCompileShadersTestBedCommandlet::Main(const FString& Params)
 					}
 					ShaderFormatListString.Append(*ShaderFormat.ToString());
 				}
-				UE_LOG(LogCompileShadersTestBedCommandlet, Display, TEXT("Compile global shaders for %s"), ShaderFormatListString.GetData());
+				UE_LOG(LogCompileShadersTestBedCommandlet, Display, TEXT("Compile global shaders for %s"), *ShaderFormatListString);
 
 				for (int32 FormatIndex = 0; FormatIndex < DesiredShaderFormats.Num(); FormatIndex++)
 				{
