PR #10862: Same approach for FBX SM/SK tangents/binormals import

#rnx

diff --git a/Engine/Source/Editor/UnrealEd/Private/Fbx/FbxSkeletalMeshImport.cpp b/Engine/Source/Editor/UnrealEd/Private/Fbx/FbxSkeletalMeshImport.cpp
--- a/Engine/Source/Editor/UnrealEd/Private/Fbx/FbxSkeletalMeshImport.cpp
+++ b/Engine/Source/Editor/UnrealEd/Private/Fbx/FbxSkeletalMeshImport.cpp
@@ -3495,6 +3495,14 @@ bool UnFbx::FFbxImporter::FillSkelMeshImporterFromFbx( FSkeletalMeshImportData&
 		TangentMappingMode = LayerElementTangent->GetMappingMode();
 	}
 
+	FbxLayerElement::EReferenceMode BinormalReferenceMode(FbxLayerElement::eDirect);
+	FbxLayerElement::EMappingMode BinormalMappingMode(FbxLayerElement::eByControlPoint);
+	if (LayerElementBinormal)
+	{
+		BinormalReferenceMode = LayerElementBinormal->GetReferenceMode();
+		BinormalMappingMode = LayerElementBinormal->GetMappingMode();
+	}
+
 	//
 	// create the points / wedges / faces
 	//
@@ -3557,26 +3565,28 @@ bool UnFbx::FFbxImporter::FillSkelMeshImporterFromFbx( FSkeletalMeshImportData&
 			//
 			if( bHasNormalInformation )
 			{
-				int32 TmpIndex = LocalIndex*3 + VertexIndex;
-				//normals may have different reference and mapping mode than tangents and binormals
-				int32 NormalMapIndex = (NormalMappingMode == FbxLayerElement::eByControlPoint) ? 
+				const int32 TmpIndex = LocalIndex*3 + VertexIndex;
+				// normals may have different reference and mapping mode than tangents and binormals
+				const int32 NormalMapIndex = (NormalMappingMode == FbxLayerElement::eByControlPoint) ?
 										ControlPointIndex : TmpIndex;
-				int32 NormalValueIndex = (NormalReferenceMode == FbxLayerElement::eDirect) ? 
+				const int32 NormalValueIndex = (NormalReferenceMode == FbxLayerElement::eDirect) ?
 										NormalMapIndex : LayerElementNormal->GetIndexArray().GetAt(NormalMapIndex);
 
-				//tangents and binormals share the same reference, mapping mode and index array
-				int32 TangentMapIndex = TmpIndex;
-
+				// tangents and binormals share the same reference, mapping mode and index array
 				FbxVector4 TempValue;
 
 				if( bHasTangentInformation )
 				{
-					TempValue = LayerElementTangent->GetDirectArray().GetAt(TangentMapIndex);
+					const int32 TangentMapIndex = (TangentMappingMode == FbxLayerElement::eByControlPoint) ? ControlPointIndex : TmpIndex;
+					const int32 TangentValueIndex = (TangentReferenceMode == FbxLayerElement::eDirect) ? TangentMapIndex : LayerElementTangent->GetIndexArray().GetAt(TangentMapIndex);
+					TempValue = LayerElementTangent->GetDirectArray().GetAt(TangentValueIndex);
 					TempValue = TotalMatrixForNormal.MultT(TempValue);
 					Triangle.TangentX[ UnrealVertexIndex ] = (FVector3f)Converter.ConvertDir(TempValue);
 					Triangle.TangentX[ UnrealVertexIndex ].Normalize();
 
-					TempValue = LayerElementBinormal->GetDirectArray().GetAt(TangentMapIndex);
+					const int32 BinormalMapIndex = (BinormalMappingMode == FbxLayerElement::eByControlPoint) ? ControlPointIndex : TmpIndex;
+					const int32 BinormalValueIndex = (BinormalReferenceMode == FbxLayerElement::eDirect) ? BinormalMapIndex : LayerElementBinormal->GetIndexArray().GetAt(BinormalMapIndex);
+					TempValue = LayerElementBinormal->GetDirectArray().GetAt(BinormalValueIndex);
 					TempValue = TotalMatrixForNormal.MultT(TempValue);
 					Triangle.TangentY[ UnrealVertexIndex ] = (FVector3f)-Converter.ConvertDir(TempValue);
 					Triangle.TangentY[ UnrealVertexIndex ].Normalize();
@@ -3590,8 +3600,7 @@ bool UnFbx::FFbxImporter::FillSkelMeshImporterFromFbx( FSkeletalMeshImportData&
 			}
 			else
 			{
-				int32 NormalIndex;
-				for( NormalIndex = 0; NormalIndex < 3; ++NormalIndex )
+				for( int32 NormalIndex = 0; NormalIndex < 3; ++NormalIndex )
 				{
 					Triangle.TangentX[ NormalIndex ] = FVector3f::ZeroVector;
 					Triangle.TangentY[ NormalIndex ] = FVector3f::ZeroVector;
