[ContentBrowser] Add Missing chip Icon and Version and use the Pointy version for the Drag and Drop Widget
Added Small fix for Action Button
Added Small text fix when dragging over collection panel/item

diff --git a/Engine/Source/Editor/ContentBrowser/Private/SCollectionView.cpp b/Engine/Source/Editor/ContentBrowser/Private/SCollectionView.cpp
--- a/Engine/Source/Editor/ContentBrowser/Private/SCollectionView.cpp
+++ b/Engine/Source/Editor/ContentBrowser/Private/SCollectionView.cpp
@@ -1259,6 +1259,27 @@ bool SCollectionView::ValidateDragDropOnCollectionTree(const FGeometry& Geometry
 			Operation->SetCursorOverride(EMouseCursor::SlashedCircle);
 		}
 	}
+	else if (Operation->IsOfType<FAssetDragDropOp>())
+	{
+		TSharedPtr<FAssetDragDropOp> DragDropOp = StaticCastSharedPtr<FAssetDragDropOp>(Operation);
+
+		if (!DragDropOp->HasAssets())
+		{
+			DragDropOp->SetCursorOverride(EMouseCursor::SlashedCircle);
+			DragDropOp->SetToolTip(LOCTEXT("CollectionView_DragDrop_NoAsset", "There is no asset being dragged"), FAppStyle::GetBrush(TEXT("Graph.ConnectorFeedback.Error")));
+			return false;
+		}
+
+		const int32 NumDraggedItems = DragDropOp->GetAssets().Num();
+		const FText FirstItemText = FText::FromName(DragDropOp->GetAssets()[0].AssetName);
+		const FText AddToCollectionText = (NumDraggedItems > 1)
+			? FText::Format(LOCTEXT("CollectionView_DragDrop_MultipleItems", "Add '{0}' and {1} {1}|plural(one=other,other=others)"), FirstItemText, NumDraggedItems - 1)
+			: FText::Format(LOCTEXT("CollectionView_DragDrop_SingularItems", "Add '{0}'"), FirstItemText);
+
+		DragDropOp->SetToolTip(AddToCollectionText, FAppStyle::GetBrush(TEXT("Graph.ConnectorFeedback.OK")));
+		return true;
+
+	}
 
 	return false;
 }
@@ -1349,6 +1370,21 @@ bool SCollectionView::ValidateDragDropOnCollectionItem(TSharedRef<FCollectionIte
 		TSharedPtr<FAssetDragDropOp> DragDropOp = StaticCastSharedPtr<FAssetDragDropOp>(Operation);
 		OutIsKnownDragOperation = true;
 		bIsValidDrag = DragDropOp->HasAssets();
+
+		if (!bIsValidDrag)
+		{
+			DragDropOp->SetToolTip(LOCTEXT("CollectionViewItem_DragDrop_NoAsset", "There is no asset being dragged"), FAppStyle::GetBrush(TEXT("Graph.ConnectorFeedback.Error")));
+		}
+		else
+		{
+			const int32 NumDraggedItems = DragDropOp->GetAssets().Num();
+			const FText FirstItemText = FText::FromName(DragDropOp->GetAssets()[0].AssetName);
+			const FText AddToCollectionText = (NumDraggedItems > 1)
+				? FText::Format(LOCTEXT("CollectionViewItem_DragDrop_MultipleItems", "Add '{0}' and {1} {1}|plural(one=other,other=others)"), FirstItemText, NumDraggedItems - 1)
+				: FText::Format(LOCTEXT("CollectionViewItem_DragDrop_SingularItems", "Add '{0}'"), FirstItemText);
+
+			DragDropOp->SetToolTip(AddToCollectionText, FAppStyle::GetBrush(TEXT("Graph.ConnectorFeedback.OK")));
+		}
 	}
 
 	// Set the default slashed circle if this drag is invalid and a drag operation hasn't set NewDragCursor to something custom
diff --git a/Engine/Source/Editor/EditorStyle/Private/StarshipStyle.cpp b/Engine/Source/Editor/EditorStyle/Private/StarshipStyle.cpp
--- a/Engine/Source/Editor/EditorStyle/Private/StarshipStyle.cpp
+++ b/Engine/Source/Editor/EditorStyle/Private/StarshipStyle.cpp
@@ -1254,14 +1254,31 @@ void FStarshipEditorStyle::FStyle::SetupGeneralStyles()
 		Set( "AssetThumbnail.HintShadowOffset", FVector2f::UnitVector );
 		Set( "AssetThumbnail.HintShadowColorAndOpacity", FLinearColor(0.f, 0.f, 0.f, 0.5f) );
 		Set( "AssetThumbnail.HintBackground", new BOX_BRUSH( "Common/TableViewHeader", FMargin(8.0f/32.0f) ) );
-		Set( "AssetThumbnail.AssetTileViewWhiteChipBorder", new IMAGE_BRUSH_SVG("Starship/ContentBrowser/CornerChipOuter2_20", Icon20x20));
-		Set( "AssetThumbnail.AssetTileViewWhiteChip", new IMAGE_BRUSH_SVG("Starship/ContentBrowser/CornerChipInner2_20", Icon20x20));
-		Set( "AssetThumbnail.AssetTileViewWhiteChipSharp", new IMAGE_BRUSH_SVG("Starship/ContentBrowser/CornerChipSharp_20", Icon20x20));
-		// TODO: Update with correct ratio icons for the small one once ready
-		Set("AssetThumbnail.AssetTileViewWhiteChipBorderSmall", new IMAGE_BRUSH_SVG("Starship/ContentBrowser/CornerChipOuter2_20", Icon20x20));
-		Set( "AssetThumbnail.AssetTileViewWhiteChipSmall", new IMAGE_BRUSH_SVG("Starship/ContentBrowser/CornerChipInner2_20", Icon20x20));
+
+		// Chip Normal Size
+		// Round
+		Set( "AssetThumbnail.AssetTileViewWhiteChipBorder", new IMAGE_BRUSH_SVG("Starship/ContentBrowser/CornerChipOuter_20", Icon20x20));
+		Set( "AssetThumbnail.AssetTileViewWhiteChip", new IMAGE_BRUSH_SVG("Starship/ContentBrowser/CornerChipInner_20", Icon20x20));
+		// Sharp
+		Set( "AssetThumbnail.AssetTileViewWhiteChipSharpBorder", new IMAGE_BRUSH_SVG("Starship/ContentBrowser/CornerChipSharpOuter_20", Icon20x20));
+		Set( "AssetThumbnail.AssetTileViewWhiteChipSharp", new IMAGE_BRUSH_SVG("Starship/ContentBrowser/CornerChipSharpInner_20", Icon20x20));
+
+		// Chip Small Size
+		// Round
+		Set( "AssetThumbnail.AssetTileViewWhiteChipBorderSmall", new IMAGE_BRUSH_SVG("Starship/ContentBrowser/CornerChipSmallOuter_16", Icon16x16));
+		Set( "AssetThumbnail.AssetTileViewWhiteChipSmall", new IMAGE_BRUSH_SVG("Starship/ContentBrowser/CornerChipSmallInner_16", Icon16x16));
+		// Sharp
+		Set( "AssetThumbnail.AssetTileViewWhiteChipSharpBorderSmall", new IMAGE_BRUSH_SVG("Starship/ContentBrowser/CornerChipSharpSmallOuter_16", Icon16x16));
+		Set( "AssetThumbnail.AssetTileViewWhiteChipSharpSmall", new IMAGE_BRUSH_SVG("Starship/ContentBrowser/CornerChipSharpSmallInner_16", Icon16x16));
+
+		// AssetBorder
+		// Round
 		Set( "AssetThumbnail.AssetBorder", new FSlateRoundedBoxBrush(FStyleColors::Transparent, FVector4(4.0f, 4.0f, 4.0f, 4.0f), FStyleColors::Secondary, 2.f));
 		Set( "AssetThumbnail.AssetBorderSmall", new FSlateRoundedBoxBrush(FStyleColors::Transparent, FVector4(4.0f, 4.0f, 4.0f, 4.0f), FStyleColors::Secondary, 1.f));
+		// Sharp
+		Set( "AssetThumbnail.AssetBorderSharp", new FSlateRoundedBoxBrush(FStyleColors::Transparent, 0.f, FStyleColors::Secondary, 2.f));
+		Set( "AssetThumbnail.AssetBorderSharpSmall", new FSlateRoundedBoxBrush(FStyleColors::Transparent, 0.f, FStyleColors::Secondary, 1.f));
+
 		FLinearColor OverlayColorAssetStatusOverlay = FStyleColors::Panel.GetSpecifiedColor();
 		OverlayColorAssetStatusOverlay.A = 0.75f;
 		Set("AssetThumbnail.AssetThumbnailStatusBar", new FSlateRoundedBoxBrush(OverlayColorAssetStatusOverlay, 2.f));
diff --git a/Engine/Source/Editor/UnrealEd/Private/AssetThumbnail.cpp b/Engine/Source/Editor/UnrealEd/Private/AssetThumbnail.cpp
--- a/Engine/Source/Editor/UnrealEd/Private/AssetThumbnail.cpp
+++ b/Engine/Source/Editor/UnrealEd/Private/AssetThumbnail.cpp
@@ -118,6 +118,7 @@ public:
 		SLATE_ARGUMENT(bool, AllowAssetStatusThumbnailOverlay)
 		SLATE_ATTRIBUTE(TSharedPtr<IToolTip>, TooltipOverride)
 		SLATE_ARGUMENT(bool, ShowAssetChip)
+		SLATE_ARGUMENT(bool, UseRoundAssetChip)
 		SLATE_ARGUMENT(bool, SmallAssetChipBorder)
 		SLATE_ATTRIBUTE(const FSlateBrush*, AssetChipBorderImageOverride)
 		SLATE_ARGUMENT(bool, CanDisplayEditModePrimitiveTools)
@@ -308,11 +309,25 @@ public:
 		{
 			if (bShowAssetChip)
 			{
+				FString ThumbnailBorder = TEXT(".AssetBorder");
+				if (InArgs._SmallAssetChipBorder)
+				{
+					ThumbnailBorder = TEXT(".AssetBorderSmall");
+				}
+
+				FString ChipBorder = InArgs._UseRoundAssetChip ?
+					InArgs._SmallAssetChipBorder ? TEXT(".AssetTileViewWhiteChipBorderSmall") : TEXT(".AssetTileViewWhiteChipBorder")
+					: InArgs._SmallAssetChipBorder ? TEXT(".AssetTileViewWhiteChipSharpBorderSmall") : TEXT(".AssetTileViewWhiteChipSharpBorder");
+
+				FString Chip = InArgs._UseRoundAssetChip ?
+					InArgs._SmallAssetChipBorder ? TEXT(".AssetTileViewWhiteChipSmall") : TEXT(".AssetTileViewWhiteChip")
+					: InArgs._SmallAssetChipBorder ? TEXT(".AssetTileViewWhiteChipSharpSmall") : TEXT(".AssetTileViewWhiteChipSharp");
+
 				// The asset color chip
 				LayeredImage = SNew(SLayeredImage)
-					.Image(FAppStyle::GetBrush(Style, InArgs._SmallAssetChipBorder ? ".AssetTileViewWhiteChipBorderSmall" : ".AssetTileViewWhiteChipBorder"))
+					.Image(FAppStyle::GetBrush(Style, StringCast<ANSICHAR>(*ChipBorder).Get()))
 					.ColorAndOpacity(FStyleColors::Secondary);
-				LayeredImage->AddLayer(FAppStyle::GetBrush(Style, InArgs._SmallAssetChipBorder ? ".AssetTileViewWhiteChipSmall" : ".AssetTileViewWhiteChip"), AssetColor);
+				LayeredImage->AddLayer(FAppStyle::GetBrush(Style, StringCast<ANSICHAR>(*Chip).Get()), AssetColor);
 
 				// Need this to be negative for the Overlay so that it overlaps with the actual border which will keep the ratio correct and won't let the thumbnail over the limits.
 				const FMargin SmallAssetChipMarginOverlay(-1.f);
@@ -330,7 +345,7 @@ public:
 
 				// Image overlay to match the Hovered/Selected status, by design only part of the chip should be highlighted
 				const TSharedRef<SImage> OverlayImage = SNew(SImage)
-					.Image(FAppStyle::GetBrush(Style, InArgs._SmallAssetChipBorder ? ".AssetBorderSmall" : ".AssetBorder"))
+					.Image(FAppStyle::GetBrush(Style, StringCast<ANSICHAR>(*ThumbnailBorder).Get()))
 					.Visibility(EVisibility::HitTestInvisible);
 
 				OverlayWidget->AddSlot()
@@ -342,7 +357,7 @@ public:
 				// Border used to contain the Thumbnail in the ratio
 				const TSharedRef<SBorder> OverlayBorder = SNew(SBorder)
 					.Padding(InArgs._SmallAssetChipBorder ? SmallAssetChipMargin : 2.f)
-					.BorderImage(FAppStyle::GetBrush(Style, InArgs._SmallAssetChipBorder ? ".AssetBorderSmall" : ".AssetBorder"));
+					.BorderImage(FAppStyle::GetBrush(Style, StringCast<ANSICHAR>(*ThumbnailBorder).Get()));
 
 
 				if (AssetChipBorderImageOverride.IsSet())
@@ -617,6 +632,11 @@ public:
 	{
 		SCompoundWidget::OnMouseEnter( MyGeometry, MouseEvent );
 
+		if (AssetThumbnailEditMode.IsValid())
+		{
+			SetHover(true);
+		}
+
 		if ( bAllowRealTimeOnHovered )
 		{
 			AssetThumbnail->SetRealTime( true );
@@ -629,8 +649,7 @@ public:
 
 		if (AssetThumbnailEditMode.IsValid())
 		{
-			// Keep the Hovered state the same as the ThumbnailEdit mode during OnMouseLeave
-			SetHover(AssetThumbnailEditMode->IsHovered());
+			SetHover(AssetThumbnailEditMode->IsEditingThumbnail());
 		}
 
 		if ( bAllowRealTimeOnHovered )
@@ -1641,6 +1660,7 @@ TSharedRef<SWidget> FAssetThumbnail::MakeThumbnailWidget( const FAssetThumbnailC
 			.TooltipOverride(InConfig.TooltipOverride)
 			.AllowAssetStatusThumbnailOverlay(InConfig.bAllowAssetStatusThumbnailOverlay)
 			.ShowAssetChip(InConfig.bShowAssetChip)
+			.UseRoundAssetChip(InConfig.bUseRoundAssetChip)
 			.IsEditModeVisible(InConfig.IsEditModeVisible)
 			.SmallAssetChipBorder(InConfig.bSmallAssetChipBorder)
 			.AssetChipBorderImageOverride(InConfig.AssetChipBorderImageOverride)
diff --git a/Engine/Source/Editor/UnrealEd/Private/DragAndDrop/AssetDragDropOp.cpp b/Engine/Source/Editor/UnrealEd/Private/DragAndDrop/AssetDragDropOp.cpp
--- a/Engine/Source/Editor/UnrealEd/Private/DragAndDrop/AssetDragDropOp.cpp
+++ b/Engine/Source/Editor/UnrealEd/Private/DragAndDrop/AssetDragDropOp.cpp
@@ -91,6 +91,7 @@ TSharedPtr<SWidget> FAssetDragDropOp::GetDefaultDecorator() const
 			FAssetThumbnailConfig AssetThumbnailConfig;
 			AssetThumbnailConfig.bShowAssetChip = true;
 			AssetThumbnailConfig.bSmallAssetChipBorder = true;
+			AssetThumbnailConfig.bUseRoundAssetChip = false;
 			AssetThumbnailConfig.AssetChipBorderImageOverride = FAppStyle::GetBrush("ContentBrowser.ThumbnailDragDropBackground");
 			ThumbnailWidget = AssetThumbnail->MakeThumbnailWidget(AssetThumbnailConfig);
 		}
diff --git a/Engine/Source/Editor/UnrealEd/Public/AssetThumbnail.h b/Engine/Source/Editor/UnrealEd/Public/AssetThumbnail.h
--- a/Engine/Source/Editor/UnrealEd/Public/AssetThumbnail.h
+++ b/Engine/Source/Editor/UnrealEd/Public/AssetThumbnail.h
@@ -69,6 +69,7 @@ struct FAssetThumbnailConfig
 		, AssetSystemInfoProvider(nullptr)
 		, bAllowAssetStatusThumbnailOverlay( false )
 		, bShowAssetChip(true)
+		, bUseRoundAssetChip(true)
 		, bSmallAssetChipBorder(false)
 		, AssetChipBorderImageOverride()
 		, bCanDisplayEditModePrimitiveTools(false)
@@ -93,6 +94,7 @@ struct FAssetThumbnailConfig
 	TAttribute<TSharedPtr<IToolTip>> TooltipOverride;
 	bool bAllowAssetStatusThumbnailOverlay;
 	bool bShowAssetChip;
+	bool bUseRoundAssetChip;
 	bool bSmallAssetChipBorder;
 	/** Border override brush needs to have an OutlineWidth of 2 if normal size or 1 if small size  */
 	TAttribute<const FSlateBrush*> AssetChipBorderImageOverride;
