Replace IsExternallyReferenceable(),SetIsExternallyReferenceable() with GetAssetAccessSpecifier(),SetAssetAccessSpecifier(EAssetAccessSpecifier::Public)
#rnx

diff --git a/Engine/Source/Developer/AssetTools/Private/AssetRenameManager.cpp b/Engine/Source/Developer/AssetTools/Private/AssetRenameManager.cpp
--- a/Engine/Source/Developer/AssetTools/Private/AssetRenameManager.cpp
+++ b/Engine/Source/Developer/AssetTools/Private/AssetRenameManager.cpp
@@ -1429,7 +1429,7 @@ void FAssetRenameManager::SetupPublicAssets(TArray<FAssetRenameDataWithReference
 		}
 
 		// If the asset is already public, it can be referenced from anywhere so we're good
-		if (Package->IsExternallyReferenceable())
+		if (Package->GetAssetAccessSpecifier() == EAssetAccessSpecifier::Public)
 		{
 			continue;
 		}
@@ -1563,7 +1563,7 @@ void FAssetRenameManager::SetupPublicAssets(TArray<FAssetRenameDataWithReference
 		// Make assets public
 		for (UPackage* Package : PackagesToMakePublic)
 		{
-			Package->SetIsExternallyReferenceable(true);
+			Package->SetAssetAccessSpecifier(EAssetAccessSpecifier::Public);
 		}
 	}
 }
diff --git a/Engine/Source/Developer/AssetTools/Private/AssetTools.cpp b/Engine/Source/Developer/AssetTools/Private/AssetTools.cpp
--- a/Engine/Source/Developer/AssetTools/Private/AssetTools.cpp
+++ b/Engine/Source/Developer/AssetTools/Private/AssetTools.cpp
@@ -1756,8 +1756,7 @@ UObject* UAssetToolsImpl::CreateAsset(const FString& AssetName, const FString& P
 
 	if( NewObj )
 	{
-
-		Pkg->SetIsExternallyReferenceable(ShouldCreateAssetsAsExternallyReferenceableForPath(PackageName));
+		Pkg->SetAssetAccessSpecifier(ShouldCreateAssetsAsExternallyReferenceableForPath(PackageName) ? EAssetAccessSpecifier::Public : EAssetAccessSpecifier::Private);
 
 		// Notify the asset registry
 		FAssetRegistryModule::AssetCreated(NewObj);
@@ -4031,7 +4030,7 @@ TArray<UObject*> UAssetToolsImpl::ImportAssetsInternal(const TArray<FString>& Fi
 				continue;
 			}
 
-			Pkg->SetIsExternallyReferenceable(ShouldCreateAssetsAsExternallyReferenceableForPath(PackageName));
+			Pkg->SetAssetAccessSpecifier(ShouldCreateAssetsAsExternallyReferenceableForPath(PackageName) ? EAssetAccessSpecifier::Public : EAssetAccessSpecifier::Private);
 
 			Factory->SetAutomatedAssetImportData(Params.ImportData);
 			Factory->SetAssetImportTask(Params.AssetImportTask);
diff --git a/Engine/Source/Editor/BlueprintEditorLibrary/Private/BlueprintEditorLibrary.cpp b/Engine/Source/Editor/BlueprintEditorLibrary/Private/BlueprintEditorLibrary.cpp
--- a/Engine/Source/Editor/BlueprintEditorLibrary/Private/BlueprintEditorLibrary.cpp
+++ b/Engine/Source/Editor/BlueprintEditorLibrary/Private/BlueprintEditorLibrary.cpp
@@ -925,7 +925,7 @@ UBlueprint* UBlueprintEditorLibrary::CreateBlueprintAssetWithParent(const FStrin
 		BlueprintGeneratedClass);
 	ensure(BP); // FKismetEditorUtilities::CreateBlueprint does not return null, if it does we should clean up the UPackage - somehow
 
-	Pkg->SetIsExternallyReferenceable(true);
+	Pkg->SetAssetAccessSpecifier(EAssetAccessSpecifier::Public);
 
 	// Notify the asset registry
 	FAssetRegistryModule::AssetCreated(BP);
diff --git a/Engine/Source/Editor/GameProjectGeneration/Private/SNewClassDialog.cpp b/Engine/Source/Editor/GameProjectGeneration/Private/SNewClassDialog.cpp
--- a/Engine/Source/Editor/GameProjectGeneration/Private/SNewClassDialog.cpp
+++ b/Engine/Source/Editor/GameProjectGeneration/Private/SNewClassDialog.cpp
@@ -1154,8 +1154,8 @@ void SNewClassDialog::FinishClicked()
 				UBlueprint* NewBP = FKismetEditorUtilities::CreateBlueprint(const_cast<UClass*>(ParentClassInfo.BaseClass), Package, FName(*NewClassName), BPTYPE_Normal);
 				if (NewBP)
 				{
-					// Set the default "IsExternallyReferenceable" state
-					Package->SetIsExternallyReferenceable(IAssetTools::Get().ShouldCreateAssetsAsExternallyReferenceableForPath(PackagePath));
+					// Set the default "AssetAccessSpecifier" state
+					Package->SetAssetAccessSpecifier(IAssetTools::Get().ShouldCreateAssetsAsExternallyReferenceableForPath(PackagePath) ? EAssetAccessSpecifier::Public : EAssetAccessSpecifier::Private);
 
 					// Notify the asset registry
 					FAssetRegistryModule::AssetCreated(NewBP);
diff --git a/Engine/Source/Editor/UnrealEd/Private/ObjectTools.cpp b/Engine/Source/Editor/UnrealEd/Private/ObjectTools.cpp
--- a/Engine/Source/Editor/UnrealEd/Private/ObjectTools.cpp
+++ b/Engine/Source/Editor/UnrealEd/Private/ObjectTools.cpp
@@ -1040,8 +1040,7 @@ namespace ObjectTools
 			}
 
 			// Duplicating an asset should respect the reference controls of the original.
-			bool isObjectExternallyReferenceable = Object->GetOutermost()->IsExternallyReferenceable();
-			DupObject->GetOutermost()->SetIsExternallyReferenceable(isObjectExternallyReferenceable);
+			DupObject->GetOutermost()->SetAssetAccessSpecifier(Object->GetOutermost()->GetAssetAccessSpecifier());
 
 			// When duplicating a World Composition map, make sure to properly initialize WorldTileInfo
 			if (Object->GetOutermost()->GetWorldTileInfo())
@@ -4338,7 +4337,7 @@ namespace ObjectTools
 							NewPackage->SetPackageFlags(PKG_NewlyCreated);
 						}
 
-						NewPackage->SetIsExternallyReferenceable(OldPackage->IsExternallyReferenceable());
+						NewPackage->SetAssetAccessSpecifier(OldPackage->GetAssetAccessSpecifier());
 
 						// When renaming a World Composition map, make sure to properly initialize WorldTileInfo
 						if (OldPackage->GetWorldTileInfo())
diff --git a/Engine/Source/Runtime/CoreUObject/Private/UObject/LinkerLoad.cpp b/Engine/Source/Runtime/CoreUObject/Private/UObject/LinkerLoad.cpp
--- a/Engine/Source/Runtime/CoreUObject/Private/UObject/LinkerLoad.cpp
+++ b/Engine/Source/Runtime/CoreUObject/Private/UObject/LinkerLoad.cpp
@@ -5106,7 +5106,7 @@ static FString GetPackageObjectFullName(FLinkerLoad* Linker, const FPackageIndex
 
 bool FLinkerLoad::IsPackageReferenceAllowed(UPackage* InPackage)
 {
-	if (InPackage && !InPackage->IsExternallyReferenceable())
+	if (InPackage && (InPackage->GetAssetAccessSpecifier() == EAssetAccessSpecifier::Private))
 	{
 #if WITH_EDITOR
 		// Package loaded for diff is not always in its original location (usually in /Temp/) 
diff --git a/Engine/Source/Runtime/CoreUObject/Private/UObject/SavePackage2.cpp b/Engine/Source/Runtime/CoreUObject/Private/UObject/SavePackage2.cpp
--- a/Engine/Source/Runtime/CoreUObject/Private/UObject/SavePackage2.cpp
+++ b/Engine/Source/Runtime/CoreUObject/Private/UObject/SavePackage2.cpp
@@ -1019,7 +1019,7 @@ ESavePackageResult ValidateImports(FSaveContext& SaveContext)
 		// Enforce that Private content can only be directly referenced by something within the same Mount Point
 		// This only applies to direct imports. Transitive imports (A -> PackageB.B, B -> B.Class -> PackageC) are
 		// allowed even for private content.
-		if (!ImportPackage->IsExternallyReferenceable() && DirectImports.Contains(Import))
+		if ((ImportPackage->GetAssetAccessSpecifier() == EAssetAccessSpecifier::Private) && DirectImports.Contains(Import))
 		{
 			FName MountPointName = FPackageName::GetPackageMountPoint(PackageName);
 
