[Viewport Toolbar] Grid Snap size menu entries are not properly updated + UV Editor and UMG Editor show wrong Grid Snap Size Menu entries

- Adds new CreateCustomLocationGridSnapMenu method for creating Grid Snap size menu, using a delegate as an input for the grid sizes array
- Deprecates previous method
- Adds a "default" method which can be used to get the same grid size menu Level Editor and other asset editors (e.g. Static Mesh) use
- Applies CreateCustomLocationGridSnapMenu to UV Editor and UMG Editor: this addresses the deprecation, and most of all fixes a bug which had both menus show the wrong grid sizes list

diff --git a/Engine/Plugins/Editor/UVEditor/Source/UVEditor/Private/SUVEditor2DViewportToolBar.cpp b/Engine/Plugins/Editor/UVEditor/Source/UVEditor/Private/SUVEditor2DViewportToolBar.cpp
--- a/Engine/Plugins/Editor/UVEditor/Source/UVEditor/Private/SUVEditor2DViewportToolBar.cpp
+++ b/Engine/Plugins/Editor/UVEditor/Source/UVEditor/Private/SUVEditor2DViewportToolBar.cpp
@@ -338,7 +338,22 @@ TSharedRef<SWidget> SUVEditor2DViewportToolBar::FillLocationGridSnapMenu()
 		}
 	);
 
-	return CreateLocationGridSnapMenu(ExecuteDelegate, IsCheckedDelegate, GridSizes, CommandList);
+	FLocationGridValuesArrayDelegate GridSizesDelegate =
+			UE::UnrealEd::FLocationGridValuesArrayDelegate::CreateLambda(
+				[GridSizes]()
+				{
+					return GridSizes;
+				}
+			);
+
+	UE::UnrealEd::FLocationGridSnapMenuOptions MenuOptions;
+	MenuOptions.MenuName = "UVEditor_LocationGridCheckboxList";
+	MenuOptions.ExecuteDelegate = ExecuteDelegate;
+	MenuOptions.IsCheckedDelegate = IsCheckedDelegate;
+	MenuOptions.GridValuesArrayDelegate = GridSizesDelegate;
+	MenuOptions.CommandList = CommandList;
+
+	return UE::UnrealEd::CreateLocationGridSnapMenu(MenuOptions);
 }
 
 TSharedRef<SWidget> SUVEditor2DViewportToolBar::FillRotationGridSnapMenu()
diff --git a/Engine/Source/Editor/UMGEditor/Private/Designer/SDesignerToolBar.cpp b/Engine/Source/Editor/UMGEditor/Private/Designer/SDesignerToolBar.cpp
--- a/Engine/Source/Editor/UMGEditor/Private/Designer/SDesignerToolBar.cpp
+++ b/Engine/Source/Editor/UMGEditor/Private/Designer/SDesignerToolBar.cpp
@@ -155,7 +155,21 @@ TSharedRef<SWidget> SDesignerToolBar::FillLocationGridSnapMenu()
 			}
 		);
 
-	return UE::UnrealEd::CreateLocationGridSnapMenu(ExecuteDelegate, IsCheckedDelegate, GridSizes, CommandList);
+	UE::UnrealEd::FLocationGridValuesArrayDelegate GridSizesDelegate =
+		UE::UnrealEd::FLocationGridValuesArrayDelegate::CreateLambda(
+			[GridSizes]()
+			{
+				return GridSizes;
+			}
+		);
+
+	UE::UnrealEd::FLocationGridSnapMenuOptions MenuOptions;
+	MenuOptions.MenuName = "UMG_LocationGridCheckboxList";
+	MenuOptions.ExecuteDelegate = ExecuteDelegate;
+	MenuOptions.IsCheckedDelegate = IsCheckedDelegate;
+	MenuOptions.GridValuesArrayDelegate = GridSizesDelegate;
+	MenuOptions.CommandList = CommandList;
+	return UE::UnrealEd::CreateLocationGridSnapMenu(MenuOptions);
 }
 
 void SDesignerToolBar::SetGridSize(int32 InGridSize)
diff --git a/Engine/Source/Editor/UnrealEd/Private/STransformViewportToolbar.cpp b/Engine/Source/Editor/UnrealEd/Private/STransformViewportToolbar.cpp
--- a/Engine/Source/Editor/UnrealEd/Private/STransformViewportToolbar.cpp
+++ b/Engine/Source/Editor/UnrealEd/Private/STransformViewportToolbar.cpp
@@ -830,7 +830,10 @@ TSharedRef<SWidget> STransformViewportToolBar::FillLocationGridSnapMenu()
 			}
 		);
 
-	return UE::UnrealEd::CreateLocationGridSnapMenu(ExecuteDelegate, IsCheckedDelegate, GridSizes, CommandList);
+	UE::UnrealEd::FLocationGridSnapMenuOptions MenuOptions;
+	MenuOptions.CommandList = CommandList;
+
+	return UE::UnrealEd::CreateLocationGridSnapMenu(MenuOptions);
 }
 
 TSharedRef<SWidget> STransformViewportToolBar::FillRotationGridSnapMenu()
diff --git a/Engine/Source/Editor/UnrealEd/Private/ViewportToolbar/UnrealEdViewportToolbar.cpp b/Engine/Source/Editor/UnrealEd/Private/ViewportToolbar/UnrealEdViewportToolbar.cpp
--- a/Engine/Source/Editor/UnrealEd/Private/ViewportToolbar/UnrealEdViewportToolbar.cpp
+++ b/Engine/Source/Editor/UnrealEd/Private/ViewportToolbar/UnrealEdViewportToolbar.cpp
@@ -302,38 +302,12 @@ FToolMenuEntry CreateLocationSnapCheckboxMenu()
 	}
 	else
 	{
-
 		FNewToolMenuDelegate MakeMenuDelegate = FNewToolMenuDelegate::CreateLambda(
 			[LocationSnapName](UToolMenu* InToolMenu)
 			{
-				UnrealEd::FLocationGridCheckboxListExecuteActionDelegate ExecuteDelegate =
-					UnrealEd::FLocationGridCheckboxListExecuteActionDelegate::CreateUObject(
-						GEditor, &UEditorEngine::SetGridSize
-					);
-
-				UnrealEd::FLocationGridCheckboxListIsCheckedDelegate IsCheckedDelegate =
-					UnrealEd::FLocationGridCheckboxListIsCheckedDelegate::CreateLambda(
-						[](int CurrGridSizeIndex)
-						{
-							const ULevelEditorViewportSettings* ViewportSettings =
-								GetDefault<ULevelEditorViewportSettings>();
-							return ViewportSettings->CurrentPosGridSize == CurrGridSizeIndex;
-						}
-					);
-
-				const ULevelEditorViewportSettings* ViewportSettings = GetDefault<ULevelEditorViewportSettings>();
-				TArray<float> GridSizes = ViewportSettings->bUsePowerOf2SnapSize ? ViewportSettings->Pow2GridSizes
-																				 : ViewportSettings->DecimalGridSizes;
-
 				InToolMenu->AddMenuEntry(
 					LocationSnapName,
-					FToolMenuEntry::InitWidget(
-						LocationSnapName,
-						UnrealEd::CreateLocationGridSnapMenu(
-							ExecuteDelegate,
-							IsCheckedDelegate,
-							GridSizes
-						),
+					FToolMenuEntry::InitWidget(LocationSnapName, UnrealEd::CreateLocationGridSnapMenu(),
 						FText()
 					)
 				);
@@ -1964,53 +1938,116 @@ TSharedRef<SWidget> CreateLocationGridSnapMenu(
 	const TSharedPtr<FUICommandList>& InCommandList
 )
 {
-	const FName CheckboxListMenu("LocationGridCheckboxList");
-	if (!UToolMenus::Get()->IsMenuRegistered(CheckboxListMenu))
+	FLocationGridSnapMenuOptions MenuOptions;
+	MenuOptions.MenuName = "LocationGridCheckboxList";
+	MenuOptions.ExecuteDelegate = InExecuteDelegate;
+	MenuOptions.IsCheckedDelegate = InIsCheckedDelegate;
+	MenuOptions.IsEnabledDelegate = InIsEnabledDelegate;
+	MenuOptions.CommandList = InCommandList;
+
+	return CreateLocationGridSnapMenu(MenuOptions);
+}
+
+TSharedRef<SWidget> CreateLocationGridSnapMenu(const FLocationGridSnapMenuOptions& InMenuOptions)
+{
+	const FName MenuName = InMenuOptions.MenuName.IsNone() ? "LocationGridCheckboxList" : InMenuOptions.MenuName;
+
+	if (!UToolMenus::Get()->IsMenuRegistered(MenuName))
 	{
-		if (UToolMenu* Menu = UToolMenus::Get()->RegisterMenu(CheckboxListMenu, NAME_None))
+		if (UToolMenu* Menu = UToolMenus::Get()->RegisterMenu(MenuName, NAME_None))
 		{
-			FToolMenuSection& Section = Menu->AddSection("Snap", LOCTEXT("LocationSnapText", "Snap Sizes"));
-			for (int32 CurrGridSizeIndex = 0; CurrGridSizeIndex < InGridSizes.Num(); ++CurrGridSizeIndex)
-			{
-				const float CurGridSize = InGridSizes[CurrGridSizeIndex];
+			Menu->AddDynamicSection(
+				NAME_None,
+				FNewToolMenuDelegate::CreateLambda(
+					[InMenuOptions](UToolMenu* InMenu)
+					{
+						TAttribute<bool> IsEnabledDelegate = InMenuOptions.IsEnabledDelegate;
+
+						UnrealEd::FLocationGridCheckboxListExecuteActionDelegate ExecuteDelegate =
+							InMenuOptions.ExecuteDelegate.IsBound()
+								? InMenuOptions.ExecuteDelegate
+								: UnrealEd::FLocationGridCheckboxListExecuteActionDelegate::CreateUObject(
+									  GEditor, &UEditorEngine::SetGridSize
+								  );
+
+						UnrealEd::FLocationGridCheckboxListIsCheckedDelegate IsCheckedDelegate =
+							InMenuOptions.IsCheckedDelegate.IsBound()
+								? InMenuOptions.IsCheckedDelegate
+								: UnrealEd::FLocationGridCheckboxListIsCheckedDelegate::CreateLambda(
+									  [](int CurrGridSizeIndex)
+									  {
+										  const ULevelEditorViewportSettings* ViewportSettings =
+											  GetDefault<ULevelEditorViewportSettings>();
+										  return ViewportSettings->CurrentPosGridSize == CurrGridSizeIndex;
+									  }
+								  );
+
+						FLocationGridValuesArrayDelegate GridValuesDelegate =
+							InMenuOptions.GridValuesArrayDelegate.IsBound()
+								? InMenuOptions.GridValuesArrayDelegate
+								: FLocationGridValuesArrayDelegate::CreateLambda(
+									  []()
+									  {
+										  const ULevelEditorViewportSettings* ViewportSettings =
+											  GetDefault<ULevelEditorViewportSettings>();
+										  TArray<float> GridSizes = ViewportSettings->bUsePowerOf2SnapSize
+																	  ? ViewportSettings->Pow2GridSizes
+																	  : ViewportSettings->DecimalGridSizes;
+
+										  return GridSizes;
+									  }
+								  );
+
+						const TArray<float> InGridSizes = GridValuesDelegate.Execute();
+
+						FToolMenuSection& Section = InMenu->AddSection("Snap", LOCTEXT("LocationSnapText", "Snap Sizes"));
+
+						for (int32 CurrGridSizeIndex = 0; CurrGridSizeIndex < InGridSizes.Num(); ++CurrGridSizeIndex)
+						{
+							const float CurGridSize = InGridSizes[CurrGridSizeIndex];
 
-				Section.AddMenuEntry(
-					NAME_None,
-					FText::AsNumber(CurGridSize),
-					FText::Format(LOCTEXT("LocationGridSize_ToolTip", "Sets grid size to {0}"), FText::AsNumber(CurGridSize)),
-					FSlateIcon(),
-					FUIAction(
-						FExecuteAction::CreateLambda(
-							[CurrGridSizeIndex, InExecuteDelegate]()
-							{
-								InExecuteDelegate.Execute(CurrGridSizeIndex);
-							}
-						),
-						FCanExecuteAction::CreateLambda(
-							[InIsEnabledDelegate]()
-							{
-								return InIsEnabledDelegate.Get();
-							}
-						),
-						FIsActionChecked::CreateLambda(
-							[CurrGridSizeIndex, InIsCheckedDelegate]()
-							{
-								return InIsCheckedDelegate.Execute(CurrGridSizeIndex);
-							}
-						)
-					),
-					EUserInterfaceActionType::RadioButton
-				);
-			}
+							Section.AddMenuEntry(
+								NAME_None,
+								FText::AsNumber(CurGridSize),
+								FText::Format(
+									LOCTEXT("LocationGridSize_ToolTip", "Sets grid size to {0}"), FText::AsNumber(CurGridSize)
+								),
+								FSlateIcon(),
+								FUIAction(
+									FExecuteAction::CreateLambda(
+										[CurrGridSizeIndex, ExecuteDelegate]()
+										{
+											ExecuteDelegate.Execute(CurrGridSizeIndex);
+										}
+									),
+									FCanExecuteAction::CreateLambda(
+										[IsEnabledDelegate]()
+										{
+											return IsEnabledDelegate.Get();
+										}
+									),
+									FIsActionChecked::CreateLambda(
+										[CurrGridSizeIndex, IsCheckedDelegate]()
+										{
+											return IsCheckedDelegate.Execute(CurrGridSizeIndex);
+										}
+									)
+								),
+								EUserInterfaceActionType::RadioButton
+							);
+						}
+					}
+				)
+			);
 		}
 	}
 
 	FToolMenuContext MenuContext;
 	{
-		MenuContext.AppendCommandList(InCommandList);
+		MenuContext.AppendCommandList(InMenuOptions.CommandList);
 	}
 
-	return UToolMenus::Get()->GenerateWidget(CheckboxListMenu, MenuContext);
+	return UToolMenus::Get()->GenerateWidget(MenuName, MenuContext);
 }
 
 FText GetScaleGridLabel()
diff --git a/Engine/Source/Editor/UnrealEd/Public/ViewportToolbar/UnrealEdViewportToolbar.h b/Engine/Source/Editor/UnrealEd/Public/ViewportToolbar/UnrealEdViewportToolbar.h
--- a/Engine/Source/Editor/UnrealEd/Public/ViewportToolbar/UnrealEdViewportToolbar.h
+++ b/Engine/Source/Editor/UnrealEd/Public/ViewportToolbar/UnrealEdViewportToolbar.h
@@ -79,6 +79,8 @@ DECLARE_DELEGATE_RetVal_TwoParams(bool, FRotationGridCheckboxListIsCheckedDelega
 DECLARE_DELEGATE_OneParam(FLocationGridCheckboxListExecuteActionDelegate, int);
 DECLARE_DELEGATE_RetVal_OneParam(bool, FLocationGridCheckboxListIsCheckedDelegate, int);
 
+DECLARE_DELEGATE_RetVal(TArray<float>, FLocationGridValuesArrayDelegate);
+
 DECLARE_DELEGATE_OneParam(FScaleGridCheckboxListExecuteActionDelegate, int);
 DECLARE_DELEGATE_RetVal_OneParam(bool, FScaleGridCheckboxListIsCheckedDelegate, int);
 
@@ -107,6 +109,8 @@ UNREALED_API TSharedRef<SWidget> CreateRotationGridSnapMenu(
 );
 
 UNREALED_API FText GetLocationGridLabel();
+
+UE_DEPRECATED(5.6, "Use CreateLocationGridSnapMenu() using FLocationGridSnapMenuOptions as argument instead.")
 UNREALED_API TSharedRef<SWidget> CreateLocationGridSnapMenu(
 	const FLocationGridCheckboxListExecuteActionDelegate& InExecuteDelegate,
 	const FLocationGridCheckboxListIsCheckedDelegate& InIsCheckedDelegate,
@@ -115,6 +119,18 @@ UNREALED_API TSharedRef<SWidget> CreateLocationGridSnapMenu(
 	const TSharedPtr<FUICommandList>& InCommandList = {}
 );
 
+struct UNREALED_API FLocationGridSnapMenuOptions
+{
+	FName MenuName;
+	FLocationGridCheckboxListExecuteActionDelegate ExecuteDelegate;
+	FLocationGridCheckboxListIsCheckedDelegate IsCheckedDelegate;
+	FLocationGridValuesArrayDelegate GridValuesArrayDelegate;
+	TAttribute<bool> IsEnabledDelegate = TAttribute<bool>(true);
+	TSharedPtr<FUICommandList> CommandList;
+};
+
+UNREALED_API TSharedRef<SWidget> CreateLocationGridSnapMenu(const FLocationGridSnapMenuOptions& InMenuOptions = {});
+
 UNREALED_API FText GetScaleGridLabel();
 UNREALED_API TSharedRef<SWidget> CreateScaleGridSnapMenu(
 	const FScaleGridCheckboxListExecuteActionDelegate& InExecuteDelegate,
