Substrate - fixed FrontMaterial pin showing up when Substrate is disabled (due to 40645535).

#rnx
#tests PC editor

diff --git a/Engine/Source/Editor/UnrealEd/Classes/MaterialGraph/MaterialGraph.h b/Engine/Source/Editor/UnrealEd/Classes/MaterialGraph/MaterialGraph.h
--- a/Engine/Source/Editor/UnrealEd/Classes/MaterialGraph/MaterialGraph.h
+++ b/Engine/Source/Editor/UnrealEd/Classes/MaterialGraph/MaterialGraph.h
@@ -64,6 +64,11 @@ struct FMaterialInputInfo
 		const bool bFrontMaterialConnected = Material->HasSubstrateFrontMaterialConnected();
 		const bool bShowAsSubstrate = bSubstrateEnabled && (!bSubstrateHiddenMaterialConversion || bFrontMaterialConnected);
 
+		if (!bSubstrateEnabled && Property == MP_FrontMaterial)
+		{
+			return false;
+		}
+
 		if (Material->bUseMaterialAttributes && !bIgnoreMaterialAttributes)
 		{
 			// On the root node both MaterialAttributes and FrontMaterial are visible when Substrate is enabled. 
