[LUF HUD Gizmo] Fixed a bug caused by the SetLookAt not using the LUF transforms on SetInitialViewTransform

diff --git a/Engine/Source/Editor/UnrealEd/Private/EditorViewportClient.cpp b/Engine/Source/Editor/UnrealEd/Private/EditorViewportClient.cpp
--- a/Engine/Source/Editor/UnrealEd/Private/EditorViewportClient.cpp
+++ b/Engine/Source/Editor/UnrealEd/Private/EditorViewportClient.cpp
@@ -759,14 +759,16 @@ void FEditorViewportClient::SetInitialViewTransform(ELevelViewportType InViewpor
 
 	FViewportCameraTransform& ViewTransform = (InViewportType == LVT_Perspective) ? ViewTransformPerspective : ViewTransformOrthographic;
 	// If using the LUF coord system, invert Y position and rotate yaw by -90* so the gizmo is centered since the Left axis draw is inverse to the Y
-	ViewTransform.SetLocation(FVector(ViewLocation.X, (bUsingLUFCoordinateSysem ? -ViewLocation.Y : ViewLocation.Y), ViewLocation.Z));
-	ViewTransform.SetRotation(ViewRotation + (bUsingLUFCoordinateSysem ? FRotator(0, -90, 0) : FRotator(0, 0, 0)));
+	FVector LocalViewLocation = FVector(ViewLocation.X, (bUsingLUFCoordinateSysem ? -ViewLocation.Y : ViewLocation.Y), ViewLocation.Z);
+	FRotator LocalViewRotation = ViewRotation + (bUsingLUFCoordinateSysem ? FRotator(0, -90, 0) : FRotator(0, 0, 0));
+	ViewTransform.SetLocation(LocalViewLocation);
+	ViewTransform.SetRotation(LocalViewRotation);
 
 	// Make a look at location in front of the camera
-	const FQuat CameraOrientation = FQuat::MakeFromEuler(ViewRotation.Euler());
+	const FQuat CameraOrientation = FQuat::MakeFromEuler(LocalViewRotation.Euler());
 	FVector Direction = CameraOrientation.RotateVector( FVector(1,0,0) );
 
-	ViewTransform.SetLookAt(ViewLocation + Direction * OrbitConstants::IntialLookAtDistance);
+	ViewTransform.SetLookAt(LocalViewLocation + Direction * OrbitConstants::IntialLookAtDistance);
 	ViewTransform.SetOrthoZoom(InOrthoZoom);
 }
 
