Rewrite TryAppendClassDigests. In the future I will modify this function to add the value of the dependency. But to do so I need TryAppendClassDigests to work directly with a list of FCookDependency and not with a prefiltered list of FTopLevelAssetPath.

#incrementalcook
#rnx

diff --git a/Engine/Source/Editor/UnrealEd/Private/Cooker/CookPackageArtifacts.cpp b/Engine/Source/Editor/UnrealEd/Private/Cooker/CookPackageArtifacts.cpp
--- a/Engine/Source/Editor/UnrealEd/Private/Cooker/CookPackageArtifacts.cpp
+++ b/Engine/Source/Editor/UnrealEd/Private/Cooker/CookPackageArtifacts.cpp
@@ -166,29 +166,13 @@ FBuildDependencySet::ECurrentKeyResult FBuildDependencySet::TryCalculateCurrentK
 			}
 			break;
 		}
+
 		case ECookDependency::NativeClass:
 		{
-			TArray<FTopLevelAssetPath, TInlineAllocator<10>> ClassPaths;
-			ClassPaths.Reserve(Batch.Num());
-			for (const FCookDependency& ClassDependency : Batch)
-			{
-				FTopLevelAssetPath& ClassPath = ClassPaths.Emplace_GetRef(ClassDependency.GetClassPath());
-				if (!ClassPath.IsValid())
-				{
-					Context.LogError(FString::Printf(
-						TEXT("ClassDependency failed: %s is not a valid TopLevelAssetPath."),
-						*WriteToString<256>(ClassDependency.GetClassPath())));
-					ClassPaths.Pop();
-					continue;
-				}
-			}
-			FString AppendDigestError;
-			if (!UE::EditorDomain::TryAppendClassDigests(KeyBuilder, ClassPaths, &AppendDigestError))
-			{
-				Context.LogError(FString::Printf(TEXT("ClassDependency failed: %s"), *AppendDigestError));
-			}
+			UE::EditorDomain::TryAppendClassDigests(Batch, KeyBuilder, Context);
 			break;
 		}
+
 		case ECookDependency::RedirectionTarget:
 		{
 			TArray<FName, TInlineAllocator<10>> PackageNames;
diff --git a/Engine/Source/Editor/UnrealEd/Private/EditorDomain/EditorDomainUtils.cpp b/Engine/Source/Editor/UnrealEd/Private/EditorDomain/EditorDomainUtils.cpp
--- a/Engine/Source/Editor/UnrealEd/Private/EditorDomain/EditorDomainUtils.cpp
+++ b/Engine/Source/Editor/UnrealEd/Private/EditorDomain/EditorDomainUtils.cpp
@@ -158,6 +158,8 @@ void Register()
 namespace UE::EditorDomain
 {
 
+void PrecacheClassDigest(const FTopLevelAssetPath& ClassPath);
+
 FStringView RemoveConfigComment(FStringView Line)
 {
 	int32 CommentIndex;
@@ -457,55 +459,57 @@ FPackageDigest CalculatePackageDigest(const FAssetPackageData& PackageData, FNam
 	return Result;
 }
 
-bool TryAppendClassDigests(FBlake3& Writer, TConstArrayView<FTopLevelAssetPath> ClassPaths, FString* OutErrorMessage)
+bool TryAppendClassDigests(const TConstArrayView<UE::Cook::FCookDependency>& Dependencies, FBlake3& Writer, UE::Cook::FCookDependencyContext& Context)
 {
 	FClassDigestMap& ClassDigests = GetClassDigests();
-	bool bHasTriedPrecacheClassDigests = false;
-	int32 NextClass = 0;
-	while (NextClass < ClassPaths.Num())
+
+	FReadScopeLock ClassDigestsScopeLock(ClassDigests.Lock);
+
+	for (const UE::Cook::FCookDependency& ClassDependency : Dependencies)
 	{
+		FStringView ClassPath = ClassDependency.GetClassPath();
+		FTopLevelAssetPath AssetPath(ClassPath);
+		if (!AssetPath.IsValid())
 		{
-			FReadScopeLock ClassDigestsScopeLock(ClassDigests.Lock);
-			while (NextClass < ClassPaths.Num())
-			{
-				const FTopLevelAssetPath& ClassPath = ClassPaths[NextClass];
-				FClassDigestData* ExistingData = ClassDigests.Map.Find(ClassPath);
-				if (!ExistingData)
-				{
-					break;
-				}
-				NextClass++;
-
-				// We only support hashes for the ClosestNative class
-				FClassDigestData* NativeData = ExistingData;
-				if (ExistingData->ClosestNative != ClassPath)
-				{
-					NativeData = ClassDigests.Map.Find(ExistingData->ClosestNative);
-					checkf(NativeData, TEXT("Classes are only stored in a ClosestNative field if they exist"));
-				}
-				check(NativeData->bNative);
-				Writer.Update(&NativeData->InclusiveSchemaHash, sizeof(NativeData->InclusiveSchemaHash));
-			}
+			Context.LogError(FString::Printf(TEXT("ClassDependency failed: %s is not a valid TopLevelAssetPath."), *AssetPath.ToString()));
+			continue;
 		}
 
-		if (NextClass < ClassPaths.Num())
+		FClassDigestData* ExistingData = ClassDigests.Map.Find(AssetPath);
+		if (!ExistingData && (!IsInGameThread() || GIsSavingPackage))
 		{
-			// EDITORDOMAIN_TODO: Remove the clauses !IsInGameThread || GIsSavingPackage once FindObject no longer asserts if GIsSavingPackage
-			if (bHasTriedPrecacheClassDigests || !IsInGameThread() || GIsSavingPackage)
+			Context.LogError(FString::Printf(TEXT("%s is not a valid class."), *AssetPath.ToString()));
+			return false;
+		}
+		else if (!ExistingData)
+		{
+			// PrecacheClassDigest will try to acquire a write lock. It can't upgrade a read lock to a write lock. That's why
+			// I release the read lock here.
+			ClassDigests.Lock.ReadUnlock();
+
+			PrecacheClassDigest(AssetPath);
+
+			ClassDigests.Lock.ReadLock();
+
+			ExistingData = ClassDigests.Map.Find(AssetPath);
+			if (!ExistingData)
 			{
-				if (OutErrorMessage)
-				{
-					*OutErrorMessage = FString::Printf(TEXT("%s is not a valid class."),
-						*ClassPaths[NextClass].ToString());
-				}
+				Context.LogError(FString::Printf(TEXT("%s is not a valid class."), *AssetPath.ToString()));
 				return false;
 			}
-			TConstArrayView<FTopLevelAssetPath> RemainingClasses =
-				TConstArrayView<FTopLevelAssetPath>(ClassPaths).RightChop(NextClass);
-			PrecacheClassDigests(RemainingClasses);
-			bHasTriedPrecacheClassDigests = true;
+		}	
+
+		// We only support hashes for the ClosestNative class
+		FClassDigestData* NativeData = ExistingData;
+		if (ExistingData->ClosestNative != AssetPath)
+		{
+			NativeData = ClassDigests.Map.Find(ExistingData->ClosestNative);
+			checkf(NativeData, TEXT("Classes are only stored in a ClosestNative field if they exist"));
 		}
+		check(NativeData->bNative);
+		Writer.Update(&NativeData->InclusiveSchemaHash, sizeof(NativeData->InclusiveSchemaHash));
 	}
+
 	return true;
 }
 
@@ -819,6 +823,15 @@ void PrecacheClassDigests(TConstArrayView<FTopLevelAssetPath> ClassPaths)
 	}
 }
 
+void PrecacheClassDigest(const FTopLevelAssetPath& ClassPath)
+{
+	if (ClassPath.IsNull())
+		return;
+
+	FPrecacheClassDigest Digester;
+	Digester.GetRecursive(ClassPath, true /* bAllowRedirects */);
+}
+
 /** Construct GGlobalConstructClasses from the classes specified by config */
 void InitializeGlobalConstructClasses()
 {
diff --git a/Engine/Source/Editor/UnrealEd/Private/EditorDomain/EditorDomainUtils.h b/Engine/Source/Editor/UnrealEd/Private/EditorDomain/EditorDomainUtils.h
--- a/Engine/Source/Editor/UnrealEd/Private/EditorDomain/EditorDomainUtils.h
+++ b/Engine/Source/Editor/UnrealEd/Private/EditorDomain/EditorDomainUtils.h
@@ -28,6 +28,9 @@ class FPackagePath;
 class IAssetRegistry;
 class UPackage;
 
+namespace UE::Cook { class FCookDependency; }
+namespace UE::Cook { struct FCookDependencyContext; }
+
 namespace UE::DerivedData { struct FCacheKey; }
 
 namespace UE::EditorDomain
@@ -89,8 +92,9 @@ ESaveStorageResult SaveStorageResultFromString(FUtf8StringView Text);
  */
 FPackageDigest CalculatePackageDigest(IAssetRegistry& AssetRegistry, FName PackageName);
 FPackageDigest CalculatePackageDigest(const FAssetPackageData& PackageData, FName PackageName);
-/** FindOrAdd class digest for the given classpath, and append its hash onto the writer. Returns false if not found. */
-bool TryAppendClassDigests(FBlake3& Writer, TConstArrayView<FTopLevelAssetPath> ClassPaths, FString* OutErrorMessage);
+
+/** FindOrAdd class digests for the given class dependencies, and append hashes onto the writer. Returns false if not found. */
+bool TryAppendClassDigests(const TConstArrayView<UE::Cook::FCookDependency>& Dependencies, FBlake3& Writer, UE::Cook::FCookDependencyContext& Context);
 
 /** For any ClassNames not already in ClassDigests, look up their UStruct and add them. */
 void PrecacheClassDigests(TConstArrayView<FTopLevelAssetPath> ClassNames);
