Cooker: CookGenerationHelper: Suppress log spam when the lifetime-management functions such as SetKeepForGeneratorSave are called when already set.
#rnx

diff --git a/Engine/Source/Editor/UnrealEd/Private/Cooker/CookGenerationHelper.h b/Engine/Source/Editor/UnrealEd/Private/Cooker/CookGenerationHelper.h
--- a/Engine/Source/Editor/UnrealEd/Private/Cooker/CookGenerationHelper.h
+++ b/Engine/Source/Editor/UnrealEd/Private/Cooker/CookGenerationHelper.h
@@ -742,34 +742,46 @@ inline const TMap<FName, FAssetPackageData>& FGenerationHelper::GetPreviousGener
 
 inline void FGenerationHelper::SetKeepForIterative()
 {
-	UE_LOG(LogCookGenerationHelper, Verbose, TEXT("%s: SetKeepForIterative."),
-		*WriteToString<256>(GetOwner().GetPackageName()));
+	if (!ReferenceFromKeepForIterative)
+	{
+		UE_LOG(LogCookGenerationHelper, Verbose, TEXT("%s: SetKeepForIterative."),
+			*WriteToString<256>(GetOwner().GetPackageName()));
 
-	ReferenceFromKeepForIterative = this;
+		ReferenceFromKeepForIterative = this;
+	}
 }
 
 inline void FGenerationHelper::ClearKeepForIterative()
 {
-	UE_LOG(LogCookGenerationHelper, Verbose, TEXT("%s: ClearKeepForIterative."),
-		*WriteToString<256>(GetOwner().GetPackageName()));
+	if (ReferenceFromKeepForIterative)
+	{
+		UE_LOG(LogCookGenerationHelper, Verbose, TEXT("%s: ClearKeepForIterative."),
+			*WriteToString<256>(GetOwner().GetPackageName()));
 
-	ReferenceFromKeepForIterative.SafeRelease();
+		ReferenceFromKeepForIterative.SafeRelease();
+	}
 }
 
 inline void FGenerationHelper::SetKeepForQueueResults()
 {
-	UE_LOG(LogCookGenerationHelper, Verbose, TEXT("%s: SetKeepForQueueResults."),
-		*WriteToString<256>(GetOwner().GetPackageName()));
+	if (!ReferenceFromKeepForQueueResults)
+	{
+		UE_LOG(LogCookGenerationHelper, Verbose, TEXT("%s: SetKeepForQueueResults."),
+			*WriteToString<256>(GetOwner().GetPackageName()));
 
-	ReferenceFromKeepForQueueResults = this;
+		ReferenceFromKeepForQueueResults = this;
+	}
 }
 
 inline void FGenerationHelper::ClearKeepForQueueResults()
 {
-	UE_LOG(LogCookGenerationHelper, Verbose, TEXT("%s: ClearKeepForQueueResults."),
-		*WriteToString<256>(GetOwner().GetPackageName()));
+	if (ReferenceFromKeepForQueueResults)
+	{
+		UE_LOG(LogCookGenerationHelper, Verbose, TEXT("%s: ClearKeepForQueueResults."),
+			*WriteToString<256>(GetOwner().GetPackageName()));
 
-	ReferenceFromKeepForQueueResults.SafeRelease();
+		ReferenceFromKeepForQueueResults.SafeRelease();
+	}
 }
 
 inline bool FGenerationHelper::IsWaitingForQueueResults() const
@@ -779,59 +791,97 @@ inline bool FGenerationHelper::IsWaitingForQueueResults() const
 
 inline void FGenerationHelper::SetKeepForGeneratorSave()
 {
-	UE_LOG(LogCookGenerationHelper, Verbose, TEXT("%s: SetKeepForGeneratorSave."),
-		*WriteToString<256>(GetOwner().GetPackageName()));
+	if (!ReferenceFromKeepForGeneratorSave)
+	{
+		UE_LOG(LogCookGenerationHelper, Verbose, TEXT("%s: SetKeepForGeneratorSave."),
+			*WriteToString<256>(GetOwner().GetPackageName()));
 
-	ReferenceFromKeepForGeneratorSave = this;
+		ReferenceFromKeepForGeneratorSave = this;
+	}
 }
 
 inline void FGenerationHelper::ClearKeepForGeneratorSave()
 {
-	UE_LOG(LogCookGenerationHelper, Verbose, TEXT("%s: ClearKeepForGeneratorSave."),
-		*WriteToString<256>(GetOwner().GetPackageName()));
+	if (ReferenceFromKeepForGeneratorSave)
+	{
+		UE_LOG(LogCookGenerationHelper, Verbose, TEXT("%s: ClearKeepForGeneratorSave."),
+			*WriteToString<256>(GetOwner().GetPackageName()));
 
-	ReferenceFromKeepForGeneratorSave.SafeRelease();
+		ReferenceFromKeepForGeneratorSave.SafeRelease();
+	}
 }
 
 inline void FGenerationHelper::SetKeepForAllSavedOrGC()
 {
-	UE_LOG(LogCookGenerationHelper, Verbose, TEXT("%s: SetKeepForAllSavedOrGC."),
-		*WriteToString<256>(GetOwner().GetPackageName()));
+	if (!bKeepForAllSavedOrGC)
+	{
+		UE_LOG(LogCookGenerationHelper, Verbose, TEXT("%s: SetKeepForAllSavedOrGC."),
+			*WriteToString<256>(GetOwner().GetPackageName()));
 
-	ReferenceFromKeepForAllSavedOrGC = this;
-	bKeepForAllSavedOrGC = true;
+		ReferenceFromKeepForAllSavedOrGC = this;
+		bKeepForAllSavedOrGC = true;
+	}
+	else
+	{
+		// If bKeepForAllSavedOrGC, ReferenceFromKeepForAllSavedOrGC must be true as well.
+		check(ReferenceFromKeepForAllSavedOrGC);
+	}
 }
 
 inline void FGenerationHelper::ClearKeepForAllSavedOrGC()
 {
-	UE_LOG(LogCookGenerationHelper, Verbose, TEXT("%s: ClearKeepForAllSavedOrGC."),
-		*WriteToString<256>(GetOwner().GetPackageName()));
+	if (bKeepForAllSavedOrGC)
+	{
+		UE_LOG(LogCookGenerationHelper, Verbose, TEXT("%s: ClearKeepForAllSavedOrGC."),
+			*WriteToString<256>(GetOwner().GetPackageName()));
 
-	bKeepForAllSavedOrGC = false;
-	if (!bKeepForAllSavedOrGC && !bKeepForCompletedAllSavesMessage)
+		bKeepForAllSavedOrGC = false;
+		if (!bKeepForAllSavedOrGC && !bKeepForCompletedAllSavesMessage)
+		{
+			ReferenceFromKeepForAllSavedOrGC.SafeRelease();
+		}
+	}
+	else
 	{
-		ReferenceFromKeepForAllSavedOrGC.SafeRelease();
+		// ReferenceFromKeepForAllSavedOrGC must only be set if at least one of bKeepForAllSavedOrGC or bKeepForCompletedAllSavesMessage is set
+		check(bKeepForAllSavedOrGC || bKeepForCompletedAllSavesMessage || !ReferenceFromKeepForAllSavedOrGC);
 	}
 }
 
 inline void FGenerationHelper::SetKeepForCompletedAllSavesMessage()
 {
-	UE_LOG(LogCookGenerationHelper, Verbose, TEXT("%s: SetKeepForCompletedAllSavesMessage."),
-		*WriteToString<256>(GetOwner().GetPackageName()));
+	if (!bKeepForCompletedAllSavesMessage)
+	{
+		UE_LOG(LogCookGenerationHelper, Verbose, TEXT("%s: SetKeepForCompletedAllSavesMessage."),
+			*WriteToString<256>(GetOwner().GetPackageName()));
 
-	ReferenceFromKeepForAllSavedOrGC = this;
-	bKeepForCompletedAllSavesMessage = true;
+		ReferenceFromKeepForAllSavedOrGC = this;
+		bKeepForCompletedAllSavesMessage = true;
+	}
+	else
+	{
+		// If bKeepForCompletedAllSavesMessage, ReferenceFromKeepForAllSavedOrGC must be true as well.
+		check(ReferenceFromKeepForAllSavedOrGC);
+	}
 }
 
 inline void FGenerationHelper::ClearKeepForCompletedAllSavesMessage()
 {
-	UE_LOG(LogCookGenerationHelper, Verbose, TEXT("%s: ClearKeepForCompletedAllSavesMessage."),
-		*WriteToString<256>(GetOwner().GetPackageName()));
+	if (bKeepForCompletedAllSavesMessage)
+	{
+		UE_LOG(LogCookGenerationHelper, Verbose, TEXT("%s: ClearKeepForCompletedAllSavesMessage."),
+			*WriteToString<256>(GetOwner().GetPackageName()));
 
-	bKeepForCompletedAllSavesMessage = false;
-	if (!bKeepForAllSavedOrGC && !bKeepForCompletedAllSavesMessage)
+		bKeepForCompletedAllSavesMessage = false;
+		if (!bKeepForAllSavedOrGC && !bKeepForCompletedAllSavesMessage)
+		{
+			ReferenceFromKeepForAllSavedOrGC.SafeRelease();
+		}
+	}
+	else
 	{
-		ReferenceFromKeepForAllSavedOrGC.SafeRelease();
+		// ReferenceFromKeepForAllSavedOrGC must only be set if at least one of bKeepForAllSavedOrGC or bKeepForCompletedAllSavesMessage is set
+		check(bKeepForAllSavedOrGC || bKeepForCompletedAllSavesMessage || !ReferenceFromKeepForAllSavedOrGC);
 	}
 }
 
