When encountering an in-memory package (eg: ones generated by verse) that the AssetRegistry is unaware of, ensure that the IsIterativeEnabled method does an equivalent evaluation as an on-disk package based on the imported classes for the in-memory package.  This allows us to correctly exclude VerseClass from incremental cook until its dependencies are declared correctly.

diff --git a/Engine/Source/Editor/UnrealEd/Private/TargetDomain/TargetDomainUtils.cpp b/Engine/Source/Editor/UnrealEd/Private/TargetDomain/TargetDomainUtils.cpp
--- a/Engine/Source/Editor/UnrealEd/Private/TargetDomain/TargetDomainUtils.cpp
+++ b/Engine/Source/Editor/UnrealEd/Private/TargetDomain/TargetDomainUtils.cpp
@@ -34,8 +34,14 @@ bool IsIterativeEnabled(FName PackageName, bool bAllowAllClasses)
 		return false;
 	}
 	
+	TArray<FName>* ImportedClasses = nullptr;
+	TArray<FName> InMemoryImportedClasses;
 	TOptional<FAssetPackageData> PackageDataOpt = AssetRegistry->GetAssetPackageDataCopy(PackageName);
-	if(!PackageDataOpt)
+	if (PackageDataOpt)
+	{
+		ImportedClasses = &PackageDataOpt->ImportedClasses;
+	}
+	else
 	{
 		TStringBuilder<256> PackageNameStr(EInPlace::InPlace, PackageName);
 		const UPackage* Package = FindPackage(nullptr, *PackageNameStr);
@@ -50,13 +56,27 @@ bool IsIterativeEnabled(FName PackageName, bool bAllowAllClasses)
 			// Strange case where the package does not exist on disk nor in memory.
 			return false;
 		}
-
-		//This is an in memory package
-		return true;
+		TSet<UClass*> PackageClasses;
+		ForEachObjectWithPackage(Package, [&PackageClasses, Package](UObject* Object)
+			{
+				UClass* Class = Object->GetClass();
+				if (!Class->IsInPackage(Package)) // Imported classes list does not include classes in the package
+				{
+					PackageClasses.Add(Object->GetClass());
+				}
+				return true;
+			});
+		InMemoryImportedClasses.Reserve(PackageClasses.Num());
+		for (UClass* Class : PackageClasses)
+		{
+			TStringBuilder<256> ClassPath;
+			Class->GetPathName(nullptr, ClassPath);
+			InMemoryImportedClasses.Add(FName(ClassPath));
+		}
+		InMemoryImportedClasses.Sort(FNameLexicalLess());
+		ImportedClasses = &InMemoryImportedClasses;
 	}
 
-	const FAssetPackageData& PackageData = *PackageDataOpt.GetPtrOrNull();
-
 	if (!bAllowAllClasses)
 	{
 		auto LogInvalidDueTo = [](FName PackageName, FName ClassPath)
@@ -66,7 +86,7 @@ bool IsIterativeEnabled(FName PackageName, bool bAllowAllClasses)
 
 		UE::EditorDomain::FClassDigestMap& ClassDigests = UE::EditorDomain::GetClassDigests();
 		FReadScopeLock ClassDigestsScopeLock(ClassDigests.Lock);
-		for (FName ClassName : PackageData.ImportedClasses)
+		for (FName ClassName : *ImportedClasses)
 		{
 			FTopLevelAssetPath ClassPath(WriteToString<256>(ClassName).ToView());
 			UE::EditorDomain::FClassDigestData* ExistingData = nullptr;
