IncrementalCook: CookPackageSplitter: Fix CookGenerationHelper to track whether a package was iteratively skipped separately per platform. Previously platforms were overwriting each other's data and our lifetime counters were getting corrupted and causing warnings such as "RecordResults received for generated package %s, but its ParentGenerationHelper has already been destructed..."
#rnx
#incrementalcook

diff --git a/Engine/Source/Editor/UnrealEd/Private/CookOnTheFlyServer.cpp b/Engine/Source/Editor/UnrealEd/Private/CookOnTheFlyServer.cpp
--- a/Engine/Source/Editor/UnrealEd/Private/CookOnTheFlyServer.cpp
+++ b/Engine/Source/Editor/UnrealEd/Private/CookOnTheFlyServer.cpp
@@ -4147,7 +4147,7 @@ UE::Cook::EPollStatus UCookOnTheFlyServer::PrepareSaveInternal(UE::Cook::FPackag
 			if (PackageData.GetSaveSubState() <= ESaveSubState::Generation_TryGenerateList)
 			{
 				// Keep it referenced even if we are only precaching, so we do not recreate it
-				GenerationHelper->SetKeepForGeneratorSave();
+				GenerationHelper->SetKeepForGeneratorSaveAllPlatforms();
 				if (bPrecaching)
 				{
 					// Do not proceed to GetGenerateList when precaching; do that only when we're ready to save the package
@@ -7747,7 +7747,7 @@ void UCookOnTheFlyServer::PopulateCookedPackages(TArrayView<const ITargetPlatfor
 				continue;
 			}
 			TRefCountPtr<FGenerationHelper> GenerationHelper = Generator->CreateUninitializedGenerationHelper();
-			GenerationHelper->SetPreviousGeneratedPackages(MoveTemp(Pair.Value.Generated));
+			GenerationHelper->SetPreviousGeneratedPackages(TargetPlatform, MoveTemp(Pair.Value.Generated));
 		}
 
 		PlatformAssetRegistry.SetPreviousAssetRegistry(MoveTemp(PreviousAssetRegistry));
@@ -9430,7 +9430,7 @@ void UCookOnTheFlyServer::CookByTheBookFinishedInternal()
 			{
 				// One reason it might still be around is the keepforiterativeflag, if it were in the oplog but never
 				// cooked. Clear that flag now and then retest whether it is still referenced.
-				GenerationHelper->ClearKeepForIterative();
+				GenerationHelper->ClearKeepForIterativeAllPlatforms();
 				GenerationHelper.SafeRelease();
 				if (PackageData->GetGenerationHelper())
 				{
diff --git a/Engine/Source/Editor/UnrealEd/Private/Cooker/CookDirector.cpp b/Engine/Source/Editor/UnrealEd/Private/Cooker/CookDirector.cpp
--- a/Engine/Source/Editor/UnrealEd/Private/Cooker/CookDirector.cpp
+++ b/Engine/Source/Editor/UnrealEd/Private/Cooker/CookDirector.cpp
@@ -364,8 +364,9 @@ void FCookDirector::AssignRequests(TArrayView<FPackageData*> Requests, TArray<FW
 TMap<FPackageData*, FAssignPackageExtraData> FCookDirector::GetAssignPackageExtraDatas(
 	TConstArrayView<FPackageData*> Requests) const
 {
+	TConstArrayView<const ITargetPlatform*> SessionPlatforms = COTFS.PlatformManager->GetSessionPlatforms();
 	FMPCollectorServerTickPackageContext Context;
-	Context.Platforms = COTFS.PlatformManager->GetSessionPlatforms();
+	Context.Platforms = SessionPlatforms;
 
 	TMap<FPackageData*, FAssignPackageExtraData> Results;
 	for (FPackageData* Request : Requests)
@@ -382,10 +383,16 @@ TMap<FPackageData*, FAssignPackageExtraData> FCookDirector::GetAssignPackageExtr
 		};
 
 		TRefCountPtr<FGenerationHelper> GenerationHelper = Request->GetGenerationHelper();
-		if (GenerationHelper && !GenerationHelper->GetPreviousGeneratedPackages().IsEmpty())
+		if (GenerationHelper)
 		{
 			ExtraData = GetOrAllocateExtraData(Request);
-			ExtraData->GeneratorPreviousGeneratedPackages = GenerationHelper->GetPreviousGeneratedPackages();
+			ExtraData->GeneratorPerPlatformPreviousGeneratedPackages.Reserve(SessionPlatforms.Num());
+			for (const ITargetPlatform* TargetPlatform : SessionPlatforms)
+			{
+				TMap<FName, FAssetPackageData>& PreviousGeneratedPackages
+					= ExtraData->GeneratorPerPlatformPreviousGeneratedPackages.FindOrAdd(TargetPlatform);
+				PreviousGeneratedPackages = GenerationHelper->GetPreviousGeneratedPackages(TargetPlatform);
+			}
 		}
 
 		Context.PackageName = Request->GetPackageName();
diff --git a/Engine/Source/Editor/UnrealEd/Private/Cooker/CookGenerationHelper.cpp b/Engine/Source/Editor/UnrealEd/Private/Cooker/CookGenerationHelper.cpp
--- a/Engine/Source/Editor/UnrealEd/Private/Cooker/CookGenerationHelper.cpp
+++ b/Engine/Source/Editor/UnrealEd/Private/Cooker/CookGenerationHelper.cpp
@@ -12,6 +12,7 @@
 #include "Cooker/IWorkerRequests.h"
 #include "Cooker/PackageTracker.h"
 #include "EditorDomain/EditorDomainUtils.h"
+#include "HAL/PlatformMath.h"
 #include "Interfaces/ITargetPlatform.h"
 #include "Misc/CommandLine.h"
 #include "Misc/ConfigCacheIni.h"
@@ -156,8 +157,17 @@ namespace UE::Cook
 // FGenerationHelper
 
 FGenerationHelper::FGenerationHelper(FPackageData& InOwner)
-: OwnerInfo(InOwner, true /* bInGenerator */)
+: OwnerInfo(*this, InOwner, true /* bInGenerator */)
 {
+	UCookOnTheFlyServer& COTFS = InOwner.GetPackageDatas().GetCookOnTheFlyServer();
+	TConstArrayView<const ITargetPlatform*> TargetPlatforms = COTFS.PlatformManager->GetSessionPlatforms();
+	PlatformDatas.Reserve(TargetPlatforms.Num());
+	OwnerInfo.PlatformDatas.Reserve(TargetPlatforms.Num());
+	for (const ITargetPlatform* TargetPlatform : TargetPlatforms)
+	{
+		PlatformDatas.Add(TargetPlatform);
+		OwnerInfo.PlatformDatas.Add(TargetPlatform);
+	}
 }
 
 FGenerationHelper::~FGenerationHelper()
@@ -299,16 +309,19 @@ void FGenerationHelper::Uninitialize()
 		Info.Uninitialize();
 	}
 	OwnerPackage.Reset();
-	// Keep PreviousGeneratedPackages; they are allowed in the uninitialized state
 	ExternalActorDependencies.Empty();
+
+	// Keep PlatformDatas, they are allowed in the uninitialized state.
+	// Keep PlatformData.PreviousGeneratedPackages; they are allowed in the uninitialized state
+	// PlatformData.ReferenceFromKeepForIterative
+	// PlatformData.ReferenceFromKeepForGeneratorSave
+	// Keep PlatformData.NumSaved; it is allowed in the uninitialized state
+
 	check(OwnerObjectsToMove.IsEmpty()); // We can not still be in the save state, so this should be empty
 	// Do not modify the reference tracking variables
-	// ReferenceFromKeepForIterative
 	// ReferenceFromKeepForQueueResults
-	// ReferenceFromKeepForGeneratorSave
 	// ReferenceFromKeepForAllSavedOrGC
 	// Keep MPCookNextAssignmentIndex; it is allowed in the uninitialized state
-	// Keep NumSaved; it is allowed in the uninitialized state
 	// InitializeStatus was modified above
 	// Keep DoesGeneratedRequireGeneratorValue; it is allowed in the uninitialized state
 	// Keep bUseInternalReferenceToAvoidGarbageCollect; it is allowed in the uninitialized state
@@ -319,41 +332,69 @@ void FGenerationHelper::Uninitialize()
 	// Keep bKeepForAllSavedOrGC ; it is allowed in the uninitialized state
 	// Keep bKeepForCompletedAllSavesMessage; it is allowed in the uninitialized state
 	// Keep bNeedConfirmGeneratorPackageDestroyed; it is allowed in the uninitialized state
+	// Keep bSentAllSavesCompleted; it is allowed in the uninitialized state
+}
+
+void FGenerationHelper::ModifyNumSaved(const ITargetPlatform* TargetPlatform, int32 Delta)
+{
+	FPlatformData& TargetPlatformData = FindCheckedPlatformData(TargetPlatform);
+	TargetPlatformData.NumSaved += Delta;
+	int32 NumAllSaved = PackagesToGenerate.Num() + 1;
+	check(0 <= TargetPlatformData.NumSaved && TargetPlatformData.NumSaved <= NumAllSaved);
+	if (TargetPlatformData.NumSaved != NumAllSaved)
+	{
+		return;
+	}
+	OnNumSavedUpdated();
 }
 
-void FGenerationHelper::ModifyNumSaved(int32 Delta)
+void FGenerationHelper::OnNumSavedUpdated()
 {
-	NumSaved += Delta;
-	if (bDeferEvents)
+	if (bDeferEvents || bSentAllSavesCompleted)
 	{
 		return;
 	}
 
 	int32 NumAllSaved = PackagesToGenerate.Num() + 1;
-	check(0 <= NumSaved && NumSaved <= NumAllSaved);
-	if (NumSaved == NumAllSaved)
-	{
-		UE_LOG(LogCookGenerationHelper, Verbose, TEXT("%s: All saves completed (%d/%d)."),
-			*WriteToString<256>(GetOwner().GetPackageName()), NumSaved, NumAllSaved);
-
-		UCookOnTheFlyServer& COTFS = GetOwner().GetPackageDatas().GetCookOnTheFlyServer();
-		// Only send OnAllSavesCompleted from director; clients have incomplete information and
-		// could send it spuriously.
-		// Additionally, only send it if we have completed queueing, to avoid sending it prematurely.
-		// ModifyNumSaved(1 == 1) will occur when the generator package is iteratively skipped,
-		// and ModifyNumSaved(k == k, k < expectednumber) can occur if we save some generated packages
-		// (or mark them iteratively skippable) before getting the full list of packages from the worker
-		// that called QueueGeneratedPackages.
-		if (!COTFS.CookWorkerClient && bHasFinishedQueueGeneratedPackages)
-		{
-			if (COTFS.CookDirector)
-			{
-				FName PackageName = GetOwner().GetPackageName();
-				FGeneratorEventMessage Message(EGeneratorEvent::AllSavesCompleted, PackageName);
-				COTFS.CookDirector->BroadcastMessage(MoveTemp(Message), ECookBroadcastTiming::AfterAssignPackages);
-			}
-			OnAllSavesCompleted(COTFS);
+	int32 NumSaved = 0;
+
+	// Save events are combined for all platforms; early exit for now if any platform is incomplete.
+	for (TPair<const ITargetPlatform*, FPlatformData>& PlatformPair : PlatformDatas)
+	{
+		FPlatformData& PlatformData = PlatformPair.Value;
+		check(0 <= PlatformData.NumSaved && PlatformData.NumSaved <= NumAllSaved);
+		if (PlatformData.NumSaved != NumAllSaved)
+		{
+			return;
 		}
+		NumSaved = PlatformData.NumSaved;
+	}
+
+	UCookOnTheFlyServer& COTFS = GetOwner().GetPackageDatas().GetCookOnTheFlyServer();
+	// Only send OnAllSavesCompleted from director; clients have incomplete information and
+	// could send it spuriously.
+	// Additionally, only send it if we have completed queueing, to avoid sending it prematurely.
+	// ModifyNumSaved(1 == 1) will occur when the generator package is iteratively skipped,
+	// and ModifyNumSaved(k == k, k < expectednumber) can occur if we save some generated packages
+	// (or mark them iteratively skippable) before getting the full list of packages from the worker
+	// that called QueueGeneratedPackages.
+	if (!COTFS.CookWorkerClient && !bHasFinishedQueueGeneratedPackages)
+	{
+		return;
+	}
+
+	FName PackageName = GetOwner().GetPackageName();
+	bSentAllSavesCompleted = true;
+	UE_LOG(LogCookGenerationHelper, Verbose, TEXT("%s: All saves completed (%d/%d)."),
+		*WriteToString<256>(PackageName), NumSaved, NumAllSaved);
+	if (!COTFS.CookWorkerClient)
+	{
+		if (COTFS.CookDirector)
+		{
+			FGeneratorEventMessage Message(EGeneratorEvent::AllSavesCompleted, PackageName);
+			COTFS.CookDirector->BroadcastMessage(MoveTemp(Message), ECookBroadcastTiming::AfterAssignPackages);
+		}
+		OnAllSavesCompleted(COTFS);
 	}
 }
 
@@ -370,7 +411,7 @@ FGenerationHelper::FScopeDeferEvents::~FScopeDeferEvents()
 	GenerationHelper->bDeferEvents = bOldDeferEvents;
 	if (!bOldDeferEvents)
 	{
-		GenerationHelper->ModifyNumSaved(0);
+		GenerationHelper->OnNumSavedUpdated();
 	}
 }
 
@@ -402,12 +443,17 @@ void FGenerationHelper::DemoteStalledPackages(UCookOnTheFlyServer& COTFS, bool b
 					// and demote them to idle.
 					if (bFromAllSavesCompleted)
 					{
+						int32 MinNumSaved = PackagesToGenerate.Num() + 1;
+						for (TPair<const ITargetPlatform*, FPlatformData>& PlatformPair : PlatformDatas)
+						{
+							MinNumSaved = FMath::Min(MinNumSaved, PlatformPair.Value.NumSaved);
+						}
 						UE_LOG(LogCook, Error,
 							TEXT("Package %s is still stalled on the CookDirector during FGenerationHelper::OnAllSavesCompleted. ")
 							TEXT("This is unexpected; all stalled packages on a GenerationHelper should have completed saving and therefore unstalled before OnAllSavesCompleted is called.")
 							TEXT("\n\tNumAllSaved == %d. NumSaved == %d. Info.HasSaved == %s."),
-							*Info.PackageData->GetPackageName().ToString(), PackagesToGenerate.Num() + 1, NumSaved,
-							Info.HasSaved() ? TEXT("true") : TEXT("false"));
+							*Info.PackageData->GetPackageName().ToString(), PackagesToGenerate.Num() + 1, MinNumSaved,
+							Info.HasSavedEveryPlatform() ? TEXT("true") : TEXT("false"));
 						FDebug::DumpStackTraceToLog(ELogVerbosity::Warning);
 
 						COTFS.DemoteToIdle(*Info.PackageData, ESendFlags::QueueAddAndRemove, ESuppressCookReason::RetractedByCookDirector);
@@ -462,10 +508,16 @@ void FGenerationHelper::ValidateSaveStalledState(UCookOnTheFlyServer& COTFS, FPa
 void FGenerationHelper::DiagnoseWhyNotShutdown()
 {
 	TStringBuilder<256> Lines;
+	int32 MinNumSaved = PackagesToGenerate.Num() + 1;
+	for (TPair<const ITargetPlatform*, FPlatformData>& PlatformPair : PlatformDatas)
+	{
+		MinNumSaved = FMath::Min(MinNumSaved, PlatformPair.Value.NumSaved);
+	}
+
 	int32 ExpectedNumSaved = PackagesToGenerate.Num() + 1;
-	if (NumSaved != ExpectedNumSaved)
+	if (MinNumSaved != ExpectedNumSaved)
 	{
-		Lines.Appendf(TEXT("\tNumSaved == %d, ExpectedNumSaved == %d.\n"), NumSaved, ExpectedNumSaved);
+		Lines.Appendf(TEXT("\tNumSaved == %d, ExpectedNumSaved == %d.\n"), MinNumSaved, ExpectedNumSaved);
 	}
 	UCookOnTheFlyServer& COTFS = GetOwner().GetPackageDatas().GetCookOnTheFlyServer();
 	uint32 ExpectedRefCount = 1;
@@ -510,7 +562,7 @@ void FGenerationHelper::DiagnoseWhyNotShutdown()
 						LexToString(Info.PackageData->GetSuppressCookReason()));
 				}
 			}
-			if (!Info.HasSaved())
+			if (!Info.HasSavedEveryPlatform())
 			{
 				Lines.Appendf(TEXT("\t%s%s has not marked saved.\n"),
 					Info.IsGenerator() ? TEXT("OwnerInfo") : TEXT("GeneratedPackage "),
@@ -529,21 +581,27 @@ void FGenerationHelper::DiagnoseWhyNotShutdown()
 		TestInfo(Info);
 	}
 
-	if (ReferenceFromKeepForIterative)
+	for (TPair<const ITargetPlatform*, FPlatformData>& PlatformPair : PlatformDatas)
 	{
-		Lines.Append(TEXT("\tReferenceFromKeepForIterative is set.\n"));
-		++ExpectedRefCount;
+		FPlatformData& PlatformData = PlatformPair.Value;
+		if (PlatformData.ReferenceFromKeepForIterative)
+		{
+			Lines.Appendf(TEXT("\tReferenceFromKeepForIterative is set for platform %s.\n"),
+				*PlatformPair.Key->PlatformName());
+			++ExpectedRefCount;
+		}
+		if (PlatformData.ReferenceFromKeepForGeneratorSave)
+		{
+			Lines.Appendf(TEXT("\tReferenceFromKeepForGeneratorSave is set for platform %s.\n"),
+				*PlatformPair.Key->PlatformName());
+			++ExpectedRefCount;
+		}
 	}
 	if (ReferenceFromKeepForQueueResults)
 	{
 		Lines.Append(TEXT("\tReferenceFromKeepForQueueResults is set.\n"));
 		++ExpectedRefCount;
 	}
-	if (ReferenceFromKeepForGeneratorSave)
-	{
-		Lines.Append(TEXT("\tReferenceFromKeepForGeneratorSave is set.\n"));
-		++ExpectedRefCount;
-	}
 	if (bKeepForAllSavedOrGC)
 	{
 		Lines.Append(TEXT("\tbKeepForAllSavedOrGC is true.\n"));
@@ -769,8 +827,8 @@ void FGenerationHelper::ClearSelfReferences()
 	// Any references we release might be the last reference and cause *this to be deleted,
 	// so create a local reference to keep it alive until the end of the function.
 	TRefCountPtr<FGenerationHelper> LocalRef(this);
-	ClearKeepForIterative();
-	ClearKeepForGeneratorSave();
+	ClearKeepForIterativeAllPlatforms();
+	ClearKeepForGeneratorSaveAllPlatforms();
 	ClearKeepForQueueResults();
 	ClearKeepForAllSavedOrGC();
 	ClearKeepForCompletedAllSavesMessage();
@@ -938,7 +996,10 @@ bool FGenerationHelper::TryGenerateList()
 	}
 	PackagesToGenerate.Reserve(GeneratorDatas.Num());
 
-	NumSaved = 0;
+	for (TPair<const ITargetPlatform*, FPlatformData>& PlatformPair : PlatformDatas)
+	{
+		PlatformPair.Value.NumSaved = 0;
+	}
 	for (ICookPackageSplitter::FGeneratedPackage& SplitterData : GeneratorDatas)
 	{
 		if (!SplitterData.GetCreateAsMap().IsSet())
@@ -986,7 +1047,7 @@ bool FGenerationHelper::TryGenerateList()
 		}
 		if (!GeneratedInfo)
 		{
-			GeneratedInfo = &PackagesToGenerate.Emplace_GetRef(*PackageData, false /* bInGenerator */);
+			GeneratedInfo = &PackagesToGenerate.Emplace_GetRef(*this, *PackageData, false /* bInGenerator */);
 		}
 		GeneratedInfo->RelativePath = MoveTemp(SplitterData.RelativePath);
 		GeneratedInfo->GeneratedRootPath = MoveTemp(SplitterData.GeneratedRootPath);
@@ -1022,9 +1083,22 @@ bool FGenerationHelper::TryGenerateList()
 		// Create the Hash from the GenerationHash and Dependencies
 		GeneratedInfo->CreatePackageHash();
 
-		NumSaved += GeneratedInfo->HasSaved() ? 1 : 0;
+		for (TPair<const ITargetPlatform*, FPlatformData>& PlatformPair : PlatformDatas)
+		{
+			FPlatformData& GeneratorPlatformData = PlatformPair.Value;
+			FCookGenerationInfo::FPlatformData& InfoPlatformData = GeneratedInfo->FindCheckedPlatformData(PlatformPair.Key);
+			GeneratorPlatformData.NumSaved += InfoPlatformData.HasSaved() ? 1 : 0;
+		}
+	}
+
+	int32 MinNumSaved = PackagesToGenerate.Num() + 1;
+	for (TPair<const ITargetPlatform*, FPlatformData>& PlatformPair : PlatformDatas)
+	{
+		FPlatformData& GeneratorPlatformData = PlatformPair.Value;
+		FCookGenerationInfo::FPlatformData& InfoPlatformData = OwnerInfo.FindCheckedPlatformData(PlatformPair.Key);
+		GeneratorPlatformData.NumSaved += InfoPlatformData.HasSaved() ? 1 : 0;
+		MinNumSaved = FMath::Min(MinNumSaved, GeneratorPlatformData.NumSaved);
 	}
-	NumSaved += OwnerInfo.HasSaved() ? 1 : 0;
 
 	if (!AlreadyExistingInfoPackageToIndex.IsEmpty())
 	{
@@ -1046,9 +1120,9 @@ bool FGenerationHelper::TryGenerateList()
 	UE_LOG(LogCookGenerationHelper, Verbose,
 		TEXT("%s: TryGenerateList: %d packages to generate, %d previously recorded, %d/%d already saved."),
 		*WriteToString<256>(GetOwner().GetPackageName()), PackagesToGenerate.Num(),
-		NumAlreadyExisting, NumSaved, PackagesToGenerate.Num() + 1);
+		NumAlreadyExisting, MinNumSaved, PackagesToGenerate.Num() + 1);
 
-	ModifyNumSaved(0);
+	OnNumSavedUpdated();
 
 	bGeneratedList = true;
 	return true;
@@ -1228,7 +1302,7 @@ void FGenerationHelper::StartOwnerSave()
 		*WriteToString<256>(GetOwner().GetPackageName()),
 		*GetRegisteredSplitterType()->GetSplitterDebugName(),
 		*WriteToString<256>(GetSplitDataObjectName()));
-	SetKeepForGeneratorSave();
+	SetKeepForGeneratorSaveAllPlatforms();
 }
 
 void FGenerationHelper::StartQueueGeneratedPackages(UCookOnTheFlyServer& COTFS)
@@ -1243,64 +1317,84 @@ void FGenerationHelper::StartQueueGeneratedPackages(UCookOnTheFlyServer& COTFS)
 	NotifyStartQueueGeneratedPackages(COTFS, FWorkerId::Local());
 
 	bool bHybridIterativeEnabled = COTFS.bHybridIterativeEnabled;
-	if (!PreviousGeneratedPackages.IsEmpty())
+	bool bHasPreviousResultsForSomePlatform = false;
+	for (TPair<const ITargetPlatform*, FPlatformData>& PlatformPair : PlatformDatas)
 	{
-		TSet<FName> RemainingPreviousPackages;
-		RemainingPreviousPackages.Reserve(PreviousGeneratedPackages.Num());
-		for (const TPair<FName, FAssetPackageData>& Pair : PreviousGeneratedPackages)
-		{
-			RemainingPreviousPackages.Add(Pair.Key);
-		}
+		bHasPreviousResultsForSomePlatform |= !PlatformPair.Value.PreviousGeneratedPackages.IsEmpty();
+	}
 
+	if (bHasPreviousResultsForSomePlatform)
+	{
 		FPackageData& OwnerPackageData = GetOwner();
 		TArray<const ITargetPlatform*, TInlineAllocator<1>> PlatformsToCook;
 		OwnerPackageData.GetPlatformsNeedingCommit(PlatformsToCook, COTFS.GetCookPhase());
 
-		int32 NumIterativeUnmodified = 0;
-		int32 NumIterativeModified = 0;
-		int32 NumIterativeRemoved = 0;
-		int32 NumIterativePrevious = RemainingPreviousPackages.Num();
-
-		for (FCookGenerationInfo& GeneratedInfo : PackagesToGenerate)
+		for (TPair<const ITargetPlatform*, FPlatformData>& PlatformPair : PlatformDatas)
 		{
-			FName GeneratedPackageName = GeneratedInfo.PackageData->GetPackageName();
-			FAssetPackageData* PreviousAssetData = PreviousGeneratedPackages.Find(GeneratedPackageName);
-			RemainingPreviousPackages.Remove(GeneratedPackageName);
-			if (PreviousAssetData)
+			const ITargetPlatform* TargetPlatform = PlatformPair.Key;
+			FPlatformData& PlatformData = PlatformPair.Value;
+			TMap<FName, FAssetPackageData>& PreviousGeneratedPackages = PlatformData.PreviousGeneratedPackages;
+			if (PreviousGeneratedPackages.IsEmpty())
 			{
-				if (!bHybridIterativeEnabled)
-				{
-					bool bIterativelyUnmodified;
-					GeneratedInfo.IterativeCookValidateOrClear(*this, PlatformsToCook, PreviousAssetData->GetPackageSavedHash(),
-						bIterativelyUnmodified);
-					++(bIterativelyUnmodified ? NumIterativeUnmodified : NumIterativeModified);
-				}
-				else
+				continue;
+			}
+
+			if (!PlatformsToCook.Contains(TargetPlatform))
+			{
+				continue;
+			}
+
+			TSet<FName> RemainingPreviousPackages;
+			RemainingPreviousPackages.Reserve(PreviousGeneratedPackages.Num());
+			for (const TPair<FName, FAssetPackageData>& Pair : PreviousGeneratedPackages)
+			{
+				RemainingPreviousPackages.Add(Pair.Key);
+			}
+
+			int32 NumIterativeUnmodified = 0;
+			int32 NumIterativeModified = 0;
+			int32 NumIterativeRemoved = 0;
+			int32 NumIterativePrevious = RemainingPreviousPackages.Num();
+
+			for (FCookGenerationInfo& GeneratedInfo : PackagesToGenerate)
+			{
+				FName GeneratedPackageName = GeneratedInfo.PackageData->GetPackageName();
+				FAssetPackageData* PreviousAssetData = PreviousGeneratedPackages.Find(GeneratedPackageName);
+				RemainingPreviousPackages.Remove(GeneratedPackageName);
+				if (PreviousAssetData)
 				{
-					// Copy the current value for the package's hash into the PreviousPackageData, for use by incremental cook's
-					// calculation in FRequestCluster::TryCalculateIterativelyUnmodified
-					PreviousAssetData->SetPackageSavedHash(GeneratedInfo.PackageHash);
+					if (!bHybridIterativeEnabled)
+					{
+						bool bIterativelyUnmodified;
+						GeneratedInfo.IterativeCookValidateOrClear(*this, TargetPlatform,
+							PreviousAssetData->GetPackageSavedHash(), bIterativelyUnmodified);
+						++(bIterativelyUnmodified ? NumIterativeUnmodified : NumIterativeModified);
+					}
+					else
+					{
+						// Copy the current value for the package's hash into the PreviousPackageData, for use by
+						// incremental cook's calculation in FRequestCluster::TryCalculateIterativelyUnmodified
+						PreviousAssetData->SetPackageSavedHash(GeneratedInfo.PackageHash);
+					}
 				}
 			}
-		}
-		if (!RemainingPreviousPackages.IsEmpty())
-		{
-			NumIterativeRemoved = RemainingPreviousPackages.Num();
-			for (FName PreviousPackageName : RemainingPreviousPackages)
+			if (!RemainingPreviousPackages.IsEmpty())
 			{
-				for (const ITargetPlatform* TargetPlatform : PlatformsToCook)
+				NumIterativeRemoved = RemainingPreviousPackages.Num();
+				for (FName PreviousPackageName : RemainingPreviousPackages)
 				{
 					COTFS.DeleteOutputForPackage(PreviousPackageName, TargetPlatform);
 				}
 			}
-		}
 
-		if (!bHybridIterativeEnabled)
-		{
-			UE_LOG(LogCook, Display, TEXT("Found %d cooked package(s) in package store for generator package %s."),
-				NumIterativePrevious, *WriteToString<256>(GetOwner().GetPackageName()));
-			UE_LOG(LogCook, Display, TEXT("Keeping %d. Recooking %d. Removing %d."),
-				NumIterativeUnmodified, NumIterativeModified, NumIterativeRemoved);
+			if (!bHybridIterativeEnabled)
+			{
+				UE_LOG(LogCook, Display,
+					TEXT("Found %d cooked package(s) in package store for generator package %s."),
+					NumIterativePrevious, *WriteToString<256>(GetOwner().GetPackageName()));
+				UE_LOG(LogCook, Display, TEXT("Keeping %d. Recooking %d. Removing %d."),
+					NumIterativeUnmodified, NumIterativeModified, NumIterativeRemoved);
+			}
 		}
 	}
 }
@@ -1360,9 +1454,16 @@ void FGenerationHelper::EndQueueGeneratedPackagesOnDirector(UCookOnTheFlyServer&
 
 	// The save message for the owner may have come in before this GenerationHelper was created and thus
 	// MarkPackageSavedRemotely was not called. Check for that case now and marked saved if so.
-	if (GetOwner().HasAnyCookedPlatform())
+	FPackageData& OwnerPackageData = GetOwner();
+	FCookGenerationInfo& LocalOwnerInfo = GetOwnerInfo();
+	for (TPair<const ITargetPlatform*, FCookGenerationInfo::FPlatformData>& PlatformPair
+		: LocalOwnerInfo.PlatformDatas)
 	{
-		GetOwnerInfo().SetHasSaved(*this, true, SourceWorkerId);
+		const ITargetPlatform* TargetPlatform = PlatformPair.Key;
+		if (OwnerPackageData.HasCookedPlatform(TargetPlatform, true /* bIncludeFailed */))
+		{
+			PlatformPair.Value.SetHasSaved(*this, LocalOwnerInfo, TargetPlatform, true /* bValue */, SourceWorkerId);
+		}
 	}
 }
 
@@ -1371,14 +1472,26 @@ void FGenerationHelper::OnRequestFencePassed(UCookOnTheFlyServer& COTFS)
 	// This function should only be called in response to a subscription that is sent from the cook director
 	check(!COTFS.CookWorkerClient);
 
-	if (OwnerInfo.IsIterativelySkipped())
+	const ITargetPlatform* FirstNotSkippedTargetPlatform = nullptr;
+	for (TPair<const ITargetPlatform*, FPlatformData>& PlatformPair : PlatformDatas)
 	{
-		// PumpRequests has completed and we marked ourselves and all generated packages as iteratively skipped,
-		// so we no longer need the PreviouslyCookedData or this entire GenerationHelper
-		UE_LOG(LogCookGenerationHelper, Verbose, TEXT("%s: RequestFencePassed and found to be entirely iteratively skipped."),
-			*WriteToString<256>(GetOwner().GetPackageName()));
-		ClearKeepForIterative();
-		PreviousGeneratedPackages.Empty();
+		const ITargetPlatform* TargetPlatform = PlatformPair.Key;
+		FPlatformData& PlatformData = PlatformPair.Value;
+		if (OwnerInfo.FindCheckedPlatformData(TargetPlatform).IsIterativelySkipped())
+		{
+			// PumpRequests has completed and we marked ourselves and all generated packages as iteratively skipped,
+			// so we no longer need the PreviouslyCookedData or this entire GenerationHelper
+			UE_LOG(LogCookGenerationHelper, Verbose,
+				TEXT("%s: RequestFencePassed and found to be entirely iteratively skipped for %s."),
+				*WriteToString<256>(GetOwner().GetPackageName()), *TargetPlatform->PlatformName());
+			PlatformData.ClearKeepForIterative(*this, TargetPlatform);
+			PlatformData.PreviousGeneratedPackages.Empty();
+		}
+		else
+		{
+			FirstNotSkippedTargetPlatform = FirstNotSkippedTargetPlatform
+				? FirstNotSkippedTargetPlatform : TargetPlatform;
+		}
 	}
 
 	if (bHasFinishedQueueGeneratedPackages)
@@ -1387,10 +1500,10 @@ void FGenerationHelper::OnRequestFencePassed(UCookOnTheFlyServer& COTFS)
 		// and assigned to local ReadyRequests or to a CookWorker. Send OnQueuedGeneratedPackagesFencePassed
 		// to ourselves and all cookworkers.
 
-		// Call ModifyNumSaved to check for whether all packages have already been saved by the time we reach
+		// Call OnNumSavedUpdated to check for whether all packages have already been saved by the time we reach
 		// the request fence. This can happen in iterative cooks, or in race conditions if we sent all packages
 		// out for saving before receiving the EndQueueGeneratedPackagesOnDirector message.
-		ModifyNumSaved(0);
+		OnNumSavedUpdated();
 
 		if (COTFS.CookDirector)
 		{
@@ -1402,8 +1515,11 @@ void FGenerationHelper::OnRequestFencePassed(UCookOnTheFlyServer& COTFS)
 	}
 	else
 	{
-		UE_LOG(LogCookGenerationHelper, Verbose, TEXT("%s: RequestFencePassed, but not yet HasFinishedQueueGeneratedPackages."),
-			*WriteToString<256>(GetOwner().GetPackageName()));
+		if (FirstNotSkippedTargetPlatform)
+		{
+			UE_LOG(LogCookGenerationHelper, Verbose, TEXT("%s: RequestFencePassed, but not yet HasFinishedQueueGeneratedPackages."),
+				*WriteToString<256>(GetOwner().GetPackageName()));
+		}
 	}
 }
 
@@ -1416,8 +1532,12 @@ void FGenerationHelper::OnQueuedGeneratedPackagesFencePassed(UCookOnTheFlyServer
 	// We no longer need PreviousGeneratedPackages or KeepForIterative, because they are used only in 
 	// StartQueueGeneratedPackages or the request cluster that they end up in in PumpRequests, both of which
 	// are now finished. Clear them on the director and any CookWorkers that received them to free memory.
-	ClearKeepForIterative();
-	PreviousGeneratedPackages.Empty();
+	for (TPair<const ITargetPlatform*, FPlatformData>& PlatformPair : PlatformDatas)
+	{
+		FPlatformData& PlatformData = PlatformPair.Value;
+		PlatformData.ClearKeepForIterative(*this, PlatformPair.Key);
+		PlatformData.PreviousGeneratedPackages.Empty();
+	}
 }
 
 UPackage* FGenerationHelper::TryCreateGeneratedPackage(FCookGenerationInfo& GeneratedInfo, bool bResetToEmpty)
@@ -1506,9 +1626,11 @@ void FGenerationHelper::FinishGeneratorPlatformSave(FPackageData& PackageData, b
 		*WriteToString<256>(GetOwner().GetPackageName()));
 }
 
-void FGenerationHelper::FinishGeneratedPlatformSave(FPackageData& PackageData, FAssetPackageData& OutAssetPackageData,
-	TArray<FAssetDependency>& OutDependencies, FBuildResultDependenciesMap& OutBuildResultDependencies)
+void FGenerationHelper::FinishGeneratedPlatformSave(FPackageData& PackageData, const ITargetPlatform* TargetPlatform,
+	FAssetPackageData& OutAssetPackageData, TArray<FAssetDependency>& OutDependencies,
+	FBuildResultDependenciesMap& OutBuildResultDependencies)
 {
+	check(TargetPlatform);
 	ConditionalInitialize();
 
 	FCookGenerationInfo* Info = FindInfo(PackageData);
@@ -1520,6 +1642,8 @@ void FGenerationHelper::FinishGeneratedPlatformSave(FPackageData& PackageData, F
 	}
 	UCookOnTheFlyServer& COTFS = Info->PackageData->GetPackageDatas().GetCookOnTheFlyServer();
 
+	FCookGenerationInfo::FPlatformData& InfoPlatformData = Info->FindCheckedPlatformData(TargetPlatform);
+
 	// There should be no package dependencies present for the package from the global assetregistry
 	// because it is newly created. Add on the package dependencies declared for it from the CookPackageSplitter.
 	// These dependencies are used in the AssetRegistry entry and runtime dependencies for the generated package
@@ -1551,7 +1675,7 @@ void FGenerationHelper::FinishGeneratedPlatformSave(FPackageData& PackageData, F
 	}
 	ImportedClasses.Sort(FNameLexicalLess());
 
-	FAssetPackageData& AssetPackageData = Info->AssetPackageData.Emplace();
+	FAssetPackageData& AssetPackageData = InfoPlatformData.AssetPackageData.Emplace();
 	AssetPackageData.FileVersionUE = GPackageFileUEVersion;
 	AssetPackageData.FileVersionLicenseeUE = GPackageFileLicenseeUEVersion;
 	AssetPackageData.SetIsLicenseeVersion(FEngineVersion::Current().IsLicenseeVersion());
@@ -1564,17 +1688,50 @@ void FGenerationHelper::FinishGeneratedPlatformSave(FPackageData& PackageData, F
 
 	OutAssetPackageData = AssetPackageData;
 
-	UE_LOG(LogCookGenerationHelper, Verbose, TEXT("%s generated package %s: Generated package saved."),
+	UE_LOG(LogCookGenerationHelper, Verbose, TEXT("%s generated package %s: Generated package saved for platform %s."),
 		*WriteToString<256>(GetOwner().GetPackageName()),
-		*WriteToString<256>(Info->PackageData->GetPackageName()));
+		*WriteToString<256>(Info->PackageData->GetPackageName()), *TargetPlatform->PlatformName());
 }
 
-const FAssetPackageData* FGenerationHelper::GetAssetPackageData(FName PackageName)
+const FAssetPackageData* FGenerationHelper::GetAssetPackageData(FName PackageName, const ITargetPlatform* TargetPlatform)
 {
-	const FAssetPackageData* Result = PreviousGeneratedPackages.Find(PackageName);
-	if (Result)
+	FPlatformData* PlatformData = PlatformDatas.Find(TargetPlatform);
+	if (PlatformData)
 	{
-		return Result;
+		const FAssetPackageData* Result = PlatformData->PreviousGeneratedPackages.Find(PackageName);
+		if (Result)
+		{
+			return Result;
+		}
+	}
+
+	FCookGenerationInfo* Info = FindInfoNoInitialize(PackageName);
+	if (!Info || Info->IsGenerator())
+	{
+		return nullptr;
+	}
+
+	FCookGenerationInfo::FPlatformData* InfoPlatformData = Info->PlatformDatas.Find(TargetPlatform);
+	if (InfoPlatformData)
+	{
+		const FAssetPackageData* Result = InfoPlatformData->AssetPackageData.GetPtrOrNull();
+		if (Result)
+		{
+			return Result;
+		}
+	}
+	return nullptr;
+}
+
+const FAssetPackageData* FGenerationHelper::GetAssetPackageDataAnyPlatform(FName PackageName)
+{
+	for (const TPair<const ITargetPlatform*, FPlatformData>& Pair : PlatformDatas)
+	{
+		const FAssetPackageData* Result = Pair.Value.PreviousGeneratedPackages.Find(PackageName);
+		if (Result)
+		{
+			return Result;
+		}
 	}
 
 	FCookGenerationInfo* Info = FindInfoNoInitialize(PackageName);
@@ -1583,12 +1740,20 @@ const FAssetPackageData* FGenerationHelper::GetAssetPackageData(FName PackageNam
 		return nullptr;
 	}
 
-	return Info->AssetPackageData.GetPtrOrNull();
+	for (const TPair<const ITargetPlatform*, FCookGenerationInfo::FPlatformData>& InfoPair : Info->PlatformDatas)
+	{
+		const FAssetPackageData* Result = InfoPair.Value.AssetPackageData.GetPtrOrNull();
+		if (Result)
+		{
+			return Result;
+		}
+	}
+	return nullptr;
 }
 
-UE::EditorDomain::FPackageDigest FGenerationHelper::GetPackageDigest(FName PackageName)
+UE::EditorDomain::FPackageDigest FGenerationHelper::GetPackageDigest(FName PackageName, const ITargetPlatform* TargetPlatform)
 {
-	const FAssetPackageData* AssetPackageData = GetAssetPackageData(PackageName);
+	const FAssetPackageData* AssetPackageData = GetAssetPackageData(PackageName, TargetPlatform);
 	if (!AssetPackageData)
 	{
 		return UE::EditorDomain::FPackageDigest();
@@ -1666,15 +1831,21 @@ void FGenerationHelper::ResetSaveState(FCookGenerationInfo& Info, UPackage* Pack
 		// reference.
 		if (Info.IsGenerator())
 		{
-			Info.SetHasSaved(*this, true, FWorkerId::Local());
-
-			// Now that we've finished saving, we know that we will not call QueueGeneratedPackages again, so we can
-			// teardown iterative results as well
-			ClearKeepForIterative();
-			PreviousGeneratedPackages.Empty();
-
-			// And also teardown data needed during save
-			ClearKeepForGeneratorSave();
+			for (TPair<const ITargetPlatform*, FPlatformData>& PlatformPair : PlatformDatas)
+			{
+				const ITargetPlatform* TargetPlatform = PlatformPair.Key;
+				FPlatformData& PlatformData = PlatformPair.Value;
+				Info.FindCheckedPlatformData(TargetPlatform).SetHasSaved(*this, Info, TargetPlatform,
+					true /* bValue */, FWorkerId::Local());
+
+				// Now that we've finished saving, we know that we will not call QueueGeneratedPackages again, so we can
+				// teardown iterative results as well
+				PlatformData.ClearKeepForIterative(*this, TargetPlatform);
+				PlatformData.PreviousGeneratedPackages.Empty();
+
+				// And also teardown data needed during save
+				PlatformData.ClearKeepForGeneratorSave(*this, TargetPlatform);
+			}
 		}
 		else
 		{
@@ -1802,13 +1973,14 @@ void FGenerationHelper::FetchExternalActorDependencies()
 	ExternalActorDependencies.Shrink();
 }
 
-void FGenerationHelper::SetPreviousGeneratedPackages(TMap<FName, FAssetPackageData>&& Packages)
+void FGenerationHelper::SetPreviousGeneratedPackages(const ITargetPlatform* TargetPlatform, TMap<FName, FAssetPackageData>&& Packages)
 {
 	UE_LOG(LogCookGenerationHelper, Verbose,
-		TEXT("%s: SetPreviousGeneratedPackages with %d generated packages."),
-		*WriteToString<256>(GetOwner().GetPackageName()), Packages.Num());
-	SetKeepForIterative();
-	PreviousGeneratedPackages = MoveTemp(Packages);
+		TEXT("%s: SetPreviousGeneratedPackages with %d generated packages for platform %s."),
+		*WriteToString<256>(GetOwner().GetPackageName()), Packages.Num(), *TargetPlatform->PlatformName());
+	FPlatformData& PlatformData = FindCheckedPlatformData(TargetPlatform);
+	PlatformData.SetKeepForIterative(*this, TargetPlatform);
+	PlatformData.PreviousGeneratedPackages = MoveTemp(Packages);
 }
 
 template <typename T>
@@ -2016,29 +2188,36 @@ void FGenerationHelper::TrackGeneratedPackageListedRemotely(UCookOnTheFlyServer&
 		if (!FindInfoNoInitialize(PackageData))
 		{
 			bool bGenerator = false; // Cannot be the generator, if it were we would have found it
-			PackagesToGenerate.Emplace(PackageData, bGenerator);
+			PackagesToGenerate.Emplace(*this, PackageData, bGenerator);
 		}
 	}
-	FAssetPackageData* PreviousAssetData = PreviousGeneratedPackages.Find(PackageData.GetPackageName());
-	if (PreviousAssetData)
+
+	for (TPair<const ITargetPlatform*, FPlatformData>& PlatformPair : PlatformDatas)
 	{
-		// Copy the current value for the package's hash into the PreviousPackageData, for use by incremental cook's
-		// calculation in FRequestCluster::TryCalculateIterativelyUnmodified
-		PreviousAssetData->SetPackageSavedHash(CurrentPackageHash);
+		FPlatformData& PlatformData = PlatformPair.Value;
+		FAssetPackageData* PreviousAssetData = PlatformData.PreviousGeneratedPackages.Find(PackageData.GetPackageName());
+		if (PreviousAssetData)
+		{
+			// Copy the current value for the package's hash into the PreviousPackageData, for use by incremental cook's
+			// calculation in FRequestCluster::TryCalculateIterativelyUnmodified
+			PreviousAssetData->SetPackageSavedHash(CurrentPackageHash);
+		}
 	}
 }
 
 void FGenerationHelper::MarkPackageSavedRemotely(UCookOnTheFlyServer& COTFS, FPackageData& PackageData,
-	FWorkerId SourceWorkerId)
+	const ITargetPlatform* TargetPlatform, FWorkerId SourceWorkerId)
 {
 	FCookGenerationInfo* Info = FindInfoNoInitialize(PackageData);
 	if (Info)
 	{
-		Info->SetHasSaved(*this, true, SourceWorkerId);
+		Info->FindCheckedPlatformData(TargetPlatform).SetHasSaved(*this, *Info, TargetPlatform,
+			true /* bValue */, SourceWorkerId);
 	}
 }
 
-void FGenerationHelper::MarkPackageIterativelySkipped(FPackageData& PackageData, bool bIncrementallySkipped)
+void FGenerationHelper::MarkPackageIterativelySkipped(FPackageData& PackageData,
+	const ITargetPlatform* TargetPlatform, bool bIncrementallySkipped)
 {
 	FCookGenerationInfo* Info = FindInfoNoInitialize(PackageData);
 	if (!Info)
@@ -2067,8 +2246,9 @@ void FGenerationHelper::MarkPackageIterativelySkipped(FPackageData& PackageData,
 	}
 	if (bIncrementallySkipped)
 	{
-		Info->SetHasSaved(*this, true, FWorkerId::Local());
-		Info->SetIterativelySkipped(true);
+		FCookGenerationInfo::FPlatformData& PlatformData = Info->FindCheckedPlatformData(TargetPlatform);
+		PlatformData.SetHasSaved(*this, *Info, TargetPlatform, true /* bValue */, FWorkerId::Local());
+		PlatformData.SetIterativelySkipped(true);
 		if (bGenerator)
 		{
 			// The entire generator package has been skipped. Wait for the current cluster to complete
@@ -2137,7 +2317,7 @@ void FGenerationHelper::PostGarbageCollect(const TRefCountPtr<FGenerationHelper>
 	{
 		if (FindObject<UPackage>(nullptr, *Info.PackageData->GetPackageName().ToString()))
 		{
-			if (!Info.PackageData->IsKeepReferencedDuringGC() && !Info.HasSaved() && !bHasIssuedWarning)
+			if (!Info.PackageData->IsKeepReferencedDuringGC() && !Info.HasSavedEveryPlatform() && !bHasIssuedWarning)
 			{
 				UE_LOG(LogCook, Warning,
 					TEXT("PackageSplitter found a package it generated that was not removed from memory during garbage collection. This will cause errors later during population.")
@@ -2311,12 +2491,21 @@ void FGenerationHelper::UpdateSaveAfterGarbageCollect(const FPackageData& Packag
 	}
 }
 
-FCookGenerationInfo::FCookGenerationInfo(FPackageData& InPackageData, bool bInGenerator)
+FCookGenerationInfo::FCookGenerationInfo(FGenerationHelper& GenerationHelper, FPackageData& InPackageData, bool bInGenerator)
 	: PackageData(&InPackageData)
-	, bCreateAsMap(false), bHasCreatedPackage(false), bHasSaved(false), bTakenOverCachedCookedPlatformData(false)
+	, bCreateAsMap(false), bHasCreatedPackage(false), bTakenOverCachedCookedPlatformData(false)
 	, bIssuedUndeclaredMovedObjectsWarning(false), bGenerator(bInGenerator), bHasCalledPopulate(false)
-	, bIterativelySkipped(false)
 {
+	// The Info for the generator is created in the constructor of the GenerationHelper, and PlatformDatas is not yet readable.
+	// The GenerationHelper constructor will set the PlatformDatas on this Info later on in its constructor.
+	if (!bInGenerator)
+	{
+		for (TPair<const ITargetPlatform*, FGenerationHelper::FPlatformData>& PlatformPair
+			: GenerationHelper.PlatformDatas)
+		{
+			PlatformDatas.Add(PlatformPair.Key);
+		}
+	}
 }
 
 void FCookGenerationInfo::Uninitialize()
@@ -2344,14 +2533,18 @@ void FCookGenerationInfo::Uninitialize()
 		bGenerator ? TEXT("generator") : TEXT("generated"), *GetPackageName(),
 		LexToString(PackageData->GetState()), LexToString(PackageData->GetSaveSubState()));
 	// Keep SavedOnWorker; it is allowed in the uninitialized state
+
+	// Keep PlatformDatas; they are allowed in the unitialized state
+	// Keep PlatformData.AssetPackageData; it is allowed in the uninitialized state
+	// Keep PlatformData.bHasSaved; it is allowed in the uninitialized state
+	// Keep PlatformData.bIterativelySkipped; it is allowed in the uninitialized state
+
 	bCreateAsMap = false;
 	bHasCreatedPackage = false;
-	// Keep bHasSaved; it is allowed in the uninitialized state
 	bTakenOverCachedCookedPlatformData = false;
 	bIssuedUndeclaredMovedObjectsWarning = false;
 	// Keep bGenerator; it is allowed in the uninitialized state
 	bHasCalledPopulate = false;
-	// Keep bIterativelySkipped; it is allowed in the uninitialized state
 }
 
 void FCachedObjectInOuterGeneratorInfo::Initialize(UObject* Object)
@@ -2510,7 +2703,7 @@ void FCookGenerationInfo::CreatePackageHash()
 }
 
 void FCookGenerationInfo::IterativeCookValidateOrClear(FGenerationHelper& GenerationHelper,
-	TConstArrayView<const ITargetPlatform*> RequestedPlatforms, const FIoHash& PreviousPackageHash,
+	const ITargetPlatform* TargetPlatform, const FIoHash& PreviousPackageHash,
 	bool& bOutIterativelyUnmodified)
 {
 	UCookOnTheFlyServer& COTFS = GenerationHelper.GetOwner().GetPackageDatas().GetCookOnTheFlyServer();
@@ -2528,48 +2721,37 @@ void FCookGenerationInfo::IterativeCookValidateOrClear(FGenerationHelper& Genera
 				bOutIterativelyUnmodified = false;
 				break;
 			}
-			for (const ITargetPlatform* TargetPlatform : RequestedPlatforms)
-			{
-				FPackagePlatformData* DependencyPlatformData = DependencyData->FindPlatformData(TargetPlatform);
-				if (!DependencyPlatformData || !DependencyPlatformData->IsIterativelyUnmodified())
-				{
-					bOutIterativelyUnmodified = false;
-					break;
-				}
-			}
-			if (!bOutIterativelyUnmodified)
+			FPackagePlatformData* DependencyPlatformData = DependencyData->FindPlatformData(TargetPlatform);
+			if (!DependencyPlatformData || !DependencyPlatformData->IsIterativelyUnmodified())
 			{
+				bOutIterativelyUnmodified = false;
 				break;
 			}
 		}
 	}
 
-	bool bFirstPlatform = true;
-	for (const ITargetPlatform* TargetPlatform : RequestedPlatforms)
+	if (bOutIterativelyUnmodified)
 	{
-		if (bOutIterativelyUnmodified)
-		{
-			PackageData->FindOrAddPlatformData(TargetPlatform).SetIterativelyUnmodified(true);
-		}
-		bool bShouldIterativelySkip = bOutIterativelyUnmodified;
-		ICookedPackageWriter& PackageWriter = COTFS.FindOrCreatePackageWriter(TargetPlatform);
-		PackageWriter.UpdatePackageModificationStatus(PackageData->GetPackageName(), bOutIterativelyUnmodified,
-			bShouldIterativelySkip);
-		if (bShouldIterativelySkip)
-		{
-			PackageData->SetPlatformCooked(TargetPlatform, ECookResult::Succeeded);
-			if (bFirstPlatform)
-			{
-				COOK_STAT(++DetailedCookStats::NumPackagesIterativelySkipped);
-			}
-			// Declare the package to the EDLCookInfo verification so we don't warn about missing exports from it
-			UE::SavePackageUtilities::EDLCookInfoAddIterativelySkippedPackage(PackageData->GetPackageName());
-		}
-		else
+		PackageData->FindOrAddPlatformData(TargetPlatform).SetIterativelyUnmodified(true);
+	}
+	bool bShouldIterativelySkip = bOutIterativelyUnmodified;
+	ICookedPackageWriter& PackageWriter = COTFS.FindOrCreatePackageWriter(TargetPlatform);
+	PackageWriter.UpdatePackageModificationStatus(PackageData->GetPackageName(), bOutIterativelyUnmodified,
+		bShouldIterativelySkip);
+	if (bShouldIterativelySkip)
+	{
+		PackageData->SetPlatformCooked(TargetPlatform, ECookResult::Succeeded);
+		TConstArrayView<const ITargetPlatform*> SessionPlatforms = COTFS.PlatformManager->GetSessionPlatforms();
+		if (SessionPlatforms.Num() > 0 && TargetPlatform == SessionPlatforms[0])
 		{
-			COTFS.DeleteOutputForPackage(PackageData->GetPackageName(), TargetPlatform);
+			COOK_STAT(++DetailedCookStats::NumPackagesIterativelySkipped);
 		}
-		bFirstPlatform = false;
+		// Declare the package to the EDLCookInfo verification so we don't warn about missing exports from it
+		UE::SavePackageUtilities::EDLCookInfoAddIterativelySkippedPackage(PackageData->GetPackageName());
+	}
+	else
+	{
+		COTFS.DeleteOutputForPackage(PackageData->GetPackageName(), TargetPlatform);
 	}
 }
 
diff --git a/Engine/Source/Editor/UnrealEd/Private/Cooker/CookGenerationHelper.h b/Engine/Source/Editor/UnrealEd/Private/Cooker/CookGenerationHelper.h
--- a/Engine/Source/Editor/UnrealEd/Private/Cooker/CookGenerationHelper.h
+++ b/Engine/Source/Editor/UnrealEd/Private/Cooker/CookGenerationHelper.h
@@ -14,6 +14,7 @@
 #include "Cooker/MPCollector.h"
 #include "CookPackageSplitter.h"
 #include "Hash/Blake3.h"
+#include "Interfaces/ITargetPlatform.h"
 #include "IO/IoHash.h"
 #include "Misc/Optional.h"
 #include "Templates/RefCounting.h"
@@ -27,6 +28,7 @@
 class ITargetPlatform;
 class UCookOnTheFlyServer;
 namespace UE::Cook { class FCookGCDiagnosticContext; }
+namespace UE::Cook { struct FGenerationHelper; }
 namespace UE::EditorDomain { struct FPackageDigest; }
 
 namespace UE::Cook::CookPackageSplitter
@@ -79,14 +81,14 @@ public:
 struct FCookGenerationInfo
 {
 public:
-	FCookGenerationInfo(FPackageData& InPackageData, bool bInGenerator);
+	FCookGenerationInfo(FGenerationHelper& GenerationHelper, FPackageData& InPackageData, bool bInGenerator);
 
 	bool IsCreateAsMap() const;
 	void SetIsCreateAsMap(bool bValue);
 	bool HasCreatedPackage() const;
 	void SetHasCreatedPackage(bool bValue);
-	bool HasSaved() const;
-	void SetHasSaved(FGenerationHelper& GenerationHelper, bool bValue, FWorkerId SourceWorkerId);
+	bool HasSavedEveryPlatform() const;
+	void SetAllPlatformsSaved(FGenerationHelper& GenerationHelper);
 	bool HasTakenOverCachedCookedPlatformData() const;
 	void SetHasTakenOverCachedCookedPlatformData(bool bValue);
 	bool HasIssuedUndeclaredMovedObjectsWarning() const;
@@ -95,8 +97,6 @@ public:
 	void SetIsGenerator(bool bValue);
 	bool HasCalledPopulate() const;
 	void SetHasCalledPopulate(bool bValue);
-	bool IsIterativelySkipped() const;
-	void SetIterativelySkipped(bool bValue);
 
 	/**
 	 * Steal the list of cached objects to call BeginCacheForCookedPlatformData on from the PackageData,
@@ -124,9 +124,8 @@ public:
 	 * test whether they need to be invalidated now and clear the results if so. Only called in legacy iterative;
 	 * incremental cooks handle invalidation by querying the TargetDomainDigest during the RequestCluster.
 	 */
-	void IterativeCookValidateOrClear(FGenerationHelper& GeneratorHelper,
-		TConstArrayView<const ITargetPlatform*> RequestedPlatforms, const FIoHash& PreviousHash,
-		bool& bOutIterativelyUnmodified);
+	void IterativeCookValidateOrClear(FGenerationHelper& GeneratorHelper, const ITargetPlatform* TargetPlatform,
+		const FIoHash& PreviousHash, bool& bOutIterativelyUnmodified);
 
 	TConstArrayView<FAssetDependency> GetDependencies() const;
 
@@ -152,22 +151,51 @@ public:
 	TArray<TWeakObjectPtr<UPackage>> KeepReferencedPackages;
 	TMap<UObject*, FCachedObjectInOuterGeneratorInfo> CachedObjectsInOuterInfo;
 	FWorkerId SavedOnWorker = FWorkerId::Invalid();
-	TOptional<FAssetPackageData> AssetPackageData;
 
 private:
+	struct FPlatformData
+	{
+		TOptional<FAssetPackageData> AssetPackageData;
+		bool bHasSaved : 1;
+		bool bIterativelySkipped : 1;
+
+		FPlatformData();
+		bool IsIterativelySkipped() const;
+		void SetIterativelySkipped(bool bValue);
+		bool HasSaved() const;
+		void SetHasSaved(FGenerationHelper& GenerationHelper, FCookGenerationInfo& Info,
+			const ITargetPlatform* TargetPlatform, bool bValue, FWorkerId SourceWorkerId);
+	};
+
+private:
+	FPlatformData& FindCheckedPlatformData(const ITargetPlatform* TargetPlatform);
+
+private:
+
+	TMap<const ITargetPlatform*, FPlatformData> PlatformDatas;
 	bool bCreateAsMap : 1;
 	bool bHasCreatedPackage : 1;
-	bool bHasSaved : 1;
 	bool bTakenOverCachedCookedPlatformData : 1;
 	bool bIssuedUndeclaredMovedObjectsWarning : 1;
 	bool bGenerator : 1;
 	bool bHasCalledPopulate : 1;
-	bool bIterativelySkipped : 1;
+
+	friend UE::Cook::FGenerationHelper;
 };
 
 /**
  * Helper that wraps an ICookPackageSplitter, gets/caches packages to generate, and is a reference-counted
  * collection of cached data and helper functions to save and list the generated packages.
+ * 
+ * TODO: MultiplatformCook: We do not yet support Packages that use a GenerationHelper being cooked for some platforms
+ * in a multiplatformcook and not others. And if the package is not incrementally skipped for one platform then it is
+ * not incrementally skipped for any platform. We don't support per-platform incremental cooks because allowing
+ * per-platform results would require extra tests throughout PumpSaves and the events on GenerationHelper to check for
+ * which platforms are enabled, and we don't have those yet. We don't support per-platform reachability for the same
+ * reason, and additionally because reachability might be discovered for one platform after the save already started
+ * for another platform, and we would need to handle the state transition to restart the generation with the new
+ * platform accommodated. We will eventually need to implement both per-platform incremental cooks and per-platform
+ * reachability to fully support per-platform customization of the cook graph.
  */
 struct FGenerationHelper : public FThreadSafeRefCountedObject
 {
@@ -320,15 +348,30 @@ public:
 	 * Called when a generated package is in FSaveCookedPackageContext::FinishPlatformSave.
 	 * Calculates AssetRegistryData.
 	 */
-	void FinishGeneratedPlatformSave(FPackageData& PackageData, FAssetPackageData& OutAssetPackageData,
-		TArray< FAssetDependency>& OutARDependencies, FBuildResultDependenciesMap& OutBuildResultDependencies);
+	void FinishGeneratedPlatformSave(FPackageData& PackageData, const ITargetPlatform* TargetPlatform,
+		FAssetPackageData& OutAssetPackageData, TArray<FAssetDependency>& OutARDependencies,
+		FBuildResultDependenciesMap& OutBuildResultDependencies);
+	/**
+	 * Return the AssetPackageData that was created for the given generated PackageName, or loaded for it from a
+	 * previous cook. Will return nullptr if not loaded from a previous incremental cook and it has not yet been saved.
+	 * Returns the AssetPackageData for the given platform, or for the shared data between platforms if nullptr.
+	 * Does not call initialize.
+	 */
+	const FAssetPackageData* GetAssetPackageData(FName PackageName, const ITargetPlatform* TargetPlatform);
+	/**
+	 * Returns GetAssetPackageData for each known platform, returning the first non-null value, or null if no values
+	 * found.
+	 */
+	const FAssetPackageData* GetAssetPackageDataAnyPlatform(FName PackageName);
 	/**
-	 * Return the AssetPackageData that was created for the given PackageData, or loaded for it from a previous cook.
-	 * Will return nullptr if not loaded from a previous incremental cook and it has not yet been saved.
+	 * Calculate the PackageDigest for the given generated packagename and given platform based on its
+	 * AssetPackageData (e.g. list of classes in the package) in a previous cook.
+	 * TargetPlatform specifies which platform to look for results, or nullptr to look up results recorded for the
+	 * shared data between platforms. Returns empty digest if not loaded from a previous incremental cook and it has not
+	 * yet been saved.
 	 * Does not call initialize.
 	 */
-	const FAssetPackageData* GetAssetPackageData(FName PackageName);
-	UE::EditorDomain::FPackageDigest GetPackageDigest(FName PackageName);
+	UE::EditorDomain::FPackageDigest GetPackageDigest(FName PackageName, const ITargetPlatform* TargetPlatform);
 
 	/**
 	 * Clear any data that should only be held when an FPackageData is in the save state, for the given Info. The given
@@ -351,9 +394,9 @@ public:
 	 */
 	void FetchExternalActorDependencies();
 	/** Iterative cook: store the list of generated packages from the last cook. Does not call Initialize. */
-	void SetPreviousGeneratedPackages(TMap<FName, FAssetPackageData>&& Packages);
+	void SetPreviousGeneratedPackages(const ITargetPlatform* TargetPlatform, TMap<FName, FAssetPackageData>&& Packages);
 	/** Return the information set by SetPreviousGeneratedPackages if not yet cleared. Does not call Initialize. */
-	const TMap<FName, FAssetPackageData>& GetPreviousGeneratedPackages() const;
+	const TMap<FName, FAssetPackageData>& GetPreviousGeneratedPackages(const ITargetPlatform* TargetPlatform) const;
 
 	/**
 	 * Callback during garbage collection. Does not call initialize. Caller must pass in a refcount to show
@@ -376,13 +419,17 @@ public:
 
 	// Self-references that keep this GenerationHelper in memory and referenced from the packages that use it until
 	// the self-references are cleared. These Set/Clear functions do not call initialize.
-	void SetKeepForIterative();
-	void ClearKeepForIterative();
+	void SetKeepForIterative(const ITargetPlatform* TargetPlatform);
+	void SetKeepForIterativeAllPlatforms();
+	void ClearKeepForIterative(const ITargetPlatform* TargetPlatform);
+	void ClearKeepForIterativeAllPlatforms();
 	void SetKeepForQueueResults();
 	void ClearKeepForQueueResults();
 	bool IsWaitingForQueueResults() const;
-	void SetKeepForGeneratorSave();
-	void ClearKeepForGeneratorSave();
+	void SetKeepForGeneratorSave(const ITargetPlatform* TargetPlatform);
+	void SetKeepForGeneratorSaveAllPlatforms();
+	void ClearKeepForGeneratorSave(const ITargetPlatform* TargetPlatform);
+	void ClearKeepForGeneratorSaveAllPlatforms();
 	void SetKeepForAllSavedOrGC();
 	void ClearKeepForAllSavedOrGC();
 	void SetKeepForCompletedAllSavesMessage();
@@ -409,7 +456,8 @@ public:
 	 * Called on the director when the generator or one of the generated packages was saved on a remote worker.
 	 * Used to manage KeepForGCOrAllSaved lifetime. Does not call Initialize.
 	 */
-	void MarkPackageSavedRemotely(UCookOnTheFlyServer& COTFS, FPackageData& PackageData, FWorkerId SourceWorkerId);
+	void MarkPackageSavedRemotely(UCookOnTheFlyServer& COTFS, FPackageData& PackageData,
+		const ITargetPlatform* TargetPlatform, FWorkerId SourceWorkerId);
 	/** Scope object that blocks calls to OnAllSavesCompleted. Checks and calls it if necessary when destructed. */
 	struct FScopeDeferEvents
 	{
@@ -424,7 +472,8 @@ public:
 	 * Called on the director when the generator or one of the generated packages was found to be iteratively skipped
 	 * and marked already cooked in an incremental cook. Does not call Initialize.
 	 */
-	void MarkPackageIterativelySkipped(FPackageData& PackageData, bool bIncrementallySkipped);
+	void MarkPackageIterativelySkipped(FPackageData& PackageData, const ITargetPlatform* TargetPlatform,
+		bool bIncrementallySkipped);
 	/**
 	 * Called on the director and every CookWorker when all saves have been completed; this indicates that
 	 * some of our contract points are complete and the splitter can be destroyed. Does not call Initialize.
@@ -485,8 +534,22 @@ private:
 		Invalid,
 		Valid,
 	};
+	struct FPlatformData
+	{
+		TMap<FName, FAssetPackageData> PreviousGeneratedPackages;
+		TRefCountPtr<FGenerationHelper> ReferenceFromKeepForIterative;
+		TRefCountPtr<FGenerationHelper> ReferenceFromKeepForGeneratorSave;
+		int32 NumSaved = 0;
+
+		void SetKeepForIterative(FGenerationHelper& GenerationHelper, const ITargetPlatform* TargetPlatform);
+		void ClearKeepForIterative(FGenerationHelper& GenerationHelper, const ITargetPlatform* TargetPlatform);
+		void SetKeepForGeneratorSave(FGenerationHelper& GenerationHelper, const ITargetPlatform* TargetPlatform);
+		void ClearKeepForGeneratorSave(FGenerationHelper& GenerationHelper, const ITargetPlatform* TargetPlatform);
+	};
 
 private:
+	FPlatformData& FindCheckedPlatformData(const ITargetPlatform* TargetPlatform);
+	const FPlatformData& FindCheckedPlatformData(const ITargetPlatform* TargetPlatform) const;
 	void ConditionalInitialize() const;
 	FCookGenerationInfo* FindInfoNoInitialize(const FPackageData& PackageData);
 	FCookGenerationInfo* FindInfoNoInitialize(FName PackageName);
@@ -494,7 +557,8 @@ private:
 	void PreGarbageCollectGCLifetimeData();
 	void PostGarbageCollectGCLifetimeData(FCookGCDiagnosticContext& Context);
 	void Uninitialize();
-	void ModifyNumSaved(int32 Delta);
+	void ModifyNumSaved(const ITargetPlatform* TargetPlatform, int32 Delta);
+	void OnNumSavedUpdated();
 	void VerifyGeneratorPackageGarbageCollected(FCookGCDiagnosticContext& Context);
 	void DemoteStalledPackages(UCookOnTheFlyServer& COTFS, bool bFromAllSavesCompleted);
 
@@ -511,15 +575,12 @@ private:
 	/** Recorded list of packages to generate from the splitter, and data we need about them */
 	TArray<FCookGenerationInfo> PackagesToGenerate;
 	TWeakObjectPtr<UPackage> OwnerPackage;
-	TMap<FName, FAssetPackageData> PreviousGeneratedPackages;
+	TMap<const ITargetPlatform*, FPlatformData> PlatformDatas;
 	TArray<FName> ExternalActorDependencies;
 	TArray<FWeakObjectPtr> OwnerObjectsToMove;
-	TRefCountPtr<FGenerationHelper> ReferenceFromKeepForIterative;
 	TRefCountPtr<FGenerationHelper> ReferenceFromKeepForQueueResults;
-	TRefCountPtr<FGenerationHelper> ReferenceFromKeepForGeneratorSave;
 	TRefCountPtr<FGenerationHelper> ReferenceFromKeepForAllSavedOrGC;
 	int32 MPCookNextAssignmentIndex = 0;
-	int32 NumSaved = 0;
 	EInitializeStatus InitializeStatus = EInitializeStatus::Uninitialized;
 	ICookPackageSplitter::EGeneratedRequiresGenerator DoesGeneratedRequireGeneratorValue =
 		ICookPackageSplitter::EGeneratedRequiresGenerator::None;
@@ -533,6 +594,7 @@ private:
 	bool bNeedConfirmGeneratorPackageDestroyed = false;
 	bool bHasFinishedQueueGeneratedPackages = false;
 	bool bDeferEvents = false;
+	bool bSentAllSavesCompleted = false;
 
 	friend FCookGenerationInfo;
 };
@@ -543,6 +605,21 @@ private:
 ///////////////////////////////////////////////////////
 
 
+inline FCookGenerationInfo::FPlatformData& FCookGenerationInfo::FindCheckedPlatformData(
+	const ITargetPlatform* TargetPlatform)
+{
+	FPlatformData* PlatformData = PlatformDatas.Find(TargetPlatform);
+	if (!PlatformData)
+	{
+		UE_LOG(LogCook, Error,
+			TEXT("FCookGenerationInfo::FindCheckedPlatformData called for %s on %s, but doesn't already have it. This is a cook programming error."),
+			*TargetPlatform->PlatformName(), *GetPackageName());
+		FDebug::DumpStackTraceToLog(ELogVerbosity::Warning);
+		PlatformData = &PlatformDatas.FindOrAdd(TargetPlatform);
+	}
+	return *PlatformData;
+}
+
 inline bool FCookGenerationInfo::IsCreateAsMap() const
 {
 	return bCreateAsMap;
@@ -563,22 +640,24 @@ inline void FCookGenerationInfo::SetHasCreatedPackage(bool bValue)
 	bHasCreatedPackage = bValue;
 }
 
-inline bool FCookGenerationInfo::HasSaved() const
+inline bool FCookGenerationInfo::HasSavedEveryPlatform() const
 {
+	bool bHasSaved = true;
+	for (const TPair<const ITargetPlatform*, FPlatformData>& InfoPlatformPair : PlatformDatas)
+	{
+		bHasSaved &= InfoPlatformPair.Value.bHasSaved;
+	}
 	return bHasSaved;
 }
 
-inline void FCookGenerationInfo::SetHasSaved(FGenerationHelper& GenerationHelper, bool bValue,
-	FWorkerId SourceWorkerId)
+inline void FCookGenerationInfo::SetAllPlatformsSaved(FGenerationHelper& GenerationHelper)
 {
-	if (bValue != bHasSaved)
+	for (TPair<const ITargetPlatform*, FCookGenerationInfo::FPlatformData>& PlatformPair : PlatformDatas)
 	{
-		bHasSaved = bValue;
-		GenerationHelper.ModifyNumSaved(bValue ? 1 : -1);
-		if (bHasSaved)
-		{
-			SavedOnWorker = SourceWorkerId;
-		}
+		const ITargetPlatform* TargetPlatform = PlatformPair.Key;
+		FPlatformData& InfoPlatformData = PlatformPair.Value;
+		InfoPlatformData.SetHasSaved(GenerationHelper, *this, TargetPlatform,
+			true /* bValue */, FWorkerId::Local());
 	}
 }
 
@@ -622,24 +701,69 @@ inline void FCookGenerationInfo::SetHasCalledPopulate(bool bValue)
 	bHasCalledPopulate = bValue;
 }
 
-inline bool FCookGenerationInfo::IsIterativelySkipped() const
+inline TConstArrayView<FAssetDependency> FCookGenerationInfo::GetDependencies() const
+{
+	return PackageDependencies;
+}
+
+inline FString FCookGenerationInfo::GetPackageName() const
+{
+	return PackageData->GetPackageName().ToString();
+}
+
+inline FCookGenerationInfo::FPlatformData::FPlatformData()
+	: bHasSaved(false), bIterativelySkipped(false)
+{
+}
+
+inline bool FCookGenerationInfo::FPlatformData::IsIterativelySkipped() const
 {
 	return bIterativelySkipped;
 }
 
-inline void FCookGenerationInfo::SetIterativelySkipped(bool bValue)
+inline void FCookGenerationInfo::FPlatformData::SetIterativelySkipped(bool bValue)
 {
 	bIterativelySkipped = bValue;
 }
 
-inline TConstArrayView<FAssetDependency> FCookGenerationInfo::GetDependencies() const
+inline bool FCookGenerationInfo::FPlatformData::HasSaved() const
 {
-	return PackageDependencies;
+	return bHasSaved;
 }
 
-inline FString FCookGenerationInfo::GetPackageName() const
+inline void FCookGenerationInfo::FPlatformData::SetHasSaved(FGenerationHelper& GenerationHelper,
+	FCookGenerationInfo& Info, const ITargetPlatform* TargetPlatform, bool bValue, FWorkerId SourceWorkerId)
 {
-	return PackageData->GetPackageName().ToString();
+	if (bValue != bHasSaved)
+	{
+		bHasSaved = bValue;
+		GenerationHelper.ModifyNumSaved(TargetPlatform, bValue ? 1 : -1);
+		if (bHasSaved)
+		{
+			Info.SavedOnWorker = SourceWorkerId;
+		}
+	}
+}
+
+inline FGenerationHelper::FPlatformData& FGenerationHelper::FindCheckedPlatformData(
+	const ITargetPlatform* TargetPlatform)
+{
+	FPlatformData* PlatformData = PlatformDatas.Find(TargetPlatform);
+	if (!PlatformData)
+	{
+		UE_LOG(LogCook, Error,
+			TEXT("FGenerationHelper::FindCheckedPlatformData called for %s on %s, but doesn't already have it. This is a cook programming error."),
+			*TargetPlatform->PlatformName(), *GetOwner().GetPackageName().ToString());
+		FDebug::DumpStackTraceToLog(ELogVerbosity::Warning);
+		PlatformData = &PlatformDatas.FindOrAdd(TargetPlatform);
+	}
+	return *PlatformData;
+}
+
+inline const FGenerationHelper::FPlatformData& FGenerationHelper::FindCheckedPlatformData(
+	const ITargetPlatform* TargetPlatform) const
+{
+	return const_cast<FGenerationHelper&>(*this).FindCheckedPlatformData(TargetPlatform);
 }
 
 inline bool FGenerationHelper::IsInitialized() const
@@ -735,33 +859,60 @@ inline TArray<FName> FGenerationHelper::ReleaseExternalActorDependencies()
 	return Result;
 }
 
-inline const TMap<FName, FAssetPackageData>& FGenerationHelper::GetPreviousGeneratedPackages() const
+inline const TMap<FName, FAssetPackageData>& FGenerationHelper::GetPreviousGeneratedPackages(
+	const ITargetPlatform* TargetPlatform) const
 {
-	return PreviousGeneratedPackages;
+	return FindCheckedPlatformData(TargetPlatform).PreviousGeneratedPackages;
 }
 
-inline void FGenerationHelper::SetKeepForIterative()
+inline void FGenerationHelper::FPlatformData::SetKeepForIterative(FGenerationHelper& GenerationHelper,
+	const ITargetPlatform* TargetPlatform)
 {
 	if (!ReferenceFromKeepForIterative)
 	{
-		UE_LOG(LogCookGenerationHelper, Verbose, TEXT("%s: SetKeepForIterative."),
-			*WriteToString<256>(GetOwner().GetPackageName()));
+		UE_LOG(LogCookGenerationHelper, Verbose, TEXT("%s: SetKeepForIterative for platform %s."),
+			*WriteToString<256>(GenerationHelper.GetOwner().GetPackageName()), *TargetPlatform->PlatformName());
+		ReferenceFromKeepForIterative = &GenerationHelper;
+	}
+}
+
+inline void FGenerationHelper::SetKeepForIterative(const ITargetPlatform* TargetPlatform)
+{
+	FindCheckedPlatformData(TargetPlatform).SetKeepForIterative(*this, TargetPlatform);
+}
 
-		ReferenceFromKeepForIterative = this;
+inline void FGenerationHelper::SetKeepForIterativeAllPlatforms()
+{
+	for (TPair<const ITargetPlatform*, FPlatformData>& PlatformPair : PlatformDatas)
+	{
+		PlatformPair.Value.SetKeepForIterative(*this, PlatformPair.Key);
 	}
 }
 
-inline void FGenerationHelper::ClearKeepForIterative()
+inline void FGenerationHelper::FPlatformData::ClearKeepForIterative(FGenerationHelper& GenerationHelper,
+	const ITargetPlatform* TargetPlatform)
 {
 	if (ReferenceFromKeepForIterative)
 	{
-		UE_LOG(LogCookGenerationHelper, Verbose, TEXT("%s: ClearKeepForIterative."),
-			*WriteToString<256>(GetOwner().GetPackageName()));
-
+		UE_LOG(LogCookGenerationHelper, Verbose, TEXT("%s: ClearKeepForIterative for platform %s."),
+			*WriteToString<256>(GenerationHelper.GetOwner().GetPackageName()), *TargetPlatform->PlatformName());
 		ReferenceFromKeepForIterative.SafeRelease();
 	}
 }
 
+inline void FGenerationHelper::ClearKeepForIterative(const ITargetPlatform* TargetPlatform)
+{
+	FindCheckedPlatformData(TargetPlatform).ClearKeepForIterative(*this, TargetPlatform);
+}
+
+inline void FGenerationHelper::ClearKeepForIterativeAllPlatforms()
+{
+	for (TPair<const ITargetPlatform*, FPlatformData>& PlatformPair : PlatformDatas)
+	{
+		PlatformPair.Value.ClearKeepForIterative(*this, PlatformPair.Key);
+	}
+}
+
 inline void FGenerationHelper::SetKeepForQueueResults()
 {
 	if (!ReferenceFromKeepForQueueResults)
@@ -789,28 +940,54 @@ inline bool FGenerationHelper::IsWaitingForQueueResults() const
 	return ReferenceFromKeepForQueueResults.IsValid();
 }
 
-inline void FGenerationHelper::SetKeepForGeneratorSave()
+inline void FGenerationHelper::FPlatformData::SetKeepForGeneratorSave(FGenerationHelper& GenerationHelper,
+	const ITargetPlatform* TargetPlatform)
 {
 	if (!ReferenceFromKeepForGeneratorSave)
 	{
-		UE_LOG(LogCookGenerationHelper, Verbose, TEXT("%s: SetKeepForGeneratorSave."),
-			*WriteToString<256>(GetOwner().GetPackageName()));
+		UE_LOG(LogCookGenerationHelper, Verbose, TEXT("%s: SetKeepForGeneratorSave for platform %s."),
+			*WriteToString<256>(GenerationHelper.GetOwner().GetPackageName()), *TargetPlatform->PlatformName());
+		ReferenceFromKeepForGeneratorSave = &GenerationHelper;
+	}
+}
+
+inline void FGenerationHelper::SetKeepForGeneratorSave(const ITargetPlatform* TargetPlatform)
+{
+	FindCheckedPlatformData(TargetPlatform).SetKeepForGeneratorSave(*this, TargetPlatform);
+}
 
-		ReferenceFromKeepForGeneratorSave = this;
+inline void FGenerationHelper::SetKeepForGeneratorSaveAllPlatforms()
+{
+	for (TPair<const ITargetPlatform*, FPlatformData>& PlatformPair : PlatformDatas)
+	{
+		PlatformPair.Value.SetKeepForGeneratorSave(*this, PlatformPair.Key);
 	}
 }
 
-inline void FGenerationHelper::ClearKeepForGeneratorSave()
+inline void FGenerationHelper::FPlatformData::ClearKeepForGeneratorSave(FGenerationHelper& GenerationHelper,
+	const ITargetPlatform* TargetPlatform)
 {
 	if (ReferenceFromKeepForGeneratorSave)
 	{
-		UE_LOG(LogCookGenerationHelper, Verbose, TEXT("%s: ClearKeepForGeneratorSave."),
-			*WriteToString<256>(GetOwner().GetPackageName()));
-
+		UE_LOG(LogCookGenerationHelper, Verbose, TEXT("%s: ClearKeepForGeneratorSave for platform %s."),
+			*WriteToString<256>(GenerationHelper.GetOwner().GetPackageName()), *TargetPlatform->PlatformName());
 		ReferenceFromKeepForGeneratorSave.SafeRelease();
 	}
 }
 
+inline void FGenerationHelper::ClearKeepForGeneratorSave(const ITargetPlatform* TargetPlatform)
+{
+	FindCheckedPlatformData(TargetPlatform).ClearKeepForGeneratorSave(*this, TargetPlatform);
+}
+
+inline void FGenerationHelper::ClearKeepForGeneratorSaveAllPlatforms()
+{
+	for (TPair<const ITargetPlatform*, FPlatformData>& PlatformPair : PlatformDatas)
+	{
+		PlatformPair.Value.ClearKeepForGeneratorSave(*this, PlatformPair.Key);
+	}
+}
+
 inline void FGenerationHelper::SetKeepForAllSavedOrGC()
 {
 	if (!bKeepForAllSavedOrGC)
diff --git a/Engine/Source/Editor/UnrealEd/Private/Cooker/CookPackageArtifacts.cpp b/Engine/Source/Editor/UnrealEd/Private/Cooker/CookPackageArtifacts.cpp
--- a/Engine/Source/Editor/UnrealEd/Private/Cooker/CookPackageArtifacts.cpp
+++ b/Engine/Source/Editor/UnrealEd/Private/Cooker/CookPackageArtifacts.cpp
@@ -41,7 +41,8 @@ constexpr uint32 PackageArtifactsVersion = 0x00000004;
 static const FUtf8StringView PackageArtifactsAttachmentKey = UTF8TEXTVIEW("CookPackageArtifacts");
 static const FUtf8StringView BuildDefinitionsAttachmentKey = UTF8TEXTVIEW("BuildDefinitionsAttachmentKey");
 
-bool FBuildDependencySet::HasKeyMatch(FName PackageName, FGenerationHelper* GenerationHelper)
+bool FBuildDependencySet::HasKeyMatch(FName PackageName, const ITargetPlatform* TargetPlatform,
+	FGenerationHelper* GenerationHelper)
 {
 	if (!bValid)
 	{
@@ -53,7 +54,7 @@ bool FBuildDependencySet::HasKeyMatch(FName PackageName, FGenerationHelper* Gene
 	}
 	if (CurrentKey.IsZero())
 	{
-		ECurrentKeyResult Result = TryCalculateCurrentKey(PackageName, GenerationHelper);
+		ECurrentKeyResult Result = TryCalculateCurrentKey(PackageName, TargetPlatform, GenerationHelper);
 		if (Result != ECurrentKeyResult::Success)
 		{
 			CurrentKey.Reset();
@@ -64,7 +65,8 @@ bool FBuildDependencySet::HasKeyMatch(FName PackageName, FGenerationHelper* Gene
 }
 
 FBuildDependencySet::ECurrentKeyResult FBuildDependencySet::TryCalculateCurrentKey(FName PackageName,
-	FGenerationHelper* GenerationHelper, TArray<TPair<ELogVerbosity::Type, FString>>* OutMessages)
+	const ITargetPlatform* TargetPlatform, FGenerationHelper* GenerationHelper,
+	TArray<TPair<ELogVerbosity::Type, FString>>* OutMessages)
 {
 	IAssetRegistry* AssetRegistry = IAssetRegistry::Get();
 	if (PackageName.IsNone())
@@ -131,7 +133,7 @@ FBuildDependencySet::ECurrentKeyResult FBuildDependencySet::TryCalculateCurrentK
 				UE::EditorDomain::FPackageDigest PackageDigest = EditorDomain->GetPackageDigest(DependencyPackageName);
 				if (!PackageDigest.IsSuccessful() && GenerationHelper)
 				{
-					PackageDigest = GenerationHelper->GetPackageDigest(DependencyPackageName);
+					PackageDigest = GenerationHelper->GetPackageDigest(DependencyPackageName, TargetPlatform);
 				}
 				if (PackageDigest.IsSuccessful())
 				{
@@ -756,15 +758,15 @@ bool FPackageArtifacts::HasSaveResults() const
 	return bHasSaveResults;
 }
 
-bool FPackageArtifacts::HasKeyMatch(FGenerationHelper* GenerationHelper)
+bool FPackageArtifacts::HasKeyMatch(const ITargetPlatform* TargetPlatform, FGenerationHelper* GenerationHelper)
 {
-	return bValid && SaveBuildDependencies.HasKeyMatch(PackageName, GenerationHelper);
+	return bValid && SaveBuildDependencies.HasKeyMatch(PackageName, TargetPlatform, GenerationHelper);
 }
 
-FBuildDependencySet::ECurrentKeyResult FPackageArtifacts::TryCalculateCurrentKey(
+FBuildDependencySet::ECurrentKeyResult FPackageArtifacts::TryCalculateCurrentKey(const ITargetPlatform* TargetPlatform,
 	FGenerationHelper* GenerationHelper, TArray<TPair<ELogVerbosity::Type, FString>>* OutMessages)
 {
-	return SaveBuildDependencies.TryCalculateCurrentKey(PackageName, GenerationHelper, OutMessages);
+	return SaveBuildDependencies.TryCalculateCurrentKey(PackageName, TargetPlatform, GenerationHelper, OutMessages);
 }
 
 void FPackageArtifacts::Empty()
@@ -1134,7 +1136,7 @@ FPackageArtifacts FPackageArtifacts::Collect(const UPackage* Package, const ITar
 	TArray<FCookDependency>& LoadDependencies = InResultDependencies.FindOrAdd(BuildResult::NAME_Load);
 	Result.LoadBuildDependencies.SetName(BuildResult::NAME_Load);
 	Result.LoadBuildDependencies.SetNormalizedDependencies(MoveTemp(LoadDependencies));
-		Result.LoadBuildDependencies.TryCalculateCurrentKey(PackageName,
+		Result.LoadBuildDependencies.TryCalculateCurrentKey(PackageName, TargetPlatform,
 			nullptr /* GenerationHelper */, nullptr /* OutMessages */);
 		Result.LoadBuildDependencies.StoreCurrentKey();
 		Result.LoadBuildDependencies.SetValid(true);
@@ -1150,7 +1152,7 @@ FPackageArtifacts FPackageArtifacts::Collect(const UPackage* Package, const ITar
 	Result.SaveBuildDependencies.SetName(BuildResult::NAME_Save);
 	Result.SaveBuildDependencies.SetNormalizedDependencies(MoveTemp(SaveDependencies));
 	FBuildDependencySet::ECurrentKeyResult CurrentKeyResult = Result.SaveBuildDependencies.TryCalculateCurrentKey(
-		PackageName, GenerationHelper, OutMessages);
+		PackageName, TargetPlatform, GenerationHelper, OutMessages);
 	if (CurrentKeyResult == FBuildDependencySet::ECurrentKeyResult::Error)
 	{
 		return FPackageArtifacts();
diff --git a/Engine/Source/Editor/UnrealEd/Private/Cooker/CookPackageArtifacts.h b/Engine/Source/Editor/UnrealEd/Private/Cooker/CookPackageArtifacts.h
--- a/Engine/Source/Editor/UnrealEd/Private/Cooker/CookPackageArtifacts.h
+++ b/Engine/Source/Editor/UnrealEd/Private/Cooker/CookPackageArtifacts.h
@@ -99,7 +99,7 @@ public:
 	 */
 	const FIoHash& GetCurrentKey() const;
 	/** Call TryCalculateCurrentKey if not yet called, and return whether StoredKey == CurrentKey. */
-	bool HasKeyMatch(FName PackageName, FGenerationHelper* GenerationHelper);
+	bool HasKeyMatch(FName PackageName, const ITargetPlatform* TargetPlatform, FGenerationHelper* GenerationHelper);
 
 	enum class ECurrentKeyResult
 	{
@@ -113,8 +113,8 @@ public:
 	 * @param GenerationHelper - If non-null, provides the lookup for the AssetPackageData of generated packages.
 	 *                           Must be provided if any generated packages are in the dependencies.
 	 */
-	ECurrentKeyResult TryCalculateCurrentKey(FName PackageName, FGenerationHelper* GenerationHelper,
-		TArray<TPair<ELogVerbosity::Type, FString>>* OutMessages = nullptr);
+	ECurrentKeyResult TryCalculateCurrentKey(FName PackageName, const ITargetPlatform* TargetPlatform,
+		FGenerationHelper* GenerationHelper, TArray<TPair<ELogVerbosity::Type, FString>>* OutMessages = nullptr);
 
 	/** Clear data (except Name) and free memory. */
 	void Empty();
@@ -221,11 +221,12 @@ public:
 	 * Calculate the current key(s) from the BuildDependencies stored on this PackageArtifacts, and store it in
 	 * GetCurrentKey().
 	 * 
+	 * @param TargetPlatform - Which targetplatform we are collecting for.
 	 * @param GenerationHelper - If non-null, provides the lookup for the AssetPackageData of generated packages.
 	 *                           Must be provided if any generated packages are in the dependencies.
 	 */
-	FBuildDependencySet::ECurrentKeyResult TryCalculateCurrentKey(FGenerationHelper* GenerationHelper,
-		TArray<TPair<ELogVerbosity::Type, FString>>* OutMessages = nullptr);
+	FBuildDependencySet::ECurrentKeyResult TryCalculateCurrentKey(const ITargetPlatform* TargetPlatform,
+		FGenerationHelper* GenerationHelper, TArray<TPair<ELogVerbosity::Type, FString>>* OutMessages = nullptr);
 	void Empty();
 
 	// Construction functions
@@ -252,7 +253,7 @@ public:
 	void GetTransitiveBuildDependencies(TArray<FName, AllocatorType>& OutDependencies) const;
 	const FIoHash& GetStoredKey() const;
 	const FIoHash& GetCurrentKey() const;
-	bool HasKeyMatch(FGenerationHelper* GenerationHelper);
+	bool HasKeyMatch(const ITargetPlatform* TargetPlatform, FGenerationHelper* GenerationHelper);
 
 private:
 	FBuildDependencySet LoadBuildDependencies;
diff --git a/Engine/Source/Editor/UnrealEd/Private/Cooker/CookPackageData.cpp b/Engine/Source/Editor/UnrealEd/Private/Cooker/CookPackageData.cpp
--- a/Engine/Source/Editor/UnrealEd/Private/Cooker/CookPackageData.cpp
+++ b/Engine/Source/Editor/UnrealEd/Private/Cooker/CookPackageData.cpp
@@ -373,7 +373,7 @@ void FPackageData::ClearInProgressData(EStateChangeReason StateChangeReason)
 		// state after a garbage collect. Call it here in case we cancel the save of the packagedata after demotion.
 		// The other self-references (iterative, queued packages) should persist even when the packagedata is not
 		// in progress.
-		GenerationHelper->ClearKeepForGeneratorSave();
+		GenerationHelper->ClearKeepForGeneratorSaveAllPlatforms();
 	}
 	SetParentGenerationHelper(nullptr, StateChangeReason);
 }
@@ -1609,7 +1609,7 @@ void FPackageData::SetParentGenerationHelper(FGenerationHelper* InGenerationHelp
 		}
 		if (InfoOfPackageInGenerator)
 		{
-			InfoOfPackageInGenerator->SetHasSaved(*ParentGenerationHelper, true, FWorkerId::Local());
+			InfoOfPackageInGenerator->SetAllPlatformsSaved(*ParentGenerationHelper);
 		}
 	}
 	ParentGenerationHelper = InGenerationHelper;
diff --git a/Engine/Source/Editor/UnrealEd/Private/Cooker/CookRequestCluster.cpp b/Engine/Source/Editor/UnrealEd/Private/Cooker/CookRequestCluster.cpp
--- a/Engine/Source/Editor/UnrealEd/Private/Cooker/CookRequestCluster.cpp
+++ b/Engine/Source/Editor/UnrealEd/Private/Cooker/CookRequestCluster.cpp
@@ -1517,6 +1517,22 @@ bool FRequestCluster::FGraphSearch::FExploreEdgesContext::TryCalculateIterativel
 	Vertex->GetUnreadyDependencies().Reset();
 	Vertex->SetWaitingOnUnreadyDependencies(false);
 	bool bAllPlatformsAreReady = true;
+
+	TRefCountPtr<FGenerationHelper> GenerationHelper;
+	FPackageData* ParentPackageData = nullptr;
+	if (!PackageData->IsGenerated())
+	{
+		GenerationHelper = PackageData->GetGenerationHelper();
+	}
+	else
+	{
+		GenerationHelper = PackageData->GetOrFindParentGenerationHelper();
+		if (GenerationHelper)
+		{
+			ParentPackageData = &GenerationHelper->GetOwner();
+		}
+	}
+
 	for (int32 PlatformIndex : PlatformsToProcess)
 	{
 		if (PlatformIndex == CookerLoadingPlatformIndex)
@@ -1535,15 +1551,8 @@ bool FRequestCluster::FGraphSearch::FExploreEdgesContext::TryCalculateIterativel
 		FPackagePlatformData& PackagePlatformData = PackageData->FindOrAddPlatformData(TargetPlatform);
 
 		UE::Cook::FPackageArtifacts& Artifacts = QueryPlatformData.CookAttachments.Artifacts;
-		FPackageData* ParentPackageData = nullptr;
-		TRefCountPtr<FGenerationHelper> GenerationHelper;
-		if (!PackageData->IsGenerated())
-		{
-			GenerationHelper = PackageData->GetGenerationHelper();
-		}
-		else
+		if (PackageData->IsGenerated())
 		{
-			GenerationHelper = PackageData->GetOrFindParentGenerationHelper();
 			if (!GenerationHelper)
 			{
 				// Generated packages for which we do not have the GenerationHelper available are not iteratively
@@ -1551,7 +1560,7 @@ bool FRequestCluster::FGraphSearch::FExploreEdgesContext::TryCalculateIterativel
 				SetIsIterativelyUnmodified(PlatformIndex, false, PackagePlatformData);
 				continue;
 			}
-			ParentPackageData = &GenerationHelper->GetOwner();
+			check(ParentPackageData); // Set above in the IsGenerated && GenerationHelper case
 			// If a generator is marked iteratively unmodified, then by contract we are not required to test its
 			// generated packages; they are all marked iteratively unmodified as well
 			const FPackagePlatformData* ParentPlatformData =
@@ -1565,7 +1574,7 @@ bool FRequestCluster::FGraphSearch::FExploreEdgesContext::TryCalculateIterativel
 				}
 			}
 		}
-		if (!Artifacts.HasKeyMatch(GenerationHelper.GetReference()))
+		if (!Artifacts.HasKeyMatch(TargetPlatform, GenerationHelper.GetReference()))
 		{
 			SetIsIterativelyUnmodified(PlatformIndex, false, PackagePlatformData);
 			continue;
@@ -1661,6 +1670,39 @@ bool FRequestCluster::FGraphSearch::FExploreEdgesContext::TryCalculateIterativel
 		return false;
 	}
 
+	if (GenerationHelper && !PackageData->IsGenerated())
+	{
+		// We don't currently support separate incrementally skippable results per platform for a generator, see the
+		// class comment on FGenerationHelper. Therefore if any platform was found to be incrementally modified, then
+		// set all platforms to modified.
+		bool bAllUnmodified = true;
+		for (int32 PlatformIndex : PlatformsToProcess)
+		{
+			if (PlatformIndex == CookerLoadingPlatformIndex)
+			{
+				continue;
+			}
+			FQueryPlatformData& QueryPlatformData = Vertex->GetPlatformData()[PlatformIndex];
+			check(QueryPlatformData.bIterativelyUnmodified.IsSet()); // Otherwise we early exited in !bAllPlatformsAreReady
+			bAllUnmodified &= *QueryPlatformData.bIterativelyUnmodified;
+		}
+		if (!bAllUnmodified)
+		{
+			for (int32 PlatformIndex : PlatformsToProcess)
+			{
+				if (PlatformIndex == CookerLoadingPlatformIndex)
+				{
+					continue;
+				}
+				FFetchPlatformData& FetchPlatformData = GraphSearch.FetchPlatforms[PlatformIndex];
+				const ITargetPlatform* TargetPlatform = FetchPlatformData.Platform;
+				FPackagePlatformData& PackagePlatformData = PackageData->FindOrAddPlatformData(TargetPlatform);
+				SetIsIterativelyUnmodified(PlatformIndex, false, PackagePlatformData);
+				PackagePlatformData.SetIterativelyUnmodified(false);
+			}
+		}
+	}
+
 	Vertex->GetUnreadyDependencies().Empty(); // Already reset above, call empty to free the memory
 	TArray<FVertexData*>& Listeners = Vertex->GetIterativelyModifiedListeners();
 	if (!Listeners.IsEmpty())
@@ -2100,11 +2142,13 @@ void FRequestCluster::FGraphSearch::FExploreEdgesContext::ProcessPlatformAttachm
 				PackagePlatformData.SetIterativelySkipped(true);
 				if (TRefCountPtr<FGenerationHelper> GenerationHelper = PackageData->GetGenerationHelper())
 				{
-					GenerationHelper->MarkPackageIterativelySkipped(*PackageData, true /* bIncrementallySkipped */);
+					GenerationHelper->MarkPackageIterativelySkipped(*PackageData, TargetPlatform,
+						true /* bIncrementallySkipped */);
 				}
 				if (ParentGenerationHelper)
 				{
-					ParentGenerationHelper->MarkPackageIterativelySkipped(*PackageData, true /* bIncrementallySkipped */);
+					ParentGenerationHelper->MarkPackageIterativelySkipped(*PackageData, TargetPlatform,
+						true /* bIncrementallySkipped */);
 				}
 				if (CookResult == ECookResult::Succeeded && PlatformIndex == FirstSessionPlatformIndex)
 				{
@@ -2127,11 +2171,13 @@ void FRequestCluster::FGraphSearch::FExploreEdgesContext::ProcessPlatformAttachm
 			{
 				if (TRefCountPtr<FGenerationHelper> GenerationHelper = PackageData->GetGenerationHelper())
 				{
-					GenerationHelper->MarkPackageIterativelySkipped(*PackageData, false /* bIncrementallySkipped */);
+					GenerationHelper->MarkPackageIterativelySkipped(*PackageData, TargetPlatform,
+						false /* bIncrementallySkipped */);
 				}
 				if (ParentGenerationHelper)
 				{
-					ParentGenerationHelper->MarkPackageIterativelySkipped(*PackageData, false /* bIncrementallySkipped */);
+					ParentGenerationHelper->MarkPackageIterativelySkipped(*PackageData, TargetPlatform,
+						false /* bIncrementallySkipped */);
 				}
 			}
 		}
@@ -2153,7 +2199,7 @@ void FRequestCluster::FGraphSearch::FExploreEdgesContext::ProcessPlatformAttachm
 			{
 				GenerationHelper = PackageData->GetOrFindParentGenerationHelper();
 			}
-			if (PlatformAgnosticQueryData.CookAttachments.Artifacts.HasKeyMatch(GenerationHelper.GetReference()))
+			if (PlatformAgnosticQueryData.CookAttachments.Artifacts.HasKeyMatch(nullptr, GenerationHelper.GetReference()))
 			{
 				Cluster.BuildDefinitions.AddBuildDefinitionList(PackageName, TargetPlatform,
 					PlatformAgnosticQueryData.CookAttachments.BuildDefinitions.Definitions);
@@ -2205,8 +2251,9 @@ FRequestCluster::FindOrAddVertex(FName PackageName, FGenerationHelper* ParentGen
 		PackageData = COTFS.PackageDatas->TryAddPackageDataByPackageName(PackageName);
 		if (!PackageData && ParentGenerationHelper && ICookPackageSplitter::IsUnderGeneratedPackageSubPath(NameBuffer))
 		{
-			const FAssetPackageData* PreviousPackageData =
-				ParentGenerationHelper->GetAssetPackageData(PackageName);
+			// Look up the AssetPackageData for the generated package, from any previously recorded platform;
+			// we just need to know whether it was a .map or .uasset, which should be the same per platform.
+			const FAssetPackageData* PreviousPackageData = ParentGenerationHelper->GetAssetPackageDataAnyPlatform(PackageName);
 			if (PreviousPackageData)
 			{
 				bool bIsMap = PreviousPackageData->Extension == EPackageExtension::Map;
diff --git a/Engine/Source/Editor/UnrealEd/Private/Cooker/CookSavePackage.cpp b/Engine/Source/Editor/UnrealEd/Private/Cooker/CookSavePackage.cpp
--- a/Engine/Source/Editor/UnrealEd/Private/Cooker/CookSavePackage.cpp
+++ b/Engine/Source/Editor/UnrealEd/Private/Cooker/CookSavePackage.cpp
@@ -442,7 +442,7 @@ void FSaveCookedPackageContext::FinishPlatform()
 			OverrideAssetPackageData.Emplace();
 			OverridePackageDependencies.Emplace();
 			GeneratedPackageBuildResultDependencies.Emplace();
-			GenerationHelper->FinishGeneratedPlatformSave(PackageData, *OverrideAssetPackageData,
+			GenerationHelper->FinishGeneratedPlatformSave(PackageData, TargetPlatform, *OverrideAssetPackageData,
 				 *OverridePackageDependencies, *GeneratedPackageBuildResultDependencies);
 			AssetPackageData = OverrideAssetPackageData.GetPtrOrNull();
 		}
diff --git a/Engine/Source/Editor/UnrealEd/Private/Cooker/CookWorkerClient.cpp b/Engine/Source/Editor/UnrealEd/Private/Cooker/CookWorkerClient.cpp
--- a/Engine/Source/Editor/UnrealEd/Private/Cooker/CookWorkerClient.cpp
+++ b/Engine/Source/Editor/UnrealEd/Private/Cooker/CookWorkerClient.cpp
@@ -925,10 +925,15 @@ void FCookWorkerClient::AssignPackages(FAssignPackagesMessage& Message)
 				PackageData.SetGenerated(AssignData.ParentGenerator);
 				PackageData.SetDoesGeneratedRequireGenerator(AssignData.DoesGeneratedRequireGenerator);
 			}
-			if (!AssignData.GeneratorPreviousGeneratedPackages.IsEmpty())
+			if (!AssignData.GeneratorPerPlatformPreviousGeneratedPackages.IsEmpty())
 			{
 				TRefCountPtr<FGenerationHelper> GenerationHelper = PackageData.CreateUninitializedGenerationHelper();
-				GenerationHelper->SetPreviousGeneratedPackages(MoveTemp(AssignData.GeneratorPreviousGeneratedPackages));
+				for (TPair<uint8, TMap<FName, FAssetPackageData>>& PlatformPair
+					: AssignData.GeneratorPerPlatformPreviousGeneratedPackages)
+				{
+					const ITargetPlatform* TargetPlatform = OrderedSessionPlatforms[PlatformPair.Key];
+					GenerationHelper->SetPreviousGeneratedPackages(TargetPlatform, MoveTemp(PlatformPair.Value));
+				}
 			}
 			if (!AssignData.PerPackageCollectorMessages.IsEmpty())
 			{
diff --git a/Engine/Source/Editor/UnrealEd/Private/Cooker/CookWorkerServer.cpp b/Engine/Source/Editor/UnrealEd/Private/Cooker/CookWorkerServer.cpp
--- a/Engine/Source/Editor/UnrealEd/Private/Cooker/CookWorkerServer.cpp
+++ b/Engine/Source/Editor/UnrealEd/Private/Cooker/CookWorkerServer.cpp
@@ -733,7 +733,18 @@ void FCookWorkerServer::SendPendingPackages()
 		FAssignPackageExtraData* ExtraData = PackagesToAssignExtraDatas.Find(PackageData);
 		if (ExtraData)
 		{
-			AssignData.GeneratorPreviousGeneratedPackages = MoveTemp(ExtraData->GeneratorPreviousGeneratedPackages);
+			AssignData.GeneratorPerPlatformPreviousGeneratedPackages.Reserve(
+				ExtraData->GeneratorPerPlatformPreviousGeneratedPackages.Num());
+			for (TPair<const ITargetPlatform*, TMap<FName, FAssetPackageData>>& PlatformPair
+				: ExtraData->GeneratorPerPlatformPreviousGeneratedPackages)
+			{
+				int32 PlatformIdInt = OrderedSessionPlatforms.IndexOfByKey(PlatformPair.Key);
+				check(0 <= PlatformIdInt && PlatformIdInt <= MAX_uint8);
+				uint8 PlatformId = static_cast<uint8>(PlatformIdInt);
+				TMap<FName, FAssetPackageData>& DestMap
+					= AssignData.GeneratorPerPlatformPreviousGeneratedPackages.Add(PlatformId);
+				DestMap = MoveTemp(PlatformPair.Value);
+			}
 			AssignData.PerPackageCollectorMessages = MoveTemp(ExtraData->PerPackageCollectorMessages);
 		}
 	}
@@ -1036,63 +1047,62 @@ void FCookWorkerServer::RecordResults(FPackageResultsMessage& Message)
 				*WriteCookStatus(*PackageData, COTFS.GetSessionPlatforms()));
 			continue;
 		}
+		int32 NumPlatforms = OrderedSessionPlatforms.Num();
+		if (Result.GetPlatforms().Num() != NumPlatforms)
+		{
+			UE_LOG(LogCook, Error,
+				TEXT("CookWorkerServer %d received FPackageResultsMessage for package %s with an invalid number of platform results: expected %d, actual %d. Ignoring it."),
+				ProfileId, *Result.GetPackageName().ToString(), NumPlatforms, Result.GetPlatforms().Num());
+			continue;
+		}
 		bRetiredAnyPackages = true;
 
 		// MPCOOKTODO: Refactor FSaveCookedPackageContext::FinishPlatform and ::FinishPackage so we can call them from
 		// here to reduce duplication
-		if (Result.GetSuppressCookReason() == ESuppressCookReason::NotSuppressed)
+		bool bResultIsSaveResult = Result.GetSuppressCookReason() == ESuppressCookReason::NotSuppressed;
+		if (bResultIsSaveResult)
 		{
-			int32 NumPlatforms = OrderedSessionPlatforms.Num();
-			if (Result.GetPlatforms().Num() != NumPlatforms)
-			{
-				UE_LOG(LogCook, Error,
-					TEXT("CookWorkerServer %d received FPackageResultsMessage for package %s with an invalid number of platform results: expected %d, actual %d. Ignoring it."),
-					ProfileId, *Result.GetPackageName().ToString(), NumPlatforms, Result.GetPlatforms().Num());
-			}
-			else
+			HandleReceivedPackagePlatformMessages(*PackageData, nullptr, Result.ReleaseMessages());
+			for (int32 PlatformIndex = 0; PlatformIndex < NumPlatforms; ++PlatformIndex)
 			{
-				HandleReceivedPackagePlatformMessages(*PackageData, nullptr, Result.ReleaseMessages());
-				for (int32 PlatformIndex = 0; PlatformIndex < NumPlatforms; ++PlatformIndex)
+				ITargetPlatform* TargetPlatform = OrderedSessionPlatforms[PlatformIndex];
+				FPackageRemoteResult::FPlatformResult& PlatformResult = Result.GetPlatforms()[PlatformIndex];
+				FPackagePlatformData& ExistingData = PackageData->FindOrAddPlatformData(TargetPlatform);
+				if (ExistingData.IsCommitted())
 				{
-					ITargetPlatform* TargetPlatform = OrderedSessionPlatforms[PlatformIndex];
-					FPackageRemoteResult::FPlatformResult& PlatformResult = Result.GetPlatforms()[PlatformIndex];
-					FPackagePlatformData& ExistingData = PackageData->FindOrAddPlatformData(TargetPlatform);
-					if (ExistingData.IsCommitted())
+					if (PlatformResult.WasCommitted())
 					{
-						if (PlatformResult.WasCommitted())
-						{
-							UE_LOG(LogCook, Display,
-								TEXT("CookWorkerServer %d received FPackageResultsMessage for package %s, platform %s, but that platform has already been committed. Ignoring the results for that platform."),
-								ProfileId, *Result.GetPackageName().ToString(), *TargetPlatform->PlatformName());
-						}
+						UE_LOG(LogCook, Display,
+							TEXT("CookWorkerServer %d received FPackageResultsMessage for package %s, platform %s, but that platform has already been committed. Ignoring the results for that platform."),
+							ProfileId, *Result.GetPackageName().ToString(), *TargetPlatform->PlatformName());
+					}
+				}
+				else
+				{
+					bool bWasCooked = PlatformResult.GetCookResults() != ECookResult::Invalid
+						&& PlatformResult.GetCookResults() != ECookResult::NotAttempted;
+					if (!ExistingData.NeedsCooking(TargetPlatform) && bWasCooked)
+					{
+						UE_LOG(LogCook, Display,
+							TEXT("CookWorkerServer %d received FPackageResultsMessage for package %s, platform %s, but that platform has already been cooked. Ignoring the results for that platform."),
+							ProfileId, *Result.GetPackageName().ToString(), *TargetPlatform->PlatformName());
 					}
 					else
 					{
-						bool bWasCooked = PlatformResult.GetCookResults() != ECookResult::Invalid
-							&& PlatformResult.GetCookResults() != ECookResult::NotAttempted;
-						if (!ExistingData.NeedsCooking(TargetPlatform) && bWasCooked)
+						if (bWasCooked)
 						{
-							UE_LOG(LogCook, Display,
-								TEXT("CookWorkerServer %d received FPackageResultsMessage for package %s, platform %s, but that platform has already been cooked. Ignoring the results for that platform."),
-								ProfileId, *Result.GetPackageName().ToString(), *TargetPlatform->PlatformName());
+							PackageData->SetPlatformCooked(TargetPlatform, PlatformResult.GetCookResults());
 						}
-						else
+						else if (PlatformResult.WasCommitted())
 						{
-							if (bWasCooked)
-							{
-								PackageData->SetPlatformCooked(TargetPlatform, PlatformResult.GetCookResults());
-							}
-							else if (PlatformResult.WasCommitted())
-							{
-								PackageData->SetPlatformCommitted(TargetPlatform);
-							}
-							HandleReceivedPackagePlatformMessages(*PackageData, TargetPlatform,
-								PlatformResult.ReleaseMessages());
+							PackageData->SetPlatformCommitted(TargetPlatform);
 						}
+						HandleReceivedPackagePlatformMessages(*PackageData, TargetPlatform,
+							PlatformResult.ReleaseMessages());
 					}
 				}
-				COTFS.RecordExternalActorDependencies(Result.GetExternalActorDependencies());
 			}
+			COTFS.RecordExternalActorDependencies(Result.GetExternalActorDependencies());
 		}
 
 		// For generator and generated packages, after we handle all of their messages above, execute their state
@@ -1117,7 +1127,15 @@ void FCookWorkerServer::RecordResults(FPackageResultsMessage& Message)
 			else
 			{
 				DeferGenerationHelperEvents.Emplace(ParentGenerationHelper);
-				ParentGenerationHelper->MarkPackageSavedRemotely(COTFS, *PackageData, GetWorkerId());
+				for (int32 PlatformIndex = 0; PlatformIndex < NumPlatforms; ++PlatformIndex)
+				{
+					ITargetPlatform* TargetPlatform = OrderedSessionPlatforms[PlatformIndex];
+					FPackageRemoteResult::FPlatformResult& PlatformResult = Result.GetPlatforms()[PlatformIndex];
+					if (!bResultIsSaveResult || PlatformResult.WasCommitted())
+					{
+						ParentGenerationHelper->MarkPackageSavedRemotely(COTFS, *PackageData, TargetPlatform, GetWorkerId());
+					}
+				}
 				EStateChangeReason StateChangeReason =
 					Result.GetSuppressCookReason() == ESuppressCookReason::NotSuppressed
 					? EStateChangeReason::Saved
@@ -1129,7 +1147,15 @@ void FCookWorkerServer::RecordResults(FPackageResultsMessage& Message)
 		if (GenerationHelper)
 		{
 			DeferGenerationHelperEvents.Emplace(GenerationHelper);
-			GenerationHelper->MarkPackageSavedRemotely(COTFS, *PackageData, GetWorkerId());
+			for (int32 PlatformIndex = 0; PlatformIndex < NumPlatforms; ++PlatformIndex)
+			{
+				ITargetPlatform* TargetPlatform = OrderedSessionPlatforms[PlatformIndex];
+				FPackageRemoteResult::FPlatformResult& PlatformResult = Result.GetPlatforms()[PlatformIndex];
+				if (!bResultIsSaveResult || PlatformResult.WasCommitted())
+				{
+					GenerationHelper->MarkPackageSavedRemotely(COTFS, *PackageData, TargetPlatform, GetWorkerId());
+				}
+			}
 			GenerationHelper.SafeRelease();
 		}
 
@@ -1398,11 +1424,20 @@ void FAssignPackageData::Write(FCbWriter& Writer,
 	WriteToCompactBinary(Writer, NeedCommitPlatforms, OrderedSessionPlatforms);
 	{
 		Writer.BeginArray();
-		for (const TPair<FName, FAssetPackageData>& Pair : GeneratorPreviousGeneratedPackages)
+		for (const TPair<uint8, TMap<FName, FAssetPackageData>>& PlatformPair
+			: GeneratorPerPlatformPreviousGeneratedPackages)
 		{
 			Writer.BeginArray();
-			Writer << Pair.Key;
-			Pair.Value.NetworkWrite(Writer);
+			Writer << PlatformPair.Key;
+			Writer.BeginArray();
+			for (const TPair<FName, FAssetPackageData>& PackagePair : PlatformPair.Value)
+			{
+				Writer.BeginArray();
+				Writer << PackagePair.Key;
+				PackagePair.Value.NetworkWrite(Writer);
+				Writer.EndArray();
+			}
+			Writer.EndArray();
 			Writer.EndArray();
 		}
 		Writer.EndArray();
@@ -1440,33 +1475,55 @@ bool FAssignPackageData::TryRead(FCbFieldView Field, TConstArrayView<const ITarg
 	}
 	bOk = LoadFromCompactBinary(*It++, NeedCommitPlatforms, OrderedSessionPlatforms) & bOk;
 	{
-		FCbFieldView ArrayFieldView = *It++;
+		FCbFieldView PlatformArrayFieldView = *It++;
 		bool bGeneratorPreviousGeneratedPackagesOk = false;
-		const uint64 Length = ArrayFieldView.AsArrayView().Num();
-		if (Length <= MAX_int32)
+		const uint64 PlatformLength = PlatformArrayFieldView.AsArrayView().Num();
+		if (PlatformLength <= MAX_int32)
 		{
-			GeneratorPreviousGeneratedPackages.Empty((int32)Length);
-			bGeneratorPreviousGeneratedPackagesOk = !ArrayFieldView.HasError();
-			for (const FCbFieldView& ElementField : ArrayFieldView)
+			GeneratorPerPlatformPreviousGeneratedPackages.Empty((int32)PlatformLength);
+			bGeneratorPreviousGeneratedPackagesOk = !PlatformArrayFieldView.HasError();
+			for (const FCbFieldView& PlatformIt : PlatformArrayFieldView)
 			{
-				FCbFieldViewIterator PairIt = ElementField.CreateViewIterator();
-				bool bElementOk = false;
-				FName Key;
-				FAssetPackageData Value;
-				if (LoadFromCompactBinary(*PairIt++, Key))
+				bool bPlatformPairOk = false;
+				FCbFieldViewIterator PlatformPairIt = PlatformIt.CreateViewIterator();
+				uint8 PlatformIndex;
+				TMap<FName, FAssetPackageData> PackagesMap;
+				if (LoadFromCompactBinary(*PlatformPairIt++, PlatformIndex))
 				{
-					if (Value.TryNetworkRead(*PairIt++))
+					FCbFieldView PackagesArrayFieldView = *PlatformPairIt++;
+					const uint64 PackagesLength = PackagesArrayFieldView.AsArrayView().Num();
+					if (PackagesLength <= MAX_int32)
 					{
-						GeneratorPreviousGeneratedPackages.Add(Key, MoveTemp(Value));
-						bElementOk = true;
+						PackagesMap.Empty((int32)PackagesLength);
+						bPlatformPairOk = !PackagesArrayFieldView.HasError();
+						for (const FCbFieldView& PackagesElementField : PackagesArrayFieldView)
+						{
+							FCbFieldViewIterator PairIt = PackagesElementField.CreateViewIterator();
+							bool bElementOk = false;
+							FName Key;
+							FAssetPackageData Value;
+							if (LoadFromCompactBinary(*PairIt++, Key))
+							{
+								if (Value.TryNetworkRead(*PairIt++))
+								{
+									PackagesMap.Add(Key, MoveTemp(Value));
+									bElementOk = true;
+								}
+							}
+							bPlatformPairOk &= bElementOk;
+						}
 					}
 				}
-				bGeneratorPreviousGeneratedPackagesOk &= bElementOk;
+				if (bPlatformPairOk)
+				{
+					GeneratorPerPlatformPreviousGeneratedPackages.Add(PlatformIndex, MoveTemp(PackagesMap));
+				}
+				bGeneratorPreviousGeneratedPackagesOk &= bPlatformPairOk;
 			}
 		}
 		else
 		{
-			GeneratorPreviousGeneratedPackages.Empty();
+			GeneratorPerPlatformPreviousGeneratedPackages.Empty();
 		}
 		bOk &= bGeneratorPreviousGeneratedPackagesOk;
 	}
diff --git a/Engine/Source/Editor/UnrealEd/Private/Cooker/CookWorkerServer.h b/Engine/Source/Editor/UnrealEd/Private/Cooker/CookWorkerServer.h
--- a/Engine/Source/Editor/UnrealEd/Private/Cooker/CookWorkerServer.h
+++ b/Engine/Source/Editor/UnrealEd/Private/Cooker/CookWorkerServer.h
@@ -50,7 +50,7 @@ enum class ENotifyRemote
  */
 struct FAssignPackageExtraData
 {
-	TMap<FName, FAssetPackageData> GeneratorPreviousGeneratedPackages;
+	TMap<const ITargetPlatform*, TMap<FName, FAssetPackageData>> GeneratorPerPlatformPreviousGeneratedPackages;
 	TArray<UE::CompactBinaryTCP::FMarshalledMessage> PerPackageCollectorMessages;
 };
 
@@ -262,7 +262,7 @@ struct FAssignPackageData
 	EUrgency Urgency;
 	EReachability Reachability;
 	FDiscoveredPlatformSet NeedCommitPlatforms;
-	TMap<FName, FAssetPackageData> GeneratorPreviousGeneratedPackages;
+	TMap<uint8, TMap<FName, FAssetPackageData>> GeneratorPerPlatformPreviousGeneratedPackages;
 	TArray<UE::CompactBinaryTCP::FMarshalledMessage> PerPackageCollectorMessages;
 	ICookPackageSplitter::EGeneratedRequiresGenerator DoesGeneratedRequireGenerator
 		= ICookPackageSplitter::EGeneratedRequiresGenerator::None;
