[Transforms] Modify startup ordering for AxisDisplayInfo so that the non-settings information is available early on.

diff --git a/Engine/Source/Editor/UnrealEd/Private/EditorAxisDisplayInfo.cpp b/Engine/Source/Editor/UnrealEd/Private/EditorAxisDisplayInfo.cpp
--- a/Engine/Source/Editor/UnrealEd/Private/EditorAxisDisplayInfo.cpp
+++ b/Engine/Source/Editor/UnrealEd/Private/EditorAxisDisplayInfo.cpp
@@ -5,6 +5,7 @@
 #include "HAL/IConsoleManager.h"
 #include "Internationalization/Internationalization.h"
 #include "Settings/EditorStyleSettings.h"
+#include "Editor.h"
 
 #define LOCTEXT_NAMESPACE "EditorAxisDisplayInfo"
 
@@ -20,21 +21,7 @@ namespace UE::Editor::Private
 
 FEditorAxisDisplayInfo::FEditorAxisDisplayInfo()
 {
-	if (FProperty* XAxisColorProperty = UEditorStyleSettings::StaticClass()->FindPropertyByName(GET_MEMBER_NAME_CHECKED(UEditorStyleSettings, XAxisColor)))
-	{
-		XAxisColorProperty->SetMetaData(TEXT("DisplayName"), *FText::Format(LOCTEXT("XAxisColorDisplayName", "{0} Axis Color"), GetAxisDisplayName(EAxisList::Forward)).ToString());
-		XAxisColorProperty->SetMetaData(TEXT("ToolTip"), *FText::Format(LOCTEXT("XAxisColorToolTip", "The color used for the {0} axis"), GetAxisDisplayName(EAxisList::Forward)).ToString());
-	}
-	if (FProperty* YAxisColorProperty = UEditorStyleSettings::StaticClass()->FindPropertyByName(GET_MEMBER_NAME_CHECKED(UEditorStyleSettings, YAxisColor)))
-	{
-		YAxisColorProperty->SetMetaData(TEXT("DisplayName"), *FText::Format(LOCTEXT("YAxisColorDisplayName", "{0} Axis Color"), GetAxisDisplayName(EAxisList::Left)).ToString());
-		YAxisColorProperty->SetMetaData(TEXT("ToolTip"), *FText::Format(LOCTEXT("YAxisColorToolTip", "The color used for the {0} axis"), GetAxisDisplayName(EAxisList::Left)).ToString());
-	}
-	if (FProperty* ZAxisColorProperty = UEditorStyleSettings::StaticClass()->FindPropertyByName(GET_MEMBER_NAME_CHECKED(UEditorStyleSettings, ZAxisColor)))
-	{
-		ZAxisColorProperty->SetMetaData(TEXT("DisplayName"), *FText::Format(LOCTEXT("ZAxisColorDisplayName", "{0} Axis Color"), GetAxisDisplayName(EAxisList::Up)).ToString());
-		ZAxisColorProperty->SetMetaData(TEXT("ToolTip"), *FText::Format(LOCTEXT("ZAxisColorToolTip", "The color used for the {0} axis"), GetAxisDisplayName(EAxisList::Up)).ToString());
-	}
+	FEditorDelegates::OnEditorBoot.AddRaw(this, &FEditorAxisDisplayInfo::InitSettingsInfo);
 }
 
 EAxisList::Type FEditorAxisDisplayInfo::GetAxisDisplayCoordinateSystem()
@@ -241,4 +228,23 @@ EAxisList::Type FEditorAxisDisplayInfo::MapAxis(EAxisList::Type Axis)
 	return Axis;
 }
 
+void FEditorAxisDisplayInfo::InitSettingsInfo(double EditorStartupTime)
+{
+	if (FProperty* XAxisColorProperty = UEditorStyleSettings::StaticClass()->FindPropertyByName(GET_MEMBER_NAME_CHECKED(UEditorStyleSettings, XAxisColor)))
+	{
+		XAxisColorProperty->SetMetaData(TEXT("DisplayName"), *FText::Format(LOCTEXT("XAxisColorDisplayName", "{0} Axis Color"), GetAxisDisplayName(EAxisList::Forward)).ToString());
+		XAxisColorProperty->SetMetaData(TEXT("ToolTip"), *FText::Format(LOCTEXT("XAxisColorToolTip", "The color used for the {0} axis"), GetAxisDisplayName(EAxisList::Forward)).ToString());
+	}
+	if (FProperty* YAxisColorProperty = UEditorStyleSettings::StaticClass()->FindPropertyByName(GET_MEMBER_NAME_CHECKED(UEditorStyleSettings, YAxisColor)))
+	{
+		YAxisColorProperty->SetMetaData(TEXT("DisplayName"), *FText::Format(LOCTEXT("YAxisColorDisplayName", "{0} Axis Color"), GetAxisDisplayName(EAxisList::Left)).ToString());
+		YAxisColorProperty->SetMetaData(TEXT("ToolTip"), *FText::Format(LOCTEXT("YAxisColorToolTip", "The color used for the {0} axis"), GetAxisDisplayName(EAxisList::Left)).ToString());
+	}
+	if (FProperty* ZAxisColorProperty = UEditorStyleSettings::StaticClass()->FindPropertyByName(GET_MEMBER_NAME_CHECKED(UEditorStyleSettings, ZAxisColor)))
+	{
+		ZAxisColorProperty->SetMetaData(TEXT("DisplayName"), *FText::Format(LOCTEXT("ZAxisColorDisplayName", "{0} Axis Color"), GetAxisDisplayName(EAxisList::Up)).ToString());
+		ZAxisColorProperty->SetMetaData(TEXT("ToolTip"), *FText::Format(LOCTEXT("ZAxisColorToolTip", "The color used for the {0} axis"), GetAxisDisplayName(EAxisList::Up)).ToString());
+	}
+}
+
 #undef LOCTEXT_NAMESPACE
diff --git a/Engine/Source/Editor/UnrealEd/Private/EditorEngine.cpp b/Engine/Source/Editor/UnrealEd/Private/EditorEngine.cpp
--- a/Engine/Source/Editor/UnrealEd/Private/EditorEngine.cpp
+++ b/Engine/Source/Editor/UnrealEd/Private/EditorEngine.cpp
@@ -940,6 +940,9 @@ void UEditorEngine::InitEditor(IEngineLoop* InEngineLoop)
 	}
 	InitDerivedDataBuildWorkers();
 
+	// Startup UnrealEdMisc
+	FUnrealEdMisc::Get();
+
 	// Call base.
 	UEngine::Init(InEngineLoop);
 
diff --git a/Engine/Source/Editor/UnrealEd/Private/UnrealEdMisc.cpp b/Engine/Source/Editor/UnrealEd/Private/UnrealEdMisc.cpp
--- a/Engine/Source/Editor/UnrealEd/Private/UnrealEdMisc.cpp
+++ b/Engine/Source/Editor/UnrealEd/Private/UnrealEdMisc.cpp
@@ -229,6 +229,13 @@ FUnrealEdMisc::FUnrealEdMisc() :
 #endif
 	}
 #endif
+
+	// Register the EditorAxisDisplayInfo modular feature
+	if (!IModularFeatures::Get().IsModularFeatureAvailable(IAxisDisplayInfo::GetModularFeatureName()))
+	{
+		EditorAxisDisplayInfo = MakeUnique<FEditorAxisDisplayInfo>();
+		IModularFeatures::Get().RegisterModularFeature(IAxisDisplayInfo::GetModularFeatureName(), EditorAxisDisplayInfo.Get());
+	}
 }
 
 FUnrealEdMisc::~FUnrealEdMisc()
@@ -256,13 +263,6 @@ void FUnrealEdMisc::OnInit()
 	CmdExec = MakeUnique<FConsoleCommandExecutor>();
 	IModularFeatures::Get().RegisterModularFeature(IConsoleCommandExecutor::ModularFeatureName(), CmdExec.Get());
 
-	// Register the EditorAxisDisplayInfo modular feature
-	if (!IModularFeatures::Get().IsModularFeatureAvailable(IAxisDisplayInfo::GetModularFeatureName()))
-	{
-		EditorAxisDisplayInfo = MakeUnique<FEditorAxisDisplayInfo>();
-		IModularFeatures::Get().RegisterModularFeature(IAxisDisplayInfo::GetModularFeatureName(), EditorAxisDisplayInfo.Get());
-	}
-
 	// Register all callback notifications
 	FEditorDelegates::SelectedProps.AddRaw(this, &FUnrealEdMisc::CB_SelectedProps);
 	FEditorDelegates::DisplayLoadErrors.AddRaw(this, &FUnrealEdMisc::CB_DisplayLoadErrors);
@@ -1147,7 +1147,7 @@ void FUnrealEdMisc::OnExit()
 	}
 
 	// Unregister the EditorAxisDisplayInfo modular feature
-	if (EditorAxisDisplayInfo.IsValid())
+	if (EditorAxisDisplayInfo)
 	{
 		IModularFeatures::Get().UnregisterModularFeature(IAxisDisplayInfo::GetModularFeatureName(), EditorAxisDisplayInfo.Get());
 		EditorAxisDisplayInfo.Reset();
diff --git a/Engine/Source/Editor/UnrealEd/Public/EditorAxisDisplayInfo.h b/Engine/Source/Editor/UnrealEd/Public/EditorAxisDisplayInfo.h
--- a/Engine/Source/Editor/UnrealEd/Public/EditorAxisDisplayInfo.h
+++ b/Engine/Source/Editor/UnrealEd/Public/EditorAxisDisplayInfo.h
@@ -22,10 +22,12 @@ public:
 	virtual bool UseForwardRightUpDisplayNames() override;
 
 private:
-
 	// Maps the given axis from FLU -> XYZ if the AxisDisplayCoordinateSystem is XYZ
 	EAxisList::Type MapAxis(EAxisList::Type Axis);
 
+	// Inits info stored in settings, like the axis colors
+	void InitSettingsInfo(double EditorStartupTime);
+
 	TOptional<bool> bUseForwardRightUpDisplayNames;
 	TOptional<EAxisList::Type> AxisDisplayCoordinateSystem;
 
