Allow derived FAssetEditorToolkit to call RestoreLayout without reading the user layout data.
Calling RestoreLayout with bLoadUserLayout will reset the layout to the code-defined version.

diff --git a/Engine/Source/Editor/UnrealEd/Private/Toolkits/AssetEditorToolkit.cpp b/Engine/Source/Editor/UnrealEd/Private/Toolkits/AssetEditorToolkit.cpp
--- a/Engine/Source/Editor/UnrealEd/Private/Toolkits/AssetEditorToolkit.cpp
+++ b/Engine/Source/Editor/UnrealEd/Private/Toolkits/AssetEditorToolkit.cpp
@@ -871,6 +871,31 @@ void FAssetEditorToolkit::MapToolkitCommands()
 	}
 }
 
+void FAssetEditorToolkit::RestoreFromLayout(const TSharedRef<FTabManager::FLayout>& NewLayout, bool bLoadUserLayout)
+{
+	TSharedPtr< class SStandaloneAssetEditorToolkitHost > HostWidget = StandaloneHost.Pin();
+	if (HostWidget.IsValid())
+	{
+		// Save the old layout
+		FLayoutSaveRestore::SaveToConfig(LayoutIni, TabManager->PersistLayout());
+
+		TSharedRef<FTabManager::FLayout> UserConfiguredNewLayout = NewLayout;
+		if (bLoadUserLayout)
+		{
+			// Load the potentially previously saved new layout
+			UserConfiguredNewLayout = FLayoutSaveRestore::LoadFromConfig(LayoutIni, NewLayout);
+		}
+
+		for (TSharedPtr<FLayoutExtender> LayoutExtender : LayoutExtenders)
+		{
+			NewLayout->ProcessExtensions(*LayoutExtender);
+		}
+
+		// Apply the new layout
+		HostWidget->RestoreFromLayout(UserConfiguredNewLayout);
+	}
+}
+
 FAssetEditorModeManager* FAssetEditorToolkit::GetAssetEditorModeManager() const
 {
 	PRAGMA_DISABLE_DEPRECATION_WARNINGS
@@ -1528,7 +1553,6 @@ void FAssetEditorToolkit::RegenerateMenusAndToolbars()
 	}
 }
 
-
 void FAssetEditorToolkit::RegisterDrawer(FWidgetDrawerConfig&& Drawer, int32 SlotIndex)
 {
 	TSharedPtr< class SStandaloneAssetEditorToolkitHost > HostWidget = StandaloneHost.Pin();
@@ -1540,23 +1564,8 @@ void FAssetEditorToolkit::RegisterDrawer(FWidgetDrawerConfig&& Drawer, int32 Slo
 
 void FAssetEditorToolkit::RestoreFromLayout(const TSharedRef<FTabManager::FLayout>& NewLayout)
 {
-	TSharedPtr< class SStandaloneAssetEditorToolkitHost > HostWidget = StandaloneHost.Pin();
-	if (HostWidget.IsValid())
-	{
-		// Save the old layout
-		FLayoutSaveRestore::SaveToConfig(LayoutIni, TabManager->PersistLayout());
-
-		// Load the potentially previously saved new layout
-		TSharedRef<FTabManager::FLayout> UserConfiguredNewLayout = FLayoutSaveRestore::LoadFromConfig(LayoutIni, NewLayout);
-
-		for (TSharedPtr<FLayoutExtender> LayoutExtender : LayoutExtenders)
-		{
-			NewLayout->ProcessExtensions(*LayoutExtender);
-		}
-
-		// Apply the new layout
-		HostWidget->RestoreFromLayout(UserConfiguredNewLayout);
-	}
+	constexpr bool bLoadUserLayout = true;
+	RestoreFromLayout(NewLayout, bLoadUserLayout);
 }
 
 bool FAssetEditorToolkit::IsActuallyAnAsset() const
diff --git a/Engine/Source/Editor/UnrealEd/Public/Toolkits/AssetEditorToolkit.h b/Engine/Source/Editor/UnrealEd/Public/Toolkits/AssetEditorToolkit.h
--- a/Engine/Source/Editor/UnrealEd/Public/Toolkits/AssetEditorToolkit.h
+++ b/Engine/Source/Editor/UnrealEd/Public/Toolkits/AssetEditorToolkit.h
@@ -455,6 +455,14 @@ protected:
 	/** Can be overriden to prevent registering the default commands.  */
 	UNREALED_API virtual void MapToolkitCommands();
 
+	/** 
+	 * Applies the passed in layout.
+	 * @param NewLayout The layout to restore.
+	 * @param bLoadUserLayout Whether to attempt loading the saved user-modified version.
+	 * @note Must be called after InitAssetEditor. 
+	 */
+	UNREALED_API void RestoreFromLayout(const TSharedRef<FTabManager::FLayout>& NewLayout, bool bLoadUserLayout);
+
 private:
 	// Callback for persisting the Asset Editor's layout.
 	void HandleTabManagerPersistLayout( const TSharedRef<FTabManager::FLayout>& LayoutToSave )
