Anim In Engine: Custom frame rate when exporting animation sequences, includes linked anim sequence support, bp/python export is implicit with change to the AnimSeqExport options.

diff --git a/Engine/Plugins/MovieScene/SequencerScripting/Source/SequencerScriptingEditor/Private/SequencerTools.cpp b/Engine/Plugins/MovieScene/SequencerScripting/Source/SequencerScriptingEditor/Private/SequencerTools.cpp
--- a/Engine/Plugins/MovieScene/SequencerScripting/Source/SequencerScriptingEditor/Private/SequencerTools.cpp
+++ b/Engine/Plugins/MovieScene/SequencerScripting/Source/SequencerScriptingEditor/Private/SequencerTools.cpp
@@ -461,7 +461,8 @@ bool USequencerToolsFunctionLibrary::LinkAnimSequence(ULevelSequence*  Sequence,
 				for (FLevelSequenceAnimSequenceLinkItem& LevelAnimLinkItem : LevelAnimLink->AnimSequenceLinks)
 				{
 					if (LevelAnimLinkItem.IsEqual(Binding.BindingID, ExportOptions->bUseCustomTimeRange,
-						ExportOptions->CustomStartFrame, ExportOptions->CustomEndFrame, ExportOptions->CustomDisplayRate))
+						ExportOptions->CustomStartFrame, ExportOptions->CustomEndFrame, ExportOptions->CustomDisplayRate,
+						ExportOptions->bUseCustomFrameRate, ExportOptions->CustomFrameRate))
 					{
 						bAddItem = false;
 						UAnimSequence* OtherAnimSequence = LevelAnimLinkItem.ResolveAnimSequence();
@@ -496,6 +497,9 @@ bool USequencerToolsFunctionLibrary::LinkAnimSequence(ULevelSequence*  Sequence,
 						LevelAnimLinkItem.CustomStartFrame = ExportOptions->CustomStartFrame;
 						LevelAnimLinkItem.CustomEndFrame = ExportOptions->CustomEndFrame;
 						LevelAnimLinkItem.CustomDisplayRate = ExportOptions->CustomDisplayRate;
+						LevelAnimLinkItem.bUseCustomFrameRate = ExportOptions->bUseCustomFrameRate;
+						LevelAnimLinkItem.CustomFrameRate = ExportOptions->CustomFrameRate;
+
 
 						break;
 					}
@@ -527,6 +531,8 @@ bool USequencerToolsFunctionLibrary::LinkAnimSequence(ULevelSequence*  Sequence,
 				LevelAnimLinkItem.CustomStartFrame = ExportOptions->CustomStartFrame;
 				LevelAnimLinkItem.CustomEndFrame = ExportOptions->CustomEndFrame;
 				LevelAnimLinkItem.CustomDisplayRate = ExportOptions->CustomDisplayRate;
+				LevelAnimLinkItem.bUseCustomFrameRate = ExportOptions->bUseCustomFrameRate;
+				LevelAnimLinkItem.CustomFrameRate = ExportOptions->CustomFrameRate;
 
 				LevelAnimLink->AnimSequenceLinks.Add(LevelAnimLinkItem);
 				AssetUserDataInterface->AddAssetUserData(LevelAnimLink);
diff --git a/Engine/Source/Editor/MovieSceneTools/Private/MovieSceneToolHelpers.cpp b/Engine/Source/Editor/MovieSceneTools/Private/MovieSceneToolHelpers.cpp
--- a/Engine/Source/Editor/MovieSceneTools/Private/MovieSceneToolHelpers.cpp
+++ b/Engine/Source/Editor/MovieSceneTools/Private/MovieSceneToolHelpers.cpp
@@ -4168,6 +4168,10 @@ bool MovieSceneToolHelpers::BakeToSkelMeshToCallbacks(const FAnimExportSequenceP
 		}
 		AnimTrackAdapter.SetRange(StartFrameInTick,EndFrameInTick);
 	}
+	if (ExportOptions->bUseCustomFrameRate)
+	{
+		AnimTrackAdapter.SetFrameRate(ExportOptions->CustomFrameRate);
+	}
 
 	int32 LocalStartFrame = AnimTrackAdapter.GetLocalStartFrame();
 	int32 AnimationLength = AnimTrackAdapter.GetLength();
@@ -4343,8 +4347,7 @@ bool MovieSceneToolHelpers::ExportToAnimSequence(UAnimSequence* AnimSequence, UA
 	}
 
 	FAnimRecorderInstance AnimationRecorder;
-	const FFrameRate SampleRate = AESP.MovieSceneSequence->GetMovieScene()->GetDisplayRate();
-
+	const FFrameRate SampleRate = ExportOptions->bUseCustomFrameRate ? ExportOptions->CustomFrameRate :  AESP.MovieSceneSequence->GetMovieScene()->GetDisplayRate();
 	// Used by function that retrieves the current timecode 
 	FQualifiedFrameTime CurrentTimecode;
 
diff --git a/Engine/Source/Editor/MovieSceneTools/Private/TrackEditors/SkeletalAnimationTrackEditor.cpp b/Engine/Source/Editor/MovieSceneTools/Private/TrackEditors/SkeletalAnimationTrackEditor.cpp
--- a/Engine/Source/Editor/MovieSceneTools/Private/TrackEditors/SkeletalAnimationTrackEditor.cpp
+++ b/Engine/Source/Editor/MovieSceneTools/Private/TrackEditors/SkeletalAnimationTrackEditor.cpp
@@ -1547,6 +1547,8 @@ void FSkeletalAnimationTrackEditor::OnSequencerSaved(ISequencer& )
 						const FFrameNumber SavedCustomStartFrame = AnimSeqExportOption->CustomStartFrame;
 						const FFrameNumber SavedCustomEndFrame = AnimSeqExportOption->CustomEndFrame;
 						const FFrameRate SavedCustomDisplayRate = AnimSeqExportOption->CustomDisplayRate;
+						const bool bSavedUseCustomFrameRate = AnimSeqExportOption->bUseCustomFrameRate;
+						const FFrameRate SavedCustomFrameRate = AnimSeqExportOption->CustomFrameRate;
 
 						AnimSeqExportOption->bExportMorphTargets = Item.bExportMorphTargets;
 						AnimSeqExportOption->bExportAttributeCurves = Item.bExportAttributeCurves;
@@ -1565,6 +1567,8 @@ void FSkeletalAnimationTrackEditor::OnSequencerSaved(ISequencer& )
 						AnimSeqExportOption->CustomStartFrame = Item.CustomStartFrame;
 						AnimSeqExportOption->CustomEndFrame = Item.CustomEndFrame;
 						AnimSeqExportOption->CustomDisplayRate = Item.CustomDisplayRate;
+						AnimSeqExportOption->bUseCustomFrameRate = Item.bUseCustomFrameRate;
+						AnimSeqExportOption->CustomFrameRate = Item.CustomFrameRate;
 
 						FAnimExportSequenceParameters AESP;
 						AESP.Player = SequencerPtr.Get();
@@ -1593,7 +1597,8 @@ void FSkeletalAnimationTrackEditor::OnSequencerSaved(ISequencer& )
 						AnimSeqExportOption->CustomStartFrame = SavedCustomStartFrame;
 						AnimSeqExportOption->CustomEndFrame = SavedCustomEndFrame;
 						AnimSeqExportOption->CustomDisplayRate = SavedCustomDisplayRate;
-
+						AnimSeqExportOption->bUseCustomFrameRate = bSavedUseCustomFrameRate;
+						AnimSeqExportOption->CustomFrameRate = SavedCustomFrameRate;
 						//save the anim sequence to disk to make sure they are in sync
 						UPackage* const Package = AnimSequence->GetOutermost();
 						FString const PackageName = Package->GetName();
@@ -1771,7 +1776,9 @@ bool FSkeletalAnimationTrackEditor::CreateAnimationSequence(const TArray<UObject
 						for (FLevelSequenceAnimSequenceLinkItem& LevelAnimLinkItem : LevelAnimLink->AnimSequenceLinks)
 						{
 							if (LevelAnimLinkItem.IsEqual(Binding, AnimSeqExportOption->bUseCustomTimeRange,
-								AnimSeqExportOption->CustomStartFrame, AnimSeqExportOption->CustomEndFrame, AnimSeqExportOption->CustomDisplayRate))
+								AnimSeqExportOption->CustomStartFrame, AnimSeqExportOption->CustomEndFrame, AnimSeqExportOption->CustomDisplayRate,
+								AnimSeqExportOption->bUseCustomFrameRate, AnimSeqExportOption->CustomFrameRate
+								))
 							{
 								bAddItem = false;
 								UAnimSequence* OtherAnimSequence = LevelAnimLinkItem.ResolveAnimSequence();
@@ -1804,6 +1811,8 @@ bool FSkeletalAnimationTrackEditor::CreateAnimationSequence(const TArray<UObject
 								LevelAnimLinkItem.CustomStartFrame = AnimSeqExportOption->CustomStartFrame;
 								LevelAnimLinkItem.CustomEndFrame = AnimSeqExportOption->CustomEndFrame;
 								LevelAnimLinkItem.CustomDisplayRate = AnimSeqExportOption->CustomDisplayRate;
+								LevelAnimLinkItem.bUseCustomFrameRate = AnimSeqExportOption->bUseCustomFrameRate;
+								LevelAnimLinkItem.CustomFrameRate = AnimSeqExportOption->CustomFrameRate;
 								break;
 							}
 						}
@@ -1834,6 +1843,8 @@ bool FSkeletalAnimationTrackEditor::CreateAnimationSequence(const TArray<UObject
 						LevelAnimLinkItem.CustomStartFrame = AnimSeqExportOption->CustomStartFrame;
 						LevelAnimLinkItem.CustomEndFrame = AnimSeqExportOption->CustomEndFrame;
 						LevelAnimLinkItem.CustomDisplayRate = AnimSeqExportOption->CustomDisplayRate;
+						LevelAnimLinkItem.bUseCustomFrameRate = AnimSeqExportOption->bUseCustomFrameRate;
+						LevelAnimLinkItem.CustomFrameRate = AnimSeqExportOption->CustomFrameRate;
 
 						LevelAnimLink->AnimSequenceLinks.Add(LevelAnimLinkItem);
 						AssetUserDataInterface->AddAssetUserData(LevelAnimLink);
diff --git a/Engine/Source/Editor/UnrealEd/Classes/Exporters/AnimSeqExportOption.h b/Engine/Source/Editor/UnrealEd/Classes/Exporters/AnimSeqExportOption.h
--- a/Engine/Source/Editor/UnrealEd/Classes/Exporters/AnimSeqExportOption.h
+++ b/Engine/Source/Editor/UnrealEd/Classes/Exporters/AnimSeqExportOption.h
@@ -94,10 +94,19 @@ public:
 	UPROPERTY(EditAnywhere, BlueprintReadWrite, AdvancedDisplay, meta = (EditCondition = bUseCustomTimeRange),  Category = Export)
 	FFrameNumber CustomEndFrame = 120;
 
-	/** Custom display rate, should be set from the movie scene/sequencer display rate */
+	/** Custom display rate for use when specifying custom start and end frame, should be set from the movie scene/sequencer display rate */
 	UPROPERTY(BlueprintReadWrite, AdvancedDisplay, meta = (EditCondition = bUseCustomTimeRange), Category = Export)
 	FFrameRate CustomDisplayRate = FFrameRate(30,1);
 
+	/** Whether or not to use custom frame rate when recording the anim sequence, if false will use Sequencers display rate */
+	UPROPERTY(EditAnywhere, BlueprintReadWrite, AdvancedDisplay, Category = Export)
+	bool bUseCustomFrameRate = false;
+
+	/** Custom frame rate that the anim sequence will be recorded at*/
+	UPROPERTY(EditAnywhere, BlueprintReadWrite, meta = (EditCondition = bUseCustomFrameRate), AdvancedDisplay, Category = Export)
+	FFrameRate CustomFrameRate = FFrameRate(30, 1);
+
+
 	void ResetToDefault()
 	{
 		bExportTransforms = true;
@@ -118,5 +127,7 @@ public:
 		CustomStartFrame = 0;
 		CustomEndFrame = 120;
 		CustomDisplayRate = FFrameRate(30, 1);
+		bUseCustomFrameRate = false;
+		CustomFrameRate = FFrameRate(30, 1);
 	}
 };
diff --git a/Engine/Source/Editor/UnrealEd/Private/Fbx/FbxMainExport.cpp b/Engine/Source/Editor/UnrealEd/Private/Fbx/FbxMainExport.cpp
--- a/Engine/Source/Editor/UnrealEd/Private/Fbx/FbxMainExport.cpp
+++ b/Engine/Source/Editor/UnrealEd/Private/Fbx/FbxMainExport.cpp
@@ -1862,7 +1862,7 @@ TRange<FFrameNumber> FLevelSequenceAnimTrackAdapter::GetSequenceRange() const
 
 	const FMovieSceneSubSequenceData* SubSequenceData = Hierarchy.FindSubData(Template);
 	TRange<FFrameNumber> Range = MovieScene->GetPlaybackRange();
-	FFrameTime StartTime = FFrameRate::TransformTime(UE::MovieScene::DiscreteInclusiveLower(MovieScene->GetPlaybackRange()).Value, MovieScene->GetTickResolution(), MovieScene->GetDisplayRate());
+	FFrameTime StartTime = FFrameRate::TransformTime(UE::MovieScene::DiscreteInclusiveLower(MovieScene->GetPlaybackRange()).Value, MovieScene->GetTickResolution(), GetDisplayRate());
 
 	if (SubSequenceData && bForceUseOfMovieScenePlaybackRange == false)
 	{
@@ -1875,7 +1875,7 @@ int32 FLevelSequenceAnimTrackAdapter::GetLocalStartFrame() const
 	TRange<FFrameNumber> Range = GetSequenceRange();
 
 	FFrameRate TickResolution = MovieScene->GetTickResolution();
-	FFrameRate DisplayRate = MovieScene->GetDisplayRate();
+	FFrameRate DisplayRate = GetDisplayRate();
 	return FFrameRate::TransformTime(FFrameTime(UE::MovieScene::DiscreteInclusiveLower(Range)), TickResolution, DisplayRate).RoundToFrame().Value;
 }
 
@@ -1883,7 +1883,7 @@ int32 FLevelSequenceAnimTrackAdapter::GetStartFrame() const
 {
 	TRange<FFrameNumber> Range = GetSequenceRange();
 	FFrameRate TickResolution = MovieScene->GetTickResolution();
-	FFrameRate DisplayRate = MovieScene->GetDisplayRate();
+	FFrameRate DisplayRate = GetDisplayRate();
 	FFrameNumber StartFrame = UE::MovieScene::DiscreteInclusiveLower(Range);
 
 	FMovieSceneInverseSequenceTransform Inverse = RootToLocalTransform.Inverse();
@@ -1896,7 +1896,7 @@ int32 FLevelSequenceAnimTrackAdapter::GetLength() const
 {
 	TRange<FFrameNumber> Range = GetSequenceRange();
 	FFrameRate TickResolution = MovieScene->GetTickResolution();
-	FFrameRate DisplayRate = MovieScene->GetDisplayRate();
+	FFrameRate DisplayRate = GetDisplayRate();
 
 	return FFrameRate::TransformTime(FFrameTime(UE::MovieScene::DiscreteSize(Range)), TickResolution, DisplayRate).RoundToFrame().Value;
 }
@@ -1904,7 +1904,7 @@ int32 FLevelSequenceAnimTrackAdapter::GetLength() const
 void FLevelSequenceAnimTrackAdapter::UpdateAnimation( int32 LocalFrame )
 {	
 	FFrameRate TickResolution = MovieScene->GetTickResolution();
-	FFrameRate DisplayRate = MovieScene->GetDisplayRate();
+	FFrameRate DisplayRate = GetDisplayRate();
 
 	FFrameTime LocalTime = FFrameRate::TransformTime(FFrameTime(LocalFrame), DisplayRate, TickResolution);
 	TOptional<FFrameTime> GlobalTime = RootToLocalTransform.Inverse().TryTransformTime(LocalTime);
@@ -1916,9 +1916,19 @@ void FLevelSequenceAnimTrackAdapter::UpdateAnimation( int32 LocalFrame )
 	}
 }
 
+void FLevelSequenceAnimTrackAdapter::SetFrameRate(FFrameRate InFrameRate)
+{
+	OptionalFrameRate = InFrameRate;
+}
+
 double FLevelSequenceAnimTrackAdapter::GetFrameRate() const
 {
-	return MovieScene->GetDisplayRate().AsDecimal();
+	return OptionalFrameRate.IsSet() ? OptionalFrameRate.GetValue().AsDecimal() : MovieScene->GetDisplayRate().AsDecimal();
+}
+
+FFrameRate FLevelSequenceAnimTrackAdapter::GetDisplayRate() const
+{
+	return OptionalFrameRate.IsSet() ? OptionalFrameRate.GetValue() : MovieScene->GetDisplayRate();
 }
 		
 UAnimSequence* FLevelSequenceAnimTrackAdapter::GetAnimSequence(int32 LocalFrame) const
@@ -1929,7 +1939,7 @@ UAnimSequence* FLevelSequenceAnimTrackAdapter::GetAnimSequence(int32 LocalFrame)
 	}
 
 	FFrameRate TickResolution = MovieScene->GetTickResolution();
-	FFrameRate DisplayRate = MovieScene->GetDisplayRate();
+	FFrameRate DisplayRate = GetDisplayRate();
 
 	FFrameTime LocalTime = FFrameRate::TransformTime(FFrameTime(LocalFrame), DisplayRate, TickResolution);
 
@@ -1952,7 +1962,7 @@ float FLevelSequenceAnimTrackAdapter::GetAnimTime(int32 LocalFrame) const
 	}
 
 	FFrameRate TickResolution = MovieScene->GetTickResolution();
-	FFrameRate DisplayRate = MovieScene->GetDisplayRate();
+	FFrameRate DisplayRate = GetDisplayRate();
 
 	FFrameTime LocalTime = FFrameRate::TransformTime(FFrameTime(LocalFrame), DisplayRate, TickResolution);
 
diff --git a/Engine/Source/Editor/UnrealEd/Private/FbxExporter.h b/Engine/Source/Editor/UnrealEd/Private/FbxExporter.h
--- a/Engine/Source/Editor/UnrealEd/Private/FbxExporter.h
+++ b/Engine/Source/Editor/UnrealEd/Private/FbxExporter.h
@@ -103,8 +103,12 @@ namespace UnFbx
 		UNREALED_API virtual float GetAnimTime(int32 LocalFrame) const override;
 		UNREALED_API virtual void SetRange(const FFrameNumber& StartFrame, const FFrameNumber& EndFrame) override;
 
+		//set the frame rate to use, either from movie scene(default) or custom one
+		UNREALED_API void SetFrameRate(FFrameRate InFrameRate);
 	private:
 		TRange<FFrameNumber> GetSequenceRange() const;
+		//will be either sequencer display rate or custom frame rate
+		FFrameRate GetDisplayRate() const; 
 
 	private:
 		IMovieScenePlayer* MovieScenePlayer;
@@ -115,6 +119,7 @@ namespace UnFbx
 		FMovieSceneSequenceTransform RootToLocalTransform;
 		UMovieSceneSkeletalAnimationTrack* AnimTrack;
 		TOptional < TPair<FFrameNumber, FFrameNumber>> OptionalRange;
+		TOptional<FFrameRate> OptionalFrameRate;
 	};
 /**
  * Main FBX Exporter class.
diff --git a/Engine/Source/Runtime/LevelSequence/Public/LevelSequenceAnimSequenceLink.h b/Engine/Source/Runtime/LevelSequence/Public/LevelSequenceAnimSequenceLink.h
--- a/Engine/Source/Runtime/LevelSequence/Public/LevelSequenceAnimSequenceLink.h
+++ b/Engine/Source/Runtime/LevelSequence/Public/LevelSequenceAnimSequenceLink.h
@@ -71,13 +71,21 @@ public:
 	UPROPERTY(EditAnywhere, BlueprintReadWrite, meta = (EditCondition = bUseCustomTimeRange), Category = Property)
 	FFrameRate CustomDisplayRate = FFrameRate(30, 1);
 
+	/** Whether or not to use custom frame rate or Sequencer display rate */
+	UPROPERTY(BlueprintReadWrite, Category = Property)
+	bool bUseCustomFrameRate = false;
+
+	/** Custom frame rate that the anim sequence may have been recorded at */
+	UPROPERTY(BlueprintReadWrite, meta = (EditCondition = bUseCustomFrameRate), Category = Property)
+	FFrameRate CustomFrameRate = FFrameRate(30, 1);
 
 	LEVELSEQUENCE_API void SetAnimSequence(UAnimSequence* InAnimSequence);
 	LEVELSEQUENCE_API UAnimSequence* ResolveAnimSequence();
 
 	LEVELSEQUENCE_API bool IsEqual(FGuid InSkelTrackGuid, bool bInUseCustomTimeRange = false,
-		FFrameNumber InCustomStartFrame = 0, FFrameNumber InCustomEndFrame = 120,
-		FFrameRate InCustomDisplayRate = FFrameRate(30,1))
+		FFrameNumber InCustomStartFrame = 0, FFrameNumber InCustomEndFrame = 120, FFrameRate InCustomDisplayRate = FFrameRate(30,1),
+		bool bInUseCustomFrameRate = false, FFrameRate InCustomFrameRate = FFrameRate(30, 1)
+		)
 	{
 		if (InSkelTrackGuid == SkelTrackGuid)
 		{
@@ -89,6 +97,13 @@ public:
 					return true;
 				}
 			}
+			if (bUseCustomFrameRate == bInUseCustomFrameRate)
+			{
+				if (bUseCustomFrameRate == false || (InCustomFrameRate == CustomFrameRate))
+				{
+					return true;
+				}
+			}
 		}
 		return false;
 	}
