Workflow Centric Application - Add call to allow derived FWorkflowCentricApplications to remove an application mode

diff --git a/Engine/Source/Editor/UnrealEd/Private/WorkflowCentricApplication.cpp b/Engine/Source/Editor/UnrealEd/Private/WorkflowCentricApplication.cpp
--- a/Engine/Source/Editor/UnrealEd/Private/WorkflowCentricApplication.cpp
+++ b/Engine/Source/Editor/UnrealEd/Private/WorkflowCentricApplication.cpp
@@ -9,6 +9,8 @@
 
 #define LOCTEXT_NAMESPACE "WorkflowCentricApplication"
 
+DEFINE_LOG_CATEGORY_STATIC(LogWorkflowCentricApplication, Log, All);
+
 /////////////////////////////////////////////////////
 // FWorkflowCentricApplication
 
@@ -154,6 +156,19 @@ void FWorkflowCentricApplication::AddApplicationMode(FName ModeName, TSharedRef<
 	ApplicationModeList.Add(ModeName, Mode);
 }
 
+bool FWorkflowCentricApplication::RemoveApplicationMode(FName ModeName)
+{
+	if (GetCurrentMode() != ModeName)
+	{
+		return ApplicationModeList.Remove(ModeName) > 0;
+	}
+	else
+	{
+		UE_LOG(LogWorkflowCentricApplication, Warning, TEXT("Cannot remove application mode \"%s\", because it's the current mode."), *ModeName.ToString());
+		return false;
+	}
+}
+
 void FWorkflowCentricApplication::RemoveAllApplicationModes()
 {
 	ApplicationModeList.Reset();
diff --git a/Engine/Source/Editor/UnrealEd/Public/WorkflowOrientedApp/WorkflowCentricApplication.h b/Engine/Source/Editor/UnrealEd/Public/WorkflowOrientedApp/WorkflowCentricApplication.h
--- a/Engine/Source/Editor/UnrealEd/Public/WorkflowOrientedApp/WorkflowCentricApplication.h
+++ b/Engine/Source/Editor/UnrealEd/Public/WorkflowOrientedApp/WorkflowCentricApplication.h
@@ -57,6 +57,11 @@ protected:
 
 	UNREALED_API virtual void AddApplicationMode(FName ModeName, TSharedRef<FApplicationMode> Mode);
 
+	// Remove an application mode.
+	// @return true if the mode was successfully removed.
+	// @note Has no effect if you call this function with the current mode.
+	UNREALED_API virtual bool RemoveApplicationMode(FName ModeName);
+
 	// Remove all application modes that were added through AddApplicationMode.
 	UNREALED_API virtual void RemoveAllApplicationModes();
 
