Implement the new viewport toolbar for the UV Editor.
- Refactor snapping submenu handling to be based on the viewport context object.
- Move the "Preserve Non Uniform Scale" setting to the transform menu.

diff --git a/Engine/Plugins/Editor/UVEditor/Source/UVEditor/Private/SUVEditor2DViewport.cpp b/Engine/Plugins/Editor/UVEditor/Source/UVEditor/Private/SUVEditor2DViewport.cpp
--- a/Engine/Plugins/Editor/UVEditor/Source/UVEditor/Private/SUVEditor2DViewport.cpp
+++ b/Engine/Plugins/Editor/UVEditor/Source/UVEditor/Private/SUVEditor2DViewport.cpp
@@ -2,10 +2,15 @@
 
 #include "SUVEditor2DViewport.h"
 
+#include "EditorViewportCommands.h"
 #include "SUVEditor2DViewportToolBar.h"
 #include "UVEditor2DViewportClient.h"
 #include "UVEditorCommands.h"
+#include "UVEditorStyle.h"
+#include "ToolMenus.h"
+#include "UVEditor2DViewportContext.h"
 #include "ContextObjects/UVToolContextObjects.h" // UUVToolViewportButtonsAPI::ESelectionMode
+#include "ViewportToolbar/UnrealEdViewportToolbar.h"
 
 #define LOCTEXT_NAMESPACE "SUVEditor2DViewport"
 
@@ -73,6 +78,63 @@ void SUVEditor2DViewport::BindCommands()
 		FIsActionChecked::CreateLambda([this]() { 
 			return StaticCastSharedPtr<FUVEditor2DViewportClient>(Client)->GetSelectionMode() == UUVToolViewportButtonsAPI::ESelectionMode::Mesh;
 		}));
+		
+	CommandList->MapAction(
+		FEditorViewportCommands::Get().LocationGridSnap,
+		FExecuteAction::CreateLambda([WeakClient = Client.ToWeakPtr()]
+		{
+			if (TSharedPtr<FUVEditor2DViewportClient> Client = StaticCastSharedPtr<FUVEditor2DViewportClient>(WeakClient.Pin()))
+			{
+				Client->SetLocationGridSnapEnabled(!Client->GetLocationGridSnapEnabled());
+			}
+		}),
+		FCanExecuteAction(),
+		FIsActionChecked::CreateLambda([WeakClient = Client.ToWeakPtr()]
+		{
+			if (TSharedPtr<FUVEditor2DViewportClient> Client = StaticCastSharedPtr<FUVEditor2DViewportClient>(WeakClient.Pin()))
+			{
+				return Client->GetLocationGridSnapEnabled();
+			}
+			return false;
+		}));
+	
+	CommandList->MapAction(
+		FEditorViewportCommands::Get().RotationGridSnap,
+		FExecuteAction::CreateLambda([WeakClient = Client.ToWeakPtr()]
+		{
+			if (TSharedPtr<FUVEditor2DViewportClient> Client = StaticCastSharedPtr<FUVEditor2DViewportClient>(WeakClient.Pin()))
+			{
+				Client->SetRotationGridSnapEnabled(!Client->GetRotationGridSnapEnabled());
+			}
+		}),
+		FCanExecuteAction(),
+		FIsActionChecked::CreateLambda([WeakClient = Client.ToWeakPtr()]
+		{
+			if (TSharedPtr<FUVEditor2DViewportClient> Client = StaticCastSharedPtr<FUVEditor2DViewportClient>(WeakClient.Pin()))
+			{
+				return Client->GetRotationGridSnapEnabled();
+			}
+			return false;
+		}));
+		
+	CommandList->MapAction(
+		FEditorViewportCommands::Get().ScaleGridSnap,
+		FExecuteAction::CreateLambda([WeakClient = Client.ToWeakPtr()]
+		{
+			if (TSharedPtr<FUVEditor2DViewportClient> Client = StaticCastSharedPtr<FUVEditor2DViewportClient>(WeakClient.Pin()))
+			{
+				Client->SetScaleGridSnapEnabled(!Client->GetScaleGridSnapEnabled());
+			}
+		}),
+		FCanExecuteAction(),
+		FIsActionChecked::CreateLambda([WeakClient = Client.ToWeakPtr()]
+		{
+			if (TSharedPtr<FUVEditor2DViewportClient> Client = StaticCastSharedPtr<FUVEditor2DViewportClient>(WeakClient.Pin()))
+			{
+				return Client->GetScaleGridSnapEnabled();
+			}
+			return false;
+		}));
 }
 
 void SUVEditor2DViewport::AddOverlayWidget(TSharedRef<SWidget> OverlaidWidget, int32 ZOrder)
@@ -95,6 +157,112 @@ TSharedPtr<SWidget> SUVEditor2DViewport::MakeViewportToolbar()
 		.Viewport2DClient(StaticCastSharedPtr<FUVEditor2DViewportClient>(Client));
 }
 
+TSharedPtr<SWidget> SUVEditor2DViewport::BuildViewportToolbar()
+{
+	const FName ToolbarName = "UVEditor2DToolbar";
+	
+	if (!UToolMenus::Get()->IsMenuRegistered(ToolbarName))
+	{
+		UToolMenu* Toolbar = UToolMenus::Get()->RegisterMenu(ToolbarName, NAME_None, EMultiBoxType::SlimHorizontalToolBar);
+		Toolbar->StyleName = "ViewportToolbar";
+		
+		{
+			FToolMenuSection& LeftSection = Toolbar->FindOrAddSection("Left");
+			
+			{
+				FToolMenuEntry& TransformEntry = LeftSection.AddSubMenu(
+					"Transform",
+					LOCTEXT("TransformsSubmenuLabel", "Transform"),
+					LOCTEXT("TransformsSubmenuTooltip", "Viewport-related transforms tools"),
+					FNewToolMenuDelegate::CreateLambda([](UToolMenu* ToolMenu)
+					{
+						FToolMenuSection& TransformToolsSection = ToolMenu->FindOrAddSection(
+							"TransformTools",
+							LOCTEXT("TransformToolsLabel", "Transform Tools")
+						);
+
+						FToolMenuEntry& SelectMode = TransformToolsSection.AddMenuEntry(FEditorViewportCommands::Get().SelectMode);
+						SelectMode.SetShowInToolbarTopLevel(true);
+						SelectMode.ToolBarData.StyleNameOverride = "ViewportToolbar.TransformTools";
+					
+						FToolMenuEntry& TranslateMode = TransformToolsSection.AddMenuEntry(FEditorViewportCommands::Get().TranslateMode);
+						TranslateMode.SetShowInToolbarTopLevel(true);
+						TranslateMode.ToolBarData.StyleNameOverride = "ViewportToolbar.TransformTools";
+					})
+				);
+				TransformEntry.Icon = FSlateIcon(FAppStyle::GetAppStyleSetName(), "LevelEditor.SelectMode");
+				TransformEntry.ToolBarData.LabelOverride = FText();
+				TransformEntry.ToolBarData.ResizeParams.ClippingPriority = 1000;
+			}
+			
+			LeftSection.AddEntry(UE::UnrealEd::CreateSnappingSubmenu());
+			
+			{
+				FToolMenuEntry& ElementSelectionEntry = LeftSection.AddSubMenu(
+					"ElementSelection",
+					LOCTEXT("MeshElementSelectionSubmenuLabel", "Mesh Element Selection"),
+					LOCTEXT("MeshElementSelectionSubmenuTooltip", "Mesh Element Selection settings in the viewport"),
+					FNewToolMenuDelegate::CreateLambda([](UToolMenu* ToolMenu)
+					{
+						FToolMenuSection& ElementSection = ToolMenu->FindOrAddSection(
+							"Element Selection",
+							LOCTEXT("ElementSelectionLabel", "Element Selection")
+						);
+						
+						FToolMenuEntryToolBarData ToolBarData;
+						ToolBarData.BlockGroupName = "ElementSelection";
+						ToolBarData.LabelOverride = FText::GetEmpty();
+						
+						FToolMenuEntry& VertexEntry = ElementSection.AddMenuEntry(FUVEditorCommands::Get().VertexSelection);
+						VertexEntry.ToolBarData = ToolBarData;
+						VertexEntry.TutorialHighlightName = TEXT("VertexSelection");
+						VertexEntry.SetShowInToolbarTopLevel(true);
+						
+						FToolMenuEntry& EdgeEntry = ElementSection.AddMenuEntry(FUVEditorCommands::Get().EdgeSelection);
+						EdgeEntry.ToolBarData = ToolBarData;
+						EdgeEntry.TutorialHighlightName = TEXT("EdgeSelection");
+						EdgeEntry.SetShowInToolbarTopLevel(true);
+						
+						FToolMenuEntry& TriangleEntry = ElementSection.AddMenuEntry(FUVEditorCommands::Get().TriangleSelection);
+						TriangleEntry.ToolBarData = ToolBarData;
+						TriangleEntry.TutorialHighlightName = TEXT("TriangleSelection");
+						TriangleEntry.SetShowInToolbarTopLevel(true);
+						
+						FToolMenuEntry& IslandEntry = ElementSection.AddMenuEntry(FUVEditorCommands::Get().IslandSelection);
+						IslandEntry.ToolBarData = ToolBarData;
+						IslandEntry.TutorialHighlightName = TEXT("IslandSelection");
+						IslandEntry.SetShowInToolbarTopLevel(true);
+						
+						FToolMenuEntry& MeshEntry = ElementSection.AddMenuEntry(FUVEditorCommands::Get().FullMeshSelection);
+						MeshEntry.ToolBarData = ToolBarData;
+						MeshEntry.TutorialHighlightName = TEXT("FullMeshSelection");
+						MeshEntry.SetShowInToolbarTopLevel(true);
+					})
+				);
+				
+				ElementSelectionEntry.ToolBarData.ResizeParams.ClippingPriority = 950;
+				ElementSelectionEntry.ToolBarData.LabelOverride = FText();
+				ElementSelectionEntry.Icon = FSlateIcon(FUVEditorStyle::StyleName, "UVEditor.ElementSelection");
+			}
+		}
+		
+		{
+			FToolMenuSection& RightSection = Toolbar->FindOrAddSection("Right");
+			RightSection.Alignment = EToolMenuSectionAlign::Last;
+		}
+	}
+	
+	FToolMenuContext Context;
+	{
+		UUVEditor2DViewportContext* ContextObject = NewObject<UUVEditor2DViewportContext>();
+		ContextObject->Viewport = SharedThis(this);
+		Context.AddObject(ContextObject);
+		Context.AppendCommandList(GetCommandList());
+	}
+	
+	return UToolMenus::Get()->GenerateWidget(ToolbarName, Context);
+}
+
 bool SUVEditor2DViewport::IsWidgetModeActive(UE::Widget::EWidgetMode Mode) const
 {
 	return static_cast<FUVEditor2DViewportClient*>(Client.Get())->AreWidgetButtonsEnabled() 
diff --git a/Engine/Plugins/Editor/UVEditor/Source/UVEditor/Private/SUVEditor3DViewport.cpp b/Engine/Plugins/Editor/UVEditor/Source/UVEditor/Private/SUVEditor3DViewport.cpp
--- a/Engine/Plugins/Editor/UVEditor/Source/UVEditor/Private/SUVEditor3DViewport.cpp
+++ b/Engine/Plugins/Editor/UVEditor/Source/UVEditor/Private/SUVEditor3DViewport.cpp
@@ -6,6 +6,10 @@
 #include "SUVEditor3DViewportToolBar.h"
 #include "UVEditorCommands.h"
 #include "UVEditor3DViewportClient.h"
+#include "ViewportToolbar/UnrealEdViewportToolbar.h"
+#include "ViewportToolbar/UnrealEdViewportToolbarContext.h"
+#include "ToolMenus.h"
+#include "UVEditorStyle.h"
 
 #define LOCTEXT_NAMESPACE "SUVEditor3DViewport"
 
@@ -46,4 +50,60 @@ TSharedPtr<SWidget> SUVEditor3DViewport::MakeViewportToolbar()
 		.CommandList(CommandList);
 }
 
+TSharedPtr<SWidget> SUVEditor3DViewport::BuildViewportToolbar()
+{
+	const FName ToolbarName = "UVEditor3DViewportToolbar";
+	
+	if (!UToolMenus::Get()->IsMenuRegistered(ToolbarName))
+	{
+		UToolMenu* Toolbar = UToolMenus::Get()->RegisterMenu(ToolbarName, NAME_None, EMultiBoxType::SlimHorizontalToolBar);
+		Toolbar->StyleName = "ViewportToolbar";
+		
+		Toolbar->FindOrAddSection("Left");
+		
+		FToolMenuSection& RightSection = Toolbar->FindOrAddSection("Right");
+		RightSection.Alignment = EToolMenuSectionAlign::Last;
+		
+		{
+			RightSection.AddEntry(UE::UnrealEd::CreateCameraSubmenu());
+			
+			UToolMenu* CameraMenu = UToolMenus::Get()->ExtendMenu(UToolMenus::JoinMenuPaths(ToolbarName, "Camera"));
+			FToolMenuSection& MovementSection = CameraMenu->FindOrAddSection("Movement");
+			
+			FToolMenuEntry& OrbitEntry = MovementSection.AddMenuEntry(FUVEditorCommands::Get().EnableOrbitCamera);
+			OrbitEntry.SetShowInToolbarTopLevel(true);
+			OrbitEntry.Icon = FSlateIcon(FUVEditorStyle::Get().GetStyleSetName(), "UVEditor.OrbitCamera");
+			OrbitEntry.ToolBarData.LabelOverride = FText();
+			OrbitEntry.TutorialHighlightName = "OrbitCamera";
+			
+			FToolMenuEntry& FlyEntry = MovementSection.AddMenuEntry(FUVEditorCommands::Get().EnableFlyCamera);
+			FlyEntry.SetShowInToolbarTopLevel(true);
+			FlyEntry.Icon = FSlateIcon(FUVEditorStyle::Get().GetStyleSetName(), "UVEditor.FlyCamera");
+			FlyEntry.ToolBarData.LabelOverride = FText();
+			OrbitEntry.TutorialHighlightName = "FlyCamera";
+			
+			MovementSection.AddSeparator("FocusSeparator");
+			
+			FToolMenuEntry& FocusEntry = MovementSection.AddMenuEntry(FUVEditorCommands::Get().SetFocusCamera);
+			FocusEntry.SetShowInToolbarTopLevel(true);
+			FocusEntry.Icon = FSlateIcon(FUVEditorStyle::Get().GetStyleSetName(), "UVEditor.FocusCamera");
+			FocusEntry.ToolBarData.LabelOverride = FText();
+			FocusEntry.TutorialHighlightName = "FocusCamera";
+		}
+		
+		RightSection.AddEntry(UE::UnrealEd::CreateViewModesSubmenu());
+	}
+	
+	FToolMenuContext Context;
+	
+	{
+		UUnrealEdViewportToolbarContext* ContextObject = UE::UnrealEd::CreateViewportToolbarDefaultContext(SharedThis(this));
+		Context.AddObject(ContextObject);
+		Context.AppendCommandList(GetCommandList());
+	}
+	
+	return UToolMenus::Get()->GenerateWidget(ToolbarName, Context);
+}
+
+
 #undef LOCTEXT_NAMESPACE
diff --git a/Engine/Plugins/Editor/UVEditor/Source/UVEditor/Private/UVEditor2DViewportContext.cpp b/Engine/Plugins/Editor/UVEditor/Source/UVEditor/Private/UVEditor2DViewportContext.cpp
new file mode 100644
--- /dev/null
+++ b/Engine/Plugins/Editor/UVEditor/Source/UVEditor/Private/UVEditor2DViewportContext.cpp
@@ -0,0 +1,129 @@
+﻿// Copyright Epic Games, Inc. All Rights Reserved.
+
+#include "UVEditor2DViewportContext.h"
+
+#include "SUVEditor2DViewport.h"
+#include "UVEditorUXSettings.h"
+#include "Settings/LevelEditorViewportSettings.h"
+
+#define LOCTEXT_NAMESPACE "UVEditor2DViewportToolbar"
+
+UUVEditor2DViewportContext::UUVEditor2DViewportContext()
+{
+	bShowSurfaceSnap = false;
+
+	GridSnapSizes.Reserve(FUVEditorUXSettings::MaxLocationSnapValue());
+	for (int32 Index = 0; Index < FUVEditorUXSettings::MaxLocationSnapValue(); ++Index)
+	{
+		GridSnapSizes.Add(FUVEditorUXSettings::LocationSnapValue(Index));
+	}	
+}
+
+TSharedPtr<FUVEditor2DViewportClient> UUVEditor2DViewportContext::GetViewportClient() const
+{
+	if (TSharedPtr<SEditorViewport> EditorViewport = Viewport.Pin())
+	{
+		return StaticCastSharedPtr<FUVEditor2DViewportClient>(EditorViewport->GetViewportClient());
+	}
+	return nullptr;
+}
+
+FText UUVEditor2DViewportContext::GetGridSnapLabel() const
+{
+	if (TSharedPtr<FUVEditor2DViewportClient> ViewportClient = GetViewportClient())
+	{
+		return FText::AsNumber(ViewportClient->GetLocationGridSnapValue());
+	}
+	return FText();
+}
+
+TArray<float> UUVEditor2DViewportContext::GetGridSnapSizes() const
+{
+	return GridSnapSizes;
+}
+
+bool UUVEditor2DViewportContext::IsGridSnapSizeActive(int32 GridSizeIndex) const
+{
+	if (TSharedPtr<FUVEditor2DViewportClient> ViewportClient = GetViewportClient())
+	{
+		const float CurrGridSize = GridSnapSizes[GridSizeIndex];
+		return FMath::IsNearlyEqual(ViewportClient->GetLocationGridSnapValue(), CurrGridSize);
+	}
+	return false;
+}
+
+void UUVEditor2DViewportContext::SetGridSnapSize(int32 GridSizeIndex)
+{
+	if (TSharedPtr<FUVEditor2DViewportClient> ViewportClient = GetViewportClient())
+	{
+		const float CurrGridSize = GridSnapSizes[GridSizeIndex];
+		ViewportClient->SetLocationGridSnapValue(CurrGridSize);
+	}
+}
+
+FText UUVEditor2DViewportContext::GetRotationSnapLabel() const
+{
+	if (TSharedPtr<FUVEditor2DViewportClient> ViewportClient = GetViewportClient())
+	{
+		return FText::Format(LOCTEXT("GridRotation - Number - DegreeSymbol", "{0}\u00b0"), FText::AsNumber(ViewportClient->GetRotationGridSnapValue()));
+	}
+	return FText();
+}
+
+bool UUVEditor2DViewportContext::IsRotationSnapActive(int32 RotationIndex, ERotationGridMode RotationMode) const
+{
+	if (TSharedPtr<FUVEditor2DViewportClient> ViewportClient = GetViewportClient())
+	{
+		const ULevelEditorViewportSettings* ViewportSettings = GetDefault<ULevelEditorViewportSettings>();
+		const TArray<float>& Sizes = RotationMode == GridMode_Common ? ViewportSettings->CommonRotGridSizes : ViewportSettings->DivisionsOf360RotGridSizes;
+		const float CurrGridAngle = Sizes[RotationIndex];
+		return FMath::IsNearlyEqual(ViewportClient->GetRotationGridSnapValue(), CurrGridAngle);
+	}
+	return false;
+}
+
+void UUVEditor2DViewportContext::SetRotationSnapSize(int32 RotationIndex, ERotationGridMode RotationMode)
+{
+	if (TSharedPtr<FUVEditor2DViewportClient> ViewportClient = GetViewportClient())
+	{
+		const ULevelEditorViewportSettings* ViewportSettings = GetDefault<ULevelEditorViewportSettings>();
+		const TArray<float>& Sizes = RotationMode == GridMode_Common ? ViewportSettings->CommonRotGridSizes : ViewportSettings->DivisionsOf360RotGridSizes;
+		return ViewportClient->SetRotationGridSnapValue(Sizes[RotationIndex]);
+	}
+}
+
+FText UUVEditor2DViewportContext::GetScaleSnapLabel() const
+{
+	if (TSharedPtr<FUVEditor2DViewportClient> ViewportClient = GetViewportClient())
+	{
+		FNumberFormattingOptions NumberFormattingOptions;
+		NumberFormattingOptions.MaximumFractionalDigits = 5;
+
+		const float CurGridAmount = ViewportClient->GetScaleGridSnapValue();
+		return FText::AsNumber(CurGridAmount, &NumberFormattingOptions);
+	}
+	return FText();
+}
+
+bool UUVEditor2DViewportContext::IsScaleSnapActive(int32 ScaleIndex) const
+{
+	if (TSharedPtr<FUVEditor2DViewportClient> ViewportClient = GetViewportClient())
+	{
+		const ULevelEditorViewportSettings* ViewportSettings = GetDefault<ULevelEditorViewportSettings>();
+		const float CurrGridAngle = ViewportSettings->ScalingGridSizes[ScaleIndex];
+		return FMath::IsNearlyEqual(ViewportClient->GetRotationGridSnapValue(), CurrGridAngle);
+	}
+	return false;
+}
+
+void UUVEditor2DViewportContext::SetScaleSnapSize(int32 ScaleIndex)
+{
+	if (TSharedPtr<FUVEditor2DViewportClient> ViewportClient = GetViewportClient())
+	{
+		const ULevelEditorViewportSettings* ViewportSettings = GetDefault<ULevelEditorViewportSettings>();
+		ViewportClient->SetScaleGridSnapValue(ViewportSettings->ScalingGridSizes[ScaleIndex]);
+	}
+}
+
+
+#undef LOCTEXT_NAMESPACE
diff --git a/Engine/Plugins/Editor/UVEditor/Source/UVEditor/Private/UVEditorCommands.cpp b/Engine/Plugins/Editor/UVEditor/Source/UVEditor/Private/UVEditorCommands.cpp
--- a/Engine/Plugins/Editor/UVEditor/Source/UVEditor/Private/UVEditorCommands.cpp
+++ b/Engine/Plugins/Editor/UVEditor/Source/UVEditor/Private/UVEditorCommands.cpp
@@ -59,8 +59,8 @@ void FUVEditorCommands::RegisterCommands()
 	UI_COMMAND(FullMeshSelection, "Mesh Selection", "Select meshes", EUserInterfaceActionType::ToggleButton, FInputChord(EKeys::Five));
 	UI_COMMAND(SelectAll, "Select All", "Select everything based on current selection mode", EUserInterfaceActionType::None, FInputChord(EKeys::A, EModifierKey::Control));
 
-	UI_COMMAND(EnableOrbitCamera, "Orbit", "Enable orbit camera", EUserInterfaceActionType::ToggleButton, FInputChord());
-	UI_COMMAND(EnableFlyCamera, "Fly", "Enable fly camera", EUserInterfaceActionType::ToggleButton, FInputChord());
+	UI_COMMAND(EnableOrbitCamera, "Orbit", "Enable orbit camera", EUserInterfaceActionType::RadioButton, FInputChord());
+	UI_COMMAND(EnableFlyCamera, "Fly", "Enable fly camera", EUserInterfaceActionType::RadioButton, FInputChord());
 	UI_COMMAND(SetFocusCamera, "Focus Camera", "Focus camera around the currently selected UVs", EUserInterfaceActionType::Button, FInputChord(EKeys::F, EModifierKey::Alt));
 
 	UI_COMMAND(ToggleBackground, "Toggle Background", "Toggle background display", EUserInterfaceActionType::ToggleButton,
diff --git a/Engine/Plugins/Editor/UVEditor/Source/UVEditor/Public/SUVEditor2DViewport.h b/Engine/Plugins/Editor/UVEditor/Source/UVEditor/Public/SUVEditor2DViewport.h
--- a/Engine/Plugins/Editor/UVEditor/Source/UVEditor/Public/SUVEditor2DViewport.h
+++ b/Engine/Plugins/Editor/UVEditor/Source/UVEditor/Public/SUVEditor2DViewport.h
@@ -20,5 +20,6 @@ public:
 	// SEditorViewport
 	virtual void BindCommands() override;
 	virtual TSharedPtr<SWidget> MakeViewportToolbar() override;
+	virtual TSharedPtr<SWidget> BuildViewportToolbar() override;
 	bool IsWidgetModeActive(UE::Widget::EWidgetMode Mode) const override;
 };
diff --git a/Engine/Plugins/Editor/UVEditor/Source/UVEditor/Public/SUVEditor3DViewport.h b/Engine/Plugins/Editor/UVEditor/Source/UVEditor/Public/SUVEditor3DViewport.h
--- a/Engine/Plugins/Editor/UVEditor/Source/UVEditor/Public/SUVEditor3DViewport.h
+++ b/Engine/Plugins/Editor/UVEditor/Source/UVEditor/Public/SUVEditor3DViewport.h
@@ -17,5 +17,6 @@ public:
 	// SAssetEditorViewport
 	virtual void BindCommands() override;
 	virtual TSharedPtr<SWidget> MakeViewportToolbar() override;
+	virtual TSharedPtr<SWidget> BuildViewportToolbar() override;
 
 };
diff --git a/Engine/Plugins/Editor/UVEditor/Source/UVEditor/Public/UVEditor2DViewportContext.h b/Engine/Plugins/Editor/UVEditor/Source/UVEditor/Public/UVEditor2DViewportContext.h
new file mode 100644
--- /dev/null
+++ b/Engine/Plugins/Editor/UVEditor/Source/UVEditor/Public/UVEditor2DViewportContext.h
@@ -0,0 +1,37 @@
+﻿// Copyright Epic Games, Inc. All Rights Reserved.
+
+#pragma once
+
+#include "CoreMinimal.h"
+#include "UVEditor2DViewportClient.h"
+#include "ViewportToolbar/UnrealEdViewportToolbarContext.h"
+
+#include "UVEditor2DViewportContext.generated.h"
+
+UCLASS()
+class UUVEditor2DViewportContext : public UUnrealEdViewportToolbarContext
+{
+	GENERATED_BODY()
+	
+public:
+
+	UUVEditor2DViewportContext();
+
+	virtual FText GetGridSnapLabel() const override;
+	virtual TArray<float> GetGridSnapSizes() const override;
+	virtual bool IsGridSnapSizeActive(int32 GridSizeIndex) const override;
+	virtual void SetGridSnapSize(int32 GridSizeIndex) override;
+	
+	virtual FText GetRotationSnapLabel() const override;
+	virtual bool IsRotationSnapActive(int32 RotationIndex, ERotationGridMode RotationMode) const override;
+	virtual void SetRotationSnapSize(int32 RotationIndex, ERotationGridMode RotationMode) override;
+	
+	virtual FText GetScaleSnapLabel() const override;
+	virtual bool IsScaleSnapActive(int32 ScaleIndex) const override;
+	virtual void SetScaleSnapSize(int32 ScaleIndex) override;
+	
+protected:
+	TSharedPtr<FUVEditor2DViewportClient> GetViewportClient() const;
+	
+	TArray<float> GridSnapSizes;
+}; 
diff --git a/Engine/Plugins/Editor/UVEditor/Source/UVEditorTools/Private/UVEditorStyle.cpp b/Engine/Plugins/Editor/UVEditor/Source/UVEditorTools/Private/UVEditorStyle.cpp
--- a/Engine/Plugins/Editor/UVEditor/Source/UVEditorTools/Private/UVEditorStyle.cpp
+++ b/Engine/Plugins/Editor/UVEditor/Source/UVEditorTools/Private/UVEditorStyle.cpp
@@ -53,6 +53,7 @@ FUVEditorStyle::FUVEditorStyle()
 	Set("UVEditor.FlyCamera", new CORE_IMAGE_BRUSH_SVG("Starship/EditorViewport/camera", ToolbarIconSize));
 	Set("UVEditor.FocusCamera", new CORE_IMAGE_BRUSH_SVG("Starship/Actors/snap-view-to-object", ToolbarIconSize));
 
+	Set("UVEditor.ElementSelection", new IMAGE_BRUSH_SVG("MeshSelect_16", ToolbarIconSize));
 	Set("UVEditor.VertexSelection", new IMAGE_BRUSH_SVG("SelectionVertices", ToolbarIconSize));
 	Set("UVEditor.EdgeSelection", new IMAGE_BRUSH_SVG("SelectionLine", ToolbarIconSize));
 	Set("UVEditor.TriangleSelection", new IMAGE_BRUSH_SVG("SelectionTriangle", ToolbarIconSize));
diff --git a/Engine/Source/Developer/ToolMenus/Public/ToolMenuDelegates.h b/Engine/Source/Developer/ToolMenus/Public/ToolMenuDelegates.h
--- a/Engine/Source/Developer/ToolMenus/Public/ToolMenuDelegates.h
+++ b/Engine/Source/Developer/ToolMenus/Public/ToolMenuDelegates.h
@@ -42,6 +42,16 @@ struct TOOLMENUS_API FToolUIAction
 {
 	FToolUIAction() {}
 	FToolUIAction(const FToolMenuExecuteAction& InExecuteAction) : ExecuteAction(InExecuteAction) {}
+	FToolUIAction(const FToolMenuExecuteAction& InExecuteAction, const FToolMenuCanExecuteAction& InCanExecuteAction, const FToolMenuGetActionCheckState& InGetActionCheckState)
+		: ExecuteAction(InExecuteAction)
+		, CanExecuteAction(InCanExecuteAction)
+		, GetActionCheckState(InGetActionCheckState)
+	{}
+	
+	FToolUIAction(const FToolMenuExecuteAction& InExectueAction, const FToolMenuGetActionCheckState& InGetActionCheckState)
+		: ExecuteAction(InExectueAction)
+		, GetActionCheckState(InGetActionCheckState)
+	{}
 
 	FToolMenuExecuteAction ExecuteAction;
 	FToolMenuCanExecuteAction CanExecuteAction;
diff --git a/Engine/Source/Editor/LevelEditor/Private/ViewportToolbar/LevelEditorSubmenus.cpp b/Engine/Source/Editor/LevelEditor/Private/ViewportToolbar/LevelEditorSubmenus.cpp
--- a/Engine/Source/Editor/LevelEditor/Private/ViewportToolbar/LevelEditorSubmenus.cpp
+++ b/Engine/Source/Editor/LevelEditor/Private/ViewportToolbar/LevelEditorSubmenus.cpp
@@ -706,12 +706,13 @@ FToolMenuEntry CreatePreviewSelectedCamerasCheckBoxSubmenu()
 		"CameraPreviewWindowTooltip", "When enabled, selecting a camera actor will display a live 'picture in picture' preview from the camera's perspective within the current editor view port.  This can be used to easily tweak camera positioning, post-processing and other settings without having to possess the camera itself.  This feature may reduce application performance when enabled."
 	);
 
-	FToolMenuEntry Entry = UnrealEd::CreateCheckboxSubmenu(
+	FToolMenuEntry Entry = FToolMenuEntry::InitSubMenu(
 		"PreviewSelectedCameras",
 		LOCTEXT("PreviewSelectedCamerasLabel", "Preview Selected Cameras"),
 		PreviewSelectedCamerasTooltip,
-		FToolMenuExecuteAction::CreateLambda(
-			[](const FToolMenuContext& InContext)
+		PreviewSelectedCamerasMenuDelegate,
+		FToolUIAction(
+			FToolMenuExecuteAction::CreateLambda([](const FToolMenuContext& InContext)
 			{
 				TWeakPtr<::SLevelViewport> LevelViewportWeak;
 				if (ULevelViewportContext* const LevelViewportContext = InContext.FindContext<ULevelViewportContext>())
@@ -720,17 +721,14 @@ FToolMenuEntry CreatePreviewSelectedCamerasCheckBoxSubmenu()
 				}
 
 				FLevelEditorActionCallbacks::TogglePreviewSelectedCameras(LevelViewportWeak);
-			}
-		),
-		FToolMenuCanExecuteAction(),
-		FToolMenuGetActionCheckState::CreateLambda(
-			[](const FToolMenuContext& InContext)
+			}),
+			FToolMenuGetActionCheckState::CreateLambda([](const FToolMenuContext& InContext)
 			{
 				return FLevelEditorActionCallbacks::IsPreviewSelectedCamerasChecked() ? ECheckBoxState::Checked
 																					  : ECheckBoxState::Unchecked;
-			}
+			})
 		),
-		PreviewSelectedCamerasMenuDelegate
+		EUserInterfaceActionType::ToggleButton
 	);
 
 	return Entry;
diff --git a/Engine/Source/Editor/UnrealEd/Private/ViewportToolbar/UnrealEdViewportToolbar.cpp b/Engine/Source/Editor/UnrealEd/Private/ViewportToolbar/UnrealEdViewportToolbar.cpp
--- a/Engine/Source/Editor/UnrealEd/Private/ViewportToolbar/UnrealEdViewportToolbar.cpp
+++ b/Engine/Source/Editor/UnrealEd/Private/ViewportToolbar/UnrealEdViewportToolbar.cpp
@@ -27,6 +27,8 @@
 #include "ToolMenuEntry.h"
 #include "ToolMenuSection.h"
 #include "ToolMenus.h"
+#include "ToolMenusEditor.h"
+#include "Tests/ToolMenusTestUtilities.h"
 #include "ViewportToolbar/UnrealEdViewportToolbarContext.h"
 #include "Widgets/Input/SNumericEntryBox.h"
 #include "Widgets/Input/SSlider.h"
@@ -101,7 +103,6 @@ bool IsSurfaceSnappingEnabled()
 	return false;
 }
 
-// TODO: Maybe export CreateSurfaceSnapOffsetEntry function, so that it can be used elsewhere, e.g. STransformViewportToolbar.cpp
 FToolMenuEntry CreateSurfaceSnapOffsetEntry()
 {
 	FText Label = LOCTEXT("SurfaceOffsetLabel", "Surface Offset");
@@ -165,344 +166,364 @@ FToolMenuEntry CreateSurfaceSnapOffsetEntry()
 	return SurfaceOffset;
 }
 
-FToolMenuEntry CreateSurfaceSnapCheckboxMenu()
+FToolMenuEntry CreateSurfaceSnapCheckboxMenu(const FToolMenuContext& InContext)
 {
-	FNewToolMenuDelegate MakeMenuDelegate = FNewToolMenuDelegate::CreateLambda(
-		[](UToolMenu* Submenu)
+	const TSharedPtr<FUICommandInfo> Command = FEditorViewportCommands::Get().SurfaceSnapping;
+
+	FUIAction Action;
+	
+	if (const FUIAction* FoundAction = InContext.GetActionForCommand(Command))
+	{
+		Action = *FoundAction;
+	}
+	else
+	{
+		// Provide a default implementation
+		// TODO: This would be better handled by ensuring each editor has a bound command.
+		Action.ExecuteAction.BindLambda([] { ToggleSurfaceSnapping(); });
+		Action.GetActionCheckState.BindLambda([]
+		{
+			return IsSurfaceSnappingEnabled() ? ECheckBoxState::Checked : ECheckBoxState::Unchecked;
+		});
+	}
+	
+	FToolMenuEntry Entry = FToolMenuEntry::InitSubMenu(
+		"SurfaceSnapping",
+		LOCTEXT("SurfaceSnapLabel", "Surface"),
+		Command->GetDescription(),
+		FNewToolMenuDelegate::CreateLambda([](UToolMenu* Submenu)
 		{
 			FToolMenuSection& SurfaceSnappingSection =
 				Submenu->FindOrAddSection("SurfaceSnapping", LOCTEXT("SurfaceSnappingLabel", "Surface Snapping"));
 
-			// Add "Rotate to surface normal" checkbox.
-			{
-				FToolMenuEntry RotateToSurfaceNormalSnapping =
-					FToolMenuEntry::InitMenuEntry(FEditorViewportCommands::Get().RotateToSurfaceNormal);
-				SurfaceSnappingSection.AddEntry(RotateToSurfaceNormalSnapping);
-			}
-
-			// Add "Surface offset" widget.
-			{
-				SurfaceSnappingSection.AddEntry(CreateSurfaceSnapOffsetEntry());
-			}
-		}
+			SurfaceSnappingSection.AddMenuEntry(FEditorViewportCommands::Get().RotateToSurfaceNormal);
+			SurfaceSnappingSection.AddEntry(CreateSurfaceSnapOffsetEntry());
+		}),
+		Action,
+		EUserInterfaceActionType::ToggleButton
 	);
-
-	FToolMenuEntry Entry = UnrealEd::CreateCheckboxSubmenu(
-		"SurfaceSnapping",
-		LOCTEXT("SurfaceSnapLabel", "Surface"),
-		FEditorViewportCommands::Get().SurfaceSnapping->GetDescription(),
-		FToolMenuExecuteAction::CreateLambda(
-			[](const FToolMenuContext& InContext)
-			{
-				ToggleSurfaceSnapping();
-			}
-		),
-		FToolMenuCanExecuteAction(),
-		FToolMenuGetActionCheckState::CreateLambda(
-			[](const FToolMenuContext& InContext)
+	
+	if (UUnrealEdViewportToolbarContext* Context = InContext.FindContext<UUnrealEdViewportToolbarContext>())
+	{
+		TWeakObjectPtr<UUnrealEdViewportToolbarContext> WeakContext = Context;
+		Entry.SetVisibility(TAttribute<EVisibility>::CreateLambda([WeakContext]
+		{
+			if (WeakContext.IsValid())
 			{
-				return UE::UnrealEd::Private::IsSurfaceSnappingEnabled() ? ECheckBoxState::Checked
-																		 : ECheckBoxState::Unchecked;
+				return WeakContext->bShowSurfaceSnap ? EVisibility::Visible : EVisibility::Hidden;
 			}
-		),
-		MakeMenuDelegate
-	);
+			return EVisibility::Visible;
+		}));
+	}
 
-	Entry.ToolBarData.LabelOverride = TAttribute<FText>::CreateLambda(
-		[]()
-		{
-			const ULevelEditorViewportSettings* const Settings = GetMutableDefault<ULevelEditorViewportSettings>();
-			return FText::AsNumber(Settings->SnapToSurface.SnapOffsetExtent);
-		}
-	);
-	Entry.Icon = TAttribute<FSlateIcon>::CreateLambda(
-		[]()
-		{
-			// todo: dynamic update works in real time for the icon in the raised entry, but not within the menu itself.
-			// In order to se the icon within the menu update, the menu needs to close and re-open
-			return GetDefault<ULevelEditorViewportSettings>()->SnapToSurface.bSnapRotation
-					 ? FSlateIcon(FAppStyle::GetAppStyleSetName(), "ViewportToolbar.SurfaceSnapRotateToNormal")
-					 : FSlateIcon(FAppStyle::GetAppStyleSetName(), "ViewportToolbar.SurfaceSnap");
-		}
-	);
+	Entry.ToolBarData.LabelOverride = TAttribute<FText>::CreateLambda([]
+	{
+		const ULevelEditorViewportSettings* const Settings = GetMutableDefault<ULevelEditorViewportSettings>();
+		return FText::AsNumber(Settings->SnapToSurface.SnapOffsetExtent);
+	});
+	
+	Entry.Icon = TAttribute<FSlateIcon>::CreateLambda([]
+	{
+		// todo: dynamic update works in real time for the icon in the raised entry, but not within the menu itself.
+		// In order to se the icon within the menu update, the menu needs to close and re-open
+		return GetDefault<ULevelEditorViewportSettings>()->SnapToSurface.bSnapRotation
+				 ? FSlateIcon(FAppStyle::GetAppStyleSetName(), "ViewportToolbar.SurfaceSnapRotateToNormal")
+				 : FSlateIcon(FAppStyle::GetAppStyleSetName(), "ViewportToolbar.SurfaceSnap");
+	});
 
 	return Entry;
 }
 
-FToolMenuEntry CreateActorSnapDistanceEntry()
-{
-	const FText Label = LOCTEXT("ActorSnapDistanceLabel", "Snap Distance");
-	const FText Tooltip = LOCTEXT("ActorSnapDistanceTooltip", "The amount of offset to apply when snapping to surfaces");
-
-	const FMargin WidgetsMargin(2.0f, 0.0f, 3.0f, 0.0f);
-
-	FToolMenuEntry SnapDistance = FToolMenuEntry::InitMenuEntry(
-		"ActorSnapDistance",
-		FUIAction(
-			FExecuteAction(),
-			FCanExecuteAction::CreateLambda(
-				[]()
-				{
-					return !!GetDefault<ULevelEditorViewportSettings>()->bEnableActorSnap;
-				}
-			)
-		),
-		// clang-format off
-		SNew(SHorizontalBox)
-		+ SHorizontalBox::Slot()
-		.VAlign(VAlign_Center)
-		.Padding(WidgetsMargin)
-		.AutoWidth()
-		[
-			SNew(STextBlock)
-			.Text(Label)
-		]
-		+ SHorizontalBox::Slot()
-		.VAlign(VAlign_Center)
-		.Padding(WidgetsMargin)
-		.AutoWidth()
-		[
-			SNew(SBox).Padding(WidgetsMargin).MinDesiredWidth(100.0f)
-			[
-				// TODO: Check how to improve performance for this widget OnValueChanged.
-				// Same functionality in LevelEditorToolBar.cpp seems to have better performance
-				SNew(SNumericEntryBox<float>)
-				.ToolTipText(Tooltip)
-				.MinValue(0.0f)
-				.MaxValue(1.0f)
-				.MaxSliderValue(1.0f)
-				.AllowSpin(true)
-				.MaxFractionalDigits(1)
-				.Font(FAppStyle::GetFontStyle(TEXT("MenuItem.Font")))
-				.OnValueChanged_Static(&FLevelEditorActionCallbacks::SetActorSnapSetting)
-				.Value_Lambda([]()
-				{
-					return FLevelEditorActionCallbacks::GetActorSnapSetting();
-				})
-			]
-		]
-		// clang-format on
-	);
-
-	return SnapDistance;
-}
-
-FToolMenuEntry CreateLocationSnapCheckboxMenu()
+FToolMenuEntry CreateLocationSnapCheckboxMenu(const FToolMenuContext& InContext)
 {
-	const FName LocationSnapName = "LocationSnap";
-	const FText LocationSnapLabel = LOCTEXT("LocationSnapLabel", "Location");
-
-	FToolMenuEntry Entry;
+	const TSharedPtr<FUICommandInfo> Command = FEditorViewportCommands::Get().LocationGridSnap;
 
-	if (!GEditor)
+	FUIAction Action;
+	
+	if (const FUIAction* FoundAction = InContext.GetActionForCommand(Command))
 	{
-		Entry = FToolMenuEntry::InitMenuEntry(
-			LocationSnapName, LocationSnapLabel, FText(), FSlateIcon(), UnrealEd::Private::DisabledAction()
-		);
+		Action = *FoundAction;
 	}
 	else
 	{
-		FNewToolMenuDelegate MakeMenuDelegate = FNewToolMenuDelegate::CreateLambda(
-			[LocationSnapName](UToolMenu* InToolMenu)
+		// Provide a default implementation
+		// TODO: This would be better handled by ensuring each editor has a bound command.
+		Action.ExecuteAction.BindStatic(&FLevelEditorActionCallbacks::LocationGridSnap_Clicked);
+		Action.GetActionCheckState.BindLambda([]
+		{
+			return FLevelEditorActionCallbacks::LocationGridSnap_IsChecked() ? ECheckBoxState::Checked : ECheckBoxState::Unchecked;
+		});
+	}
+
+	FToolMenuEntry Entry = FToolMenuEntry::InitSubMenu(
+		"GridSnapping",
+		LOCTEXT("LocationSnapLabel", "Location"),
+		Command->GetDescription(),
+		FNewToolMenuDelegate::CreateLambda([](UToolMenu* ToolMenu)
+		{
+			UUnrealEdViewportToolbarContext* Context = ToolMenu->FindContext<UUnrealEdViewportToolbarContext>();
+			if (!Context) return;
+					
+			FToolMenuSection& Section = ToolMenu->AddSection("Snap", LOCTEXT("LocationSnapText", "Snap Sizes"));
+					
+			const TArray<float> GridSizes = Context->GetGridSnapSizes();
+					
+			for (int32 GridSizeIndex = 0; GridSizeIndex < GridSizes.Num(); ++GridSizeIndex)
 			{
-				InToolMenu->AddMenuEntry(
-					LocationSnapName,
-					FToolMenuEntry::InitWidget(LocationSnapName, UnrealEd::CreateLocationGridSnapMenu(),
-						FText()
-					)
+				const float GridSize = GridSizes[GridSizeIndex];
+				Section.AddMenuEntry(
+					NAME_None,
+					FText::AsNumber(GridSize),
+					FText::Format(
+						LOCTEXT("LocationGridSize_ToolTip", "Sets grid size to {0}"), FText::AsNumber(GridSize)
+					),
+					FSlateIcon(),
+					FUIAction(
+						FExecuteAction::CreateUObject(Context, &UUnrealEdViewportToolbarContext::SetGridSnapSize, GridSizeIndex),
+						FCanExecuteAction(),
+						FIsActionChecked::CreateUObject(Context, &UUnrealEdViewportToolbarContext::IsGridSnapSizeActive, GridSizeIndex)
+					),
+					EUserInterfaceActionType::RadioButton
 				);
 			}
-		);
-
-		Entry = UnrealEd::CreateCheckboxSubmenu(
-			"GridSnapping",
-			LocationSnapLabel,
-			FEditorViewportCommands::Get().LocationGridSnap->GetDescription(),
-			FToolMenuExecuteAction::CreateLambda(
-				[](const FToolMenuContext& InContext)
-				{
-					FLevelEditorActionCallbacks::LocationGridSnap_Clicked();
-				}
-			),
-			FToolMenuCanExecuteAction(),
-			FToolMenuGetActionCheckState::CreateLambda(
-				[](const FToolMenuContext& InContext)
-				{
-					return FLevelEditorActionCallbacks::LocationGridSnap_IsChecked() ? ECheckBoxState::Checked
-																					 : ECheckBoxState::Unchecked;
-				}
-			),
-			MakeMenuDelegate
-		);
+		}),
+		Action,
+		EUserInterfaceActionType::ToggleButton
+	);
+			
+	Entry.UserInterfaceActionType = EUserInterfaceActionType::ToggleButton;
+	
+	if (UUnrealEdViewportToolbarContext* Context = InContext.FindContext<UUnrealEdViewportToolbarContext>())
+	{
+		TAttribute<FText> LabelOverride;
+		LabelOverride.BindUObject(Context, &UUnrealEdViewportToolbarContext::GetGridSnapLabel);
+			
+		Entry.ToolBarData.LabelOverride = LabelOverride;
 	}
-
-	Entry.ToolBarData.LabelOverride = TAttribute<FText>::Create(&UE::UnrealEd::GetLocationGridLabel);
+	
 	Entry.Icon = FSlateIcon(FAppStyle::GetAppStyleSetName(), "ViewportToolbar.SnapLocation");
-
+	
 	return Entry;
 }
 
-FToolMenuEntry CreateRotationSnapCheckboxMenu()
+TSharedRef<SWidget> CreateRotationGridSnapList(
+	UUnrealEdViewportToolbarContext* Context,
+	const FText& Heading,
+	const TArray<float>& Sizes,
+	ERotationGridMode GridMode
+)
 {
-	const FName RotationSnapName = "RotationSnap";
-	const FText RotationSnapLabel = LOCTEXT("RotationSnapLabel", "Rotation");
-
-	if (!GEditor)
+	FMenuBuilder Menu(true, nullptr);
+	
+	Menu.BeginSection(NAME_None, Heading);
+	
+	for (int32 Index = 0; Index < Sizes.Num(); ++Index)
 	{
-		return FToolMenuEntry::InitMenuEntry(
-			RotationSnapName, RotationSnapLabel, FText(), FSlateIcon(), UnrealEd::Private::DisabledAction()
+		FText Label = FText::Format(
+			LOCTEXT("RotationGridAngle", "{0}\u00b0"), // \u00b0 is the degree symbol
+			FText::AsNumber(Sizes[Index])
+		);
+		
+		FText ToolTip = FText::Format(
+			LOCTEXT("RotationGridAngle_ToolTip", "Sets rotation grid angle to {0}"),
+			Label
+		);
+	
+		Menu.AddMenuEntry(
+			Label,
+			ToolTip,
+			FSlateIcon(),
+			FUIAction(
+				FExecuteAction::CreateUObject(Context, &UUnrealEdViewportToolbarContext::SetRotationSnapSize, Index, GridMode),
+				FCanExecuteAction(),
+				FIsActionChecked::CreateUObject(Context, &UUnrealEdViewportToolbarContext::IsRotationSnapActive, Index, GridMode)
+			),
+			NAME_None,
+			EUserInterfaceActionType::RadioButton
 		);
 	}
+	
+	Menu.EndSection();
+	
+	return Menu.MakeWidget();
+}
 
-	const FNewToolMenuDelegate MakeMenuDelegate = FNewToolMenuDelegate::CreateLambda(
-		[RotationSnapName](UToolMenu* InToolMenu)
-		{
-			const UnrealEd::FRotationGridCheckboxListExecuteActionDelegate ExecuteDelegate =
-				UnrealEd::FRotationGridCheckboxListExecuteActionDelegate::CreateUObject(
-					GEditor, &UEditorEngine::SetRotGridSize
-				);
+TSharedRef<SWidget> CreateRotationGridSnapWidget(UUnrealEdViewportToolbarContext* Context)
+{
+	const ULevelEditorViewportSettings* ViewportSettings = GetDefault<ULevelEditorViewportSettings>();
 
-			const UnrealEd::FRotationGridCheckboxListIsCheckedDelegate IsCheckedDelegate =
-				UnrealEd::FRotationGridCheckboxListIsCheckedDelegate::CreateLambda(
-					[](int CurrGridAngleIndex, ERotationGridMode InGridMode)
-					{
-						const ULevelEditorViewportSettings* ViewportSettings = GetDefault<ULevelEditorViewportSettings>();
-						return ViewportSettings->CurrentRotGridSize == CurrGridAngleIndex
-							&& ViewportSettings->CurrentRotGridMode == InGridMode;
-					}
-				);
+	return SNew(SUniformGridPanel)
+		+ SUniformGridPanel::Slot(0, 0)
+		[
+			CreateRotationGridSnapList(
+				Context,
+				LOCTEXT("RotationCommonText", "Rotation Increment"),
+				ViewportSettings->CommonRotGridSizes,
+				GridMode_Common
+			)
+		]
+		+ SUniformGridPanel::Slot(1, 0)
+		[
+			CreateRotationGridSnapList(
+				Context,
+				LOCTEXT("RotationDivisions360DegreesText", "Divisions of 360\u00b0"),
+				ViewportSettings->DivisionsOf360RotGridSizes,
+				GridMode_DivisionsOf360
+			)
+		];
+}
 
+FToolMenuEntry CreateRotationSnapCheckboxMenu(const FToolMenuContext& InContext)
+{
+	const TSharedPtr<FUICommandInfo> Command = FEditorViewportCommands::Get().RotationGridSnap;
+	
+	FUIAction Action;
+	
+	if (const FUIAction* FoundAction = InContext.GetActionForCommand(Command))
+	{
+		Action = *FoundAction;
+	}
+	else
+	{
+		// Provide a default implementation
+		// TODO: This would be better handled by ensuring each editor has a bound command.
+		Action.ExecuteAction.BindStatic(&FLevelEditorActionCallbacks::RotationGridSnap_Clicked);
+		Action.GetActionCheckState.BindLambda([]
+		{
+			return FLevelEditorActionCallbacks::RotationGridSnap_IsChecked()
+				? ECheckBoxState::Checked
+				: ECheckBoxState::Unchecked;
+		});
+	}
+	
+	FToolMenuEntry Entry = FToolMenuEntry::InitSubMenu(
+		"RotationSnapping",
+		LOCTEXT("RotationSnapLabel", "Rotation"),
+		FEditorViewportCommands::Get().RotationGridSnap->GetDescription(),
+		FNewToolMenuDelegate::CreateLambda([](UToolMenu* InToolMenu)
+		{
+			UUnrealEdViewportToolbarContext* Context = InToolMenu->FindContext<UUnrealEdViewportToolbarContext>();
+			if (!Context)
+			{
+				return;
+			}
+			
 			InToolMenu->AddMenuEntry(
-				RotationSnapName,
+				"RotationSnap",
 				FToolMenuEntry::InitWidget(
-					RotationSnapName, UnrealEd::CreateRotationGridSnapMenu(ExecuteDelegate, IsCheckedDelegate),
+					"RotationSnap",
+					CreateRotationGridSnapWidget(Context),
 					FText()
 				)
 			);
-		}
+		}),
+		Action,
+		EUserInterfaceActionType::ToggleButton
 	);
 
-	FToolMenuEntry Entry = UnrealEd::CreateCheckboxSubmenu(
-		"RotationSnapping",
-		RotationSnapLabel,
-		FEditorViewportCommands::Get().RotationGridSnap->GetDescription(),
-		FToolMenuExecuteAction::CreateLambda(
-			[](const FToolMenuContext& InContext)
-			{
-				FLevelEditorActionCallbacks::RotationGridSnap_Clicked();
-			}
-		),
-		FToolMenuCanExecuteAction(),
-		FToolMenuGetActionCheckState::CreateLambda(
-			[](const FToolMenuContext& InContext)
-			{
-				return FLevelEditorActionCallbacks::RotationGridSnap_IsChecked() ? ECheckBoxState::Checked
-																				 : ECheckBoxState::Unchecked;
-			}
-		),
-		MakeMenuDelegate
-	);
+	TAttribute<FText> LabelOverride;
+	
+	if (UUnrealEdViewportToolbarContext* Context = InContext.FindContext<UUnrealEdViewportToolbarContext>())
+	{
+		LabelOverride.BindUObject(Context, &UUnrealEdViewportToolbarContext::GetRotationSnapLabel);
+	}
 
-	Entry.ToolBarData.LabelOverride = TAttribute<FText>::Create(&UE::UnrealEd::GetRotationGridLabel);
+	Entry.ToolBarData.LabelOverride = LabelOverride;
+	
 	Entry.Icon = FSlateIcon(FAppStyle::GetAppStyleSetName(), "ViewportToolbar.SnapRotation");
 
 	return Entry;
 }
 
-FToolMenuEntry CreateScaleSnapCheckboxMenu()
+FToolMenuEntry CreateScaleSnapCheckboxMenu(const FToolMenuContext& InContext)
 {
-	const FName ScaleSnapName = "ScaleSnap";
-	const FText ScaleSnapLabel = LOCTEXT("ScaleSnapLabel", "Scale");
-
-	if (!GEditor)
+	const TSharedPtr<FUICommandInfo> Command = FEditorViewportCommands::Get().ScaleGridSnap;
+	
+	FUIAction Action;
+	
+	if (const FUIAction* FoundAction = InContext.GetActionForCommand(Command))
 	{
-		return FToolMenuEntry::InitMenuEntry(
-			ScaleSnapName, ScaleSnapLabel, FText(), FSlateIcon(), UnrealEd::Private::DisabledAction()
-		);
+		Action = *FoundAction;
 	}
-
-	FNewToolMenuDelegate MakeMenuDelegate = FNewToolMenuDelegate::CreateLambda(
-		[ScaleSnapName](UToolMenu* InToolMenu)
+	else
+	{
+		// Provide a default implementation
+		// TODO: This would be better handled by ensuring each editor has a bound command.
+		Action.ExecuteAction.BindStatic(&FLevelEditorActionCallbacks::ScaleGridSnap_Clicked);
+		Action.GetActionCheckState.BindLambda([]
 		{
-			FCanExecuteAction CanExecuteScaleSnapping = FCanExecuteAction::CreateLambda(
-				[]()
-				{
-					return FLevelEditorActionCallbacks::ScaleGridSnap_IsChecked();
-				}
-			);
-
-			const ULevelEditorViewportSettings* ViewportSettings = GetDefault<ULevelEditorViewportSettings>();
-			TArray<float> GridSizes = ViewportSettings->ScalingGridSizes;
-
-			UnrealEd::FScaleGridCheckboxListExecuteActionDelegate ExecuteDelegate =
-				UnrealEd::FScaleGridCheckboxListExecuteActionDelegate::CreateUObject(GEditor, &UEditorEngine::SetScaleGridSize);
-
-			UnrealEd::FScaleGridCheckboxListIsCheckedDelegate IsCheckedDelegate =
-				UnrealEd::FScaleGridCheckboxListIsCheckedDelegate::CreateLambda(
-					[](int CurrGridSizeIndex)
-					{
-						const ULevelEditorViewportSettings* ViewportSettings = GetDefault<ULevelEditorViewportSettings>();
-						return ViewportSettings->CurrentScalingGridSize == CurrGridSizeIndex;
-					}
-				);
-
-			InToolMenu->AddMenuEntry(
-				ScaleSnapName,
-				FToolMenuEntry::InitWidget(
-					ScaleSnapName,
-					UnrealEd::CreateScaleGridSnapMenu(
-						ExecuteDelegate,
-						IsCheckedDelegate,
-						GridSizes,
-						TAttribute<bool>(true),
-						{},
-						true,
-						FUIAction(
-							FExecuteAction::CreateLambda(
-								[]()
-								{
-									ULevelEditorViewportSettings* Settings =
-										GetMutableDefault<ULevelEditorViewportSettings>();
-									Settings->PreserveNonUniformScale = !Settings->PreserveNonUniformScale;
-								}
-							),
-							CanExecuteScaleSnapping,
-							FIsActionChecked::CreateLambda(
-								[]()
-								{
-									return GetDefault<ULevelEditorViewportSettings>()->PreserveNonUniformScale;
-								}
-							)
-						)
-					),
-					FText()
-				)
-			);
-		}
-	);
+			return FLevelEditorActionCallbacks::ScaleGridSnap_IsChecked() ? ECheckBoxState::Checked : ECheckBoxState::Unchecked;
+		});
+	}
 
-	FToolMenuEntry Entry = UnrealEd::CreateCheckboxSubmenu(
+	const FName ScaleSnapName = "ScaleSnap";
+	
+	FToolMenuEntry Entry = FToolMenuEntry::InitSubMenu(
 		"ScaleSnapping",
-		ScaleSnapLabel,
-		FEditorViewportCommands::Get().ScaleGridSnap->GetDescription(),
-		FToolMenuExecuteAction::CreateLambda(
-			[](const FToolMenuContext& InContext)
+		LOCTEXT("ScaleSnapLabel", "Scale"),
+		Command->GetDescription(),
+		FNewToolMenuDelegate::CreateLambda([ScaleSnapName](UToolMenu* InToolMenu)
+		{
+			UUnrealEdViewportToolbarContext* Context = InToolMenu->FindContext<UUnrealEdViewportToolbarContext>();
+			if (!Context)
 			{
-				FLevelEditorActionCallbacks::ScaleGridSnap_Clicked();
+				return;
 			}
-		),
-		FToolMenuCanExecuteAction(),
-		FToolMenuGetActionCheckState::CreateLambda(
-			[](const FToolMenuContext& InContext)
+			
+			FToolMenuSection& Section = InToolMenu->FindOrAddSection("ScaleSnap", LOCTEXT("ScaleSnapSizesSectionLabel", "Snap Sizes"));
+		
+			TArray<float> ScaleSnapSizes = Context->GetScaleSnapSizes();
+			
+			FNumberFormattingOptions NumberFormattingOptions;
+			NumberFormattingOptions.MaximumFractionalDigits = 5;
+			
+			for (int32 ScaleSnapSizeIndex = 0; ScaleSnapSizeIndex < ScaleSnapSizes.Num(); ++ScaleSnapSizeIndex)
 			{
-				return FLevelEditorActionCallbacks::ScaleGridSnap_IsChecked() ? ECheckBoxState::Checked
-																			  : ECheckBoxState::Unchecked;
+				const float CurGridAmount = ScaleSnapSizes[ScaleSnapSizeIndex];
+
+				FText Label;
+				FText ToolTip;
+
+				if (GEditor->UsePercentageBasedScaling())
+				{
+					Label = FText::AsPercent(CurGridAmount / 100.0f, &NumberFormattingOptions);
+					ToolTip =
+						FText::Format(LOCTEXT("ScaleGridAmountOld_ToolTip", "Snaps scale values to {0}"), Label);
+				}
+				else
+				{
+					Label = FText::AsNumber(CurGridAmount, &NumberFormattingOptions);
+					ToolTip = FText::Format(
+						LOCTEXT("ScaleGridAmount_ToolTip", "Snaps scale values to increments of {0}"), Label
+					);
+				}
+				
+				Section.AddMenuEntry(
+					NAME_None,
+					Label,
+					ToolTip,
+					FSlateIcon(),
+					FUIAction(
+						FExecuteAction::CreateUObject(Context, &UUnrealEdViewportToolbarContext::SetScaleSnapSize, ScaleSnapSizeIndex),
+						FCanExecuteAction(),
+						FIsActionChecked::CreateUObject(Context, &UUnrealEdViewportToolbarContext::IsScaleSnapActive, ScaleSnapSizeIndex)
+					),
+					EUserInterfaceActionType::RadioButton
+				);
 			}
-		),
-		MakeMenuDelegate
+		}),
+		Action,
+		EUserInterfaceActionType::ToggleButton
 	);
 
-	Entry.ToolBarData.LabelOverride = TAttribute<FText>::Create(&UE::UnrealEd::GetScaleGridLabel);
+	TAttribute<FText> LabelOverride;
+	
+	if (UUnrealEdViewportToolbarContext* Context = InContext.FindContext<UUnrealEdViewportToolbarContext>())
+	{
+		LabelOverride.BindUObject(Context, &UUnrealEdViewportToolbarContext::GetScaleSnapLabel);
+	}
+
+	Entry.ToolBarData.LabelOverride = LabelOverride;
 	Entry.Icon = FSlateIcon(FAppStyle::GetAppStyleSetName(), "ViewportToolbar.SnapScale");
 
 	return Entry;
@@ -935,6 +956,32 @@ FToolMenuEntry CreateTransformsSubmenu()
 						"GizmoScale", GizmoScaleWidget, LOCTEXT("GizmoScaleLabel", "Gizmo Scale")
 					));
 					GizmoScaleEntry.Icon = FSlateIcon(FAppStyle::GetAppStyleSetName(), "ViewportToolbar.GizmoScale");
+					
+					
+					{
+						FUIAction PreserveNonUniformScaleAction;
+						PreserveNonUniformScaleAction.ExecuteAction.BindLambda([]
+						{
+							ULevelEditorViewportSettings* Settings = GetMutableDefault<ULevelEditorViewportSettings>();
+							Settings->PreserveNonUniformScale = !Settings->PreserveNonUniformScale;
+						});
+						PreserveNonUniformScaleAction.GetActionCheckState.BindLambda([]
+						{
+							return GetDefault<ULevelEditorViewportSettings>()->PreserveNonUniformScale ? ECheckBoxState::Checked : ECheckBoxState::Unchecked;
+						});
+						
+						GizmoSection.AddMenuEntry(
+							"PreserveNonUniformScale",
+							LOCTEXT("ScaleGridPreserveNonUniformScale", "Preserve Non-Uniform Scale"),
+							LOCTEXT(
+								"ScaleGridPreserveNonUniformScale_ToolTip",
+								"When this option is checked, scaling objects that have a non-uniform scale will preserve the ratios between each axis, snapping the axis with the largest value."
+							),
+							FSlateIcon(),
+							PreserveNonUniformScaleAction,
+							EUserInterfaceActionType::ToggleButton
+						);
+					}
 				}
 			}
 		)
@@ -1252,10 +1299,10 @@ FToolMenuEntry CreateSnappingSubmenu()
 				FToolMenuSection& SnappingSection =
 					Submenu->FindOrAddSection("Snapping", LOCTEXT("SnappingLabel", "Snapping"));
 
-				SnappingSection.AddEntry(Private::CreateSurfaceSnapCheckboxMenu()).SetShowInToolbarTopLevel(true);
-				SnappingSection.AddEntry(Private::CreateLocationSnapCheckboxMenu()).SetShowInToolbarTopLevel(true);
-				SnappingSection.AddEntry(Private::CreateRotationSnapCheckboxMenu()).SetShowInToolbarTopLevel(true);
-				SnappingSection.AddEntry(Private::CreateScaleSnapCheckboxMenu()).SetShowInToolbarTopLevel(true);
+				SnappingSection.AddEntry(Private::CreateSurfaceSnapCheckboxMenu(Submenu->Context)).SetShowInToolbarTopLevel(true);
+				SnappingSection.AddEntry(Private::CreateLocationSnapCheckboxMenu(Submenu->Context)).SetShowInToolbarTopLevel(true);
+				SnappingSection.AddEntry(Private::CreateRotationSnapCheckboxMenu(Submenu->Context)).SetShowInToolbarTopLevel(true);
+				SnappingSection.AddEntry(Private::CreateScaleSnapCheckboxMenu(Submenu->Context)).SetShowInToolbarTopLevel(true);
 			}
 		)
 	);
@@ -2184,7 +2231,7 @@ FToolMenuEntry CreateCheckboxSubmenu(
 	FToolMenuEntry CheckBoxSubmenu = FToolMenuEntry::InitSubMenu(
 		InName, InLabel, InToolTip, InMakeMenu, CheckboxMenuAction, EUserInterfaceActionType::ToggleButton
 	);
-
+	
 	return CheckBoxSubmenu;
 }
 
diff --git a/Engine/Source/Editor/UnrealEd/Private/ViewportToolbar/UnrealEdViewportToolbarContext.cpp b/Engine/Source/Editor/UnrealEd/Private/ViewportToolbar/UnrealEdViewportToolbarContext.cpp
--- a/Engine/Source/Editor/UnrealEd/Private/ViewportToolbar/UnrealEdViewportToolbarContext.cpp
+++ b/Engine/Source/Editor/UnrealEd/Private/ViewportToolbar/UnrealEdViewportToolbarContext.cpp
@@ -3,6 +3,10 @@
 #include "ViewportToolbar/UnrealEdViewportToolbarContext.h"
 
 #include "SEditorViewport.h"
+#include "Settings/LevelEditorViewportSettings.h"
+#include "ViewportToolbar/UnrealEdViewportToolbar.h"
+
+#define LOCTEXT_NAMESPACE "UnrealEdViewportToolbar"
 
 TSharedPtr<IPreviewProfileController> UUnrealEdViewportToolbarContext::GetPreviewProfileController() const
 {
@@ -12,3 +16,73 @@ TSharedPtr<IPreviewProfileController> UUnrealEdViewportToolbarContext::GetPrevie
 	}
 	return nullptr;
 }
+
+FText UUnrealEdViewportToolbarContext::GetGridSnapLabel() const
+{
+	return FText::AsNumber(GEditor->GetGridSize());
+}
+
+TArray<float> UUnrealEdViewportToolbarContext::GetGridSnapSizes() const
+{
+	const ULevelEditorViewportSettings* ViewportSettings = GetDefault<ULevelEditorViewportSettings>();
+	if (ViewportSettings->bUsePowerOf2SnapSize)
+	{
+		return ViewportSettings->Pow2GridSizes;
+	}
+	return ViewportSettings->DecimalGridSizes;
+}
+
+bool UUnrealEdViewportToolbarContext::IsGridSnapSizeActive(int32 GridSizeIndex) const
+{
+	const ULevelEditorViewportSettings* ViewportSettings = GetDefault<ULevelEditorViewportSettings>();
+	return ViewportSettings->CurrentPosGridSize == GridSizeIndex;
+}
+
+void UUnrealEdViewportToolbarContext::SetGridSnapSize(int32 GridSizeIndex)
+{
+	GEditor->SetGridSize(GridSizeIndex);
+}
+
+FText UUnrealEdViewportToolbarContext::GetRotationSnapLabel() const
+{
+	return FText::Format(
+		LOCTEXT("GridRotation - Number - DegreeSymbol", "{0}\u00b0"),
+		FText::AsNumber(GEditor->GetRotGridSize().Pitch)
+	);
+}
+
+bool UUnrealEdViewportToolbarContext::IsRotationSnapActive(int32 RotationIndex, ERotationGridMode RotationMode) const
+{
+	const ULevelEditorViewportSettings* ViewportSettings = GetDefault<ULevelEditorViewportSettings>();
+	return ViewportSettings->CurrentRotGridSize == RotationIndex
+		&& ViewportSettings->CurrentRotGridMode == RotationMode;
+}
+
+void UUnrealEdViewportToolbarContext::SetRotationSnapSize(int32 RotationIndex, ERotationGridMode RotationMode)
+{
+	GEditor->SetRotGridSize(RotationIndex, RotationMode);
+}
+
+FText UUnrealEdViewportToolbarContext::GetScaleSnapLabel() const
+{
+	return UE::UnrealEd::GetScaleGridLabel();
+}
+
+TArray<float> UUnrealEdViewportToolbarContext::GetScaleSnapSizes() const
+{
+	const ULevelEditorViewportSettings* ViewportSettings = GetDefault<ULevelEditorViewportSettings>();
+	return ViewportSettings->ScalingGridSizes;
+}
+
+bool UUnrealEdViewportToolbarContext::IsScaleSnapActive(int32 ScaleIndex) const
+{
+	const ULevelEditorViewportSettings* ViewportSettings = GetDefault<ULevelEditorViewportSettings>();
+	return ViewportSettings->CurrentScalingGridSize == ScaleIndex;
+}
+
+void UUnrealEdViewportToolbarContext::SetScaleSnapSize(int32 ScaleIndex)
+{
+	GEditor->SetScaleGridSize(ScaleIndex);
+}
+
+#undef LOCTEXT_NAMESPACE
diff --git a/Engine/Source/Editor/UnrealEd/Public/ViewportToolbar/UnrealEdViewportToolbar.h b/Engine/Source/Editor/UnrealEd/Public/ViewportToolbar/UnrealEdViewportToolbar.h
--- a/Engine/Source/Editor/UnrealEd/Public/ViewportToolbar/UnrealEdViewportToolbar.h
+++ b/Engine/Source/Editor/UnrealEd/Public/ViewportToolbar/UnrealEdViewportToolbar.h
@@ -90,6 +90,7 @@ DECLARE_DELEGATE_OneParam(FNumericEntryExecuteActionDelegateInt32, int32);
 DECLARE_DELEGATE_TwoParams(FOnViewportClientCamSpeedChanged, const TSharedRef<SEditorViewport>&, int32);
 DECLARE_DELEGATE_TwoParams(FOnViewportClientCamSpeedScalarChanged, const TSharedRef<SEditorViewport>&, float);
 
+// This will be deprecated in the future. Prefer using CreateSnappingSubmenu() to get snapping tools.
 UNREALED_API TSharedRef<SWidget> BuildRotationGridCheckBoxList(
 	FName InExtentionHook,
 	const FText& InHeading,
@@ -101,6 +102,7 @@ UNREALED_API TSharedRef<SWidget> BuildRotationGridCheckBoxList(
 );
 
 UNREALED_API FText GetRotationGridLabel();
+// This will be deprecated in the future. Prefer using CreateSnappingSubmenu() to get snapping tools.
 UNREALED_API TSharedRef<SWidget> CreateRotationGridSnapMenu(
 	const FRotationGridCheckboxListExecuteActionDelegate& InExecuteDelegate,
 	const FRotationGridCheckboxListIsCheckedDelegate& InIsCheckedDelegate,
@@ -129,9 +131,11 @@ struct UNREALED_API FLocationGridSnapMenuOptions
 	TSharedPtr<FUICommandList> CommandList;
 };
 
+// This will be deprecated in the future. Prefer using CreateSnappingSubmenu() to get snapping tools.
 UNREALED_API TSharedRef<SWidget> CreateLocationGridSnapMenu(const FLocationGridSnapMenuOptions& InMenuOptions = {});
 
 UNREALED_API FText GetScaleGridLabel();
+// This will be deprecated in the future. Prefer using CreateSnappingSubmenu() to get snapping tools.
 UNREALED_API TSharedRef<SWidget> CreateScaleGridSnapMenu(
 	const FScaleGridCheckboxListExecuteActionDelegate& InExecuteDelegate,
 	const FScaleGridCheckboxListIsCheckedDelegate& InIsCheckedDelegate,
@@ -142,6 +146,7 @@ UNREALED_API TSharedRef<SWidget> CreateScaleGridSnapMenu(
 	const FUIAction& PreserveNonUniformScaleUIAction = FUIAction()
 );
 
+UE_DEPRECATED(5.6, "Use the native FToolMenuEntry::InitSubmenu() functions.")
 UNREALED_API FToolMenuEntry CreateCheckboxSubmenu(
 	const FName InName,
 	const TAttribute<FText>& InLabel,
diff --git a/Engine/Source/Editor/UnrealEd/Public/ViewportToolbar/UnrealEdViewportToolbarContext.h b/Engine/Source/Editor/UnrealEd/Public/ViewportToolbar/UnrealEdViewportToolbarContext.h
--- a/Engine/Source/Editor/UnrealEd/Public/ViewportToolbar/UnrealEdViewportToolbarContext.h
+++ b/Engine/Source/Editor/UnrealEd/Public/ViewportToolbar/UnrealEdViewportToolbarContext.h
@@ -12,6 +12,7 @@
 
 class FAssetEditorToolkit;
 class IPreviewProfileController;
+enum ERotationGridMode : int;
 
 namespace UE::UnrealEd
 {
@@ -67,4 +68,20 @@ public:
 	FTabId PreviewSettingsTabId;
 	
 	virtual TSharedPtr<IPreviewProfileController> GetPreviewProfileController() const;
+	
+	bool bShowSurfaceSnap = true;
+	
+	virtual FText GetGridSnapLabel() const;
+	virtual TArray<float> GetGridSnapSizes() const;
+	virtual bool IsGridSnapSizeActive(int32 GridSizeIndex) const;
+	virtual void SetGridSnapSize(int32 GridSizeIndex);
+	
+	virtual FText GetRotationSnapLabel() const;
+	virtual bool IsRotationSnapActive(int32 RotationIndex, ERotationGridMode RotationMode) const;
+	virtual void SetRotationSnapSize(int32 RotationIndex, ERotationGridMode RotationMode);
+	
+	virtual FText GetScaleSnapLabel() const;
+	virtual TArray<float> GetScaleSnapSizes() const;
+	virtual bool IsScaleSnapActive(int32 ScaleIndex) const;
+	virtual void SetScaleSnapSize(int32 ScaleIndex);
 };
