MPCook: Start CookWorkers as early as possible to optimize MPCook to do more work in parallel.
#rnx

diff --git a/Engine/Source/Editor/UnrealEd/Classes/CookOnTheSide/CookOnTheFlyServer.h b/Engine/Source/Editor/UnrealEd/Classes/CookOnTheSide/CookOnTheFlyServer.h
--- a/Engine/Source/Editor/UnrealEd/Classes/CookOnTheSide/CookOnTheFlyServer.h
+++ b/Engine/Source/Editor/UnrealEd/Classes/CookOnTheSide/CookOnTheFlyServer.h
@@ -1611,6 +1611,12 @@ private:
 	/** Tracking for the ticking of tickable cook objects */
 	double LastCookableObjectTickTime = 0.;
 
+	/**
+	 * Number of packages requested in CookByTheBook, set to 0 in other modes. Used to inform heuristics such as
+	 * number of cookworkers.
+	 */
+	int32 InitialRequestCount = 0;
+
 	// Cook events that can be listenned to
 	static FOnCookByTheBookStarted CookByTheBookStartedEvent;
 	static FOnCookByTheBookFinished CookByTheBookFinishedEvent;
diff --git a/Engine/Source/Editor/UnrealEd/Private/CookOnTheFlyServer.cpp b/Engine/Source/Editor/UnrealEd/Private/CookOnTheFlyServer.cpp
--- a/Engine/Source/Editor/UnrealEd/Private/CookOnTheFlyServer.cpp
+++ b/Engine/Source/Editor/UnrealEd/Private/CookOnTheFlyServer.cpp
@@ -6342,6 +6342,7 @@ void UCookOnTheFlyServer::InitializeSession()
 	NumObjectsHistory.Initialize(GUObjectArray.GetObjectArrayNumMinusAvailable());
 	VirtualMemoryHistory.Initialize(FPlatformMemory::GetStats().UsedVirtual);
 	bKickedBuildDependencies = false;
+	InitialRequestCount = 0;
 }
 
 bool UCookOnTheFlyServer::Exec_Editor(class UWorld* InWorld, const TCHAR* Cmd, FOutputDevice& Ar)
@@ -10804,9 +10805,8 @@ void UCookOnTheFlyServer::StartCookByTheBook( const FCookByTheBookStartupOptions
 	// Functions in this section are not dependent upon each other and can be ordered arbitrarily or for async performance
 	BeginCookFinishShaderCodeLibrary(BeginContext);
 
-	// Initialize systems that nothing in StartCookByTheBook references
-	// Functions in this section are not dependent upon each other and can be ordered arbitrarily or for async performance
-	BeginCookDirector(BeginContext);
+	// Functions in this section can depend on functions before this section but not on each other,
+	// and can be ordered arbitrarily or for async performance.
 	BeginCookEditorSystems();
 	BeginCookEDLCookInfo(BeginContext);
 	BeginCookPackageWriters(BeginContext);
@@ -10817,6 +10817,11 @@ void UCookOnTheFlyServer::StartCookByTheBook( const FCookByTheBookStartupOptions
 	RecordDLCPackagesFromBaseGame(BeginContext);
 	RegisterCookByTheBookDelegates();
 
+	// Functions in this section can depend on functions before this section but not on each other,
+	// and can be ordered arbitrarily or for async performance.
+	BeginCookDirector(BeginContext);
+
+	// BroadcastCookStarted by contract is sent after all internal CookByTheBookStarted is complete.
 	BroadcastCookStarted();
 }
 
@@ -11494,6 +11499,7 @@ void UCookOnTheFlyServer::GenerateInitialRequests(FBeginCookContext& BeginContex
 				EMessageSeverity::Warning);
 		}
 	}
+	InitialRequestCount = WorkerRequests->GetNumExternalRequests();
 
 	const FString& CreateReleaseVersion = BeginContext.StartupOptions->CreateReleaseVersion;
 	const FString& BasedOnReleaseVersion = BeginContext.StartupOptions->BasedOnReleaseVersion;
diff --git a/Engine/Source/Editor/UnrealEd/Private/Cooker/CookDirector.cpp b/Engine/Source/Editor/UnrealEd/Private/Cooker/CookDirector.cpp
--- a/Engine/Source/Editor/UnrealEd/Private/Cooker/CookDirector.cpp
+++ b/Engine/Source/Editor/UnrealEd/Private/Cooker/CookDirector.cpp
@@ -317,6 +317,11 @@ void FCookDirector::FRunnableShunt::Stop()
 void FCookDirector::StartCook(const FBeginCookContext& InBeginContext)
 {
 	BeginCookContext.Set(InBeginContext);
+
+	// We launch the CookWorkers during StartCook, so that their startup time overlaps with
+	// work the Director has to do during the first tick, for PumpRequests and AssignRequests.
+	FScopeLock CommunicationScopeLock(&CommunicationLock);
+	InitializeWorkers();
 }
 
 void FCookDirector::AssignRequests(TArrayView<FPackageData*> Requests, TArray<FWorkerId>& OutAssignments,
@@ -326,7 +331,7 @@ void FCookDirector::AssignRequests(TArrayView<FPackageData*> Requests, TArray<FW
 	TArray<TRefCountPtr<FCookWorkerServer>> LocalRemoteWorkers;
 	{
 		FScopeLock CommunicationScopeLock(&CommunicationLock);
-		InitializeWorkers(Requests.Num());
+		InitializeWorkers();
 	}
 	LocalRemoteWorkers = CopyRemoteWorkers();;
 
@@ -662,7 +667,7 @@ void FCookDirector::BroadcastMessage(UE::CompactBinaryTCP::FMarshalledMessage&&
 
 	{
 		FScopeLock CommunicationScopeLock(&CommunicationLock);
-		InitializeWorkers(0);
+		InitializeWorkers();
 	}
 	TArray<TRefCountPtr<FCookWorkerServer>> LocalRemoteWorkers = CopyRemoteWorkers();
 	for (TRefCountPtr<FCookWorkerServer>& RemoteWorker : LocalRemoteWorkers)
@@ -1244,12 +1249,13 @@ bool FCookDirector::TryCreateWorkerConnectSocket()
 }
 
 
-void FCookDirector::InitializeWorkers(int32 InitialRequestCount)
+void FCookDirector::InitializeWorkers()
 {
 	if (bWorkersInitialized)
 	{
 		return;
 	}
+	int32 InitialRequestCount = COTFS.InitialRequestCount;
 	bWorkersInitialized = true;
 
 	check(!CommunicationThread);
diff --git a/Engine/Source/Editor/UnrealEd/Private/Cooker/CookDirector.h b/Engine/Source/Editor/UnrealEd/Private/Cooker/CookDirector.h
--- a/Engine/Source/Editor/UnrealEd/Private/Cooker/CookDirector.h
+++ b/Engine/Source/Editor/UnrealEd/Private/Cooker/CookDirector.h
@@ -188,10 +188,10 @@ private:
 	/** Initialization helper: create the listen socket. */
 	bool TryCreateWorkerConnectSocket();
 	/**
-	 * Construct CookWorkerServers and communication thread if not yet constructed.
-	 * The CookWorkerServers are constructed to Uninitialized; the worker process is created later.
+	 * Construct CookWorkerServers and communication thread if not yet constructed. The CookWorkerServers are
+	 * constructed to Uninitialized; the worker process is created asynchronously during TickCommunication.
 	 */
-	void InitializeWorkers(int32 InitialRequestCount);
+	void InitializeWorkers();
 	/**
 	 * Copy to snapshot variables the data required on the communication thread that can only be read from the
 	 * scheduler thread.
