[Deprecation] Remove deprecated functions from component visualzier

diff --git a/Engine/Plugins/ChaosVD/Source/ChaosVD/Private/Visualizers/ChaosVDSolverCollisionDataComponentVisualizer.cpp b/Engine/Plugins/ChaosVD/Source/ChaosVD/Private/Visualizers/ChaosVDSolverCollisionDataComponentVisualizer.cpp
--- a/Engine/Plugins/ChaosVD/Source/ChaosVD/Private/Visualizers/ChaosVDSolverCollisionDataComponentVisualizer.cpp
+++ b/Engine/Plugins/ChaosVD/Source/ChaosVD/Private/Visualizers/ChaosVDSolverCollisionDataComponentVisualizer.cpp
@@ -46,7 +46,7 @@ void FChaosVDSolverCollisionDataComponentVisualizer::RegisterVisualizerMenus()
 	CreateGenericVisualizerMenu<UChaosVDCollisionDataVisualizationSettings, EChaosVDCollisionVisualizationFlags>(SChaosVDViewportToolbar::ShowMenuName, MenuSection, MenuSectionLabel, FlagsMenuLabel, FlagsMenuTooltip, FlagsMenuIcon, SettingsMenuLabel, SettingsMenuTooltip);
 }
 
-bool FChaosVDSolverCollisionDataComponentVisualizer::ShowWhenSelected()
+bool FChaosVDSolverCollisionDataComponentVisualizer::ShouldShowForSelectedSubcomponents(const UActorComponent* Component)
 {
 	return false;
 }
diff --git a/Engine/Plugins/ChaosVD/Source/ChaosVD/Private/Visualizers/ChaosVDSolverCollisionDataComponentVisualizer.h b/Engine/Plugins/ChaosVD/Source/ChaosVD/Private/Visualizers/ChaosVDSolverCollisionDataComponentVisualizer.h
--- a/Engine/Plugins/ChaosVD/Source/ChaosVD/Private/Visualizers/ChaosVDSolverCollisionDataComponentVisualizer.h
+++ b/Engine/Plugins/ChaosVD/Source/ChaosVD/Private/Visualizers/ChaosVDSolverCollisionDataComponentVisualizer.h
@@ -33,7 +33,7 @@ public:
 
 	virtual void RegisterVisualizerMenus() override;
 
-	virtual bool ShowWhenSelected() override;
+	virtual bool ShouldShowForSelectedSubcomponents(const UActorComponent* Component) override;
 	virtual void DrawVisualization(const UActorComponent* Component, const FSceneView* View, FPrimitiveDrawInterface* PDI) override;
 
 	virtual bool CanHandleClick(const HChaosVDComponentVisProxy& VisProxy) override;
diff --git a/Engine/Source/Editor/UnrealEd/Private/ComponentVisualizer.cpp b/Engine/Source/Editor/UnrealEd/Private/ComponentVisualizer.cpp
--- a/Engine/Source/Editor/UnrealEd/Private/ComponentVisualizer.cpp
+++ b/Engine/Source/Editor/UnrealEd/Private/ComponentVisualizer.cpp
@@ -216,43 +216,6 @@ bool FComponentVisualizer::ShouldAutoSelectElementOnHandleClick() const
 	return ComponentVisualizers::GAutoSelectComponentWithPointSelection;
 }
 
-FPropertyNameAndIndex FComponentVisualizer::GetComponentPropertyName(const UActorComponent* Component)
-{
-	if (Component)
-	{
-		return GetActorPropertyNameAndIndexForComponent(Component->GetOwner(), const_cast<UActorComponent*>(Component));
-	}
-
-	// Didn't find actor property referencing this component
-	return FPropertyNameAndIndex();
-}
-
-UActorComponent* FComponentVisualizer::GetComponentFromPropertyName(const AActor* CompOwner, const FPropertyNameAndIndex& Property)
-{
-	UActorComponent* ResultComp = NULL;
-	if(CompOwner && Property.IsValid())
-	{
-		UClass* ActorClass = CompOwner->GetClass();
-		FProperty* Prop = FindFProperty<FProperty>(ActorClass, Property.Name);
-		if (FObjectProperty* ObjectProp = CastField<FObjectProperty>(Prop))
-		{
-			UObject* Object = ObjectProp->GetObjectPropertyValue(ObjectProp->ContainerPtrToValuePtr<void>(CompOwner, Property.Index));
-			ResultComp = Cast<UActorComponent>(Object);
-		}
-		else if (FArrayProperty* ArrayProp = CastField<FArrayProperty>(Prop))
-		{
-			if (FObjectProperty* InnerProp = CastField<FObjectProperty>(ArrayProp->Inner))
-			{
-				FScriptArrayHelper ArrayHelper(ArrayProp, ArrayProp->ContainerPtrToValuePtr<void>(CompOwner));
-				UObject* Object = InnerProp->GetObjectPropertyValue(ArrayHelper.GetRawPtr(Property.Index));
-				ResultComp = Cast<UActorComponent>(Object);
-			}
-		}
-	}
-
-	return ResultComp;
-}
-
 void FComponentVisualizer::NotifyPropertyModified(UActorComponent* Component, FProperty* Property, EPropertyChangeType::Type PropertyChangeType)
 {
 	TArray<FProperty*> Properties;
diff --git a/Engine/Source/Editor/UnrealEd/Public/ComponentVisualizer.h b/Engine/Source/Editor/UnrealEd/Public/ComponentVisualizer.h
--- a/Engine/Source/Editor/UnrealEd/Public/ComponentVisualizer.h
+++ b/Engine/Source/Editor/UnrealEd/Public/ComponentVisualizer.h
@@ -161,10 +161,7 @@ public:
 
 	/** */
 	virtual void OnRegister() {}
-	/** Only show this visualizer if the actor is selected */
-	UE_DEPRECATED(5.4, "This function is unused and will be removed in a future version. Component visualizers are only shown for the active selection. Use bDebugDraw on the specific component, or the editor setting to control drawing of subcomponents for selected actors.")
-	virtual bool ShowWhenSelected() { return true; }
-	/** Show this visualizer if the component is directly is selected */
+	/** Show this visualizer if the component is indirectly selected */
 	virtual bool ShouldShowForSelectedSubcomponents(const UActorComponent* Component) { return true; }
 	/** Draw visualization for the supplied component */
 	virtual void DrawVisualization(const UActorComponent* Component, const FSceneView* View, FPrimitiveDrawInterface* PDI) {}
@@ -210,15 +207,7 @@ public:
 	virtual bool IsVisualizingArchetype() const { return false; }
 
 	// So deprecated code expecting this as an inner class still works
-	using FPropertyNameAndIndex = ::FPropertyNameAndIndex;
-
-	/** Find the name of the property that points to this component */
-	UE_DEPRECATED(4.24, "Please use the FComponentPropertyPath class to build property name paths for components.")
-	static UNREALED_API FPropertyNameAndIndex GetComponentPropertyName(const UActorComponent* Component);
-
-	/** Get a component pointer from the property name */
-	UE_DEPRECATED(4.24, "Please use the FComponentPropertyPath::GetComponent() to retrieve a component pointer from a property name path.")
-	static UNREALED_API UActorComponent* GetComponentFromPropertyName(const AActor* CompOwner, const FPropertyNameAndIndex& Property);
+	using FPropertyNameAndIndex UE_DEPRECATED(5.6, "Please update to use the external class instead of internal identifier") = ::FPropertyNameAndIndex;
 
 	/** Notify that a component property has been modified */
 	static UNREALED_API void NotifyPropertyModified(UActorComponent* Component, FProperty* Property, EPropertyChangeType::Type PropertyChangeType = EPropertyChangeType::Unspecified);
