NetDriver
* Fix memory leak if a PIE session is stopped while the PendingNetGame is still active

diff --git a/Engine/Source/Editor/UnrealEd/Private/PlayLevel.cpp b/Engine/Source/Editor/UnrealEd/Private/PlayLevel.cpp
--- a/Engine/Source/Editor/UnrealEd/Private/PlayLevel.cpp
+++ b/Engine/Source/Editor/UnrealEd/Private/PlayLevel.cpp
@@ -340,6 +340,19 @@ void UEditorEngine::EndPlayMap()
 			{
 				TeardownPlaySession(ThisContext);
 
+				// Destroy the PendingNetGame netdriver if it never finished connecting.
+				if (ThisContext.PendingNetGame && ThisContext.PendingNetGame->NetDriver)
+				{
+					if (ThisContext.PendingNetGame->NetDriver->ServerConnection)
+					{
+						CancelPending(ThisContext);
+					}
+					else
+					{
+						DestroyNamedNetDriver(ThisContext.PendingNetGame, ThisContext.PendingNetGame->NetDriver->NetDriverName);
+					}
+				}
+
 				ShutdownWorldNetDriver(ThisContext.World());
 			}
 
diff --git a/Engine/Source/Runtime/Engine/Classes/Engine/PendingNetGame.h b/Engine/Source/Runtime/Engine/Classes/Engine/PendingNetGame.h
--- a/Engine/Source/Runtime/Engine/Classes/Engine/PendingNetGame.h
+++ b/Engine/Source/Runtime/Engine/Classes/Engine/PendingNetGame.h
@@ -135,12 +135,7 @@ public:
 	//~ Begin UObject Interface.
 	virtual void Serialize( FArchive& Ar ) override;
 
-	virtual void FinishDestroy() override
-	{
-		NetDriver = NULL;
-		
-		Super::FinishDestroy();
-	}
+	virtual void FinishDestroy() override;
 	
 	static void AddReferencedObjects(UObject* InThis, FReferenceCollector& Collector);
 	//~ End UObject Interface.
diff --git a/Engine/Source/Runtime/Engine/Private/PendingNetGame.cpp b/Engine/Source/Runtime/Engine/Private/PendingNetGame.cpp
--- a/Engine/Source/Runtime/Engine/Private/PendingNetGame.cpp
+++ b/Engine/Source/Runtime/Engine/Private/PendingNetGame.cpp
@@ -100,6 +100,15 @@ void UPendingNetGame::InitNetDriver(UNetDriver* InNetDriver)
 	}
 }
 
+void UPendingNetGame::FinishDestroy()
+{
+	ensureMsgf(NetDriver == nullptr, TEXT("PendingNetGame's NetDriver was not shutdown before destruction."));
+	
+	NetDriver = nullptr;
+
+	Super::FinishDestroy();
+}
+
 void UPendingNetGame::BeginHandshake()
 {
 	// Kick off the connection handshake
