UnrealEd/Lightmass: add DevHttp shared library to list of required dependencies

DevHttp is a new module recently added after having been split
out from DerivedDataCache, so we add it as a required dependency
for UnrealLightmass everywhere that DerivedDataCache is a required
dependency.

diff --git a/Engine/Source/Editor/UnrealEd/Private/Lightmass/Lightmass.cpp b/Engine/Source/Editor/UnrealEd/Private/Lightmass/Lightmass.cpp
--- a/Engine/Source/Editor/UnrealEd/Private/Lightmass/Lightmass.cpp
+++ b/Engine/Source/Editor/UnrealEd/Private/Lightmass/Lightmass.cpp
@@ -3009,6 +3009,7 @@ bool FLightmassProcessor::BeginRun()
 		TEXT("Binaries/Win32/UnrealLightmass-CoreUObject.dll"),
 		TEXT("Binaries/Win32/UnrealLightmass-CorePreciseFP.dll"),
 		TEXT("Binaries/Win32/UnrealLightmass-DerivedDataCache.dll"),
+		TEXT("Binaries/Win32/UnrealLightmass-DevHttp.dll"),
 		TEXT("Binaries/Win32/UnrealLightmass-Sockets.dll"),
 		TEXT("Binaries/Win32/UnrealLightmass-Zen.dll"),
 		TEXT("Binaries/Win32/UnrealLightmass-Projects.dll"),
@@ -3030,6 +3031,7 @@ bool FLightmassProcessor::BeginRun()
 		TEXT("Binaries/Win64/UnrealLightmass-CoreUObject.dll"),
 		TEXT("Binaries/Win64/UnrealLightmass-CorePreciseFP.dll"),
 		TEXT("Binaries/Win64/UnrealLightmass-DerivedDataCache.dll"),
+		TEXT("Binaries/Win64/UnrealLightmass-DevHttp.dll"),
 		TEXT("Binaries/Win64/UnrealLightmass-Sockets.dll"),
 		TEXT("Binaries/Win64/UnrealLightmass-Zen.dll"),
 		TEXT("Binaries/Win64/UnrealLightmass-Projects.dll"),
@@ -3050,6 +3052,7 @@ bool FLightmassProcessor::BeginRun()
 		TEXT("Binaries/Mac/UnrealLightmass-CoreUObject.dylib"),
 		TEXT("Binaries/Mac/UnrealLightmass-CorePreciseFP.dylib"),
 		TEXT("Binaries/Mac/UnrealLightmass-DerivedDataCache.dylib"),
+		TEXT("Binaries/Mac/UnrealLightmass-DevHttp.dylib"),
 		TEXT("Binaries/Mac/UnrealLightmass-Sockets.dylib"),
 		TEXT("Binaries/Mac/UnrealLightmass-Zen.dylib"),
 		TEXT("Binaries/Mac/UnrealLightmass-Json.dylib"),
@@ -3071,6 +3074,7 @@ bool FLightmassProcessor::BeginRun()
 		TEXT("Binaries/Linux/libUnrealLightmass-CoreUObject.so"),
 		TEXT("Binaries/Linux/libUnrealLightmass-CorePrecise.so"),
 		TEXT("Binaries/Linux/libUnrealLightmass-DerivedDataCache.so"),
+		TEXT("Binaries/Linux/libUnrealLightmass-DevHttp.so"),
 		TEXT("Binaries/Linux/libUnrealLightmass-Sockets.so"),
 		TEXT("Binaries/Linux/libUnrealLightmass-Zen.so"),
 		TEXT("Binaries/Linux/libUnrealLightmass-Json.so"),
