MPCook: Fix crash when a CookGenerationHelper is still allocated when the cooker is destroyed; PackageData shutdown should skip initialization in any GenerationHelper functions it calls.
#rnx
#incrementalcook

diff --git a/Engine/Source/Editor/UnrealEd/Private/Cooker/CookGenerationHelper.h b/Engine/Source/Editor/UnrealEd/Private/Cooker/CookGenerationHelper.h
--- a/Engine/Source/Editor/UnrealEd/Private/Cooker/CookGenerationHelper.h
+++ b/Engine/Source/Editor/UnrealEd/Private/Cooker/CookGenerationHelper.h
@@ -88,7 +88,6 @@ public:
 	bool HasCreatedPackage() const;
 	void SetHasCreatedPackage(bool bValue);
 	bool HasSavedEveryPlatform() const;
-	void SetAllPlatformsSaved(FGenerationHelper& GenerationHelper);
 	bool HasTakenOverCachedCookedPlatformData() const;
 	void SetHasTakenOverCachedCookedPlatformData(bool bValue);
 	bool HasIssuedUndeclaredMovedObjectsWarning() const;
@@ -458,6 +457,9 @@ public:
 	 */
 	void MarkPackageSavedRemotely(UCookOnTheFlyServer& COTFS, FPackageData& PackageData,
 		const ITargetPlatform* TargetPlatform, FWorkerId SourceWorkerId);
+	/** Called when a Package is demoted to idle, to update the saved count for it. Does not call Initialize. */
+	void SetAllPlatformsSaved(FPackageData& PackageData, FCookGenerationInfo* Info);
+
 	/** Scope object that blocks calls to OnAllSavesCompleted. Checks and calls it if necessary when destructed. */
 	struct FScopeDeferEvents
 	{
@@ -650,17 +652,6 @@ inline bool FCookGenerationInfo::HasSavedEveryPlatform() const
 	return bHasSaved;
 }
 
-inline void FCookGenerationInfo::SetAllPlatformsSaved(FGenerationHelper& GenerationHelper)
-{
-	for (TPair<const ITargetPlatform*, FCookGenerationInfo::FPlatformData>& PlatformPair : PlatformDatas)
-	{
-		const ITargetPlatform* TargetPlatform = PlatformPair.Key;
-		FPlatformData& InfoPlatformData = PlatformPair.Value;
-		InfoPlatformData.SetHasSaved(GenerationHelper, *this, TargetPlatform,
-			true /* bValue */, FWorkerId::Local());
-	}
-}
-
 inline bool FCookGenerationInfo::HasTakenOverCachedCookedPlatformData() const
 {
 	return bTakenOverCachedCookedPlatformData;
@@ -1072,6 +1063,24 @@ inline int32& FGenerationHelper::GetMPCookNextAssignmentIndex()
 	return MPCookNextAssignmentIndex;
 }
 
+inline void FGenerationHelper::SetAllPlatformsSaved(FPackageData& PackageData, FCookGenerationInfo* Info)
+{
+	if (!Info)
+	{
+		Info = FindInfoNoInitialize(PackageData);
+	}
+	if (Info)
+	{
+		for (TPair<const ITargetPlatform*, FCookGenerationInfo::FPlatformData>& PlatformPair : Info->PlatformDatas)
+		{
+			const ITargetPlatform* TargetPlatform = PlatformPair.Key;
+			FCookGenerationInfo::FPlatformData& InfoPlatformData = PlatformPair.Value;
+			InfoPlatformData.SetHasSaved(*this, *Info, TargetPlatform,
+				true /* bValue */, FWorkerId::Local());
+		}
+	}
+}
+
 inline FGenerationHelper::FDirectorAPI FGenerationHelper::GetDirectorAPI() const
 {
 	return FDirectorAPI(const_cast<FGenerationHelper&>(*this));
diff --git a/Engine/Source/Editor/UnrealEd/Private/Cooker/CookPackageData.cpp b/Engine/Source/Editor/UnrealEd/Private/Cooker/CookPackageData.cpp
--- a/Engine/Source/Editor/UnrealEd/Private/Cooker/CookPackageData.cpp
+++ b/Engine/Source/Editor/UnrealEd/Private/Cooker/CookPackageData.cpp
@@ -1603,14 +1603,7 @@ void FPackageData::SetParentGenerationHelper(FGenerationHelper* InGenerationHelp
 	if (ParentGenerationHelper && !InGenerationHelper && IsTerminalStateChange(StateChangeReason))
 	{
 		// The package's progress is completed and we will not come back to it; report the package was saved.
-		if (!InfoOfPackageInGenerator)
-		{
-			InfoOfPackageInGenerator = ParentGenerationHelper->FindInfo(*this);
-		}
-		if (InfoOfPackageInGenerator)
-		{
-			InfoOfPackageInGenerator->SetAllPlatformsSaved(*ParentGenerationHelper);
-		}
+		ParentGenerationHelper->SetAllPlatformsSaved(*this, InfoOfPackageInGenerator);
 	}
 	ParentGenerationHelper = InGenerationHelper;
 }
