Move virtual texture pending mips to a new virtual texture submenu.
Other virtual texture debug views will be added there.

diff --git a/Engine/Source/Editor/EditorStyle/Private/StarshipStyle.cpp b/Engine/Source/Editor/EditorStyle/Private/StarshipStyle.cpp
--- a/Engine/Source/Editor/EditorStyle/Private/StarshipStyle.cpp
+++ b/Engine/Source/Editor/EditorStyle/Private/StarshipStyle.cpp
@@ -4786,7 +4786,6 @@ void FStarshipEditorStyle::FStyle::SetupLevelEditorStyle()
 		Set( "EditorViewport.TexStreamAccMeshUVDensityMode", new IMAGE_BRUSH("Icons/icon_ViewMode_TextureStreamingAccuracy_16px", Icon16x16));
 		Set( "EditorViewport.TexStreamAccMaterialTextureScaleMode", new IMAGE_BRUSH( "Icons/icon_ViewMode_TextureStreamingAccuracy_16px", Icon16x16 ) );
 		Set( "EditorViewport.RequiredTextureResolutionMode", new IMAGE_BRUSH( "Icons/icon_ViewMode_TextureStreamingAccuracy_16px", Icon16x16 ) );
-		Set( "EditorViewport.VirtualTexturePendingMipsMode", new IMAGE_BRUSH( "Icons/icon_ViewMode_TextureStreamingAccuracy_16px", Icon16x16 ) );
 		Set( "EditorViewport.StationaryLightOverlapMode", new IMAGE_BRUSH( "Icons/icon_ViewMode_StationaryLightOverlap_16px", Icon16x16 ) );
 		Set( "EditorViewport.LightmapDensityMode", new IMAGE_BRUSH( "Icons/icon_ViewMode_LightmapDensity_16px", Icon16x16 ) );
 
@@ -4809,6 +4808,7 @@ void FStarshipEditorStyle::FStyle::SetupLevelEditorStyle()
 		Set( "EditorViewport.VisualizeGroomMode",     new IMAGE_BRUSH("Common/icon_ShowHair_16x", Icon16x16 ) );
 		Set( "EditorViewport.VisualizeVirtualShadowMapMode", new IMAGE_BRUSH_SVG("Starship/Common/VirtualShadowMap_16", Icon16x16 ) );
 		Set( "EditorViewport.VisualizeActorColorationMode", new IMAGE_BRUSH_SVG("Starship/Common/ActorColoration_16", Icon16x16));
+		Set( "EditorViewport.VisualizeVirtualTextureMode", new IMAGE_BRUSH("Icons/AssetIcons/Texture2D_16x", Icon16x16));
 		Set( "EditorViewport.CollisionPawn",          new IMAGE_BRUSH_SVG("Starship/Common/PlayerCollision", Icon16x16 ) );
 		Set( "EditorViewport.CollisionVisibility",    new IMAGE_BRUSH_SVG("Starship/Common/VisibilityCollision", Icon16x16 ) );
 
diff --git a/Engine/Source/Editor/UnrealEd/Private/ViewportToolbar/UnrealEdViewportToolbar.cpp b/Engine/Source/Editor/UnrealEd/Private/ViewportToolbar/UnrealEdViewportToolbar.cpp
--- a/Engine/Source/Editor/UnrealEd/Private/ViewportToolbar/UnrealEdViewportToolbar.cpp
+++ b/Engine/Source/Editor/UnrealEd/Private/ViewportToolbar/UnrealEdViewportToolbar.cpp
@@ -1631,13 +1631,6 @@ void PopulateViewModesMenu(UToolMenu* InMenu)
 								IsViewModeSupported, Section, BaseViewportActions.RequiredTextureResolutionMode, VMI_RequiredTextureResolution
 							);
 						}
-
-						if (AllowDebugViewShaderMode(DVSM_RequiredTextureResolution, GMaxRHIShaderPlatform, FeatureLevel))
-						{
-							UE::UnrealEd::Private::ViewModesSubmenu::AddModeIfSupported(
-								IsViewModeSupported, Section, BaseViewportActions.VirtualTexturePendingMipsMode, VMI_VirtualTexturePendingMips
-							);
-						}
 					}
 				}
 
@@ -1657,8 +1650,7 @@ void PopulateViewModesMenu(UToolMenu* InMenu)
 						|| InIsViewModeSupported.Execute(VMI_PrimitiveDistanceAccuracy)
 						|| InIsViewModeSupported.Execute(VMI_MeshUVDensityAccuracy)
 						|| InIsViewModeSupported.Execute(VMI_MaterialTextureScaleAccuracy)
-						|| InIsViewModeSupported.Execute(VMI_RequiredTextureResolution)
-						|| InIsViewModeSupported.Execute(VMI_VirtualTexturePendingMips);
+						|| InIsViewModeSupported.Execute(VMI_RequiredTextureResolution);
 				}
 			};
 
@@ -1688,7 +1680,7 @@ void PopulateViewModesMenu(UToolMenu* InMenu)
 									// Texture Streaming Accuracy
 									|| ViewMode == VMI_PrimitiveDistanceAccuracy || ViewMode == VMI_MeshUVDensityAccuracy
 									|| ViewMode == VMI_MaterialTextureScaleAccuracy
-									|| ViewMode == VMI_RequiredTextureResolution || ViewMode == VMI_VirtualTexturePendingMips
+									|| ViewMode == VMI_RequiredTextureResolution
 								);
 							}
 						)
@@ -1803,6 +1795,44 @@ void PopulateViewModesMenu(UToolMenu* InMenu)
 				FSlateIcon(FAppStyle::GetAppStyleSetName(), "EditorViewport.VisualizeGPUSkinCacheMode")
 			);
 		}
+
+		if (UseVirtualTexturing(GMaxRHIShaderPlatform))
+		{
+			struct Local
+			{
+				static void BuildVirtualTextureMenu(UToolMenu* Menu, IsViewModeSupportedDelegate IsViewModeSupported)
+				{
+					FToolMenuSection& Section = Menu->AddSection("VirtualTextureVisualizationSection", LOCTEXT("VirtualTexture_VisualizationSection", "Virtual Texture Visualization"));
+
+					UE::UnrealEd::Private::ViewModesSubmenu::AddModeIfSupported(
+						IsViewModeSupported, Section, FEditorViewportCommands::Get().VirtualTexturePendingMipsMode, VMI_VirtualTexturePendingMips
+					);
+				}
+			};
+
+			Section.AddSubMenu(
+				"VirtualTextureSubMenu",
+				LOCTEXT("VirtualTexture_SubMenu", "Virtual Texture"),
+				LOCTEXT("VirtualTexure_ToolTip", "Select virtual texture visualization view modes"),
+				FNewToolMenuDelegate::CreateStatic(&Local::BuildVirtualTextureMenu, IsViewModeSupported),
+				FUIAction(
+					FExecuteAction(),
+					FCanExecuteAction(),
+					FIsActionChecked::CreateLambda(
+						[WeakViewport = EditorViewport.ToWeakPtr()]()
+						{
+							const TSharedRef<SEditorViewport> ViewportRef = WeakViewport.Pin().ToSharedRef();
+							const TSharedPtr<FEditorViewportClient> ViewportClient = ViewportRef->GetViewportClient();
+							check(ViewportClient.IsValid());
+							return ViewportClient->IsViewModeEnabled(VMI_VirtualTexturePendingMips);
+						}
+					)
+				),
+				EUserInterfaceActionType::RadioButton,
+				false,
+				FSlateIcon(FAppStyle::GetAppStyleSetName(), "EditorViewport.VisualizeVirtualTextureMode")
+			);
+		}
 	}
 }
 
diff --git a/Engine/Source/Runtime/Engine/Private/ShowFlags.cpp b/Engine/Source/Runtime/Engine/Private/ShowFlags.cpp
--- a/Engine/Source/Runtime/Engine/Private/ShowFlags.cpp
+++ b/Engine/Source/Runtime/Engine/Private/ShowFlags.cpp
@@ -782,7 +782,7 @@ EViewModeIndex FindViewMode(const FEngineShowFlags& EngineShowFlags)
 	}
 	else if (EngineShowFlags.VirtualTexturePendingMips)
 	{
-		return VMI_RequiredTextureResolution;
+		return VMI_VirtualTexturePendingMips;
 	}
 	else if (EngineShowFlags.ShaderComplexity)
 	{
diff --git a/Engine/Source/Runtime/Engine/Private/ViewModeNames.cpp b/Engine/Source/Runtime/Engine/Private/ViewModeNames.cpp
--- a/Engine/Source/Runtime/Engine/Private/ViewModeNames.cpp
+++ b/Engine/Source/Runtime/Engine/Private/ViewModeNames.cpp
@@ -170,7 +170,7 @@ TArray<FText> FillViewModeDisplayNames()
 
 		else if (ViewModeIndex == VMI_VirtualTexturePendingMips)
 		{
-			ViewModeDisplayNames.Emplace(LOCTEXT("UViewModeUtils_VMI_VirtualTexturePendingMips", "Virtual Texture Pending Mips"));
+			ViewModeDisplayNames.Emplace(LOCTEXT("UViewModeUtils_VMI_VirtualTexturePendingMips", "Pending Mips"));
 		}
 
 		else if (ViewModeIndex == VMI_VisualizeLumen)
@@ -250,7 +250,7 @@ FName GetViewModeIconName(const EViewModeIndex InViewMode)
 	static FName LWCComplexityIcon("EditorViewport.LWCComplexityMode");
 	static FName PathTracingIcon("EditorViewport.PathTracingMode");
 	static FName RayTracingDebugIcon("EditorViewport.RayTracingDebugMode");
-	static FName VirtualTexturePendingMipsIcon("EditorViewport.VirtualTexturePendingMipsMode");
+	static FName VisualizeVirtualTextureIcon("EditorViewport.VisualizeVirtualTextureMode");
 	static FName VisualizeActorColorationIcon("EditorViewport.VisualizeActorColorationMode");
 
 	FName Icon = NAME_None;
@@ -331,7 +331,7 @@ FName GetViewModeIconName(const EViewModeIndex InViewMode)
 		Icon = VisualizeNaniteIcon;
 		break;
 	case VMI_VirtualTexturePendingMips:
-		Icon = VirtualTexturePendingMipsIcon;
+		Icon = VisualizeVirtualTextureIcon;
 		break;
 	case VMI_VisualizeLumen:
 		Icon = VisualizeLumenIcon;
