Allow the CookShaderCommandlet to write the ShaderSymbols.info file when r.Shaders.SymbolsInfo = 1 in the ini file.

diff --git a/Engine/Source/Editor/UnrealEd/Private/Commandlets/CookShadersCommandlet.cpp b/Engine/Source/Editor/UnrealEd/Private/Commandlets/CookShadersCommandlet.cpp
--- a/Engine/Source/Editor/UnrealEd/Private/Commandlets/CookShadersCommandlet.cpp
+++ b/Engine/Source/Editor/UnrealEd/Private/Commandlets/CookShadersCommandlet.cpp
@@ -9,11 +9,14 @@
 #include "HAL/FileManager.h"
 #include "Interfaces/ITargetPlatform.h"
 #include "Interfaces/ITargetPlatformManagerModule.h"
+#include "Interfaces/IShaderFormat.h"
+#include "Interfaces/IShaderFormatModule.h"
 #include "MaterialShared.h"
 #include "Materials/Material.h"
 #include "Materials/MaterialInstance.h"
 #include "SceneTypes.h"
 #include "ShaderCompiler.h"
+
 #include "UObject/UObjectIterator.h"
 
 DEFINE_LOG_CATEGORY_STATIC(LogCookShadersCommandlet, Log, All);
@@ -27,6 +30,7 @@ static const TCHAR* NiagaraName = TEXT("Niagara");
 // UnrealEditor-Cmd.exe <proj> -run=CookShaders -targetPlatform=<platform> -ShaderSymbolsExport=D:\ShaderSymbols\Out -material=M_UI_Base_BordersAndButtons
 //
 // Use -dpcvars="r.Shaders.Symbols=1" to force on symbols writing from the commandline, you can also edit the appropriate [Platform]Engine.ini and uncomment or add "r.Shaders.Symbols=1", especially if you want symbols enabled longer term for that specific platform.
+// To produce a new ShaderSymbols.info file, edit the cvar.Shaders.SymbolsInfo = 1 in the [Platform]Engine.ini.
 
 namespace CookShadersCommandlet {
 
@@ -482,6 +486,12 @@ int32 UCookShadersCommandlet::Main(const FString& Params)
 				Arguments.ShadersToRecompile.Empty();
 				RecompileShadersForRemote(Arguments, OutputDir);
 			}
+
+			const IShaderFormat* ShaderFormat = GetTargetPlatformManagerRef().FindShaderFormat(DesiredShaderFormats[FormatIndex]);
+			if(ShaderFormat)
+			{ 
+				ShaderFormat->NotifyShaderCompilersShutdown(DesiredShaderFormats[FormatIndex]);
+			}
 		}
 	}
 
