Allow SDlgPickPath to show the view options by default. This matches the behavior from SDlgPickAssetPath. Without this, it's impossible to pick a path from a plugin folder

diff --git a/Engine/Source/Editor/UnrealEd/Private/Dialogs/DlgPickAssetPath.cpp b/Engine/Source/Editor/UnrealEd/Private/Dialogs/DlgPickAssetPath.cpp
--- a/Engine/Source/Editor/UnrealEd/Private/Dialogs/DlgPickAssetPath.cpp
+++ b/Engine/Source/Editor/UnrealEd/Private/Dialogs/DlgPickAssetPath.cpp
@@ -21,16 +21,15 @@ void SDlgPickAssetPath::Construct(const FArguments& InArgs)
 {
 	AssetPath = FText::FromString(FPackageName::GetLongPackagePath(InArgs._DefaultAssetPath.ToString()));
 	AssetName = FText::FromString(FPackageName::GetLongPackageAssetName(InArgs._DefaultAssetPath.ToString()));
-	bAllowReadOnlyFolders = InArgs._AllowReadOnlyFolders;
 
 	FPathPickerConfig PathPickerConfig;
 	PathPickerConfig.DefaultPath = AssetPath.ToString();
 	PathPickerConfig.OnPathSelected = FOnPathSelected::CreateSP(this, &SDlgPickAssetPath::OnPathChange);
 	PathPickerConfig.bAddDefaultPath = true;
-	PathPickerConfig.bAllowReadOnlyFolders = bAllowReadOnlyFolders;
 	PathPickerConfig.bAllowClassesFolder = true;
-	PathPickerConfig.bCanShowDevelopersFolder = true;
-	PathPickerConfig.bShowViewOptions = true;
+	PathPickerConfig.bAllowReadOnlyFolders = InArgs._AllowReadOnlyFolders;
+	PathPickerConfig.bCanShowDevelopersFolder = InArgs._CanShowDevelopersFolder;
+	PathPickerConfig.bShowViewOptions = InArgs._ShowViewOptions;
 
 	FContentBrowserModule& ContentBrowserModule = FModuleManager::LoadModuleChecked<FContentBrowserModule>("ContentBrowser");
 
diff --git a/Engine/Source/Editor/UnrealEd/Private/Dialogs/DlgPickPath.cpp b/Engine/Source/Editor/UnrealEd/Private/Dialogs/DlgPickPath.cpp
--- a/Engine/Source/Editor/UnrealEd/Private/Dialogs/DlgPickPath.cpp
+++ b/Engine/Source/Editor/UnrealEd/Private/Dialogs/DlgPickPath.cpp
@@ -17,13 +17,14 @@
 void SDlgPickPath::Construct(const FArguments& InArgs)
 {
 	Path = FText::FromString(FPackageName::GetLongPackagePath(InArgs._DefaultPath.ToString()));
-	bAllowReadOnlyFolders = InArgs._AllowReadOnlyFolders;
 
 	FPathPickerConfig PathPickerConfig;
 	PathPickerConfig.DefaultPath = Path.ToString();
 	PathPickerConfig.OnPathSelected = FOnPathSelected::CreateSP(this, &SDlgPickPath::OnPathChange);
 	PathPickerConfig.bAddDefaultPath = true;
-	PathPickerConfig.bAllowReadOnlyFolders = bAllowReadOnlyFolders;
+	PathPickerConfig.bAllowReadOnlyFolders = InArgs._AllowReadOnlyFolders;
+	PathPickerConfig.bCanShowDevelopersFolder = InArgs._CanShowDevelopersFolder;
+	PathPickerConfig.bShowViewOptions = InArgs._ShowViewOptions;
 
 	FContentBrowserModule& ContentBrowserModule = FModuleManager::LoadModuleChecked<FContentBrowserModule>("ContentBrowser");
 
diff --git a/Engine/Source/Editor/UnrealEd/Public/Dialogs/DlgPickAssetPath.h b/Engine/Source/Editor/UnrealEd/Public/Dialogs/DlgPickAssetPath.h
--- a/Engine/Source/Editor/UnrealEd/Public/Dialogs/DlgPickAssetPath.h
+++ b/Engine/Source/Editor/UnrealEd/Public/Dialogs/DlgPickAssetPath.h
@@ -14,6 +14,8 @@ class SDlgPickAssetPath : public SWindow
 public:
 	SLATE_BEGIN_ARGS(SDlgPickAssetPath)
 		: _AllowReadOnlyFolders(true)
+		, _CanShowDevelopersFolder(true)
+		, _ShowViewOptions(true)
 		{}
 
 		SLATE_ARGUMENT(FText, Title)
@@ -21,6 +23,8 @@ public:
 
 		/** If true, read only folders will be displayed */
 		SLATE_ARGUMENT(bool, AllowReadOnlyFolders)
+		SLATE_ARGUMENT(bool, CanShowDevelopersFolder)
+		SLATE_ARGUMENT(bool, ShowViewOptions)
 
 	SLATE_END_ARGS()
 
@@ -53,9 +57,6 @@ protected:
 	EAppReturnType::Type UserResponse;
 	FText AssetPath;
 	FText AssetName;
-
-	/** If true, read only folders will be displayed */
-	bool bAllowReadOnlyFolders;
 };
 
 #undef LOCTEXT_NAMESPACE
diff --git a/Engine/Source/Editor/UnrealEd/Public/Dialogs/DlgPickPath.h b/Engine/Source/Editor/UnrealEd/Public/Dialogs/DlgPickPath.h
--- a/Engine/Source/Editor/UnrealEd/Public/Dialogs/DlgPickPath.h
+++ b/Engine/Source/Editor/UnrealEd/Public/Dialogs/DlgPickPath.h
@@ -14,6 +14,8 @@ class SDlgPickPath : public SWindow
 public:
 	SLATE_BEGIN_ARGS(SDlgPickPath)
 		: _AllowReadOnlyFolders(true)
+		, _CanShowDevelopersFolder(true)
+		, _ShowViewOptions(true)
 		{}
 	
 		SLATE_ARGUMENT(FText, Title)
@@ -21,6 +23,8 @@ public:
 
 		/** If true, read only folders will be displayed */
 		SLATE_ARGUMENT(bool, AllowReadOnlyFolders)
+		SLATE_ARGUMENT(bool, CanShowDevelopersFolder)
+		SLATE_ARGUMENT(bool, ShowViewOptions)
 
 	SLATE_END_ARGS()
 
@@ -45,9 +49,6 @@ protected:
 
 	EAppReturnType::Type UserResponse;
 	FText Path;
-
-	/** If true, read only folders will be displayed */
-	bool bAllowReadOnlyFolders;
 };
 
 #undef LOCTEXT_NAMESPACE
