Aspect Ratio added to Android Json Preview

diff --git a/Engine/Plugins/Runtime/AndroidDeviceProfileSelector/Source/AndroidDeviceProfileSelector/Private/AndroidDeviceProfileSelectorModule.cpp b/Engine/Plugins/Runtime/AndroidDeviceProfileSelector/Source/AndroidDeviceProfileSelector/Private/AndroidDeviceProfileSelectorModule.cpp
--- a/Engine/Plugins/Runtime/AndroidDeviceProfileSelector/Source/AndroidDeviceProfileSelector/Private/AndroidDeviceProfileSelectorModule.cpp
+++ b/Engine/Plugins/Runtime/AndroidDeviceProfileSelector/Source/AndroidDeviceProfileSelector/Private/AndroidDeviceProfileSelectorModule.cpp
@@ -14,6 +14,7 @@
 #include "Misc/FileHelper.h"
 
 #endif
+#include "Misc/ConfigCacheIni.h"
 
 IMPLEMENT_MODULE(FAndroidDeviceProfileSelectorModule, AndroidDeviceProfileSelector);
 
@@ -110,6 +111,23 @@ void FAndroidDeviceProfileSelectorModule::GetDeviceParametersFromJson(FString& J
 		OutDeviceParameters.Add(FName(TEXT("SRC_TotalPhysicalGB")), AndroidProperties.TotalPhysicalGB);
 		OutDeviceParameters.Add(FName(TEXT("SRC_HMDSystemName")), TEXT(""));
 		OutDeviceParameters.Add(FName(TEXT("SRC_SM5Available")), AndroidProperties.SM5Available ? "true" : "false");
+
+		FString Orientation;
+		bool bNeedPortrait = false;
+		GConfig->GetString(TEXT("/Script/AndroidRuntimeSettings.AndroidRuntimeSettings"), TEXT("Orientation"), Orientation, GEngineIni);
+		if (Orientation.ToLower().Equals("portrait") || Orientation.ToLower().Equals("reverseportrait") || Orientation.ToLower().Equals("sensorportrait"))
+		{
+			bNeedPortrait = true;
+		}
+
+		if (bNeedPortrait)
+		{
+			ConstrainedAspectRatio = float(DeviceSpecs.ResolutionX) / float(DeviceSpecs.ResolutionY);
+		}
+		else
+		{
+			ConstrainedAspectRatio = float(DeviceSpecs.ResolutionY) / float(DeviceSpecs.ResolutionX);
+		}
 	}
 }
 #endif
diff --git a/Engine/Source/Editor/LevelEditor/Private/LevelEditorActions.cpp b/Engine/Source/Editor/LevelEditor/Private/LevelEditorActions.cpp
--- a/Engine/Source/Editor/LevelEditor/Private/LevelEditorActions.cpp
+++ b/Engine/Source/Editor/LevelEditor/Private/LevelEditorActions.cpp
@@ -588,6 +588,7 @@ void FLevelEditorActionCallbacks::PreviewJson_Clicked(FName PlatformName, FName
 				TMap<FName, FString> DeviceParameters;
 				DPSelectorModule->GetDeviceParametersFromJson(JsonFile, DeviceParameters);
 				DPSelectorModule->SetSelectorProperties(DeviceParameters);
+				float ConstrainedAspectRatio = DPSelectorModule->GetConstrainedAspectRatio();
 
 				DevileProfileName = DPSelectorModule->GetDeviceProfileName();
 				
@@ -599,7 +600,7 @@ void FLevelEditorActionCallbacks::PreviewJson_Clicked(FName PlatformName, FName
 					auto GetPreviewFeatureLevelInfo = [&]()
 						{
 							const ERHIFeatureLevel::Type FeatureLevel = GetMaxSupportedFeatureLevel(ShaderPlatform);
-							return FPreviewPlatformInfo(FeatureLevel, ShaderPlatform, PlatformName, FDataDrivenShaderPlatformInfo::GetShaderFormat(ShaderPlatform), DevileProfile->GetFName(), true, PreviewShaderPlatformName, FText::FromName(DevileProfile->GetFName()));
+							return FPreviewPlatformInfo(FeatureLevel, ShaderPlatform, PlatformName, FDataDrivenShaderPlatformInfo::GetShaderFormat(ShaderPlatform), DevileProfile->GetFName(), true, PreviewShaderPlatformName, FText::FromName(DevileProfile->GetFName()), ConstrainedAspectRatio);
 						};
 
 					FPreviewPlatformInfo PreviewFeatureLevelInfo = GetPreviewFeatureLevelInfo();
diff --git a/Engine/Source/Editor/UnrealEd/Classes/Editor/EditorEngine.h b/Engine/Source/Editor/UnrealEd/Classes/Editor/EditorEngine.h
--- a/Engine/Source/Editor/UnrealEd/Classes/Editor/EditorEngine.h
+++ b/Engine/Source/Editor/UnrealEd/Classes/Editor/EditorEngine.h
@@ -218,11 +218,12 @@ struct FPreviewPlatformInfo
 	,	bPreviewFeatureLevelActive(false)
 	,	PreviewShaderPlatformName(NAME_None)
 	,	PreviewShaderPlatformFriendlyName()
+	,	ConstrainedAspectRatio(0.0f)
 	{
 		InternalSetFriendlyName();
 	}
 
-	FPreviewPlatformInfo(ERHIFeatureLevel::Type InFeatureLevel, EShaderPlatform InShaderPlatform = EShaderPlatform::SP_NumPlatforms, FName InPreviewPlatformName = NAME_None, FName InPreviewShaderFormatName = NAME_None, FName InDeviceProfileName = NAME_None, bool InbPreviewFeatureLevelActive = false, FName InShaderPlatformName = NAME_None, FText InPreviewShaderPlatformFriendlyName = FText())
+	FPreviewPlatformInfo(ERHIFeatureLevel::Type InFeatureLevel, EShaderPlatform InShaderPlatform = EShaderPlatform::SP_NumPlatforms, FName InPreviewPlatformName = NAME_None, FName InPreviewShaderFormatName = NAME_None, FName InDeviceProfileName = NAME_None, bool InbPreviewFeatureLevelActive = false, FName InShaderPlatformName = NAME_None, FText InPreviewShaderPlatformFriendlyName = FText(), float InConstrainedAspectRatio = 0.0f)
 	:	PreviewFeatureLevel(InFeatureLevel)
 	,	ShaderPlatform(InShaderPlatform)
 	,	PreviewPlatformName(InPreviewPlatformName)
@@ -231,6 +232,7 @@ struct FPreviewPlatformInfo
 	,	bPreviewFeatureLevelActive(InbPreviewFeatureLevelActive)
 	,	PreviewShaderPlatformName(InShaderPlatformName)
 	,	PreviewShaderPlatformFriendlyName(InPreviewShaderPlatformFriendlyName)
+	,	ConstrainedAspectRatio(InConstrainedAspectRatio)
 	{
 		InternalSetFriendlyName();
 	}
@@ -268,6 +270,12 @@ struct FPreviewPlatformInfo
 		return bPreviewFeatureLevelActive ? PreviewPlatformName : NAME_None;
 	}
 
+	/** Return the Constrained Aspect Ratio of the preview platform */
+	float GetConstrainedAspectRatio()
+	{
+		return ConstrainedAspectRatio;
+	}
+
 	/** returns the preview feature level if active, or GMaxRHIFeatureLevel otherwise */
 	UNREALED_API ERHIFeatureLevel::Type GetEffectivePreviewFeatureLevel() const;
 	
@@ -280,6 +288,9 @@ private:
 	/**Cache Preview ShaderPlatform Name*/
 	FText PreviewShaderPlatformFriendlyName;
 	
+	/** Constrained Aspect Ratio of the preview platform */
+	float ConstrainedAspectRatio;
+
 	UNREALED_API void InternalSetFriendlyName();
 };
 
diff --git a/Engine/Source/Editor/UnrealEd/Private/EditorViewportClient.cpp b/Engine/Source/Editor/UnrealEd/Private/EditorViewportClient.cpp
--- a/Engine/Source/Editor/UnrealEd/Private/EditorViewportClient.cpp
+++ b/Engine/Source/Editor/UnrealEd/Private/EditorViewportClient.cpp
@@ -1021,7 +1021,16 @@ FSceneView* FEditorViewportClient::CalcSceneView(FSceneViewFamily* ViewFamily, c
 	// no matter how we are drawn (forced or otherwise), reset our time here
 	TimeForForceRedraw = 0.0;
 
-	const bool bConstrainAspectRatio = bUseControllingActorViewInfo && ControllingActorViewInfo.bConstrainAspectRatio;
+	bool bConstrainAspectRatio = bUseControllingActorViewInfo && ControllingActorViewInfo.bConstrainAspectRatio;
+	if (GEditor->IsFeatureLevelPreviewActive())
+	{
+		if (GEditor->PreviewPlatform.GetConstrainedAspectRatio() != 0.0f)
+		{
+			AspectRatio = GEditor->PreviewPlatform.GetConstrainedAspectRatio();
+			bConstrainAspectRatio = true;
+		}
+	}
+
 	EAspectRatioAxisConstraint AspectRatioAxisConstraint = GetDefault<ULevelEditorViewportSettings>()->AspectRatioAxisConstraint;
 	if (bUseControllingActorViewInfo && ControllingActorAspectRatioAxisConstraint.IsSet())
 	{
diff --git a/Engine/Source/Runtime/Engine/Private/LocalPlayer.cpp b/Engine/Source/Runtime/Engine/Private/LocalPlayer.cpp
--- a/Engine/Source/Runtime/Engine/Private/LocalPlayer.cpp
+++ b/Engine/Source/Runtime/Engine/Private/LocalPlayer.cpp
@@ -40,6 +40,10 @@
 #include "Engine/DebugCameraController.h"
 #endif
 
+#if WITH_EDITOR
+#include "Editor.h"
+#endif
+
 DEFINE_LOG_CATEGORY(LogPlayerManagement);
 
 #if !UE_BUILD_SHIPPING
@@ -1219,6 +1223,17 @@ bool ULocalPlayer::GetProjectionData(FViewport* Viewport, FSceneViewProjectionDa
 	FMinimalViewInfo ViewInfo;
 	GetViewPoint(/*out*/ ViewInfo);
 
+#if WITH_EDITOR
+	if (GEditor != nullptr && GEditor->IsFeatureLevelPreviewActive())
+	{
+		if (GEditor->PreviewPlatform.GetConstrainedAspectRatio() != 0.0f)
+		{
+			ViewInfo.AspectRatio = GEditor->PreviewPlatform.GetConstrainedAspectRatio();
+			ViewInfo.bConstrainAspectRatio = true;
+		}
+	}
+#endif
+
 	// If stereo rendering is enabled, update the size and offset appropriately for this pass
 	const bool bNeedStereo = StereoViewIndex != INDEX_NONE && GEngine->IsStereoscopic3D();
 	const bool bIsHeadTrackingAllowed =
diff --git a/Engine/Source/Runtime/Engine/Public/IDeviceProfileSelectorModule.h b/Engine/Source/Runtime/Engine/Public/IDeviceProfileSelectorModule.h
--- a/Engine/Source/Runtime/Engine/Public/IDeviceProfileSelectorModule.h
+++ b/Engine/Source/Runtime/Engine/Public/IDeviceProfileSelectorModule.h
@@ -19,7 +19,8 @@ class IDeviceProfileSelectorModule
 	: public IModuleInterface
 {
 public:
-
+	IDeviceProfileSelectorModule() :ConstrainedAspectRatio(0.0f) {}
+	
 	/**
 	 * Run the logic to choose an appropriate device profile for this session
 	 *
@@ -47,6 +48,12 @@ public:
 	 */
 	virtual void GetDeviceParametersFromJson(FString& JsonLocation, TMap<FName, FString>& OutDeviceParameters) {}
 	
+	/**
+	 * Get the ConstrainedAspectRatio of the Device
+	 *  @return ConstrinedAspectRatio of the Device
+	 */
+	virtual float GetConstrainedAspectRatio() { return ConstrainedAspectRatio; }
+	
 	/**
 	 * Can we get device parameters from Json
 	 *  @return Wheter we can get device parameters from Json
@@ -81,4 +88,7 @@ public:
 	virtual ~IDeviceProfileSelectorModule()
 	{
 	}
+
+protected:
+	float ConstrainedAspectRatio;
 };
