[Materials] Extend DumpMaterialExpressions commandlet to output column with use count of each material expression.

#rnx

diff --git a/Engine/Source/Editor/UnrealEd/Private/Commandlets/DumpMaterialExpressionsCommandlet.cpp b/Engine/Source/Editor/UnrealEd/Private/Commandlets/DumpMaterialExpressionsCommandlet.cpp
--- a/Engine/Source/Editor/UnrealEd/Private/Commandlets/DumpMaterialExpressionsCommandlet.cpp
+++ b/Engine/Source/Editor/UnrealEd/Private/Commandlets/DumpMaterialExpressionsCommandlet.cpp
@@ -25,6 +25,7 @@
 #include "Materials/MaterialExpressionParameter.h"
 #include "Materials/MaterialExpressionPinBase.h"
 #include "Materials/MaterialFunction.h"
+#include "Materials/MaterialInstance.h"
 #include "Math/UnrealMathSSE.h"
 #include "Misc/AssertionMacros.h"
 #include "Misc/CString.h"
@@ -78,7 +79,6 @@ void WriteOutMaterialExpressions(FDiagnosticTableWriterCSV& CsvFile)
 {
 	struct FMaterialExpressionInfo
 	{
-		UMaterialExpression* MaterialExpression;
 		FString Name;
 		FString Keywords;
 		FString CreationName;
@@ -89,8 +89,9 @@ void WriteOutMaterialExpressions(FDiagnosticTableWriterCSV& CsvFile)
 		FString	Type;
 		FString ClassFlags;
 		FString ShowInCreateMenu;
+		int32 Uses = 0;
 	};
-	TArray<FMaterialExpressionInfo> MaterialExpressionInfos;
+	TMap<UMaterialExpression*, FMaterialExpressionInfo> MaterialExpressionInfos;
 
 	// Collect all default material expression objects
 	for (TObjectIterator<UClass> It; It; ++It)
@@ -151,8 +152,7 @@ void WriteOutMaterialExpressions(FDiagnosticTableWriterCSV& CsvFile)
 				FString DisplayName = Class->GetMetaData(TEXT("DisplayName"));
 				FString CreationName = DefaultExpression->GetCreationName().ToString();
 
-				FMaterialExpressionInfo ExpressionInfo;
-				ExpressionInfo.MaterialExpression = DefaultExpression;
+				FMaterialExpressionInfo& ExpressionInfo = MaterialExpressionInfos.Add(DefaultExpression);
 				ExpressionInfo.Name = Class->GetName().Mid(FCString::Strlen(TEXT("MaterialExpression")));
 				ExpressionInfo.Keywords = DefaultExpression->GetKeywords().ToString();
 				ExpressionInfo.CreationName = (CreationName.IsEmpty() ? (DisplayName.IsEmpty() ? ExpressionInfo.Name : DisplayName) : CreationName);
@@ -163,14 +163,59 @@ void WriteOutMaterialExpressions(FDiagnosticTableWriterCSV& CsvFile)
 				ExpressionInfo.Type = ExpressionType;
 				ExpressionInfo.ClassFlags = ClassFlags;
 				ExpressionInfo.ShowInCreateMenu = bShowInCreateMenu ? TEXT("Yes") : TEXT("No");
-				MaterialExpressionInfos.Add(ExpressionInfo);
 			}
 		}
 	}
 
+	// Collect all materials for current project and count how many times each material expressions is referenced
+	TArray<FAssetData> MaterialList;
+	FARFilter MaterialAssetFilter;
+	{
+		MaterialAssetFilter.bRecursiveClasses = true;
+		MaterialAssetFilter.ClassPaths.Add(UMaterial::StaticClass()->GetClassPathName());
+		MaterialAssetFilter.ClassPaths.Add(UMaterialInstance::StaticClass()->GetClassPathName());
+	}
+	IAssetRegistry& AssetRegistry = FModuleManager::Get().LoadModuleChecked<FAssetRegistryModule>("AssetRegistry").Get();
+	AssetRegistry.SearchAllAssets(true);
+	AssetRegistry.GetAssets(MaterialAssetFilter, MaterialList);
+
+	for (const FAssetData& AssetData : MaterialList)
+	{
+		UMaterialInterface* MaterialInterface = Cast<UMaterialInterface>(AssetData.GetAsset());
+		if (!MaterialInterface)
+		{
+			continue;
+		}
+
+		UMaterial* Material = MaterialInterface->GetMaterial();
+		if (!Material)
+		{
+			continue;
+		}
+
+		for (const TObjectPtr<UMaterialExpression>& MaterialExpression : Material->GetExpressions())
+		{
+			const UMaterialExpression* DefaultExpression = Cast<UMaterialExpression>(MaterialExpression->GetClass()->GetDefaultObject());
+			if (!DefaultExpression)
+			{
+				continue;
+			}
+
+			FMaterialExpressionInfo* ExpressionInfo = MaterialExpressionInfos.Find(DefaultExpression);
+			if (!ExpressionInfo)
+			{
+				continue;
+			}
+
+			// Increment use count for this material expression type
+			ExpressionInfo->Uses++;
+		}
+	}
+
 	// Write the material expression list to a text file
 	CsvFile.AddColumn(TEXT("NAME"));
 	CsvFile.AddColumn(TEXT("TYPE"));
+	CsvFile.AddColumn(TEXT("USES"));
 	CsvFile.AddColumn(TEXT("CLASS_FLAGS"));
 	CsvFile.AddColumn(TEXT("SHOW_IN_CREATE_MENU"));
 	CsvFile.AddColumn(TEXT("KEYWORDS"));
@@ -181,10 +226,12 @@ void WriteOutMaterialExpressions(FDiagnosticTableWriterCSV& CsvFile)
 	CsvFile.AddColumn(TEXT("TOOLTIP"));
 	CsvFile.CycleRow();
 
-	for (FMaterialExpressionInfo& ExpressionInfo : MaterialExpressionInfos)
+	for (auto Iter = MaterialExpressionInfos.CreateConstIterator(); Iter; ++Iter)
 	{
+		const FMaterialExpressionInfo& ExpressionInfo = Iter->Value;
 		CsvFile.AddColumn(*GetFormattedText(ExpressionInfo.Name));
 		CsvFile.AddColumn(*GetFormattedText(ExpressionInfo.Type));
+		CsvFile.AddColumn(*FString::FromInt(ExpressionInfo.Uses));
 		CsvFile.AddColumn(*GetFormattedText(ExpressionInfo.ClassFlags));
 		CsvFile.AddColumn(*GetFormattedText(ExpressionInfo.ShowInCreateMenu));
 		CsvFile.AddColumn(*GetFormattedText(ExpressionInfo.Keywords));
