Cooker:  Suppress MPCook when running as -cooklist.
#rnx

diff --git a/Engine/Source/Editor/UnrealEd/Classes/CookOnTheSide/CookOnTheFlyServer.h b/Engine/Source/Editor/UnrealEd/Classes/CookOnTheSide/CookOnTheFlyServer.h
--- a/Engine/Source/Editor/UnrealEd/Classes/CookOnTheSide/CookOnTheFlyServer.h
+++ b/Engine/Source/Editor/UnrealEd/Classes/CookOnTheSide/CookOnTheFlyServer.h
@@ -638,6 +638,8 @@ public:
 		bool bGenerateDependenciesForMaps = false;
 		/** this is a flag for dlc, will cause the cooker to error if the dlc references engine content */
 		bool bErrorOnEngineContentUse = false;
+		/** True if CookByTheBook is being run in cooklist mode and will not be loading/saving packages. */
+		bool bCookList = false;
 	};
 
 	struct FCookOnTheFlyStartupOptions
diff --git a/Engine/Source/Editor/UnrealEd/Private/Commandlets/CookCommandlet.cpp b/Engine/Source/Editor/UnrealEd/Private/Commandlets/CookCommandlet.cpp
--- a/Engine/Source/Editor/UnrealEd/Private/Commandlets/CookCommandlet.cpp
+++ b/Engine/Source/Editor/UnrealEd/Private/Commandlets/CookCommandlet.cpp
@@ -556,6 +556,7 @@ void UCookCommandlet::RunCookByTheBookList(UCookOnTheFlyServer* CookOnTheFlyServ
 
 	ECookListOptions CookListOptions = ECookListOptions::None;
 	if (Switches.Contains(TEXT("showrejected"))) CookListOptions |= ECookListOptions::ShowRejected;
+	StartupOptions.bCookList = true;
 
 	CookOnTheFlyServer->StartCookByTheBook(StartupOptions);
 	CookOnTheFlyServer->RunCookList(CookListOptions);
diff --git a/Engine/Source/Editor/UnrealEd/Private/CookOnTheFlyServer.cpp b/Engine/Source/Editor/UnrealEd/Private/CookOnTheFlyServer.cpp
--- a/Engine/Source/Editor/UnrealEd/Private/CookOnTheFlyServer.cpp
+++ b/Engine/Source/Editor/UnrealEd/Private/CookOnTheFlyServer.cpp
@@ -10989,6 +10989,7 @@ FBeginCookContext UCookOnTheFlyServer::CreateBeginCookByTheBookContext(const FCo
 	CookByTheBookOptions->bDlcLoadMainAssetRegistry = !!(CookOptions & ECookByTheBookOptions::DlcLoadMainAssetRegistry);
 	CookByTheBookOptions->bErrorOnEngineContentUse = StartupOptions.bErrorOnEngineContentUse;
 	CookByTheBookOptions->bAllowUncookedAssetReferences = FParse::Param(FCommandLine::Get(), TEXT("AllowUncookedAssetReferences"));
+	CookByTheBookOptions->bCookList = StartupOptions.bCookList;
 	CookByTheBookOptions->DlcName = StartupOptions.DLCName;
 	if (CookByTheBookOptions->bSkipHardReferences && !CookByTheBookOptions->bSkipSoftReferences)
 	{
@@ -11532,7 +11533,7 @@ void UCookOnTheFlyServer::GenerateInitialRequests(FBeginCookContext& BeginContex
 		}
 	}
 
-	if (FParse::Param(FCommandLine::Get(), TEXT("CookList")))
+	if (CookByTheBookOptions->bCookList)
 	{
 		WorkerRequests->LogAllRequestedFiles();
 	}
diff --git a/Engine/Source/Editor/UnrealEd/Private/Cooker/CookDirector.cpp b/Engine/Source/Editor/UnrealEd/Private/Cooker/CookDirector.cpp
--- a/Engine/Source/Editor/UnrealEd/Private/Cooker/CookDirector.cpp
+++ b/Engine/Source/Editor/UnrealEd/Private/Cooker/CookDirector.cpp
@@ -1279,6 +1279,12 @@ void FCookDirector::InitializeWorkers()
 			TEXT("CookDirector initialization skipped: -CookSinglePackage was requested. CookMultiprocess is disabled and the cooker is running as a single process."));
 		return;
 	}
+	if (Options.bCookList)
+	{
+		UE_LOG(LogCook, Display,
+			TEXT("CookDirector initialization skipped: -CookList was requested. CookMultiprocess is disabled and the cooker is running as a single process."));
+		return;
+	}
 
 	UE_LOG(LogCook, Display, TEXT("CookProcessCount=%d. CookMultiprocess is enabled with 1 CookDirector and %d %s."),
 		RequestedCookWorkerCount + 1, RequestedCookWorkerCount,
diff --git a/Engine/Source/Editor/UnrealEd/Private/Cooker/CookTypes.cpp b/Engine/Source/Editor/UnrealEd/Private/Cooker/CookTypes.cpp
--- a/Engine/Source/Editor/UnrealEd/Private/Cooker/CookTypes.cpp
+++ b/Engine/Source/Editor/UnrealEd/Private/Cooker/CookTypes.cpp
@@ -969,6 +969,9 @@ FCbWriter& operator<<(FCbWriter& Writer, const UE::Cook::FCookByTheBookOptions&
 	Writer << "SkipSoftReferences" << Value.bSkipSoftReferences;
 	Writer << "CookAgainstFixedBase" << Value.bCookAgainstFixedBase;
 	Writer << "DlcLoadMainAssetRegistry" << Value.bDlcLoadMainAssetRegistry;
+
+	// CookList is process-specific
+
 	Writer << "CookSoftPackageReferences" << Value.bCookSoftPackageReferences;
 	Writer.EndObject();
 	return Writer;
diff --git a/Engine/Source/Editor/UnrealEd/Private/Cooker/CookTypes.h b/Engine/Source/Editor/UnrealEd/Private/Cooker/CookTypes.h
--- a/Engine/Source/Editor/UnrealEd/Private/Cooker/CookTypes.h
+++ b/Engine/Source/Editor/UnrealEd/Private/Cooker/CookTypes.h
@@ -575,6 +575,8 @@ public:
 	bool							bCookSoftPackageReferences = false;
 	bool							bCookAgainstFixedBase = false;
 	bool							bDlcLoadMainAssetRegistry = false;
+	/** True if CookByTheBook is being run in cooklist mode and will not be loading/saving packages. */
+	bool							bCookList = false;
 
 	void ClearSessionData()
 	{
