MPCook: Add more feedback for failed assertion NumNeedCommitPlatforms > 0 in  FCookWorkerServer::SendPendingPackages.
#rnx

diff --git a/Engine/Source/Editor/UnrealEd/Private/Cooker/CookWorkerServer.cpp b/Engine/Source/Editor/UnrealEd/Private/Cooker/CookWorkerServer.cpp
--- a/Engine/Source/Editor/UnrealEd/Private/Cooker/CookWorkerServer.cpp
+++ b/Engine/Source/Editor/UnrealEd/Private/Cooker/CookWorkerServer.cpp
@@ -728,7 +728,22 @@ void FCookWorkerServer::SendPendingPackages()
 			++PlatformIndex;
 		}
 		// It should not have been added to PackagesToAssign if there are no platforms to commit
-		check(NumNeedCommitPlatforms > 0);
+		if (NumNeedCommitPlatforms == 0)
+		{
+			TStringBuilder<256> PlatformDataText;
+			for (const TPair<const ITargetPlatform*, FPackagePlatformData>& PlatformPair : PackageData->GetPlatformDatas())
+			{
+				PlatformDataText.Appendf(TEXT("{ %s: Reachable=%s, Committed=%s }, "),
+					PlatformPair.Key == CookerLoadingPlatformKey ? TEXT("CookerLoadingPlatform") : *PlatformPair.Key->PlatformName(),
+					PlatformPair.Value.IsReachable(CookPhase == ECookPhase::Cook
+						? EReachability::Runtime : EReachability::Build) ? TEXT("true") : TEXT("false"),
+					PlatformPair.Value.IsCommitted() ? TEXT("true") : TEXT("false"));
+			}
+			checkf(NumNeedCommitPlatforms > 0,
+				TEXT("Package %s was assigned to worker, but at sendmessage time it has no platforms needing commit. State = %s. CookPhase = %s. [ %s ]"),
+				*PackageData->GetPackageName().ToString(), LexToString(PackageData->GetState()),
+					LexToString(CookPhase), *PlatformDataText);
+		}
 		AssignData.NeedCommitPlatforms = FDiscoveredPlatformSet(SessionPlatformNeedsCommit);
 		FAssignPackageExtraData* ExtraData = PackagesToAssignExtraDatas.Find(PackageData);
 		if (ExtraData)
