Cooker: Fix static analysis report that condition is always false; the report is spurious, it is supposed to be always false and is present so it is documented in the log message.
#rnx

diff --git a/Engine/Source/Editor/UnrealEd/Private/Cooker/CookWorkerServer.cpp b/Engine/Source/Editor/UnrealEd/Private/Cooker/CookWorkerServer.cpp
--- a/Engine/Source/Editor/UnrealEd/Private/Cooker/CookWorkerServer.cpp
+++ b/Engine/Source/Editor/UnrealEd/Private/Cooker/CookWorkerServer.cpp
@@ -739,7 +739,7 @@ void FCookWorkerServer::SendPendingPackages()
 						? EReachability::Runtime : EReachability::Build) ? TEXT("true") : TEXT("false"),
 					PlatformPair.Value.IsCommitted() ? TEXT("true") : TEXT("false"));
 			}
-			checkf(NumNeedCommitPlatforms > 0,
+			checkf(NumNeedCommitPlatforms > 0, // -V547
 				TEXT("Package %s was assigned to worker, but at sendmessage time it has no platforms needing commit. State = %s. CookPhase = %s. [ %s ]"),
 				*PackageData->GetPackageName().ToString(), LexToString(PackageData->GetState()),
 					LexToString(CookPhase), *PlatformDataText);
