Cooker: (Second attempt) Fix static analysis report that condition is always false; the report is spurious, it is supposed to be always false and is present so it is documented in the log message.
#rnx

diff --git a/Engine/Source/Editor/UnrealEd/Private/Cooker/CookWorkerServer.cpp b/Engine/Source/Editor/UnrealEd/Private/Cooker/CookWorkerServer.cpp
--- a/Engine/Source/Editor/UnrealEd/Private/Cooker/CookWorkerServer.cpp
+++ b/Engine/Source/Editor/UnrealEd/Private/Cooker/CookWorkerServer.cpp
@@ -739,10 +739,10 @@ void FCookWorkerServer::SendPendingPackages()
 						? EReachability::Runtime : EReachability::Build) ? TEXT("true") : TEXT("false"),
 					PlatformPair.Value.IsCommitted() ? TEXT("true") : TEXT("false"));
 			}
-			checkf(NumNeedCommitPlatforms > 0, // -V547
+			checkf(NumNeedCommitPlatforms > 0,
 				TEXT("Package %s was assigned to worker, but at sendmessage time it has no platforms needing commit. State = %s. CookPhase = %s. [ %s ]"),
 				*PackageData->GetPackageName().ToString(), LexToString(PackageData->GetState()),
-					LexToString(CookPhase), *PlatformDataText);
+					LexToString(CookPhase), *PlatformDataText); // -V547
 		}
 		AssignData.NeedCommitPlatforms = FDiscoveredPlatformSet(SessionPlatformNeedsCommit);
 		FAssignPackageExtraData* ExtraData = PackagesToAssignExtraDatas.Find(PackageData);
