IncrementalCook: Display instigators for skipped packages.
#rnx
#incrementalcook

diff --git a/Engine/Source/Editor/UnrealEd/Private/Cooker/CookRequestCluster.cpp b/Engine/Source/Editor/UnrealEd/Private/Cooker/CookRequestCluster.cpp
--- a/Engine/Source/Editor/UnrealEd/Private/Cooker/CookRequestCluster.cpp
+++ b/Engine/Source/Editor/UnrealEd/Private/Cooker/CookRequestCluster.cpp
@@ -17,6 +17,7 @@
 #include "Cooker/CookRequests.h"
 #include "Cooker/CookTypes.h"
 #include "Cooker/PackageTracker.h"
+#include "CookerSettings.h"
 #include "CookOnTheSide/CookOnTheFlyServer.h"
 #include "DerivedDataCache.h"
 #include "DerivedDataRequestOwner.h"
@@ -2056,6 +2057,7 @@ void FRequestCluster::FGraphSearch::FExploreEdgesContext::ProcessPlatformAttachm
 	bool bFoundBuildDefinitions = false;
 	ICookedPackageWriter* PackageWriter = FetchPlatformData.Writer;
 	FQueryPlatformData& QueryPlatformData = Vertex->GetPlatformData()[PlatformIndex];
+	bool bReportedInstigator = false;
 
 	if (Cluster.IsIncrementalCook())
 	{
@@ -2154,6 +2156,13 @@ void FRequestCluster::FGraphSearch::FExploreEdgesContext::ProcessPlatformAttachm
 				{
 					COOK_STAT(++DetailedCookStats::NumPackagesIterativelySkipped);
 				}
+				if ((GCookProgressDisplay & (int32)ECookProgressDisplayMode::Instigators) && !bReportedInstigator)
+				{
+					bReportedInstigator = true;
+					UE_LOG(LogCook, Display, TEXT("Incrementally Skipped %s, Instigator: { %s }"),
+						*WriteToString<256>(PackageData->GetPackageName()),
+						*PackageData->GetInstigator(EReachability::Runtime).ToString());
+				}
 				// Declare the package to the EDLCookInfo verification so we don't warn about missing exports from it
 				UE::SavePackageUtilities::EDLCookInfoAddIterativelySkippedPackage(PackageName);
 			}
