Cooker: (Thrid attempt) Fix static analysis report that condition is always false; remove the condition on the check statement since suppression of the warning is not working.
#rnx

diff --git a/Engine/Source/Editor/UnrealEd/Private/Cooker/CookWorkerServer.cpp b/Engine/Source/Editor/UnrealEd/Private/Cooker/CookWorkerServer.cpp
--- a/Engine/Source/Editor/UnrealEd/Private/Cooker/CookWorkerServer.cpp
+++ b/Engine/Source/Editor/UnrealEd/Private/Cooker/CookWorkerServer.cpp
@@ -739,10 +739,10 @@ void FCookWorkerServer::SendPendingPackages()
 						? EReachability::Runtime : EReachability::Build) ? TEXT("true") : TEXT("false"),
 					PlatformPair.Value.IsCommitted() ? TEXT("true") : TEXT("false"));
 			}
-			checkf(NumNeedCommitPlatforms > 0,
+			checkf(false,
 				TEXT("Package %s was assigned to worker, but at sendmessage time it has no platforms needing commit. State = %s. CookPhase = %s. [ %s ]"),
 				*PackageData->GetPackageName().ToString(), LexToString(PackageData->GetState()),
-					LexToString(CookPhase), *PlatformDataText); // -V547
+					LexToString(CookPhase), *PlatformDataText);
 		}
 		AssignData.NeedCommitPlatforms = FDiscoveredPlatformSet(SessionPlatformNeedsCommit);
 		FAssignPackageExtraData* ExtraData = PackagesToAssignExtraDatas.Find(PackageData);
