MPCook: add comment for why SetAllPlatformSaved can skip initialization.
#rnx
#incrementalcook

diff --git a/Engine/Source/Editor/UnrealEd/Private/Cooker/CookGenerationHelper.h b/Engine/Source/Editor/UnrealEd/Private/Cooker/CookGenerationHelper.h
--- a/Engine/Source/Editor/UnrealEd/Private/Cooker/CookGenerationHelper.h
+++ b/Engine/Source/Editor/UnrealEd/Private/Cooker/CookGenerationHelper.h
@@ -1067,6 +1067,10 @@ inline void FGenerationHelper::SetAllPlatformsSaved(FPackageData& PackageData, F
 {
 	if (!Info)
 	{
+		// Note that we expect this Find to succeed for generated packages even without initialization.
+		// In the non-incremental case, the GenerationHelper must have already initialized,
+		// or the generated PackageData would never exist. In the incremental case, the Info was added
+		// during PopulateCookedSandbox to prepare the generator package for incremental cook.
 		Info = FindInfoNoInitialize(PackageData);
 	}
 	if (Info)
