Converted the subobject in the unittests from DefaultSubObjects to just normal Archetypes SubObjects.
Fixed some Cut&Paste problem following the change of DefaultSubObject in Unitttests
Added OS unittest pass for PIEDuplicate

#rnx

diff --git a/Engine/Source/Editor/UnrealEd/Private/EditorObject.cpp b/Engine/Source/Editor/UnrealEd/Private/EditorObject.cpp
--- a/Engine/Source/Editor/UnrealEd/Private/EditorObject.cpp
+++ b/Engine/Source/Editor/UnrealEd/Private/EditorObject.cpp
@@ -648,21 +648,6 @@ static const TCHAR* ImportProperties(
 
 				if (!ComponentTemplate)
 				{
-					// When importing hierarchies of objects with instanced sub-objects we don't want to use the
-					// instancing graph of our outer objects since we need our dedicated instances.
-					// For example:
-					//
-					//	A_CDO { SubObj }
-					//	PastedObj { A_1{ SubObj_1 } A_2{ SubObj_1 } }
-					//
-					//	by sharing the instancing graph, we end up with both SubObj_1 being the same pointer
-					//	by using a dedicated instancing graph, both SubObj_1 are newly created object pointer 
-					//
-					// Ideally this rule should be applied more widely but for now we limit it to the Prefab system
-					// which we could identify using the fact it uses overridable serialization.
-					const bool bShouldUseDedicatedInstancingGraph = FOverridableManager::Get().IsEnabled(SubobjectOuter);
-					FObjectInstancingGraph InstanceGraphForNewObject;
-
 					ComponentTemplate = NewObject<UObject>(
 						SubobjectOuter,
 						TemplateClass,
@@ -670,7 +655,7 @@ static const TCHAR* ImportProperties(
 						NewFlags,
 						Archetype,
 						!!SubobjectOuter,
-						bShouldUseDedicatedInstancingGraph ? &InstanceGraphForNewObject : &InstanceGraph
+						&InstanceGraph
 						);
 				}
 				else
@@ -843,10 +828,27 @@ const TCHAR* ImportObjectProperties( FImportObjectParams& InParams )
 		InParams.SubobjectOuter->PreEditChange(NULL);
 	}
 
+	// When importing hierarchies of objects with instanced sub-objects we don't want to use the
+	// instancing graph of our outer objects since we need our dedicated instances.
+	// For example:
+	//
+	//	A_CDO { SubObj }
+	//	PastedObj { A_1{ SubObj_1 } A_2{ SubObj_1 } }
+	//
+	//	by sharing the instancing graph, we end up with both SubObj_1 being the same pointer
+	//	by using a dedicated instancing graph, both SubObj_1 are newly created object pointer 
+	//
+	// Ideally this rule should be applied more widely but for now we limit it to the Prefabs
+	// which is identified by using overridable serialization.
+	const bool bShouldUseDedicatedInstancingGraph = FOverridableManager::Get().IsEnabled(InParams.SubobjectRoot);
 	FObjectInstancingGraph TempGraph;
-	FObjectInstancingGraph& InstanceGraph = InParams.InInstanceGraph ? *InParams.InInstanceGraph : TempGraph;
+	FObjectInstancingGraph& InstanceGraph = InParams.InInstanceGraph && !bShouldUseDedicatedInstancingGraph ? *InParams.InInstanceGraph : TempGraph;
 
-	if ( InParams.SubobjectRoot && InParams.SubobjectRoot != UObject::StaticClass()->GetDefaultObject() )
+	if (bShouldUseDedicatedInstancingGraph)
+	{
+		InstanceGraph.SetDestinationRoot(InParams.SubobjectOuter);
+	}
+	else if ( InParams.SubobjectRoot && InParams.SubobjectRoot != UObject::StaticClass()->GetDefaultObject() )
 	{
 		InstanceGraph.SetDestinationRoot(InParams.SubobjectRoot);
 	}
