Guard against Garbage Collection of ULevelFactory for the duration of actor import during pasting actors.
* The side effects of creating actors can cause garbage collection to occur in some cases due to public overridability of PostEditChangeEvent.
* Since ULevelActor is not referenced, the GC will mark the object destroyed while the function ULevelFactory::FactoryCreateText is running which can lead to crashes.
* The fix is to keep a reference to the factory for the duration of its use

diff --git a/Engine/Source/Editor/UnrealEd/Private/EditorActor.cpp b/Engine/Source/Editor/UnrealEd/Private/EditorActor.cpp
--- a/Engine/Source/Editor/UnrealEd/Private/EditorActor.cpp
+++ b/Engine/Source/Editor/UnrealEd/Private/EditorActor.cpp
@@ -361,7 +361,7 @@ void UUnrealEdEngine::PasteActors(TArray<AActor*>& OutPastedActors, UWorld* InWo
 	GetMutableDefault<ULevelEditorMiscSettings>()->bBSPAutoUpdate = false;
 
 	// Import the actors.
-	ULevelFactory* Factory = NewObject<ULevelFactory>();
+	TStrongObjectPtr<ULevelFactory> Factory(NewObject<ULevelFactory>());
 	Factory->FactoryCreateText(ULevel::StaticClass(), InWorld->GetCurrentLevel(), InWorld->GetCurrentLevel()->GetFName(), RF_Transactional, NULL, bDuplicate ? TEXT("move") : TEXT("paste"), Paste, Paste + FCString::Strlen(Paste), GWarn);
 
 	// Reinstate old BSP update setting, and force a rebuild - any levels whose geometry has changed while pasting will be rebuilt
