SourceControl: Don't explicitly call out 'Saving to Revision Control', but just fold the operation under 'Saving Packages'.

#rnx

diff --git a/Engine/Source/Editor/UnrealEd/Private/FileHelpers.cpp b/Engine/Source/Editor/UnrealEd/Private/FileHelpers.cpp
--- a/Engine/Source/Editor/UnrealEd/Private/FileHelpers.cpp
+++ b/Engine/Source/Editor/UnrealEd/Private/FileHelpers.cpp
@@ -4315,8 +4315,7 @@ FEditorFileUtils::EPromptReturnCode InternalPromptForCheckoutAndSave(const TArra
 
 		if (bCanSaveToRevisionControl)
 		{
-			FText SavingRevisionControlText = NSLOCTEXT("UnrealEd", "SavingRevisionControl", "Saving to Revision Control");
-			SlowTask.EnterProgressFrame(PackagesToSave.Num() + PackagesToClean.Num(), SavingRevisionControlText);
+			SlowTask.EnterProgressFrame(PackagesToSave.Num() + PackagesToClean.Num());
 
 			ISourceControlModule::Get().GetProvider().Execute(SaveOperation, USourceControlHelpers::PackageFilenames(PackagesToSave), EConcurrency::Synchronous);
 		}
