[Viewport Toolbar] Add hotkey info to level editor toolbar snapping submenus

diff --git a/Engine/Source/Editor/UnrealEd/Private/ViewportToolbar/UnrealEdViewportToolbar.cpp b/Engine/Source/Editor/UnrealEd/Private/ViewportToolbar/UnrealEdViewportToolbar.cpp
--- a/Engine/Source/Editor/UnrealEd/Private/ViewportToolbar/UnrealEdViewportToolbar.cpp
+++ b/Engine/Source/Editor/UnrealEd/Private/ViewportToolbar/UnrealEdViewportToolbar.cpp
@@ -217,6 +217,18 @@ FToolMenuEntry CreateSurfaceSnapCheckboxMenu(const FToolMenuContext& InContext)
 		}));
 	}
 
+	Entry.InputBindingLabel = TAttribute<FText>::CreateLambda(
+		[]
+		{
+			if (const TSharedPtr<FUICommandInfo>& Command = FEditorViewportCommands::Get().SurfaceSnapping)
+			{
+				return Command->GetInputText().ToUpper();
+			}
+
+			return FText();
+		}
+	);
+
 	Entry.ToolBarData.LabelOverride = TAttribute<FText>::CreateLambda([]
 	{
 		const ULevelEditorViewportSettings* const Settings = GetMutableDefault<ULevelEditorViewportSettings>();
@@ -303,6 +315,18 @@ FToolMenuEntry CreateLocationSnapCheckboxMenu(const FToolMenuContext& InContext)
 	}
 	
 	Entry.Icon = FSlateIcon(FAppStyle::GetAppStyleSetName(), "ViewportToolbar.SnapLocation");
+	Entry.InputBindingLabel = TAttribute<FText>::CreateLambda(
+		[]
+		{
+			const TSharedPtr<FUICommandInfo>& Command = FEditorViewportCommands::Get().LocationGridSnap;
+			if (Command.IsValid())
+			{
+				return Command->GetInputText().ToUpper();
+			}
+
+			return FText();
+		}
+	);
 	
 	return Entry;
 }
@@ -429,6 +453,19 @@ FToolMenuEntry CreateRotationSnapCheckboxMenu(const FToolMenuContext& InContext)
 		LabelOverride.BindUObject(Context, &UUnrealEdViewportToolbarContext::GetRotationSnapLabel);
 	}
 
+	Entry.InputBindingLabel = TAttribute<FText>::CreateLambda(
+		[]
+		{
+			const TSharedPtr<FUICommandInfo>& Command = FEditorViewportCommands::Get().RotationGridSnap;
+			if (Command.IsValid())
+			{
+				return Command->GetInputText().ToUpper();
+			}
+
+			return FText();
+		}
+	);
+
 	Entry.ToolBarData.LabelOverride = LabelOverride;
 	
 	Entry.Icon = FSlateIcon(FAppStyle::GetAppStyleSetName(), "ViewportToolbar.SnapRotation");
@@ -524,6 +561,19 @@ FToolMenuEntry CreateScaleSnapCheckboxMenu(const FToolMenuContext& InContext)
 		LabelOverride.BindUObject(Context, &UUnrealEdViewportToolbarContext::GetScaleSnapLabel);
 	}
 
+	Entry.InputBindingLabel = TAttribute<FText>::CreateLambda(
+		[]
+		{
+			const TSharedPtr<FUICommandInfo>& Command = FEditorViewportCommands::Get().ScaleGridSnap;
+			if (Command.IsValid())
+			{
+				return Command->GetInputText().ToUpper();
+			}
+
+			return FText();
+		}
+	);
+
 	Entry.ToolBarData.LabelOverride = LabelOverride;
 	Entry.Icon = FSlateIcon(FAppStyle::GetAppStyleSetName(), "ViewportToolbar.SnapScale");
 
