[ITF Gizmos] Alt while dragging with new TRS Gizmos has cursor jumping around

diff --git a/Engine/Source/Editor/UnrealEd/Private/EditorViewportClient.cpp b/Engine/Source/Editor/UnrealEd/Private/EditorViewportClient.cpp
--- a/Engine/Source/Editor/UnrealEd/Private/EditorViewportClient.cpp
+++ b/Engine/Source/Editor/UnrealEd/Private/EditorViewportClient.cpp
@@ -5184,7 +5184,7 @@ bool FEditorViewportClient::ShouldOrbitCamera() const
 
 		if (!GetDefault<ULevelEditorViewportSettings>()->bUseUE3OrbitControls)
 		{
-			bDesireOrbit = IsAltPressed() && !IsCtrlPressed() && !IsShiftPressed();
+			bDesireOrbit = (IsAltPressed() && !ModeTools->HasOngoingTransform()) && !IsCtrlPressed() && !IsShiftPressed();
 		}
 		else
 		{
