Updates the viewport toolbars to a new design.
# Old toolbars can be re-enabled with "ToolMenusViewportToolbars 0"
# This CL changes the default to "2"

[RN] major ui

diff --git a/Engine/Source/Editor/UnrealEd/Private/ViewportToolbar/UnrealEdViewportToolbar.cpp b/Engine/Source/Editor/UnrealEd/Private/ViewportToolbar/UnrealEdViewportToolbar.cpp
--- a/Engine/Source/Editor/UnrealEd/Private/ViewportToolbar/UnrealEdViewportToolbar.cpp
+++ b/Engine/Source/Editor/UnrealEd/Private/ViewportToolbar/UnrealEdViewportToolbar.cpp
@@ -41,7 +41,7 @@
 
 namespace UE::UnrealEd::Private
 {
-int32 CVarToolMenusViewportToolbarsValue = 0;
+int32 CVarToolMenusViewportToolbarsValue = 2;
 
 struct ViewModesSubmenu
 {
@@ -657,9 +657,9 @@ float GetCamSpeedScalarSliderPosition(const TSharedRef<SEditorViewport>& InEdito
 static FAutoConsoleVariableRef CVarToolMenusViewportToolbars(
 	TEXT("ToolMenusViewportToolbars"),
 	UE::UnrealEd::Private::CVarToolMenusViewportToolbarsValue,
-	TEXT("Control whether the new ToolMenus-based viewport toolbars are enabled across the editor. Set to 0 (default) "
+	TEXT("Control whether the new ToolMenus-based viewport toolbars are enabled across the editor. Set to 0 "
 		 "to show only the old viewport toolbars. Set to 1 for side-by-side mode where both the old and new viewport "
-		 "toolbars are shown. Set to 2 to show only the new viewport toolbars."),
+		 "toolbars are shown. Set to 2 (default) to show only the new viewport toolbars."),
 	ECVF_Default
 );
 
