Orbit on selection: fixes animation mode misbehavior + stores the last state

- when orbiting in animation mode, do not return a valid orbit if nothing is selected
- store last orbit location value to remember its last state

diff --git a/Engine/Plugins/Animation/ControlRig/Source/ControlRigEditor/Private/EditMode/ControlRigEditMode.cpp b/Engine/Plugins/Animation/ControlRig/Source/ControlRigEditor/Private/EditMode/ControlRigEditMode.cpp
--- a/Engine/Plugins/Animation/ControlRig/Source/ControlRigEditor/Private/EditMode/ControlRigEditMode.cpp
+++ b/Engine/Plugins/Animation/ControlRig/Source/ControlRigEditor/Private/EditMode/ControlRigEditMode.cpp
@@ -1375,11 +1375,13 @@ FVector FControlRigEditMode::GetWidgetLocation() const
 
 bool FControlRigEditMode::GetPivotForOrbit(FVector& OutPivot) const
 {
-	if (UsesTransformWidget())
+	constexpr bool bUseShape = true;
+	if (IsControlSelected(bUseShape))
 	{
 		OutPivot = GetWidgetLocation();
 		return true;
 	}
+
 	return FEdMode::GetPivotForOrbit(OutPivot);
 }
 
@@ -2998,19 +3000,42 @@ void FControlRigEditMode::BindCommands()
 		FExecuteAction::CreateRaw(this, &FControlRigEditMode::OpenSpacePickerWidget));
 }
 
-bool FControlRigEditMode::IsControlSelected() const
+bool FControlRigEditMode::IsControlSelected(const bool bUseShapes) const
 {
-	static uint32 TypeFlag = (uint32)ERigElementType::Control;
-	for (const TWeakObjectPtr<UControlRig>& RuntimeRigPtr : RuntimeControlRigs)
+	static constexpr uint32 ControlType = static_cast<uint32>(ERigElementType::Control);
+	
+	if (bUseShapes)
 	{
-		if (UControlRig* ControlRig = RuntimeRigPtr.Get())
+		for (const auto& [ControlRig, Shapes] : ControlRigShapeActors)
 		{
-			if (AreRigElementsSelected(TypeFlag,ControlRig))
+			for (const AControlRigShapeActor* ShapeActor : Shapes)
+			{
+				if (ShapeActor && ShapeActor->IsSelected())
+				{
+					return true;
+				}
+			}
+
+			if (ControlRig && AreRigElementsSelected(ControlType,ControlRig))
 			{
 				return true;
 			}
 		}
 	}
+	else
+	{
+		for (const TWeakObjectPtr<UControlRig>& RuntimeRigPtr : RuntimeControlRigs)
+		{
+			if (UControlRig* ControlRig = RuntimeRigPtr.Get())
+			{
+				if (AreRigElementsSelected(ControlType,ControlRig))
+				{
+					return true;
+				}
+			}
+		}
+	}
+	
 	return false;
 }
 
diff --git a/Engine/Plugins/Animation/ControlRig/Source/ControlRigEditor/Public/EditMode/ControlRigEditMode.h b/Engine/Plugins/Animation/ControlRig/Source/ControlRigEditor/Public/EditMode/ControlRigEditMode.h
--- a/Engine/Plugins/Animation/ControlRig/Source/ControlRigEditor/Public/EditMode/ControlRigEditMode.h
+++ b/Engine/Plugins/Animation/ControlRig/Source/ControlRigEditor/Public/EditMode/ControlRigEditMode.h
@@ -440,7 +440,7 @@ protected:
 	TObjectPtr<UControlRigDetailPanelControlProxies> ControlProxy;
 
 	/** Utility functions for UI/Some other viewport manipulation*/
-	bool IsControlSelected() const;
+	bool IsControlSelected(const bool bUseShapes = false) const;
 	bool AreRigElementSelectedAndMovable(UControlRig* InControlRig) const;
 	
 	/** Set initial transform handlers */
diff --git a/Engine/Source/Editor/UnrealEd/Private/EditorViewportClient.cpp b/Engine/Source/Editor/UnrealEd/Private/EditorViewportClient.cpp
--- a/Engine/Source/Editor/UnrealEd/Private/EditorViewportClient.cpp
+++ b/Engine/Source/Editor/UnrealEd/Private/EditorViewportClient.cpp
@@ -2541,8 +2541,14 @@ void FEditorViewportClient::InputAxisForOrbit(FViewport* InViewport, const FVect
 {
 	FVector OrbitPoint;
 	bool bHasCustomOrbitPivot = GetPivotForOrbit(OrbitPoint);
-	if ( GetDefault<ULevelEditorViewportSettings>()->bOrbitCameraAroundSelection && IsOrbitRotationMode( InViewport ) && bHasCustomOrbitPivot )
+
+	if ( GetDefault<ULevelEditorViewportSettings>()->bOrbitCameraAroundSelection && IsOrbitRotationMode( InViewport ))
 	{
+		if (bHasCustomOrbitPivot)
+		{
+			DefaultOrbitLocation = OrbitPoint;
+		}
+		
 		// Override the default orbit behavior to allow orbiting around a given pivot
 		// This uses different computations from the default orbit behavior so it must not ToggleOrbitCamera
 		const bool bEnable = false;
@@ -2553,8 +2559,8 @@ void FEditorViewportClient::InputAxisForOrbit(FViewport* InViewport, const FVect
 		const FRotator ViewRotation = GetViewRotation();
 
 		// Compute the look-at and view location centered on the orbit point
-		const FVector LookAtOffset = GetLookAtLocation() - OrbitPoint;
-		const FVector ViewLocationOffset = GetViewLocation() - OrbitPoint;
+		const FVector LookAtOffset = GetLookAtLocation() - DefaultOrbitLocation;
+		const FVector ViewLocationOffset = GetViewLocation() - DefaultOrbitLocation;
 
 		// When the roll is at 180 degrees, it means the view is upside down, so invert the yaw rotation
 		if (ViewRotation.Roll == 180.f)
@@ -2598,12 +2604,12 @@ void FEditorViewportClient::InputAxisForOrbit(FViewport* InViewport, const FVect
 
 		// Set the new rotated look-at
 		FVector RotatedLookAtOffset = DeltaRotationTransform.TransformVector(LookAtOffset);
-		FVector RotatedLookAt = OrbitPoint + RotatedLookAtOffset;
+		FVector RotatedLookAt = DefaultOrbitLocation + RotatedLookAtOffset;
 		SetLookAtLocation(RotatedLookAt);
 
 		// Set the new rotated view location
 		FVector RotatedViewOffset = DeltaRotationTransform.TransformVector(ViewLocationOffset);
-		FVector RotatedViewLocation = OrbitPoint + RotatedViewOffset;
+		FVector RotatedViewLocation = DefaultOrbitLocation + RotatedViewOffset;
 		SetViewLocation(RotatedViewLocation);
 
 		FEditorViewportStats::Using(IsPerspective() ? FEditorViewportStats::CAT_PERSPECTIVE_MOUSE_ORBIT_ROTATION : FEditorViewportStats::CAT_ORTHOGRAPHIC_MOUSE_ORBIT_ROTATION);
diff --git a/Engine/Source/Editor/UnrealEd/Public/EditorViewportClient.h b/Engine/Source/Editor/UnrealEd/Public/EditorViewportClient.h
--- a/Engine/Source/Editor/UnrealEd/Public/EditorViewportClient.h
+++ b/Engine/Source/Editor/UnrealEd/Public/EditorViewportClient.h
@@ -2005,10 +2005,10 @@ protected:
 	FPreviewScene* PreviewScene;
 
 	/** default set up for toggling back **/
-	FRotator DefaultOrbitRotation;
-	FVector DefaultOrbitLocation;
-	FVector DefaultOrbitZoom;
-	FVector DefaultOrbitLookAt;
+	FRotator DefaultOrbitRotation = FRotator::ZeroRotator;
+	FVector DefaultOrbitLocation = FVector::ZeroVector;
+	FVector DefaultOrbitZoom = FVector::ZeroVector;
+	FVector DefaultOrbitLookAt = FVector::ZeroVector;
 
 	/** 
 	 * Temporary realtime overrides and user messages that are not saved between editor sessions.
