IncrementalCook: Fix more occurrences of the warning
"RecordResults received for generated package %s, but its ParentGenerationHelper has already been destructed...", by implementing the rest of the workaround for CookPackageGenerators: If a generator is in request state or in the save state, and we find it has been requested for some but not all session platforms, mark it requested for all platforms, sending it back to request state if necessary.
#rnx
#incrementalcook

diff --git a/Engine/Source/Editor/UnrealEd/Classes/CookOnTheSide/CookOnTheFlyServer.h b/Engine/Source/Editor/UnrealEd/Classes/CookOnTheSide/CookOnTheFlyServer.h
--- a/Engine/Source/Editor/UnrealEd/Classes/CookOnTheSide/CookOnTheFlyServer.h
+++ b/Engine/Source/Editor/UnrealEd/Classes/CookOnTheSide/CookOnTheFlyServer.h
@@ -577,6 +577,11 @@ private:
 
 	/** Called when a package is cancelled and returned to idle. Notifies CookDirector when on a CookWorker. */
 	void DemoteToIdle(UE::Cook::FPackageData& PackageData, UE::Cook::ESendFlags SendFlags, UE::Cook::ESuppressCookReason Reason);
+	/**
+	 * Called when a package is sent back to requestcluster for a change in request data. Notifies CookDirector when
+	 * on a CookWorker.
+	 */
+	void DemoteToRequest(UE::Cook::FPackageData& PackageData, UE::Cook::ESendFlags SendFlags, UE::Cook::ESuppressCookReason Reason);
 	/** Called when a package completes its save and returns to idle. Notifies the CookDirector when on a CookWorker. */
 	void PromoteToSaveComplete(UE::Cook::FPackageData& PackageData, UE::Cook::ESendFlags SendFlags);
 
@@ -1227,12 +1232,17 @@ private:
 	 *
 	 * @param PackageData the PackageData used to gather all uobjects from
 	 * @param bIsPrecaching true if called for precaching
+	 * @param OutDemotionRequestedReason Value is ESuppressCookReason::NotSuppressed if no suppression is requested,
+	 *        or if returned PollStatus is not InComplete. It can be set to a different value, requesting suppression,
+	 *        and returned PollStatus will be InComplete. If caller handles demotion, the caller should
+	 *        demote the package with the given reason. It is valid to call PrepareSave again later if caller does not
+	 *        handle demotion.
 	 * @return false if time slice was reached, true if all objects have had BeginCacheForCookedPlatformData called
 	 */
 	UE::Cook::EPollStatus PrepareSave(UE::Cook::FPackageData& PackageData, UE::Cook::FCookerTimer& Timer,
-		bool bPrecaching);
+		bool bPrecaching, UE::Cook::ESuppressCookReason& OutDemotionRequestedReason);
 	UE::Cook::EPollStatus PrepareSaveInternal(UE::Cook::FPackageData& PackageData, UE::Cook::FCookerTimer& Timer,
-		bool bPrecaching);
+		bool bPrecaching, UE::Cook::ESuppressCookReason& OutDemotionRequestedReason);
 	/** Call BeginCacheForCookedPlatformData on all objects currently found in the PackageData. */
 	UE::Cook::EPollStatus CallBeginCacheOnObjects(UE::Cook::FPackageData& PackageData, UPackage* Package,
 		TArray<UE::Cook::FCachedObjectInOuter>& Objects, int32& NextIndex, UE::Cook::FCookerTimer& Timer);
diff --git a/Engine/Source/Editor/UnrealEd/Private/CookOnTheFlyServer.cpp b/Engine/Source/Editor/UnrealEd/Private/CookOnTheFlyServer.cpp
--- a/Engine/Source/Editor/UnrealEd/Private/CookOnTheFlyServer.cpp
+++ b/Engine/Source/Editor/UnrealEd/Private/CookOnTheFlyServer.cpp
@@ -2653,7 +2653,7 @@ void UCookOnTheFlyServer::PumpRequests(UE::Cook::FTickStackData& StackData, int3
 
 	OutNumPushed = 0;
 	FRequestQueue& RequestQueue = PackageDatas->GetRequestQueue();
-	FPackageDataSet& RestartedRequests = RequestQueue.GetRestartedRequests();
+	TPackageDataMap<ESuppressCookReason>& RestartedRequests = RequestQueue.GetRestartedRequests();
 	TRingBuffer<FDiscoveryQueueElement>& DiscoveryQueue = RequestQueue.GetDiscoveryQueue();
 	TRingBuffer<FPackageData*>& BuildDependencyDiscoveryQueue = RequestQueue.GetBuildDependencyDiscoveryQueue();
 	TRingBuffer<TUniquePtr<FRequestCluster>>& RequestClusters = RequestQueue.GetRequestClusters();
@@ -2885,7 +2885,7 @@ void UCookOnTheFlyServer::DemoteToIdle(UE::Cook::FPackageData& PackageData, UE::
 
 	if (PackageData.IsInProgress())
 	{
-		WorkerRequests->ReportDemoteToIdle(PackageData, Reason);
+		WorkerRequests->ReportDemotion(PackageData, Reason);
 
 		bool bHasCookResult = PackageData.HasAllCookedPlatforms(PlatformManager->GetSessionPlatforms(), true /* bIncludeFailed */);
 
@@ -2939,6 +2939,32 @@ void UCookOnTheFlyServer::DemoteToIdle(UE::Cook::FPackageData& PackageData, UE::
 	PackageData.SendToState(UE::Cook::EPackageState::Idle, SendFlags, ConvertToStateChangeReason(Reason));
 }
 
+void UCookOnTheFlyServer::DemoteToRequest(UE::Cook::FPackageData& PackageData, UE::Cook::ESendFlags SendFlags,
+	UE::Cook::ESuppressCookReason Reason)
+{
+	using namespace UE::Cook;
+
+	if (!PackageData.IsInProgress())
+	{
+		return;
+	}
+
+	UE_LOG(LogCook, Display, TEXT("DemoteToRequest: Package %s was sent back to request state, with reason %s."),
+		*WriteToString<256>(PackageData.GetPackageName()), LexToString(Reason));
+	WorkerRequests->ReportDemotion(PackageData, Reason);
+	if (CookWorkerClient)
+	{
+		PackageData.SendToState(UE::Cook::EPackageState::Idle, SendFlags, ConvertToStateChangeReason(Reason));
+	}
+	else
+	{
+		EnumRemoveFlags(SendFlags, ESendFlags::QueueAdd);
+		TPackageDataMap<ESuppressCookReason>& RestartedRequests = PackageDatas->GetRequestQueue().GetRestartedRequests();
+		PackageData.SendToState(UE::Cook::EPackageState::Request, SendFlags, ConvertToStateChangeReason(Reason));
+		RestartedRequests.Add(&PackageData, Reason);
+	}
+}
+
 void UCookOnTheFlyServer::PromoteToSaveComplete(UE::Cook::FPackageData& PackageData, UE::Cook::ESendFlags SendFlags)
 {
 	if (!PackageData.IsInProgress())
@@ -3994,11 +4020,12 @@ UE::Cook::EPollStatus UCookOnTheFlyServer::TryPopulateGeneratedPackage(UE::Cook:
 }
 
 UE::Cook::EPollStatus UCookOnTheFlyServer::PrepareSave(UE::Cook::FPackageData& PackageData,
-	UE::Cook::FCookerTimer& Timer, bool bPrecaching)
+	UE::Cook::FCookerTimer& Timer, bool bPrecaching, UE::Cook::ESuppressCookReason& OutDemotionRequestedReason)
 {
 	using namespace UE::Cook;
 
 	EPollStatus Result = EPollStatus::Incomplete;
+	OutDemotionRequestedReason = ESuppressCookReason::NotSuppressed;
 	if (PackageData.GetSaveSubState() == ESaveSubState::ReadyForSave)
 	{
 		Result = EPollStatus::Success;
@@ -4010,7 +4037,7 @@ UE::Cook::EPollStatus UCookOnTheFlyServer::PrepareSave(UE::Cook::FPackageData& P
 	else
 	{
 		UE_SCOPED_HIERARCHICAL_COOKTIMER_AND_DURATION(PrepareSave, DetailedCookStats::TickCookOnTheSidePrepareSaveTimeSec);
-		Result = PrepareSaveInternal(PackageData, Timer, bPrecaching);
+		Result = PrepareSaveInternal(PackageData, Timer, bPrecaching, OutDemotionRequestedReason);
 		if (Result == EPollStatus::Error)
 		{
 			PackageData.SetHasPrepareSaveFailed(true);
@@ -4036,7 +4063,7 @@ UE::Cook::EPollStatus UCookOnTheFlyServer::PrepareSave(UE::Cook::FPackageData& P
 }
 
 UE::Cook::EPollStatus UCookOnTheFlyServer::PrepareSaveInternal(UE::Cook::FPackageData& PackageData,
-	UE::Cook::FCookerTimer& Timer, bool bPrecaching)
+	UE::Cook::FCookerTimer& Timer, bool bPrecaching, UE::Cook::ESuppressCookReason& OutDemotionRequestedReason)
 {
 	using namespace UE::Cook;
 
@@ -4165,6 +4192,18 @@ UE::Cook::EPollStatus UCookOnTheFlyServer::PrepareSaveInternal(UE::Cook::FPackag
 							*PackageData.GetPackageName().ToString());
 						return EPollStatus::Error;
 					}
+					// TODO_COOKGENERATIONHELPER: We don't currently support separate cooking for one platform but not
+					// another for a generated package, see the class comment on FGenerationHelper. Therefore if any
+					// platform is unreachable, send this package back to the request state to add the other platforms.
+					TConstArrayView<const ITargetPlatform*> SessionPlatforms = PlatformManager->GetSessionPlatforms();
+					EReachability CurrentReachability = GetCookPhase() == ECookPhase::Cook
+						? EReachability::Runtime : EReachability::Build;
+					if (!PackageData.HasReachablePlatforms(CurrentReachability, SessionPlatforms))
+					{
+						OutDemotionRequestedReason = ESuppressCookReason::GeneratedPackageNeedsRequestUpdate;
+						return EPollStatus::Incomplete;
+					}
+
 					if (!GenerationHelper->TryGenerateList())
 					{
 						return EPollStatus::Error;
@@ -4260,7 +4299,7 @@ UE::Cook::EPollStatus UCookOnTheFlyServer::PrepareSaveInternal(UE::Cook::FPackag
 					PackageData.SetSaveSubState(ESaveSubState::LastCookedPlatformData_CallingBegin);
 					// Call this function recursively to immediately reexecute CallBeginCacheOnObjects.
 					// Note that RefreshObjectCache checked for too many recursive calls and ErrorExited if so.
-					return PrepareSaveInternal(PackageData, Timer, bPrecaching);
+					return PrepareSaveInternal(PackageData, Timer, bPrecaching, OutDemotionRequestedReason);
 				}
 				else
 				{
@@ -4828,7 +4867,8 @@ void UCookOnTheFlyServer::PumpRuntimeSaves(UE::Cook::FTickStackData& StackData,
 		}
 
 		// If BeginCacheCookPlatformData is not ready then postpone the package, exit, or wait for it as appropriate
-		EPollStatus PrepareSaveStatus = PrepareSave(PackageData, StackData.Timer, false /* bPrecaching */);
+		ESuppressCookReason DemotionReason;
+		EPollStatus PrepareSaveStatus = PrepareSave(PackageData, StackData.Timer, false /* bPrecaching */, DemotionReason);
 		if (PrepareSaveStatus != EPollStatus::Success)
 		{
 			if (PrepareSaveStatus == EPollStatus::Error)
@@ -4840,6 +4880,12 @@ void UCookOnTheFlyServer::PumpRuntimeSaves(UE::Cook::FTickStackData& StackData,
 				++OutNumPushed;
 				continue;
 			}
+			if (DemotionReason != ESuppressCookReason::NotSuppressed)
+			{
+				DemoteToRequest(PackageData, ESendFlags::QueueAdd, DemotionReason);
+				++OutNumPushed;
+				continue;
+			}
 
 			// GC is required
 			if (PackageData.IsPrepareSaveRequiresGC())
@@ -4879,8 +4925,10 @@ void UCookOnTheFlyServer::PumpRuntimeSaves(UE::Cook::FTickStackData& StackData,
 					FPlatformProcess::Sleep(0.0f);
 					// Poll the results again and check whether we are now done
 					PackageDatas->PollPendingCookedPlatformDatas(true, LastCookableObjectTickTime);
-					PrepareSaveStatus = PrepareSave(PackageData, StackData.Timer, false /* bPrecaching */);
+					PrepareSaveStatus = PrepareSave(PackageData, StackData.Timer, false /* bPrecaching */,
+						DemotionReason);
 				} while (!StackData.Timer.IsActionTimeUp() && PrepareSaveStatus == EPollStatus::Incomplete
+					&& DemotionReason == ESuppressCookReason::NotSuppressed
 					&& PackageData.GetUrgency() == EUrgency::Blocking);
 			}
 			// If we couldn't postpone or wait, then we need to exit and try again later
@@ -4908,7 +4956,7 @@ void UCookOnTheFlyServer::PumpRuntimeSaves(UE::Cook::FTickStackData& StackData,
 				}
 
 				--LeftToPrecache;
-				PrepareSave(*NextData, StackData.Timer, /*bPrecaching*/ true);
+				PrepareSave(*NextData, StackData.Timer, /*bPrecaching*/ true, DemotionReason);
 			}
 
 			// If we're in RealTimeMode, check whether the precaching overflowed our timer and if so exit before we do the potentially expensive SavePackage
@@ -10066,9 +10114,9 @@ void UCookOnTheFlyServer::CancelAllQueues()
 	{
 		DatasInStatesWithTSet.Add(PackageData);
 	}
-	for (FPackageData* PackageData : RequestQueue.GetRestartedRequests())
+	for (TPair<FPackageData*, ESuppressCookReason>& Pair : RequestQueue.GetRestartedRequests())
 	{
-		DatasInStatesWithTSet.Add(PackageData);
+		DatasInStatesWithTSet.Add(Pair.Key);
 	}
 	for (FPackageData* PackageData : DatasInStatesWithTSet)
 	{
diff --git a/Engine/Source/Editor/UnrealEd/Private/Cooker/CookDirector.cpp b/Engine/Source/Editor/UnrealEd/Private/Cooker/CookDirector.cpp
--- a/Engine/Source/Editor/UnrealEd/Private/Cooker/CookDirector.cpp
+++ b/Engine/Source/Editor/UnrealEd/Private/Cooker/CookDirector.cpp
@@ -2195,7 +2195,7 @@ FCookDirector::FRetractionHandler::ReassignPackages(const FWorkerId& FromWorker,
 	if (WorkersToSplitOver.IsEmpty())
 	{
 		// Send the packages back to the Director for reassignment
-		FPackageDataSet& RestartedRequests = Director.COTFS.PackageDatas->GetRequestQueue().GetRestartedRequests();
+		TPackageDataMap<ESuppressCookReason>& RestartedRequests = Director.COTFS.PackageDatas->GetRequestQueue().GetRestartedRequests();
 		for (FPackageData* PackageData : AssignmentPackages)
 		{
 			// UnStall the package to handle the case where it was previously in the save state on the LocalWorker
@@ -2206,7 +2206,7 @@ FCookDirector::FRetractionHandler::ReassignPackages(const FWorkerId& FromWorker,
 			if (PackageData->GetState() != EPackageState::SaveActive)
 			{
 				PackageData->SendToState(EPackageState::Request, ESendFlags::QueueRemove, EStateChangeReason::Retraction);
-				RestartedRequests.Add(PackageData);
+				RestartedRequests.Add(PackageData, ESuppressCookReason::RetractedByCookDirector);
 			}
 		}
 		// MPCOOKTODO: Add a method to PumpRequests long enough to assign the packages
diff --git a/Engine/Source/Editor/UnrealEd/Private/Cooker/CookGenerationHelper.h b/Engine/Source/Editor/UnrealEd/Private/Cooker/CookGenerationHelper.h
--- a/Engine/Source/Editor/UnrealEd/Private/Cooker/CookGenerationHelper.h
+++ b/Engine/Source/Editor/UnrealEd/Private/Cooker/CookGenerationHelper.h
@@ -194,7 +194,8 @@ private:
  * reason, and additionally because reachability might be discovered for one platform after the save already started
  * for another platform, and we would need to handle the state transition to restart the generation with the new
  * platform accommodated. We will eventually need to implement both per-platform incremental cooks and per-platform
- * reachability to fully support per-platform customization of the cook graph.
+ * reachability to fully support per-platform customization of the cook graph. See TODO_COOKGENERATIONHELPER for
+ * code blocks where we implement the workarounds for this todo.
  */
 struct FGenerationHelper : public FThreadSafeRefCountedObject
 {
diff --git a/Engine/Source/Editor/UnrealEd/Private/Cooker/CookPackageData.cpp b/Engine/Source/Editor/UnrealEd/Private/Cooker/CookPackageData.cpp
--- a/Engine/Source/Editor/UnrealEd/Private/Cooker/CookPackageData.cpp
+++ b/Engine/Source/Editor/UnrealEd/Private/Cooker/CookPackageData.cpp
@@ -369,6 +369,10 @@ void FPackageData::ClearInProgressData(EStateChangeReason StateChangeReason)
 	CompletionCallback = FCompletionCallback();
 	if (GenerationHelper)
 	{
+		// ClearKeepForGeneratorSaveAllPlatforms might drop the last reference to the GenerationHelper, and delete
+		// it out from under the ClearKeepForGeneratorSaveAllPlatforms, which is not supported by
+		// ClearKeepForGeneratorSaveAllPlatforms, so keep it referenced across that call.
+		TRefCountPtr<FGenerationHelper> KeepReferenced = GenerationHelper;
 		// ClearKeepForGeneratorSave is called when finishing the save state, but not when demoting out of the save
 		// state after a garbage collect. Call it here in case we cancel the save of the packagedata after demotion.
 		// The other self-references (iterative, queued packages) should persist even when the packagedata is not
diff --git a/Engine/Source/Editor/UnrealEd/Private/Cooker/CookPackageData.h b/Engine/Source/Editor/UnrealEd/Private/Cooker/CookPackageData.h
--- a/Engine/Source/Editor/UnrealEd/Private/Cooker/CookPackageData.h
+++ b/Engine/Source/Editor/UnrealEd/Private/Cooker/CookPackageData.h
@@ -1061,7 +1061,7 @@ public:
 	uint32 RemoveRequestExceptFromCluster(FPackageData* PackageData, FRequestCluster* ExceptFromCluster);
 	void UpdateUrgency(FPackageData* PackageData, EUrgency bOldUrgency, EUrgency NewUrgency);
 
-	FPackageDataSet& GetRestartedRequests();
+	TPackageDataMap<ESuppressCookReason>& GetRestartedRequests();
 	/**
 	 * Unlike non-discovery containers on PackageData, GetDiscoveryQueue and GetBuildDependencyDiscoveryQueue
 	 * are not ownership containers.
@@ -1082,7 +1082,7 @@ public:
 	void NotifyRequestFencePassed(FPackageDatas& PackageDatas);
 
 private:
-	FPackageDataSet RestartedRequests;
+	TPackageDataMap<ESuppressCookReason> RestartedRequests;
 	TRingBuffer<FDiscoveryQueueElement> DiscoveryQueue;
 	TRingBuffer<FPackageData*> BuildDependencyDiscoveryQueue;
 	TRingBuffer<TUniquePtr<FRequestCluster>> RequestClusters; // Not trivially relocatable, so must be TUniquePtr
@@ -1887,7 +1887,7 @@ inline void FPackageData::UpdateUpEdge(bool bOld, bool bNew, const FEdgeFunction
 	}
 }
 
-inline FPackageDataSet& FRequestQueue::GetRestartedRequests()
+inline TPackageDataMap<ESuppressCookReason>& FRequestQueue::GetRestartedRequests()
 {
 	return RestartedRequests;
 }
diff --git a/Engine/Source/Editor/UnrealEd/Private/Cooker/CookRequestCluster.cpp b/Engine/Source/Editor/UnrealEd/Private/Cooker/CookRequestCluster.cpp
--- a/Engine/Source/Editor/UnrealEd/Private/Cooker/CookRequestCluster.cpp
+++ b/Engine/Source/Editor/UnrealEd/Private/Cooker/CookRequestCluster.cpp
@@ -154,13 +154,15 @@ FRequestCluster::FRequestCluster(UCookOnTheFlyServer& InCOTFS, TArray<FFilePlatf
 	InRequests.Empty();
 }
 
-FRequestCluster::FRequestCluster(UCookOnTheFlyServer& InCOTFS, FPackageDataSet&& InRequests,
+FRequestCluster::FRequestCluster(UCookOnTheFlyServer& InCOTFS, TPackageDataMap<ESuppressCookReason>&& InRequests,
 	EReachability InExploreReachability)
 	: FRequestCluster(InCOTFS, InExploreReachability)
 {
 	ReserveInitialRequests(InRequests.Num());
-	for (FPackageData* PackageData : InRequests)
+	for (TPair<FPackageData*, ESuppressCookReason>& Pair : InRequests)
 	{
+		FPackageData* PackageData = Pair.Key;
+		ESuppressCookReason SuppressCookReason = Pair.Value;
 		check(PackageData);
 		FVertexData& Vertex = FindOrAddVertex(*PackageData);
 		// Setting bNeedsStateChange=false is important to avoid a crash: changing the state will try to remove it
@@ -168,6 +170,27 @@ FRequestCluster::FRequestCluster(UCookOnTheFlyServer& InCOTFS, FPackageDataSet&&
 		// accomplish what bNeedsStateChange=true would have tried to do more slowly.
 		check(PackageData->GetState() == EPackageState::Request);
 		SetOwnedByCluster(Vertex, true /* bOwnedByCluster */, false /* bNeedsStateChange */);
+
+		// Some restartedrequests need reachability changes, and reachability changes can only be made by a
+		// RequestCluster, so do them here.
+		switch (SuppressCookReason)
+		{
+		case ESuppressCookReason::GeneratedPackageNeedsRequestUpdate:
+		{
+			// TODO_COOKGENERATIONHELPER: We don't currently support separate cooking for one platform but not
+			// another for a generated package, see the class comment on FGenerationHelper. Therefore if any platform
+			// was found to be reachable, then mark the other platforms as reachable.
+			TConstArrayView<const ITargetPlatform*> SessionPlatforms = InCOTFS.PlatformManager->GetSessionPlatforms();
+			if (PackageData->HasReachablePlatforms(InExploreReachability, SessionPlatforms))
+			{
+				PackageData->AddReachablePlatforms(*this, InExploreReachability, SessionPlatforms,
+					FInstigator(EInstigator::GeneratedPackage, PackageData->GetPackageName()));
+			}
+			break;
+		}
+		default:
+			break;
+		}
 	}
 	InRequests.Empty();
 }
@@ -1278,6 +1301,22 @@ void FRequestCluster::FGraphSearch::VisitVertex(FVertexData& Vertex)
 		}
 	}
 
+	// TODO_COOKGENERATIONHELPER: We don't currently support separate cooking for one platform but not
+	// another for a generated package, see the class comment on FGenerationHelper. Therefore if any platform
+	// was found to be reachable, then mark the other platforms as reachable.
+	if (PackageData.GetGenerationHelper() && FirstReachableSessionPlatform != nullptr)
+	{
+		TConstArrayView<const ITargetPlatform*> SessionPlatforms = Cluster.COTFS.PlatformManager->GetSessionPlatforms();
+		if (!PackageData.HasReachablePlatforms(ClusterReachability, SessionPlatforms))
+		{
+			PackageData.AddReachablePlatforms(Cluster, ClusterReachability, SessionPlatforms,
+				FInstigator(EInstigator::GeneratedPackage, PackageData.GetPackageName()));
+			// Restart the visit with a recursive call. We will not come in here again, because HasReachablePlatforms will
+			// return true.
+			return VisitVertex(Vertex);
+		}
+	}
+
 	if (Cluster.TraversalMarkCookable())
 	{
 		bool bAnyCookable = (FirstReachableSessionPlatform == nullptr) | !bAllReachablesUncookable;
@@ -1673,9 +1712,9 @@ bool FRequestCluster::FGraphSearch::FExploreEdgesContext::TryCalculateIterativel
 
 	if (GenerationHelper && !PackageData->IsGenerated())
 	{
-		// We don't currently support separate incrementally skippable results per platform for a generator, see the
-		// class comment on FGenerationHelper. Therefore if any platform was found to be incrementally modified, then
-		// set all platforms to modified.
+		// TODO_COOKGENERATIONHELPER: We don't currently support separate incrementally skippable results per platform
+		// for a generator, see the class comment on FGenerationHelper. Therefore if any platform was found to be
+		// incrementally modified, then set all platforms to modified.
 		bool bAllUnmodified = true;
 		for (int32 PlatformIndex : PlatformsToProcess)
 		{
diff --git a/Engine/Source/Editor/UnrealEd/Private/Cooker/CookRequestCluster.h b/Engine/Source/Editor/UnrealEd/Private/Cooker/CookRequestCluster.h
--- a/Engine/Source/Editor/UnrealEd/Private/Cooker/CookRequestCluster.h
+++ b/Engine/Source/Editor/UnrealEd/Private/Cooker/CookRequestCluster.h
@@ -51,7 +51,7 @@ class FRequestCluster
 public:
 	~FRequestCluster();
 	FRequestCluster(UCookOnTheFlyServer& COTFS, TArray<FFilePlatformRequest>&& InRequests);
-	FRequestCluster(UCookOnTheFlyServer& COTFS, FPackageDataSet&& InRequests, EReachability InExploreReachability);
+	FRequestCluster(UCookOnTheFlyServer& COTFS, TPackageDataMap<ESuppressCookReason>&& InRequests, EReachability InExploreReachability);
 	FRequestCluster(UCookOnTheFlyServer& COTFS, TRingBuffer<FDiscoveryQueueElement>& DiscoveryQueue);
 	FRequestCluster(FRequestCluster&&) = delete; // Has internal pointers, would have to write manually
 	FRequestCluster(const FRequestCluster&) = delete;
diff --git a/Engine/Source/Editor/UnrealEd/Private/Cooker/CookTypes.cpp b/Engine/Source/Editor/UnrealEd/Private/Cooker/CookTypes.cpp
--- a/Engine/Source/Editor/UnrealEd/Private/Cooker/CookTypes.cpp
+++ b/Engine/Source/Editor/UnrealEd/Private/Cooker/CookTypes.cpp
@@ -86,6 +86,7 @@ const TCHAR* LexToString(ESuppressCookReason Reason)
 	case ESuppressCookReason::RetractedByCookDirector: return TEXT("RetractedByCookDirector");
 	case ESuppressCookReason::CookFilter: return TEXT("CookFilter");
 	case ESuppressCookReason::NotYetReadyForRequest: return TEXT("NotYetReadyForRequest");
+	case ESuppressCookReason::GeneratedPackageNeedsRequestUpdate: return TEXT("GeneratedPackageNeedsRequestUpdate");
 	default: return TEXT("Invalid");
 	}
 }
@@ -181,6 +182,7 @@ EStateChangeReason ConvertToStateChangeReason(ESuppressCookReason Reason)
 	case ESuppressCookReason::MultiprocessAssignmentError: return EStateChangeReason::ReassignAbortedPackages;
 	case ESuppressCookReason::RetractedByCookDirector: return EStateChangeReason::Retraction;
 	case ESuppressCookReason::NotYetReadyForRequest: return EStateChangeReason::Retraction;
+	case ESuppressCookReason::GeneratedPackageNeedsRequestUpdate: return EStateChangeReason::Retraction;
 	default: return EStateChangeReason::CookSuppressed;
 	}
 }
diff --git a/Engine/Source/Editor/UnrealEd/Private/Cooker/CookTypes.h b/Engine/Source/Editor/UnrealEd/Private/Cooker/CookTypes.h
--- a/Engine/Source/Editor/UnrealEd/Private/Cooker/CookTypes.h
+++ b/Engine/Source/Editor/UnrealEd/Private/Cooker/CookTypes.h
@@ -108,6 +108,11 @@ namespace UE::Cook
 	{
 		using TFastPointerSet<FPackageData*>::TFastPointerSet;
 	};
+	template <typename ValueType>
+	class TPackageDataMap : public TFastPointerMap<FPackageData*, ValueType>
+	{
+		using TFastPointerMap<FPackageData*, ValueType>::TFastPointerMap;
+	};
 
 	/**
 	 * External Requests to the cooker can either by cook requests for a specific file, or arbitrary callbacks that
@@ -173,6 +178,7 @@ namespace UE::Cook
 		RetractedByCookDirector,
 		CookFilter,
 		NotYetReadyForRequest,
+		GeneratedPackageNeedsRequestUpdate,
 		Count,
 		BitCount = FPlatformMath::ConstExprCeilLogTwo(Count),
 	};
diff --git a/Engine/Source/Editor/UnrealEd/Private/Cooker/CookWorkerClient.cpp b/Engine/Source/Editor/UnrealEd/Private/Cooker/CookWorkerClient.cpp
--- a/Engine/Source/Editor/UnrealEd/Private/Cooker/CookWorkerClient.cpp
+++ b/Engine/Source/Editor/UnrealEd/Private/Cooker/CookWorkerClient.cpp
@@ -173,7 +173,7 @@ void FCookWorkerClient::DoneWithInitialSettings()
 	HandleReceiveMessages(MoveTemp(DeferredInitializationMessages));
 }
 
-void FCookWorkerClient::ReportDemoteToIdle(const FPackageData& PackageData, ESuppressCookReason Reason)
+void FCookWorkerClient::ReportDemotion(const FPackageData& PackageData, ESuppressCookReason Reason)
 {
 	if (Reason == ESuppressCookReason::RetractedByCookDirector)
 	{
diff --git a/Engine/Source/Editor/UnrealEd/Private/Cooker/CookWorkerClient.h b/Engine/Source/Editor/UnrealEd/Private/Cooker/CookWorkerClient.h
--- a/Engine/Source/Editor/UnrealEd/Private/Cooker/CookWorkerClient.h
+++ b/Engine/Source/Editor/UnrealEd/Private/Cooker/CookWorkerClient.h
@@ -64,7 +64,7 @@ public:
 	void SetHasRunFinished(bool Value) { bHasRunFinished = Value; }
 
 	/** Queue a message to the server that the Package was cook-suppressed. Will be sent during Tick. */
-	void ReportDemoteToIdle(const FPackageData& PackageData, ESuppressCookReason Reason);
+	void ReportDemotion(const FPackageData& PackageData, ESuppressCookReason Reason);
 	/** Queue a message to the server that the Package was saved. Will be sent during Tick. */
 	void ReportPromoteToSaveComplete(FPackageData& PackageData);
 	/** Queue a message to the server that a package was discovered as needed in the cook. Will be sent during Tick. */
diff --git a/Engine/Source/Editor/UnrealEd/Private/Cooker/CookWorkerServer.cpp b/Engine/Source/Editor/UnrealEd/Private/Cooker/CookWorkerServer.cpp
--- a/Engine/Source/Editor/UnrealEd/Private/Cooker/CookWorkerServer.cpp
+++ b/Engine/Source/Editor/UnrealEd/Private/Cooker/CookWorkerServer.cpp
@@ -1072,9 +1072,22 @@ void FCookWorkerServer::RecordResults(FPackageResultsMessage& Message)
 		}
 		bRetiredAnyPackages = true;
 
+		bool bResultIsSaveResult = Result.GetSuppressCookReason() == ESuppressCookReason::NotSuppressed;
+		EStateChangeReason StateChangeReason;
+		bool bTerminalStateChange;
+		if (bResultIsSaveResult)
+		{
+			StateChangeReason = EStateChangeReason::Saved;
+			bTerminalStateChange = true;
+		}
+		else
+		{
+			StateChangeReason = ConvertToStateChangeReason(Result.GetSuppressCookReason());
+			bTerminalStateChange = IsTerminalStateChange(StateChangeReason);
+		}
+
 		// MPCOOKTODO: Refactor FSaveCookedPackageContext::FinishPlatform and ::FinishPackage so we can call them from
 		// here to reduce duplication
-		bool bResultIsSaveResult = Result.GetSuppressCookReason() == ESuppressCookReason::NotSuppressed;
 		if (bResultIsSaveResult)
 		{
 			HandleReceivedPackagePlatformMessages(*PackageData, nullptr, Result.ReleaseMessages());
@@ -1120,7 +1133,7 @@ void FCookWorkerServer::RecordResults(FPackageResultsMessage& Message)
 			COTFS.RecordExternalActorDependencies(Result.GetExternalActorDependencies());
 		}
 
-		// For generator and generated packages, after we handle all of their messages above, execute their state
+		// For generator and generated packages, after we handle all of their save recording above, execute their state
 		// changes in the required order:
 		// *) Defer the GenerationHelper's events so that we don't yet complete it if this was the last save.
 		// *) Mark saved on the generator, so that the generator has full context for the save.
@@ -1129,8 +1142,8 @@ void FCookWorkerServer::RecordResults(FPackageResultsMessage& Message)
 		// *) Unfreeze the GenerationHelper's events and call OnAllSavesCompleted if necessary.
 		TOptional<FGenerationHelper::FScopeDeferEvents> DeferGenerationHelperEvents;
 
-		// If generated or generator, defer events and mark saved
-		if (PackageData->IsGenerated())
+		// If generated or generator, and this is a save or other end-of-cook signal, defer events and mark saved
+		if (PackageData->IsGenerated() && bTerminalStateChange)
 		{
 			TRefCountPtr<FGenerationHelper> ParentGenerationHelper = PackageData->GetOrFindParentGenerationHelper();
 			if (!ParentGenerationHelper)
@@ -1151,10 +1164,6 @@ void FCookWorkerServer::RecordResults(FPackageResultsMessage& Message)
 						ParentGenerationHelper->MarkPackageSavedRemotely(COTFS, *PackageData, TargetPlatform, GetWorkerId());
 					}
 				}
-				EStateChangeReason StateChangeReason =
-					Result.GetSuppressCookReason() == ESuppressCookReason::NotSuppressed
-					? EStateChangeReason::Saved
-					: ConvertToStateChangeReason(Result.GetSuppressCookReason());
 				PackageData->SetParentGenerationHelper(nullptr, StateChangeReason);
 			}
 		}
@@ -1162,21 +1171,24 @@ void FCookWorkerServer::RecordResults(FPackageResultsMessage& Message)
 		if (GenerationHelper)
 		{
 			DeferGenerationHelperEvents.Emplace(GenerationHelper);
-			for (int32 PlatformIndex = 0; PlatformIndex < NumPlatforms; ++PlatformIndex)
+			if (bTerminalStateChange)
 			{
-				ITargetPlatform* TargetPlatform = OrderedSessionPlatforms[PlatformIndex];
-				FPackageRemoteResult::FPlatformResult& PlatformResult = Result.GetPlatforms()[PlatformIndex];
-				if (!bResultIsSaveResult || PlatformResult.WasCommitted())
+				for (int32 PlatformIndex = 0; PlatformIndex < NumPlatforms; ++PlatformIndex)
 				{
-					GenerationHelper->MarkPackageSavedRemotely(COTFS, *PackageData, TargetPlatform, GetWorkerId());
+					ITargetPlatform* TargetPlatform = OrderedSessionPlatforms[PlatformIndex];
+					FPackageRemoteResult::FPlatformResult& PlatformResult = Result.GetPlatforms()[PlatformIndex];
+					if (!bResultIsSaveResult || PlatformResult.WasCommitted())
+					{
+						GenerationHelper->MarkPackageSavedRemotely(COTFS, *PackageData, TargetPlatform, GetWorkerId());
+					}
 				}
 			}
 			GenerationHelper.SafeRelease();
 		}
 
-		// For all packages, transition them to idle
+		// For all packages, transition them to their next state
 		PackageData->SetWorkerAssignment(FWorkerId::Invalid(), ESendFlags::QueueNone);
-		if (Result.GetSuppressCookReason() == ESuppressCookReason::NotSuppressed)
+		if (bResultIsSaveResult)
 		{
 			ECookPhase CookPhase = COTFS.GetCookPhase();
 			if (PackageData->GetPlatformsNeedingCommitNum(CookPhase) > 0)
@@ -1205,8 +1217,16 @@ void FCookWorkerServer::RecordResults(FPackageResultsMessage& Message)
 				PackageData->SendToState(EPackageState::Request, ESendFlags::QueueAddAndRemove, EStateChangeReason::Retraction);
 			}
 		}
+		else if (!bTerminalStateChange)
+		{
+			// Non-terminal SuppressCookReasons send it back to request via COTFS.DemoteToRequest. 
+			// DemoteToRequest will also handle any request data changes indicated by the SuppressCookReason
+			COTFS.DemoteToRequest(*PackageData, ESendFlags::QueueAddAndRemove, Result.GetSuppressCookReason());
+		}
 		else
 		{
+			// Terminal SuppressCookReasons send it to idle via DemoteToIdle.
+			// DemoteToIdle will also handle any required logging.
 			COTFS.DemoteToIdle(*PackageData, ESendFlags::QueueAddAndRemove, Result.GetSuppressCookReason());
 		}
 
diff --git a/Engine/Source/Editor/UnrealEd/Private/Cooker/IWorkerRequests.h b/Engine/Source/Editor/UnrealEd/Private/Cooker/IWorkerRequests.h
--- a/Engine/Source/Editor/UnrealEd/Private/Cooker/IWorkerRequests.h
+++ b/Engine/Source/Editor/UnrealEd/Private/Cooker/IWorkerRequests.h
@@ -86,7 +86,7 @@ public:
 	virtual void OnRemoveSessionPlatform(const ITargetPlatform* TargetPlatform) = 0;
 
 	// Functions that are written locally for Directors but are sent to Director for Workers
-	virtual void ReportDemoteToIdle(UE::Cook::FPackageData& PackageData, ESuppressCookReason Reason) = 0;
+	virtual void ReportDemotion(UE::Cook::FPackageData& PackageData, ESuppressCookReason Reason) = 0;
 	virtual void ReportPromoteToSaveComplete(UE::Cook::FPackageData& PackageData) = 0;
 
 	// Functions that read from local for Directors but read from Director for Workers
diff --git a/Engine/Source/Editor/UnrealEd/Private/Cooker/WorkerRequestsLocal.cpp b/Engine/Source/Editor/UnrealEd/Private/Cooker/WorkerRequestsLocal.cpp
--- a/Engine/Source/Editor/UnrealEd/Private/Cooker/WorkerRequestsLocal.cpp
+++ b/Engine/Source/Editor/UnrealEd/Private/Cooker/WorkerRequestsLocal.cpp
@@ -117,7 +117,7 @@ void FWorkerRequestsLocal::OnRemoveSessionPlatform(const ITargetPlatform* Target
 	ExternalRequests.OnRemoveSessionPlatform(TargetPlatform);
 }
 
-void FWorkerRequestsLocal::ReportDemoteToIdle(UE::Cook::FPackageData& PackageData, ESuppressCookReason Reason)
+void FWorkerRequestsLocal::ReportDemotion(UE::Cook::FPackageData& PackageData, ESuppressCookReason Reason)
 {
 }
 
diff --git a/Engine/Source/Editor/UnrealEd/Private/Cooker/WorkerRequestsLocal.h b/Engine/Source/Editor/UnrealEd/Private/Cooker/WorkerRequestsLocal.h
--- a/Engine/Source/Editor/UnrealEd/Private/Cooker/WorkerRequestsLocal.h
+++ b/Engine/Source/Editor/UnrealEd/Private/Cooker/WorkerRequestsLocal.h
@@ -44,7 +44,7 @@ public:
 	virtual void AddEditorActionCallback(FSchedulerCallback&& Callback) override;
 	virtual void AddPublicInterfaceRequest(FFilePlatformRequest&& Request, bool bForceFrontOfQueue) override;
 
-	virtual void ReportDemoteToIdle(UE::Cook::FPackageData& PackageData, ESuppressCookReason Reason) override;
+	virtual void ReportDemotion(UE::Cook::FPackageData& PackageData, ESuppressCookReason Reason) override;
 	virtual void ReportPromoteToSaveComplete(UE::Cook::FPackageData& PackageData) override;
 	virtual void GetInitializeConfigSettings(UCookOnTheFlyServer& COTFS, const FString& OutputOverrideDirectory,
 		UE::Cook::FInitializeConfigSettings& Settings) override;
diff --git a/Engine/Source/Editor/UnrealEd/Private/Cooker/WorkerRequestsRemote.cpp b/Engine/Source/Editor/UnrealEd/Private/Cooker/WorkerRequestsRemote.cpp
--- a/Engine/Source/Editor/UnrealEd/Private/Cooker/WorkerRequestsRemote.cpp
+++ b/Engine/Source/Editor/UnrealEd/Private/Cooker/WorkerRequestsRemote.cpp
@@ -113,9 +113,9 @@ void FWorkerRequestsRemote::OnRemoveSessionPlatform(const ITargetPlatform* Targe
 	ExternalRequests.OnRemoveSessionPlatform(TargetPlatform);
 }
 
-void FWorkerRequestsRemote::ReportDemoteToIdle(UE::Cook::FPackageData& PackageData, ESuppressCookReason Reason)
+void FWorkerRequestsRemote::ReportDemotion(UE::Cook::FPackageData& PackageData, ESuppressCookReason Reason)
 {
-	CookWorkerClient.ReportDemoteToIdle(PackageData, Reason);
+	CookWorkerClient.ReportDemotion(PackageData, Reason);
 }
 
 void FWorkerRequestsRemote::ReportPromoteToSaveComplete(UE::Cook::FPackageData& PackageData)
diff --git a/Engine/Source/Editor/UnrealEd/Private/Cooker/WorkerRequestsRemote.h b/Engine/Source/Editor/UnrealEd/Private/Cooker/WorkerRequestsRemote.h
--- a/Engine/Source/Editor/UnrealEd/Private/Cooker/WorkerRequestsRemote.h
+++ b/Engine/Source/Editor/UnrealEd/Private/Cooker/WorkerRequestsRemote.h
@@ -52,7 +52,7 @@ public:
 	virtual void AddEditorActionCallback(FSchedulerCallback&& Callback) override;
 	virtual void AddPublicInterfaceRequest(FFilePlatformRequest&& Request, bool bForceFrontOfQueue) override;
 
-	virtual void ReportDemoteToIdle(UE::Cook::FPackageData& PackageData, ESuppressCookReason Reason) override;
+	virtual void ReportDemotion(UE::Cook::FPackageData& PackageData, ESuppressCookReason Reason) override;
 	virtual void ReportPromoteToSaveComplete(UE::Cook::FPackageData& PackageData) override;
 	virtual void GetInitializeConfigSettings(UCookOnTheFlyServer& COTFS, const FString& OutputOverrideDirectory,
 		UE::Cook::FInitializeConfigSettings& Settings) override;
