Only scan the paths that we will actually gather UGC localization from, rather than everything on-disk

This can be significantly faster for a cold asset registry, as it only has to scan the contents of a single plugin, rather than the contents of the entire UE project

#rnx

diff --git a/Engine/Source/Developer/Localization/Private/UserGeneratedContentLocalization.cpp b/Engine/Source/Developer/Localization/Private/UserGeneratedContentLocalization.cpp
--- a/Engine/Source/Developer/Localization/Private/UserGeneratedContentLocalization.cpp
+++ b/Engine/Source/Developer/Localization/Private/UserGeneratedContentLocalization.cpp
@@ -407,6 +407,8 @@ bool ExportLocalization(TArrayView<const TSharedRef<IPlugin>> Plugins, const FEx
 					ConfigSection.Add(TEXT("CollectionFilters"), *CollectionFilter);
 				}
 				
+				ConfigSection.Add(TEXT("SearchAllAssets"), TEXT("false"));
+
 				GatherConfig.AddGatherTextStep(GatherStepIndex++, MoveTemp(ConfigSection));
 			}
 
diff --git a/Engine/Source/Editor/UnrealEd/Classes/Commandlets/GatherTextFromAssetsCommandlet.h b/Engine/Source/Editor/UnrealEd/Classes/Commandlets/GatherTextFromAssetsCommandlet.h
--- a/Engine/Source/Editor/UnrealEd/Classes/Commandlets/GatherTextFromAssetsCommandlet.h
+++ b/Engine/Source/Editor/UnrealEd/Classes/Commandlets/GatherTextFromAssetsCommandlet.h
@@ -135,6 +135,7 @@ private:
 	bool bReportStaleGatherCache;
 	bool bFixStaleGatherCache;
 	bool bFixMissingGatherCache;
+	bool bSearchAllAssets;
 	bool bShouldGatherFromEditorOnlyData;
 	bool bShouldExcludeDerivedClasses;
 	bool bFixPackageLocalizationIdConflict;
diff --git a/Engine/Source/Editor/UnrealEd/Private/Commandlets/GatherTextFromAssetsCommandlet.cpp b/Engine/Source/Editor/UnrealEd/Private/Commandlets/GatherTextFromAssetsCommandlet.cpp
--- a/Engine/Source/Editor/UnrealEd/Private/Commandlets/GatherTextFromAssetsCommandlet.cpp
+++ b/Engine/Source/Editor/UnrealEd/Private/Commandlets/GatherTextFromAssetsCommandlet.cpp
@@ -283,6 +283,7 @@ UGatherTextFromAssetsCommandlet::UGatherTextFromAssetsCommandlet(const FObjectIn
 	, MaxUsedMemoryBytes(0)
 	, NumPackagesDupLocId(0)
 	, bSkipGatherCache(false)
+	, bSearchAllAssets(true)
 	, bShouldGatherFromEditorOnlyData(false)
 	, bShouldExcludeDerivedClasses(false)
 	, bFixPackageLocalizationIdConflict(false)
@@ -1590,20 +1591,37 @@ int32 UGatherTextFromAssetsCommandlet::Main(const FString& Params)
 			GEditor->CreateNewMapForEditing(/*bPromptForSave*/false);
 		}
 	}
+
+	FARFilter FirstPassFilter;
+	if (!BuildFirstPassFilter(FirstPassFilter))
+	{
+		return -1;
+	}
+	
 	UE_LOG(LogGatherTextFromAssetsCommandlet, Display, TEXT("Discovering assets to gather..."));
 	const double DiscoveringAssetsStartTime = FPlatformTime::Seconds();
 	{
 		UE_SCOPED_TIMER(TEXT("UGatherTextFromAssetsCommandlet::SearchAssetRegistryForAllAssets"), LogGatherTextFromAssetsCommandlet, Display);
 		FAssetRegistryModule& AssetRegistryModule = FModuleManager::LoadModuleChecked<FAssetRegistryModule>(TEXT("AssetRegistry"));
 		IAssetRegistry& AssetRegistry = AssetRegistryModule.Get();
-		AssetRegistry.SearchAllAssets(true);
+
+		if (bSearchAllAssets || FirstPassFilter.PackagePaths.IsEmpty())
+		{
+			AssetRegistry.SearchAllAssets(true);
+		}
+		else
+		{
+			TArray<FString> ScanPaths;
+			ScanPaths.Reserve(FirstPassFilter.PackagePaths.Num());
+			Algo::Transform(FirstPassFilter.PackagePaths, ScanPaths, [](FName PackagePath) { return PackagePath.ToString(); });
+
+			// Note: We don't use FirstPassFilter.SoftObjectPaths as the set of files to scan, as ScanSynchronous can perform poorly when given large numbers of files to scan
+			AssetRegistry.ScanSynchronous(ScanPaths, TArray<FString>(), UE::AssetRegistry::EScanFlags::IgnoreInvalidPathWarning);
+		}
 	}
 
 	TArray<FAssetData> AssetDataArray;
-	if (!PerformFirstPassFilter(AssetDataArray))
-	{
-		return -1;
-	}
+	ApplyFirstPassFilter(FirstPassFilter, AssetDataArray);
 
 	if (!bShouldExcludeDerivedClasses)
 	{
@@ -1817,6 +1835,11 @@ bool UGatherTextFromAssetsCommandlet::ConfigureFromScript(const FString& GatherT
 
 	GetPathArrayFromConfig(*SectionName, TEXT("ManifestDependencies"), ManifestDependenciesList, GatherTextConfigPath);
 
+	if (!GetBoolFromConfig(*SectionName, TEXT("SearchAllAssets"), bSearchAllAssets, GatherTextConfigPath))
+	{
+		bSearchAllAssets = true;
+	}
+
 	// Get whether we should gather editor-only data. Typically only useful for the localization of UE itself.
 	if (!GetBoolFromConfig(*SectionName, TEXT("ShouldGatherFromEditorOnlyData"), bShouldGatherFromEditorOnlyData, GatherTextConfigPath))
 	{
