UnrealEd: Fix shadowed enum

#rnx

diff --git a/Engine/Source/Editor/UnrealEd/Private/StaticLightingSystem/StaticLightingPrivate.h b/Engine/Source/Editor/UnrealEd/Private/StaticLightingSystem/StaticLightingPrivate.h
--- a/Engine/Source/Editor/UnrealEd/Private/StaticLightingSystem/StaticLightingPrivate.h
+++ b/Engine/Source/Editor/UnrealEd/Private/StaticLightingSystem/StaticLightingPrivate.h
@@ -411,7 +411,7 @@ private:
 	int32 NextVisibilityId;
 
 	/** Lighting comes in various stages (amortized, async, etc.), we track them here. */
-	enum LightingStage
+	enum class LightingStage
 	{
 		NotRunning,
 		Startup,
diff --git a/Engine/Source/Editor/UnrealEd/Private/StaticLightingSystem/StaticLightingSystem.cpp b/Engine/Source/Editor/UnrealEd/Private/StaticLightingSystem/StaticLightingSystem.cpp
--- a/Engine/Source/Editor/UnrealEd/Private/StaticLightingSystem/StaticLightingSystem.cpp
+++ b/Engine/Source/Editor/UnrealEd/Private/StaticLightingSystem/StaticLightingSystem.cpp
@@ -245,7 +245,7 @@ void FStaticLightingManager::ImportRequested()
 {
 	if (FStaticLightingManager::Get()->ActiveStaticLightingSystem)
 	{
-		FStaticLightingManager::Get()->ActiveStaticLightingSystem->CurrentBuildStage = FStaticLightingSystem::ImportRequested;
+		FStaticLightingManager::Get()->ActiveStaticLightingSystem->CurrentBuildStage = FStaticLightingSystem::LightingStage::ImportRequested;
 	}
 }
 
@@ -254,7 +254,7 @@ void FStaticLightingManager::DiscardRequested()
 	if (FStaticLightingManager::Get()->ActiveStaticLightingSystem)
 	{
 		FStaticLightingManager::Get()->ClearCurrentNotification();
-		FStaticLightingManager::Get()->ActiveStaticLightingSystem->CurrentBuildStage = FStaticLightingSystem::Finished;
+		FStaticLightingManager::Get()->ActiveStaticLightingSystem->CurrentBuildStage = FStaticLightingSystem::LightingStage::Finished;
 	}
 }
 
@@ -369,7 +369,7 @@ void FStaticLightingManager::UpdateBuildLighting()
 		// Note: UpdateLightingBuild can change ActiveStaticLightingSystem
 		ActiveStaticLightingSystem->UpdateLightingBuild();
 
-		if (ActiveStaticLightingSystem && ActiveStaticLightingSystem->CurrentBuildStage == FStaticLightingSystem::Finished)
+		if (ActiveStaticLightingSystem && ActiveStaticLightingSystem->CurrentBuildStage == FStaticLightingSystem::LightingStage::Finished)
 		{
 			ActiveStaticLightingSystem = nullptr;
 			StaticLightingSystems.RemoveAt(0);
@@ -508,7 +508,7 @@ FStaticLightingSystem::FStaticLightingSystem(const FLightingBuildOptions& InOpti
 	, bBuildCanceled(false)
 	, DeterministicIndex(0)
 	, NextVisibilityId(0)
-	, CurrentBuildStage(FStaticLightingSystem::NotRunning)
+	, CurrentBuildStage(FStaticLightingSystem::LightingStage::NotRunning)
 	, LightingContext(MoveTemp(context))
 	, LightmassProcessor(NULL)
 {
@@ -538,7 +538,7 @@ bool FStaticLightingSystem::BeginLightmassProcess()
 {
 	StartTime = FPlatformTime::Seconds();
 	
-	CurrentBuildStage = FStaticLightingSystem::Startup;
+	CurrentBuildStage = FStaticLightingSystem::LightingStage::Startup;
 
 	bool bRebuildDirtyGeometryForLighting = true;
 	bool bForceNoPrecomputedLighting = false;
@@ -2457,7 +2457,7 @@ bool FStaticLightingSystem::InitiateLightmassProcessor()
 		{
 			LightmassProcessor->InitiateExport();
 			bSuccessful = true;
-			CurrentBuildStage = FStaticLightingSystem::AmortizedExport;
+			CurrentBuildStage = FStaticLightingSystem::LightingStage::AmortizedExport;
 		}
 	}
 	
@@ -2470,7 +2470,7 @@ void FStaticLightingSystem::KickoffSwarm()
 	
 	if (bSuccessful)
 	{
-		CurrentBuildStage = FStaticLightingSystem::AsynchronousBuilding;
+		CurrentBuildStage = FStaticLightingSystem::LightingStage::AsynchronousBuilding;
 	}
 	else
 	{
@@ -2484,7 +2484,7 @@ bool FStaticLightingSystem::FinishLightmassProcess()
 
 	GEditor->ResetTransaction( LOCTEXT("KeepLightingTransReset", "Applying Lighting") );
 
-	CurrentBuildStage = FStaticLightingSystem::Import;
+	CurrentBuildStage = FStaticLightingSystem::LightingStage::Import;
 
 	double TimeWaitingOnUserToAccept = FPlatformTime::Seconds() - WaitForUserAcceptStartTime;
 	
@@ -2564,7 +2564,7 @@ bool FStaticLightingSystem::FinishLightmassProcess()
 
 void FStaticLightingSystem::UpdateLightingBuild()
 {
-	if (CurrentBuildStage == FStaticLightingSystem::AmortizedExport)
+	if (CurrentBuildStage == FStaticLightingSystem::LightingStage::AmortizedExport)
 	{
 		bool bCompleted = LightmassProcessor->ExecuteAmortizedMaterialExport();
 
@@ -2576,16 +2576,16 @@ void FStaticLightingSystem::UpdateLightingBuild()
 
 		if (bCompleted)
 		{
-			CurrentBuildStage = FStaticLightingSystem::SwarmKickoff;
+			CurrentBuildStage = FStaticLightingSystem::LightingStage::SwarmKickoff;
 		}
 	}
-	else if (CurrentBuildStage == FStaticLightingSystem::SwarmKickoff)
+	else if (CurrentBuildStage == FStaticLightingSystem::LightingStage::SwarmKickoff)
 	{
 		FText Text = LOCTEXT("LightKickoffSwarmMessage", "Kicking off Swarm");
 		FStaticLightingManager::Get()->SetNotificationText( Text );
 		KickoffSwarm();
 	}
-	else if (CurrentBuildStage == FStaticLightingSystem::AsynchronousBuilding)
+	else if (CurrentBuildStage == FStaticLightingSystem::LightingStage::AsynchronousBuilding)
 	{
 		bool bFinished = LightmassProcessor->Update();
 		
@@ -2609,17 +2609,17 @@ void FStaticLightingSystem::UpdateLightingBuild()
 
 			if (LightmassProcessor->IsProcessingCompletedSuccessfully())
 			{
-				CurrentBuildStage = FStaticLightingSystem::AutoApplyingImport;
+				CurrentBuildStage = FStaticLightingSystem::LightingStage::AutoApplyingImport;
 			}
 			else
 			{
 				// automatically fail lighting build (discard)
 				FStaticLightingManager::Get()->FailLightingBuild();
-				CurrentBuildStage = FStaticLightingSystem::Finished;
+				CurrentBuildStage = FStaticLightingSystem::LightingStage::Finished;
 			}
 		}
 	}
-	else if ( CurrentBuildStage == FStaticLightingSystem::AutoApplyingImport )
+	else if ( CurrentBuildStage == FStaticLightingSystem::LightingStage::AutoApplyingImport )
 	{
 		if (IsRunningCommandlet() || CanAutoApplyLighting())
 		{
@@ -2627,20 +2627,20 @@ void FStaticLightingSystem::UpdateLightingBuild()
 			FStaticLightingManager::Get()->SendBuildDoneNotification(bAutoApplyFailed);
 
 			FStaticLightingManager::ProcessLightingData();
-			CurrentBuildStage = FStaticLightingSystem::Finished;
+			CurrentBuildStage = FStaticLightingSystem::LightingStage::Finished;
 		}
 		else
 		{
 			bool bAutoApplyFailed = true;
 			FStaticLightingManager::Get()->SendBuildDoneNotification(bAutoApplyFailed);
 
-			CurrentBuildStage = FStaticLightingSystem::WaitingForImport;
+			CurrentBuildStage = FStaticLightingSystem::LightingStage::WaitingForImport;
 		}
 	}
-	else if (CurrentBuildStage == FStaticLightingSystem::ImportRequested)
+	else if (CurrentBuildStage == FStaticLightingSystem::LightingStage::ImportRequested)
 	{
 		FStaticLightingManager::ProcessLightingData();
-		CurrentBuildStage = FStaticLightingSystem::Finished;
+		CurrentBuildStage = FStaticLightingSystem::LightingStage::Finished;
 	}
 }
 
@@ -2747,12 +2747,12 @@ const FStaticLightingBuildContext& FStaticLightingSystem::GetLightingContext() c
 
 bool FStaticLightingSystem::IsAsyncBuilding() const
 {
-	return CurrentBuildStage == FStaticLightingSystem::AsynchronousBuilding;
+	return CurrentBuildStage == FStaticLightingSystem::LightingStage::AsynchronousBuilding;
 }
 
 bool FStaticLightingSystem::IsAmortizedExporting() const
 {
-	return CurrentBuildStage == FStaticLightingSystem::AmortizedExport;
+	return CurrentBuildStage == FStaticLightingSystem::LightingStage::AmortizedExport;
 }
 
 void UEditorEngine::BuildLighting(const FLightingBuildOptions& Options)
