Renamed CollectionExistsWithAnyShareType to GetCollectionsByName.
Updated GetCollectionsByName to return an array of the collections found.

diff --git a/Engine/Plugins/Runtime/AssetTags/Source/AssetTags/Private/AssetTagsSubsystem.cpp b/Engine/Plugins/Runtime/AssetTags/Source/AssetTags/Private/AssetTagsSubsystem.cpp
--- a/Engine/Plugins/Runtime/AssetTags/Source/AssetTags/Private/AssetTagsSubsystem.cpp
+++ b/Engine/Plugins/Runtime/AssetTags/Source/AssetTags/Private/AssetTagsSubsystem.cpp
@@ -322,7 +322,9 @@ bool UAssetTagsSubsystem::CollectionExists(const FName Name)
 		if (UCollectionManagerScriptingSubsystem* Subsystem = GEditor->GetEditorSubsystem<UCollectionManagerScriptingSubsystem>(); ensure(Subsystem))
 		{
 			const TSharedRef<ICollectionContainer>& CollectionContainer = FCollectionManagerModule::GetModule().Get().GetProjectCollectionContainer();
-			bExists = Subsystem->CollectionExistsWithAnyShareType(FCollectionScriptingContainerSource{ CollectionContainer->GetCollectionSource()->GetName(), CollectionContainer->GetCollectionSource()->GetTitle() }, Name);
+
+			TArray<FCollectionScriptingRef> FoundCollections;
+			bExists = Subsystem->GetCollectionsByName(FCollectionScriptingContainerSource{ CollectionContainer->GetCollectionSource()->GetName(), CollectionContainer->GetCollectionSource()->GetTitle() }, Name, FoundCollections);
 		}
 	}
 #else	// WITH_EDITOR
diff --git a/Engine/Source/Editor/UnrealEd/Private/Subsystems/CollectionManagerScriptingSubsystem.cpp b/Engine/Source/Editor/UnrealEd/Private/Subsystems/CollectionManagerScriptingSubsystem.cpp
--- a/Engine/Source/Editor/UnrealEd/Private/Subsystems/CollectionManagerScriptingSubsystem.cpp
+++ b/Engine/Source/Editor/UnrealEd/Private/Subsystems/CollectionManagerScriptingSubsystem.cpp
@@ -366,7 +366,7 @@ bool UCollectionManagerScriptingSubsystem::CollectionExists(const FCollectionScr
 	return CollectionContainer->CollectionExists(Collection, CollectionManagerScriptingSubsystemUtil::ECollectionScriptingShareType_To_ECollectionShareType(ShareType));
 }
 
-bool UCollectionManagerScriptingSubsystem::CollectionExistsWithAnyShareType(const FCollectionScriptingContainerSource Container, const FName Collection)
+bool UCollectionManagerScriptingSubsystem::GetCollectionsByName(const FCollectionScriptingContainerSource Container, const FName Collection, TArray<FCollectionScriptingRef>& OutCollections)
 {
 	const TSharedPtr<ICollectionContainer> CollectionContainer = FCollectionManagerModule::GetModule().Get().FindCollectionContainer(Container.Name);
 	if (!CollectionContainer.IsValid())
@@ -375,7 +375,20 @@ bool UCollectionManagerScriptingSubsystem::CollectionExistsWithAnyShareType(cons
 		return false;
 	}
 
-	return CollectionContainer->CollectionExists(Collection, ECollectionShareType::CST_All);
+	TArray<FCollectionNameType> FoundCollections;
+	CollectionContainer->GetCollections(Collection, FoundCollections);
+
+	OutCollections.Reserve(FoundCollections.Num());
+	Algo::Transform(FoundCollections, OutCollections, [ContainerName = Container.Name](const FCollectionNameType& Collection)
+	{
+		return FCollectionScriptingRef { 
+			.Container = ContainerName,
+			.Name = Collection.Name,
+			.ShareType = CollectionManagerScriptingSubsystemUtil::ECollectionShareType_To_ECollectionScriptingShareType(Collection.Type)
+		};
+	});
+	
+	return true;
 }
 
 bool UCollectionManagerScriptingSubsystem::GetAssetsInCollection(const FCollectionScriptingRef& Collection, TArray<FAssetData>& OutAssets)
diff --git a/Engine/Source/Editor/UnrealEd/Public/Subsystems/CollectionManagerScriptingSubsystem.h b/Engine/Source/Editor/UnrealEd/Public/Subsystems/CollectionManagerScriptingSubsystem.h
--- a/Engine/Source/Editor/UnrealEd/Public/Subsystems/CollectionManagerScriptingSubsystem.h
+++ b/Engine/Source/Editor/UnrealEd/Public/Subsystems/CollectionManagerScriptingSubsystem.h
@@ -73,7 +73,7 @@ public:
 	 * @param Container The container to retrieve collections from. "None" defaults to the base game's container.
 	 * @param OutCollections The collections found in the container.
 	 *
-	 * @return True if the specified container was found, false otherwise (see the output log for details on error).
+	 * @return True if the container exists, false otherwise (see the output log for details on error).
 	 */
 	UFUNCTION(BlueprintCallable, Category="CollectionManager")
 	bool GetCollections(const FCollectionScriptingContainerSource Container, TArray<FCollectionScriptingRef>& OutCollections);
@@ -260,21 +260,22 @@ public:
 	 * @param Collection The name of the collection to look for.
 	 * @param ShareType The share type of the collection to look for.
 	 *
-	 * @return True if the collection exists, false otherwise.
+	 * @return True if the collection exists, false otherwise (if false is due to an error, see the output log for details).
 	 */
 	UFUNCTION(BlueprintCallable, Category="CollectionManager")
 	bool CollectionExists(const FCollectionScriptingContainerSource Container, const FName Collection, const ECollectionScriptingShareType ShareType);
 
 	/**
-	 * Check whether the given collection exists in the given container (matching only by name).
+	 * Gets the given collections in the given container (matching only by name).
 	 *
 	 * @param Container The container to search. "None" defaults to the base game's container.
-	 * @param Collection The collection to look for. 
+	 * @param Collection The collection to look for.
+	 * @param OutCollections The collections found.
 	 *
-	 * @return True if the collection exists (returns the first found with a matching name), false otherwise.
+	 * @return True if the container exists, false otherwise (see the output log for details on error).
 	 */
 	UFUNCTION(BlueprintCallable, Category = "CollectionManager")
-	bool CollectionExistsWithAnyShareType(const FCollectionScriptingContainerSource Container, const FName Collection);
+	bool GetCollectionsByName(const FCollectionScriptingContainerSource Container, const FName Collection, TArray<FCollectionScriptingRef>& OutCollections);
 
 	/**
 	 * Get the assets in the given collection.
