[ITF Drag Tools] Behaviors don't receive information about the viewport losing focus, leading to potential input loss and mix ups

Forward FocustLost information to UModeManagerInteractiveToolsContext, so that it can call InputRouter->ForceTerminateAll()

An example of the issue would be the linked jira:

- Key Modifier Behavior active and registering Key Mod Press (e.g. Ctrl)
- User clicks outside of the viewport, e.g. Ctrl + Click on some outliner element to de-select or select
- Ctrl is released
- Input Router never knows about this
- Key modifier might interpret Ctrl as pressed once focus is back on viewport

diff --git a/Engine/Source/Editor/Experimental/EditorInteractiveToolsFramework/Private/EditorDragTools/EditorDragToolBehaviorTarget.cpp b/Engine/Source/Editor/Experimental/EditorInteractiveToolsFramework/Private/EditorDragTools/EditorDragToolBehaviorTarget.cpp
--- a/Engine/Source/Editor/Experimental/EditorInteractiveToolsFramework/Private/EditorDragTools/EditorDragToolBehaviorTarget.cpp
+++ b/Engine/Source/Editor/Experimental/EditorInteractiveToolsFramework/Private/EditorDragTools/EditorDragToolBehaviorTarget.cpp
@@ -92,6 +92,11 @@ void FEditorDragToolBehaviorTarget::OnClickRelease(const FInputDeviceRay& InRele
 	OnDeactivateTool().Broadcast();
 }
 
+void FEditorDragToolBehaviorTarget::OnForceEndCapture()
+{
+	OnTerminateDragSequence();
+}
+
 void FEditorDragToolBehaviorTarget::OnTerminateDragSequence()
 {
 	Start = End = FVector::ZeroVector;
diff --git a/Engine/Source/Editor/Experimental/EditorInteractiveToolsFramework/Public/EditorDragTools/EditorDragToolBehaviorTarget.h b/Engine/Source/Editor/Experimental/EditorInteractiveToolsFramework/Public/EditorDragTools/EditorDragToolBehaviorTarget.h
--- a/Engine/Source/Editor/Experimental/EditorInteractiveToolsFramework/Public/EditorDragTools/EditorDragToolBehaviorTarget.h
+++ b/Engine/Source/Editor/Experimental/EditorInteractiveToolsFramework/Public/EditorDragTools/EditorDragToolBehaviorTarget.h
@@ -88,6 +88,8 @@ public:
 
 	virtual void OnClickRelease(const FInputDeviceRay& InReleasePos) override;
 
+	virtual void OnForceEndCapture() override;
+
 	virtual void OnTerminateDragSequence() override;
 	//~ End IClickDragBehaviorTarget
 
diff --git a/Engine/Source/Editor/UnrealEd/Private/EditorModeManager.cpp b/Engine/Source/Editor/UnrealEd/Private/EditorModeManager.cpp
--- a/Engine/Source/Editor/UnrealEd/Private/EditorModeManager.cpp
+++ b/Engine/Source/Editor/UnrealEd/Private/EditorModeManager.cpp
@@ -1605,7 +1605,17 @@ bool FEditorModeTools::LostFocus( FEditorViewportClient* InViewportClient, FView
 	// TODO: Conceptually, we should probably clear FocusedViewportClient here, but also have a
 	// LastFocusedViewportClient property that we don't clear, to use in ticking.
 
-	bool bHandled = false;
+	bool bHandled = InteractiveToolsContext->LostFocus(InViewportClient, Viewport);
+
+	for (const UEdMode* const EdMode : ActiveScriptableModes)
+	{
+		if (const UModeManagerInteractiveToolsContext* const ModeManagerInteractiveToolsContext =
+				EdMode->GetModeManager()->InteractiveToolsContext)
+		{
+			ModeManagerInteractiveToolsContext->LostFocus(InViewportClient, Viewport);
+		}
+	}
+
 	ForEachEdMode<ILegacyEdModeViewportInterface>([&bHandled, InViewportClient, Viewport](ILegacyEdModeViewportInterface* Mode)
 		{
 			bHandled |= Mode->LostFocus(InViewportClient, Viewport);
diff --git a/Engine/Source/Editor/UnrealEd/Private/Tools/EdModeInteractiveToolsContext.cpp b/Engine/Source/Editor/UnrealEd/Private/Tools/EdModeInteractiveToolsContext.cpp
--- a/Engine/Source/Editor/UnrealEd/Private/Tools/EdModeInteractiveToolsContext.cpp
+++ b/Engine/Source/Editor/UnrealEd/Private/Tools/EdModeInteractiveToolsContext.cpp
@@ -1115,6 +1115,12 @@ bool UModeManagerInteractiveToolsContext::MouseMove(FEditorViewportClient* Viewp
 	return false;
 }
 
+bool UModeManagerInteractiveToolsContext::LostFocus(FEditorViewportClient* InViewportClient, FViewport* Viewport) const
+{
+	InputRouter->ForceTerminateAll();
+
+	return false;
+}
 
 bool UModeManagerInteractiveToolsContext::MouseLeave(FEditorViewportClient* ViewportClient, FViewport* Viewport)
 {
diff --git a/Engine/Source/Editor/UnrealEd/Public/Tools/EdModeInteractiveToolsContext.h b/Engine/Source/Editor/UnrealEd/Public/Tools/EdModeInteractiveToolsContext.h
--- a/Engine/Source/Editor/UnrealEd/Public/Tools/EdModeInteractiveToolsContext.h
+++ b/Engine/Source/Editor/UnrealEd/Public/Tools/EdModeInteractiveToolsContext.h
@@ -293,6 +293,7 @@ public:
 	bool MouseEnter(FEditorViewportClient* ViewportClient, FViewport* Viewport, int32 x, int32 y);
 	bool MouseLeave(FEditorViewportClient* ViewportClient, FViewport* Viewport);
 	bool MouseMove(FEditorViewportClient* ViewportClient, FViewport* Viewport, int32 x, int32 y);
+	bool LostFocus(FEditorViewportClient* InViewportClient, FViewport* Viewport) const;
 
 	bool StartTracking(FEditorViewportClient* InViewportClient, FViewport* InViewport);
 	bool CapturedMouseMove(FEditorViewportClient* InViewportClient, FViewport* InViewport, int32 InMouseX, int32 InMouseY);
