PR #12467: Fix multiple EditorSetting typo's

diff --git a/Engine/Source/Editor/UnrealEd/Public/Settings/EditorSettings.h b/Engine/Source/Editor/UnrealEd/Public/Settings/EditorSettings.h
--- a/Engine/Source/Editor/UnrealEd/Public/Settings/EditorSettings.h
+++ b/Engine/Source/Editor/UnrealEd/Public/Settings/EditorSettings.h
@@ -61,7 +61,7 @@ class UEditorSettings : public UObject
 	// =====================================================================
 
 	/**
-	 * Adjusts the Local Cache location. This affects every project on your computer that uses the UE-LocalDataCachePath environment environment variable override.
+	 * Adjusts the Local Cache location. This affects every project on your computer that uses the UE-LocalDataCachePath environment variable override.
 	 * This is usually the first location to query for previously built data.
 	 */
 	UPROPERTY(EditAnywhere, Category = DerivedDataCache, meta = (DisplayName = "Global Local DDC Path", ConfigRestartRequired = true))
@@ -69,27 +69,27 @@ class UEditorSettings : public UObject
 
 	/**
 	 * Adjusts the Shared cache location. This affects every project on your computer that uses the UE-SharedDataCachePath environment variable override.
-	 * The Shared Cache location is usually queried if we do't find previously built data in the Local cache. Colleauges should point to the same shared location so that work can be distributed. 
+	 * The Shared Cache location is usually queried if we don't find previously built data in the Local cache. Colleagues should point to the same shared location so that work can be distributed. 
 	 */
 	UPROPERTY(EditAnywhere, Category = DerivedDataCache, meta = (DisplayName = "Global Shared DDC Path", ConfigRestartRequired = true))
 	FDirectoryPath GlobalSharedDDCPath;
 
 	/**
-	 * Project specific overide for the Local Cache location. The editor must be restarted for changes to take effect.
+	 * Project specific override for the Local Cache location. The editor must be restarted for changes to take effect.
 	 * This will override the 'Global Local DDC Path'.
 	 */
 	UPROPERTY(EditAnywhere, config, Category= DerivedDataCache, AdvancedDisplay, meta = (DisplayName = "Project Local DDC Path", ConfigRestartRequired = true))
 	FDirectoryPath LocalDerivedDataCache;
 
 	/**
-	 * Project specific overide for the Shared Cache location. The editor must be restarted for changes to take effect.
+	 * Project specific override for the Shared Cache location. The editor must be restarted for changes to take effect.
 	 * This will override the 'Global Shared DDC Path'.
 	 */
 	UPROPERTY(EditAnywhere, config, Category= DerivedDataCache, AdvancedDisplay, meta = (DisplayName = "Project Shared DDC Path", ConfigRestartRequired = true))
 	FDirectoryPath SharedDerivedDataCache;
 
 	/** Whether to enable any DDC Notifications */
-	UPROPERTY(EditAnywhere, config, Category = "Derived Data Cache Notifications", meta = (DisplayName = "Enable Notifcations", ConfigRestartRequired = false))
+	UPROPERTY(EditAnywhere, config, Category = "Derived Data Cache Notifications", meta = (DisplayName = "Enable Notifications", ConfigRestartRequired = false))
 	bool bEnableDDCNotifications = true;
 
 	/** Whether to enable the Unreal Cloud DDC notification */
