PR #12602: Updated "Create Socket" Button to match the style of other similar buttons

diff --git a/Engine/Source/Editor/UnrealEd/Private/SSocketManager.cpp b/Engine/Source/Editor/UnrealEd/Private/SSocketManager.cpp
--- a/Engine/Source/Editor/UnrealEd/Private/SSocketManager.cpp
+++ b/Engine/Source/Editor/UnrealEd/Private/SSocketManager.cpp
@@ -8,7 +8,7 @@
 #include "UObject/UObjectHash.h"
 #include "UObject/UObjectIterator.h"
 #include "Widgets/Layout/SSeparator.h"
-#include "Widgets/Input/SButton.h"
+#include "Developer/ToolWidgets/Public/SPositiveActionButton.h"
 #include "Widgets/Views/SListView.h"
 #include "Styling/AppStyle.h"
 #include "Components/StaticMeshComponent.h"
@@ -314,19 +314,13 @@ void SSocketManager::Construct(const FArguments& InArgs)
 						.AutoWidth()
 						.VAlign(VAlign_Center)
 						[
-							SNew(SButton)
-							.ButtonStyle(FAppStyle::Get(), "SimpleButton")
+							SNew(SPositiveActionButton)
 							.ToolTipText(LOCTEXT("CreateSocket", "Create Socket"))
+							.Icon(FAppStyle::Get().GetBrush("Icons.Plus"))
 							.OnClicked(this, &SSocketManager::CreateSocket_Execute)
 							.Visibility(this, &SSocketManager::CreateSocket_IsVisible )
-							.VAlign(VAlign_Center)
-							[
-								SNew(SImage)
-								.ColorAndOpacity(FSlateColor::UseForeground())
-								.Image(FAppStyle::Get().GetBrush("Icons.Plus"))
-								]
-							]
 						]
+					]
 
 					+ SVerticalBox::Slot()
 					.AutoHeight()
