PR #12967: Guard against ContainerPtr assertion failure

#12967

diff --git a/Engine/Source/Editor/UnrealEd/Private/EditorServer.cpp b/Engine/Source/Editor/UnrealEd/Private/EditorServer.cpp
--- a/Engine/Source/Editor/UnrealEd/Private/EditorServer.cpp
+++ b/Engine/Source/Editor/UnrealEd/Private/EditorServer.cpp
@@ -3794,7 +3794,7 @@ bool UEditorEngine::GetPropertyColorationMatch(UObject* Object)
 					UObject* ReferencedObject = ObjectPropertyBase->GetObjectPropertyValue(ObjAddr);
 					Base = (uint8*) ReferencedObject;
 				}
-				else
+				else if (Base)
 				{
 					Base = Prop->ContainerPtrToValuePtr<uint8>(Base);
 				}
