PR #9896: Fix for "Paste Here" bug for parented actors

diff --git a/Engine/Source/Editor/UnrealEd/Private/EditorServer.cpp b/Engine/Source/Editor/UnrealEd/Private/EditorServer.cpp
--- a/Engine/Source/Editor/UnrealEd/Private/EditorServer.cpp
+++ b/Engine/Source/Editor/UnrealEd/Private/EditorServer.cpp
@@ -3614,7 +3614,7 @@ void UEditorEngine::PasteSelectedActorsFromClipboard( UWorld* InWorld, const FTe
 				TArray<FAttachData, TInlineAllocator<8>> AttachData;
 				AttachData.Reserve(NumActorsToMove);
 
-				// Break any parent attachments and move the actors.
+				// Break any parent attachments
 				AActor* SingleActor = NULL;
 				for ( FSelectionIterator It( GetSelectedActorIterator() ) ; It ; ++It )
 				{
@@ -3636,6 +3636,12 @@ void UEditorEngine::PasteSelectedActorsFromClipboard( UWorld* InWorld, const FTe
 					}
 
 					SingleActor = Actor;
+				}
+
+				// Move the actors
+				for (FSelectionIterator It(GetSelectedActorIterator()); It; ++It)
+				{
+					AActor* Actor = static_cast<AActor*>(*It);
 					Actor->SetActorLocation(Actor->GetActorLocation() + Adjust, false);
 				}
 
