PR #10651: UE-189002: Make edactHideSelected take attached actors into account

diff --git a/Engine/Source/Editor/UnrealEd/Private/EditorActor.cpp b/Engine/Source/Editor/UnrealEd/Private/EditorActor.cpp
--- a/Engine/Source/Editor/UnrealEd/Private/EditorActor.cpp
+++ b/Engine/Source/Editor/UnrealEd/Private/EditorActor.cpp
@@ -1511,6 +1511,21 @@ static void SetActorVisibility(AActor* Actor, bool bVisible)
 	}
 }
 
+static void GetAttachedActors(AActor* Actor, bool bRecurseChildren, TSet< AActor* >& OutActors)
+{
+	TArray< AActor* > ChildrenActors;
+	Actor->GetAttachedActors(ChildrenActors);
+	for (AActor* ChildActor : ChildrenActors)
+	{
+		OutActors.Add(ChildActor);
+
+		if (bRecurseChildren)
+		{
+			GetAttachedActors(ChildActor, bRecurseChildren, OutActors);
+		}
+	}
+}
+
 void UUnrealEdEngine::edactHideSelected( UWorld* InWorld )
 {
 	// Assemble a list of actors to hide.
@@ -1524,6 +1539,18 @@ void UUnrealEdEngine::edactHideSelected( UWorld* InWorld )
 		if ( !FActorEditorUtils::IsABuilderBrush(Actor) && !Actor->IsHiddenEd() )
 		{
 			ActorsToHide.Add( Actor );
+
+			TSet<AActor*> AttachedActors;
+			GetAttachedActors( Actor, true, AttachedActors );
+
+			for ( AActor* AttachedActor : AttachedActors )
+			{
+				// Don't consider already hidden actors or the builder brush
+				if ( !FActorEditorUtils::IsABuilderBrush(AttachedActor) && !AttachedActor->IsHiddenEd() )
+				{
+					ActorsToHide.Add( AttachedActor );
+				}
+			}
 		}
 	}
 	
@@ -1775,7 +1802,7 @@ void UUnrealEdEngine::edactUnHideSelectedStartup( UWorld* InWorld )
 
 void UUnrealEdEngine::edactUnhideSelected( UWorld* InWorld )
 {
-	// Assemble a list of actors to hide.
+	// Assemble a list of actors to unhide.
 	TArray<AActor*> ActorsToShow;
 	for ( FSelectionIterator It( GetSelectedActorIterator() ) ; It ; ++It )
 	{
@@ -2104,21 +2131,6 @@ void UUnrealEdEngine::edactSelectInvert( UWorld* InWorld )
 	}
 }
 
-static void GetAttachedActors( AActor* Actor, bool bRecurseChildren, TSet< AActor* >& OutActors )
-{
-	TArray< AActor* > ChildrenActors;
-	Actor->GetAttachedActors( ChildrenActors );
-	for ( AActor* ChildActor : ChildrenActors )
-	{
-		OutActors.Add( ChildActor );
-
-		if ( bRecurseChildren )
-		{
-			GetAttachedActors( ChildActor, bRecurseChildren, OutActors );
-		}
-	}
-}
-
 void UUnrealEdEngine::edactSelectAllChildren( bool bRecurseChildren )
 {
 	USelection* CurrentSelection = GetSelectedActors();
