[Gizmos] PR #11676: Fixed Translate and Rotate gizmo's not fully rendering properly at the widget matrix

diff --git a/Engine/Source/Editor/UnrealEd/Private/UnrealWidgetRender.cpp b/Engine/Source/Editor/UnrealEd/Private/UnrealWidgetRender.cpp
--- a/Engine/Source/Editor/UnrealEd/Private/UnrealWidgetRender.cpp
+++ b/Engine/Source/Editor/UnrealEd/Private/UnrealWidgetRender.cpp
@@ -555,7 +555,7 @@ void FWidget::Render_Translate( const FSceneView* View, FPrimitiveDrawInterface*
 		const FVector CameraZAxis = View->ViewMatrices.GetViewMatrix().GetColumn(2);
 
 		UMaterialInstanceDynamic* XYZMaterial = ( CurrentAxis&EAxisList::Screen) ? CurrentAxisMaterial : OpaquePlaneMaterialXY;
-		DrawSphere( PDI, InLocation, FRotator::ZeroRotator, 4.0f * Space.Scale, 10, 5, XYZMaterial->GetRenderProxy(), SDPG_Foreground );
+		DrawSphere( PDI, WidgetMatrix.GetOrigin(), FRotator::ZeroRotator, 4.0f * Space.Scale, 10, 5, XYZMaterial->GetRenderProxy(), SDPG_Foreground );
 
 		PDI->SetHitProxy( nullptr );
 	}
@@ -670,6 +670,7 @@ void FWidget::Render_Rotate( const FSceneView* View,FPrimitiveDrawInterface* PDI
 	FVector XAxis = CustomCoordSystem.TransformVector(FVector(1, 0, 0));
 	FVector YAxis = CustomCoordSystem.TransformVector(FVector(0, 1, 0));
 	FVector ZAxis = CustomCoordSystem.TransformVector(FVector(0, 0, 1));
+	FMatrix WidgetMatrix = CustomCoordSystem * FTranslationMatrix( InLocation );
 
 	EAxisList::Type DrawAxis = GetAxisToDraw(ViewportClient->GetWidgetMode());
 
@@ -686,17 +687,17 @@ void FWidget::Render_Rotate( const FSceneView* View,FPrimitiveDrawInterface* PDI
 
 			if (DrawAxis & EAxisList::X)
 			{
-				DrawRotationArc(View, PDI, EAxisList::X, InLocation, ZAxis, YAxis, DirectionToWidget, AxisColorX.ToFColor(true), Scale, XAxisDir);
+				DrawRotationArc(View, PDI, EAxisList::X, WidgetMatrix.GetOrigin(), ZAxis, YAxis, DirectionToWidget, AxisColorX.ToFColor(true), Scale, XAxisDir);
 			}
 
 			if (DrawAxis & EAxisList::Y)
 			{
-				DrawRotationArc(View, PDI, EAxisList::Y, InLocation, XAxis, ZAxis, DirectionToWidget, AxisColorY.ToFColor(true), Scale, YAxisDir);
+				DrawRotationArc(View, PDI, EAxisList::Y, WidgetMatrix.GetOrigin(), XAxis, ZAxis, DirectionToWidget, AxisColorY.ToFColor(true), Scale, YAxisDir);
 			}
 
 			if (DrawAxis & EAxisList::Z)
 			{
-				DrawRotationArc(View, PDI, EAxisList::Z, InLocation, XAxis, YAxis, DirectionToWidget, AxisColorZ.ToFColor(true), Scale, ZAxisDir);
+				DrawRotationArc(View, PDI, EAxisList::Z, WidgetMatrix.GetOrigin(), XAxis, YAxis, DirectionToWidget, AxisColorZ.ToFColor(true), Scale, ZAxisDir);
 			}
 			
 		}
@@ -716,7 +717,7 @@ void FWidget::Render_Rotate( const FSceneView* View,FPrimitiveDrawInterface* PDI
 			//set priority to foreground so axis get hit first
 			PDI->SetHitProxy(new HWidgetAxis(EAxisList::XYZ, bDisabled, HPP_Foreground));
 			{
-				DrawColoredSphere(PDI, InLocation, FRotator::ZeroRotator, ArcColor, Radii, 32, 24, TransparentPlaneMaterialXY->GetRenderProxy(), SDPG_Foreground, true);
+				DrawColoredSphere(PDI, WidgetMatrix.GetOrigin(), FRotator::ZeroRotator, ArcColor, Radii, 32, 24, TransparentPlaneMaterialXY->GetRenderProxy(), SDPG_Foreground, true);
 			}
 			PDI->SetHitProxy(nullptr);
 		}
