[Revision Control] PR #9385: ResavePackages: support providers that do not use readonly flag or checkouts

diff --git a/Engine/Source/Editor/UnrealEd/Private/Commandlets/ContentCommandlets.cpp b/Engine/Source/Editor/UnrealEd/Private/Commandlets/ContentCommandlets.cpp
--- a/Engine/Source/Editor/UnrealEd/Private/Commandlets/ContentCommandlets.cpp
+++ b/Engine/Source/Editor/UnrealEd/Private/Commandlets/ContentCommandlets.cpp
@@ -938,6 +938,11 @@ void UResavePackagesCommandlet::LoadAndSaveOnePackage(const FString& Filename)
 							VerboseMessage(TEXT("Post ForceGetStatus2"));
 						}
 					}
+					else if (!ISourceControlModule::Get().GetProvider().UsesCheckout())
+					{
+						// Handle git which does not set files readonly and does not use checkout
+						FilesToSubmit.AddUnique(*Filename);
+					}
 
 					// Update the readonly state now that source control has had a chance to run
 					bIsReadOnly = IFileManager::Get().IsReadOnly(*Filename);
@@ -1141,6 +1146,11 @@ int32 UResavePackagesCommandlet::Main( const FString& Params )
 	bSkipCheckedOutFiles = Switches.Contains(TEXT("SkipCheckedOutPackages"));
 	/** if we should auto checkin packages that were checked out**/
 	bAutoCheckIn = bAutoCheckOut && (Switches.Contains(TEXT("AutoCheckIn")) || Switches.Contains(TEXT("AutoSubmit")));
+	if (!ISourceControlModule::Get().GetProvider().UsesCheckout())
+	{
+		// Handle git which does not set files readonly and does not use checkout
+		bAutoCheckIn = Switches.Contains(TEXT("AutoCheckIn")) || Switches.Contains(TEXT("AutoSubmit"));
+	}
 	/** if we should skip trying to virtualize packages being checked in **/
 	bSkipVirtualization = Switches.Contains(TEXT("SkipVirtualization"));
 	/** determine if we are building lighting for the map packages on the pass. **/
