[ITF] Remove redundant LostFocus calls

diff --git a/Engine/Source/Editor/UnrealEd/Private/EditorModeManager.cpp b/Engine/Source/Editor/UnrealEd/Private/EditorModeManager.cpp
--- a/Engine/Source/Editor/UnrealEd/Private/EditorModeManager.cpp
+++ b/Engine/Source/Editor/UnrealEd/Private/EditorModeManager.cpp
@@ -1607,15 +1607,6 @@ bool FEditorModeTools::LostFocus( FEditorViewportClient* InViewportClient, FView
 
 	bool bHandled = InteractiveToolsContext->LostFocus(InViewportClient, Viewport);
 
-	for (const UEdMode* const EdMode : ActiveScriptableModes)
-	{
-		if (const UModeManagerInteractiveToolsContext* const ModeManagerInteractiveToolsContext =
-				EdMode->GetModeManager()->InteractiveToolsContext)
-		{
-			ModeManagerInteractiveToolsContext->LostFocus(InViewportClient, Viewport);
-		}
-	}
-
 	ForEachEdMode<ILegacyEdModeViewportInterface>([&bHandled, InViewportClient, Viewport](ILegacyEdModeViewportInterface* Mode)
 		{
 			bHandled |= Mode->LostFocus(InViewportClient, Viewport);
