Fix cook stall detector in cook on the fly:
- if no package in progress, don't detect it as a stall.
- only run the stall detector if the garbage collection was triggered by an OOM.

#rnx

diff --git a/Engine/Source/Editor/UnrealEd/Private/Commandlets/CookCommandlet.cpp b/Engine/Source/Editor/UnrealEd/Private/Commandlets/CookCommandlet.cpp
--- a/Engine/Source/Editor/UnrealEd/Private/Commandlets/CookCommandlet.cpp
+++ b/Engine/Source/Editor/UnrealEd/Private/Commandlets/CookCommandlet.cpp
@@ -773,11 +773,10 @@ void UCookCommandlet::ConditionalCollectGarbage(uint32 TickResults, UCookOnTheFl
 		NumObjectsBeforeGC, MemStatsBeforeGC, AllocatorStatsBeforeGC,
 		NumObjectsAfterGC, MemStatsAfterGC, AllocatorStatsAfterGC, GCDurationSeconds);
 
-	if (!COTFS.IsCookWorkerMode() && !bSoft)
+	// Only check stalling in this is the director and if garbage collection was triggered by OOM.
+	bool bRunStallDetector = !COTFS.IsCookWorkerMode() && bWasDueToOOM;
+	if (bRunStallDetector)
 	{
-		// Only check stalling in the director and if it's not a SoftGC.
-		// SoftGC cannot cause OOM stalls because it keeps packages in memory that will be used dater.
-
 		bool bCookStalled = COTFS.IsStalled();
 		if (bCookStalled)
 		{
diff --git a/Engine/Source/Editor/UnrealEd/Private/Cooker/StallDetector.cpp b/Engine/Source/Editor/UnrealEd/Private/Cooker/StallDetector.cpp
--- a/Engine/Source/Editor/UnrealEd/Private/Cooker/StallDetector.cpp
+++ b/Engine/Source/Editor/UnrealEd/Private/Cooker/StallDetector.cpp
@@ -10,18 +10,23 @@ namespace UE::Cook
 	// Number of iteration without progress after which the cook is considered stalled for the number of in progress packages.
 	const static int32 StalledIterationLimitPackageInProgress = 10;
 
-
 	bool FStallDetector::IsStalled(int32 NewPackageCooked, int32 NewPackageInProgress)
 	{
 		// First check if the number of cooked package is stalled
 		StalledPackageCooked.Update(NewPackageCooked);
 		if (StalledPackageCooked.StalledIterationCount < UE::Cook::StalledIterationLimitPackageCooked)
+		{
 			return false;
+		}
 
 		// Now check if the number of package in progress is stalled
 		StalledPackageInProgress.Update(NewPackageInProgress);
-		if (StalledPackageInProgress.StalledIterationCount < UE::Cook::StalledIterationLimitPackageInProgress)
+
+		// If the number of package in progress is zero, then there is no cook requested so it is not a stall.
+		if (StalledPackageInProgress.StalledIterationCount < UE::Cook::StalledIterationLimitPackageInProgress || StalledPackageInProgress.Value == 0)
+		{
 			return false;
+		}
 
 		// No progress so the cook is stalled
 		return true;
