[Viewports] PR #11772: Fixed zoom ratio to be consistent in ortho viewport.

diff --git a/Engine/Source/Editor/UnrealEd/Private/EditorViewportClient.cpp b/Engine/Source/Editor/UnrealEd/Private/EditorViewportClient.cpp
--- a/Engine/Source/Editor/UnrealEd/Private/EditorViewportClient.cpp
+++ b/Engine/Source/Editor/UnrealEd/Private/EditorViewportClient.cpp
@@ -3812,10 +3812,10 @@ void FEditorViewportClient::OnOrthoZoom( const struct FInputEventState& InputSta
 	FKey Key = InputState.GetKey();
 
 	// Scrolling the mousewheel up/down zooms the orthogonal viewport in/out.
-	int32 Delta = 25 * Scale;
+	float NewOrthoRatio = 1 + ((25 * Scale) / CAMERA_ZOOM_DAMPEN);
 	if( Key == EKeys::MouseScrollUp || Key == EKeys::Add )
 	{
-		Delta *= -1;
+		NewOrthoRatio = 1 / NewOrthoRatio;
 	}
 
 	//Extract current state
@@ -3874,7 +3874,7 @@ void FEditorViewportClient::OnOrthoZoom( const struct FInputEventState& InputSta
 	const float OldUnitsPerPixel = GetOrthoUnitsPerPixel(Viewport);
 
 	//update zoom based on input
-	const float Zoom = GetOrthoZoom() + (GetOrthoZoom() / CAMERA_ZOOM_DAMPEN) * Delta;
+	const float Zoom = GetOrthoZoom() * NewOrthoRatio;
 	SetOrthoZoom( FMath::Clamp<float>( Zoom, GetMinimumOrthoZoom(), MAX_ORTHOZOOM ) );
 
 	if (bCenterZoomAroundCursor)
@@ -5517,7 +5517,12 @@ void FEditorViewportClient::MoveViewportCamera(const FVector& InDrag, const FRot
 
 			if( ( LeftMouseButtonDown || bIsUsingTrackpad ) && RightMouseButtonDown )
 			{
-				const float Zoom = GetOrthoZoom() + (GetOrthoZoom() / CAMERA_ZOOM_DAMPEN) * InDrag.Z;
+				float NewOrthoRatio = 1 + (FMath::Abs(InDrag.Z) / CAMERA_ZOOM_DAMPEN);
+				if ( InDrag.Z < 0 )
+				{
+					NewOrthoRatio = 1 / NewOrthoRatio;
+				}
+				const float Zoom = GetOrthoZoom() * NewOrthoRatio;
 				SetOrthoZoom( FMath::Clamp<float>( Zoom, GetMinimumOrthoZoom(), MAX_ORTHOZOOM ) );
 			}
 			else
