Add an overridable function for ViewportClient users to control the calculation of the view matrix.

diff --git a/Engine/Source/Editor/UnrealEd/Private/EditorViewportClient.cpp b/Engine/Source/Editor/UnrealEd/Private/EditorViewportClient.cpp
--- a/Engine/Source/Editor/UnrealEd/Private/EditorViewportClient.cpp
+++ b/Engine/Source/Editor/UnrealEd/Private/EditorViewportClient.cpp
@@ -1050,7 +1050,7 @@ FSceneView* FEditorViewportClient::CalcSceneView(FSceneViewFamily* ViewFamily, c
 	if (bUseControllingActorViewInfo)
 	{
 		// @todo vreditor: Not stereo friendly yet
-		ViewInitOptions.ViewRotationMatrix = FInverseRotationMatrix(ModifiedViewRotation) * FMatrix(
+		ViewInitOptions.ViewRotationMatrix = CalcViewRotationMatrixForControllingActorView(ModifiedViewRotation) * FMatrix(
 			FPlane(0, 0, 1, 0),
 			FPlane(1, 0, 0, 0),
 			FPlane(0, 1, 0, 0),
@@ -6653,6 +6653,11 @@ float FEditorViewportClient::UpdateViewportClientWindowDPIScale() const
 	return DPIScale;
 }
 
+FMatrix FEditorViewportClient::CalcViewRotationMatrixForControllingActorView(const FRotator& InViewRotation) const
+{
+	return FInverseRotationMatrix(InViewRotation);
+}
+
 FMatrix FEditorViewportClient::CalcViewRotationMatrix(const FRotator& InViewRotation) const
 {
 	const FViewportCameraTransform& ViewTransform = GetViewTransform();
diff --git a/Engine/Source/Editor/UnrealEd/Public/EditorViewportClient.h b/Engine/Source/Editor/UnrealEd/Public/EditorViewportClient.h
--- a/Engine/Source/Editor/UnrealEd/Public/EditorViewportClient.h
+++ b/Engine/Source/Editor/UnrealEd/Public/EditorViewportClient.h
@@ -1213,6 +1213,9 @@ public:
 	/** Get the editor viewport widget */
 	TSharedPtr<SEditorViewport> GetEditorViewportWidget() const { return EditorViewportWidget.Pin(); }
 
+
+	UNREALED_API virtual FMatrix CalcViewRotationMatrixForControllingActorView(const FRotator& InViewRotation) const;
+	
 	/**
 	 * Computes a matrix to use for viewport location and rotation
 	 */
