UnrealEd/Lightmass: add AutoRTFM shared library to list of required dependencies

diff --git a/Engine/Source/Editor/UnrealEd/Private/Lightmass/Lightmass.cpp b/Engine/Source/Editor/UnrealEd/Private/Lightmass/Lightmass.cpp
--- a/Engine/Source/Editor/UnrealEd/Private/Lightmass/Lightmass.cpp
+++ b/Engine/Source/Editor/UnrealEd/Private/Lightmass/Lightmass.cpp
@@ -3005,6 +3005,7 @@ bool FLightmassProcessor::BeginRun()
 		TEXT("Binaries/Win32/AgentInterface.dll"),
 		TEXT("Binaries/Win32/UnrealLightmass-SwarmInterface.dll"),
 		TEXT("Binaries/Win32/UnrealLightmass-ApplicationCore.dll"),
+		TEXT("Binaries/Win32/UnrealLightmass-AutoRTFM.dll"),
 		TEXT("Binaries/Win32/UnrealLightmass-Core.dll"),
 		TEXT("Binaries/Win32/UnrealLightmass-CoreUObject.dll"),
 		TEXT("Binaries/Win32/UnrealLightmass-CorePreciseFP.dll"),
@@ -3027,6 +3028,7 @@ bool FLightmassProcessor::BeginRun()
 		TEXT("Binaries/Win64/AgentInterface.dll"),
 		TEXT("Binaries/Win64/UnrealLightmass-SwarmInterface.dll"),
 		TEXT("Binaries/Win64/UnrealLightmass-ApplicationCore.dll"),
+		TEXT("Binaries/Win64/UnrealLightmass-AutoRTFM.dll"),
 		TEXT("Binaries/Win64/UnrealLightmass-Core.dll"),
 		TEXT("Binaries/Win64/UnrealLightmass-CoreUObject.dll"),
 		TEXT("Binaries/Win64/UnrealLightmass-CorePreciseFP.dll"),
@@ -3048,6 +3050,7 @@ bool FLightmassProcessor::BeginRun()
 	{
 		TEXT("Binaries/DotNET/Mac/AgentInterface.dll"),
 		TEXT("Binaries/Mac/UnrealLightmass-ApplicationCore.dylib"),
+		TEXT("Binaries/Mac/UnrealLightmass-AutoRTFM.dylib"),
 		TEXT("Binaries/Mac/UnrealLightmass-Core.dylib"),
 		TEXT("Binaries/Mac/UnrealLightmass-CoreUObject.dylib"),
 		TEXT("Binaries/Mac/UnrealLightmass-CorePreciseFP.dylib"),
@@ -3070,6 +3073,7 @@ bool FLightmassProcessor::BeginRun()
 	{
 		TEXT("Binaries/DotNET/Linux/AgentInterface.dll"),
 		TEXT("Binaries/Linux/libUnrealLightmass-ApplicationCore.so"),
+		TEXT("Binaries/Linux/libUnrealLightmass-AutoRTFM.so"),
 		TEXT("Binaries/Linux/libUnrealLightmass-Core.so"),
 		TEXT("Binaries/Linux/libUnrealLightmass-CoreUObject.so"),
 		TEXT("Binaries/Linux/libUnrealLightmass-CorePrecise.so"),
