Updated GatherTextFromAssetsCommandlet to discover external actors for a world without needing to load it

This takes several minutes off the localization gather (7-10 in internal testing, though it would be more or less depending on the number of worlds included in the gather)

diff --git a/Engine/Source/Editor/UnrealEd/Classes/Commandlets/GatherTextFromAssetsCommandlet.h b/Engine/Source/Editor/UnrealEd/Classes/Commandlets/GatherTextFromAssetsCommandlet.h
--- a/Engine/Source/Editor/UnrealEd/Classes/Commandlets/GatherTextFromAssetsCommandlet.h
+++ b/Engine/Source/Editor/UnrealEd/Classes/Commandlets/GatherTextFromAssetsCommandlet.h
@@ -65,7 +65,7 @@ private:
 	void FilterAssetsBasedOnIncludeExcludePaths(TArray<FAssetData>& InOutAssetDataArray) const;
 	
 	bool DiscoverExternalActors(TArray<FAssetData>& InOutAssetDataArray);
-	void RemoveExistingExternalActors(TArray<FAssetData>& InOutAssetDataArray, TArray<FName>& OutPartitionedWorldPackageNames) const;
+	void RemoveExistingExternalActors(TArray<FAssetData>& InOutAssetDataArray, TSet<FName>& OutPartitionedWorldPackageNames) const;
 
 	TSet<FName> GetPackageNamesToGather(const TArray<FAssetData>& InAssetDataArray) const;
 	void PopulatePackagesPendingGather(TSet<FName> PackageNamesToGather);
diff --git a/Engine/Source/Editor/UnrealEd/Private/Commandlets/GatherTextFromAssetsCommandlet.cpp b/Engine/Source/Editor/UnrealEd/Private/Commandlets/GatherTextFromAssetsCommandlet.cpp
--- a/Engine/Source/Editor/UnrealEd/Private/Commandlets/GatherTextFromAssetsCommandlet.cpp
+++ b/Engine/Source/Editor/UnrealEd/Private/Commandlets/GatherTextFromAssetsCommandlet.cpp
@@ -552,6 +552,9 @@ bool UGatherTextFromAssetsCommandlet::BuildFirstPassFilter(FARFilter& InOutFilte
 		}
 	}
 
+	InOutFilter.bIncludeOnlyOnDiskAssets = true;
+	InOutFilter.WithoutPackageFlags = PKG_Cooked;
+
 	return true;
 }
 
@@ -764,7 +767,7 @@ void UGatherTextFromAssetsCommandlet::FilterAssetsBasedOnIncludeExcludePaths(TAr
 }
 
 /** Remove any external actors that currently exist in InOutAssetDataArray. OutPartitionedWorldPackageNames is populated with the package paths of partitioned worlds.*/
-void UGatherTextFromAssetsCommandlet::RemoveExistingExternalActors(TArray<FAssetData>& InOutAssetDataArray, TArray<FName>& OutPartitionedWorldPackageNames) const
+void UGatherTextFromAssetsCommandlet::RemoveExistingExternalActors(TArray<FAssetData>& InOutAssetDataArray, TSet<FName>& OutPartitionedWorldPackageNames) const
 {
 	InOutAssetDataArray.RemoveAll([&OutPartitionedWorldPackageNames](const FAssetData& AssetData)
 		{
@@ -794,108 +797,76 @@ bool UGatherTextFromAssetsCommandlet::DiscoverExternalActors(TArray<FAssetData>&
 	UE_LOG(LogGatherTextFromAssetsCommandlet, Display, TEXT("Discovering external actors to gather..."));
 	const double DiscoveringExternalActorsStartTime = FPlatformTime::Seconds();
 
-	TArray<FName> PartitionedWorldPackageNames;
+	TSet<FName> PartitionedWorldPackageNames;
 	RemoveExistingExternalActors(InOutAssetDataArray, PartitionedWorldPackageNames);
 
-	TArray<FName> ExternalActorPackageNames;
-	{
-		FScopedSlowTask SlowTask(PartitionedWorldPackageNames.Num(), NSLOCTEXT("GatherTextCommandlet", "DiscoverExternalActors", "Discovering External Actors..."));
-
-		FLoadPackageLogOutputRedirector LogOutputRedirector;
+	IAssetRegistry& AssetRegistry = IAssetRegistry::GetChecked();
 
-		int32 NumPackagesProcessed = 0;
+	// Append the actors that are directly known by each world (by looking for their external actors under the expected path)
+	if (PartitionedWorldPackageNames.Num() > 0)
+	{
+		FARFilter Filter;
 		for (const FName PartitionedWorldPackageName : PartitionedWorldPackageNames)
 		{
-			SlowTask.EnterProgressFrame();
-
-			if (SlowTask.ShouldCancel() || (EmbeddedContext && EmbeddedContext->ShouldAbort()))
-			{
-				UE_LOG(LogGatherTextFromAssetsCommandlet, Error, TEXT("GatherText aborted!"));
-				return false;
-			}
-
 			FNameBuilder PackageNameStr(PartitionedWorldPackageName);
+			Filter.PackagePaths.Add(*ULevel::GetExternalActorsPath(*PackageNameStr));
+		}
+		Filter.bRecursivePaths = true;
+		Filter.ClassPaths.Add(AActor::StaticClass()->GetClassPathName());
+		Filter.bRecursiveClasses = true;
+		Filter.bIncludeOnlyOnDiskAssets = true;
+		Filter.WithoutPackageFlags = PKG_Cooked;
+		
+		AssetRegistry.GetAssets(Filter, InOutAssetDataArray);
+	}
 
-			const int32 CurrentPackageNum = ++NumPackagesProcessed;
-			const float PercentageComplete = static_cast<float>(CurrentPackageNum) / static_cast<float>(PartitionedWorldPackageNames.Num()) * 100.0f;
-			UE_LOG(LogGatherTextFromAssetsCommandlet, Display, TEXT("[%6.2f%%] Loading partitioned world package: '%s'..."), PercentageComplete, *PackageNameStr);
-
-			UPackage* Package = nullptr;
-			{
-				FLoadPackageLogOutputRedirector::FScopedCapture ScopedCapture(&LogOutputRedirector, *PackageNameStr);
-				Package = LoadWorldPackageForEditor(*PackageNameStr, EWorldType::Editor, LOAD_NoWarn | LOAD_Quiet);
-			}
+	// Append additional actors for each world
+	// These are actors added via ExternalDataLayers or ContentBundles, and will be referenced by the GameFeatureData assets that add them (via a GameFeatureAction)
+	if (PartitionedWorldPackageNames.Num() > 0)
+	{
+		TArray<FAssetData> GameFeatureDataAssets;
+		{
+			FARFilter Filter;
+			Filter.ClassPaths.Add(FTopLevelAssetPath("/Script/GameFeatures", "GameFeatureData")); // cannot include UGameFeatureData as it is inside a plugin
+			Filter.bRecursiveClasses = true;
+			Filter.bIncludeOnlyOnDiskAssets = true;
+			Filter.WithoutPackageFlags = PKG_Cooked;
 
-			if (!Package)
-			{
-				UE_LOG(LogGatherTextFromAssetsCommandlet, Warning, TEXT("Failed to load partitioned world package: '%s'."), *PackageNameStr);
-				continue;
-			}
+			AssetRegistry.GetAssets(Filter, GameFeatureDataAssets);
+		}
 
-			// Tick background tasks
-			UE::Private::GatherTextFromAssetsCommandlet::TickBackgroundTasks();
-			if (EmbeddedContext)
+		TArray<FAssetData> GameFeatureDataActorDependencies;
+		{
+			FARFilter Filter;
+			for (const FAssetData& GameFeatureDataAsset : GameFeatureDataAssets)
 			{
-				EmbeddedContext->RunTick();
+				AssetRegistry.GetDependencies(GameFeatureDataAsset.PackageName, Filter.PackageNames);
 			}
+			Filter.ClassPaths.Add(AActor::StaticClass()->GetClassPathName());
+			Filter.bRecursiveClasses = true;
+			Filter.bIncludeOnlyOnDiskAssets = true;
+			Filter.WithoutPackageFlags = PKG_Cooked;
 
-			if (UWorld* World = UWorld::FindWorldInPackage(Package))
+			if (Filter.PackageNames.Num() > 0)
 			{
-				UWorld::InitializationValues IVS;
-				IVS.InitializeScenes(false);
-				IVS.AllowAudioPlayback(false);
-				IVS.RequiresHitProxies(false);
-				IVS.CreatePhysicsScene(false);
-				IVS.CreateNavigation(false);
-				IVS.CreateAISystem(false);
-				IVS.ShouldSimulatePhysics(false);
-				IVS.EnableTraceCollision(false);
-				IVS.SetTransactional(false);
-				IVS.CreateFXSystem(false);
-				IVS.CreateWorldPartition(true);
-
-				TOptional<FScopedEditorWorld> ScopeEditorWorld;
-				if (!World->IsInitialized())
-				{
-					FLoadPackageLogOutputRedirector::FScopedCapture ScopedCapture(&LogOutputRedirector, *PackageNameStr);
-					ScopeEditorWorld.Emplace(World, IVS); // Initializing FScopedEditorWorld can log warnings, so capture those like we do with loading errors
-				}
-
-				if (UWorldPartition* WorldPartition = World->GetWorldPartition())
-				{
-					FWorldPartitionHelpers::ForEachActorDescInstance(WorldPartition, AActor::StaticClass(), [&ExternalActorPackageNames](const FWorldPartitionActorDescInstance* ActorDescInstance)
-					{
-						ExternalActorPackageNames.Add(ActorDescInstance->GetActorPackage());
-						return true;
-					});
-				}
+				AssetRegistry.GetAssets(Filter, GameFeatureDataActorDependencies);
 			}
+		}
 
-			if (HasExceededMemoryLimit(/*bLog*/true))
+		for (FAssetData& GameFeatureDataActorDependency : GameFeatureDataActorDependencies)
+		{
+			const FName OptionalOuterPathName = GameFeatureDataActorDependency.GetOptionalOuterPathName();
+			if (!OptionalOuterPathName.IsNone())
 			{
-				// First try a minimal purge to only remove things that are no longer referenced or needed by other packages pending gather
-				PurgeGarbage(/*bPurgeReferencedPackages*/false);
-
-				if (HasExceededMemoryLimit(/*bLog*/false))
+				const FName OptionalOuterPackageName = FName(FPackageName::ObjectPathToPackageName(FNameBuilder(OptionalOuterPathName).ToView()));
+				if (PartitionedWorldPackageNames.Contains(OptionalOuterPackageName))
 				{
-					// If we're still over the memory limit after a minimal purge, then attempt a full purge
-					PurgeGarbage(/*bPurgeReferencedPackages*/true);
-
-					// If we're still over the memory limit after both purges, then log a warning as we may be about to OOM
-					UE_CLOG(HasExceededMemoryLimit(/*bLog*/false), LogGatherTextFromAssetsCommandlet, Warning, TEXT("Flushing failed to reduce process memory to within the requested limits; this process may OOM!"));
+					InOutAssetDataArray.Add(MoveTemp(GameFeatureDataActorDependency));
 				}
 			}
 		}
 	}
 
-	if (ExternalActorPackageNames.Num() > 0)
-	{
-		FARFilter Filter;
-		Filter.PackageNames = ExternalActorPackageNames;
-		Filter.bIncludeOnlyOnDiskAssets = true;
-		IAssetRegistry::GetChecked().GetAssets(Filter, InOutAssetDataArray);
-	}
-
 	UE_LOG(LogGatherTextFromAssetsCommandlet, Display, TEXT("Discovering external actors took %.2f seconds."), FPlatformTime::Seconds() - DiscoveringExternalActorsStartTime);
 	return true;
 }
