[Viewports] PR #10812: Viewport camera navigation: Add local up/down movement with R/F

diff --git a/Engine/Plugins/FX/Niagara/Source/NiagaraEditor/Private/Widgets/SNiagaraBakerViewport.cpp b/Engine/Plugins/FX/Niagara/Source/NiagaraEditor/Private/Widgets/SNiagaraBakerViewport.cpp
--- a/Engine/Plugins/FX/Niagara/Source/NiagaraEditor/Private/Widgets/SNiagaraBakerViewport.cpp
+++ b/Engine/Plugins/FX/Niagara/Source/NiagaraEditor/Private/Widgets/SNiagaraBakerViewport.cpp
@@ -47,8 +47,8 @@ public:
 			bool bRightKeyState = false;
 			bool bLeftKeyState = false;
 
-			bool bUpKeyState = false;
-			bool bDownKeyState = false;
+			bool bLocalUpKeyState = false;
+			bool bLocalDownKeyState = false;
 			bool bZoomOutKeyState = false;
 			bool bZoomInKeyState = false;
 
@@ -63,8 +63,8 @@ public:
 				bRightKeyState |= Viewport->KeyState(FViewportNavigationCommands::Get().Right->GetActiveChord(ChordIndex)->Key);
 				bLeftKeyState |= Viewport->KeyState(FViewportNavigationCommands::Get().Left->GetActiveChord(ChordIndex)->Key);
 
-				bUpKeyState |= Viewport->KeyState(FViewportNavigationCommands::Get().Up->GetActiveChord(ChordIndex)->Key);
-				bDownKeyState |= Viewport->KeyState(FViewportNavigationCommands::Get().Down->GetActiveChord(ChordIndex)->Key);
+				bLocalUpKeyState |= Viewport->KeyState(FViewportNavigationCommands::Get().LocalUp->GetActiveChord(ChordIndex)->Key);
+				bLocalDownKeyState |= Viewport->KeyState(FViewportNavigationCommands::Get().LocalDown->GetActiveChord(ChordIndex)->Key);
 				bZoomOutKeyState |= Viewport->KeyState(FViewportNavigationCommands::Get().FovZoomOut->GetActiveChord(ChordIndex)->Key);
 				bZoomInKeyState |= Viewport->KeyState(FViewportNavigationCommands::Get().FovZoomIn->GetActiveChord(ChordIndex)->Key);
 
@@ -82,8 +82,8 @@ public:
 			{
 				//LocalMovement.X += bLeftKeyState ? KeyboardMoveSpeed : 0.0f;
 				//LocalMovement.X -= bRightKeyState ? KeyboardMoveSpeed : 0.0f;
-				//LocalMovement.Y += bUpKeyState ? KeyboardMoveSpeed : 0.0f;
-				//LocalMovement.Y -= bDownKeyState ? KeyboardMoveSpeed : 0.0f;
+				//LocalMovement.Y += bLocalUpKeyState ? KeyboardMoveSpeed : 0.0f;
+				//LocalMovement.Y -= bLocalDownKeyState ? KeyboardMoveSpeed : 0.0f;
 				//LocalMovement.Z += bBackwardKeyState ? KeyboardMoveSpeed : 0.0f;
 				//LocalMovement.Z -= bForwardKeyState ? KeyboardMoveSpeed : 0.0f;
 			}
diff --git a/Engine/Plugins/Runtime/nDisplay/Source/DisplayClusterLightCardEditor/Private/Viewport/DisplayClusterLightCardEditorViewportClient.cpp b/Engine/Plugins/Runtime/nDisplay/Source/DisplayClusterLightCardEditor/Private/Viewport/DisplayClusterLightCardEditorViewportClient.cpp
--- a/Engine/Plugins/Runtime/nDisplay/Source/DisplayClusterLightCardEditor/Private/Viewport/DisplayClusterLightCardEditorViewportClient.cpp
+++ b/Engine/Plugins/Runtime/nDisplay/Source/DisplayClusterLightCardEditor/Private/Viewport/DisplayClusterLightCardEditorViewportClient.cpp
@@ -1323,7 +1323,7 @@ void FDisplayClusterLightCardEditorViewportClient::BeginCameraMovement(bool bHas
 		// control up/down instead of forward/back. If the up/down keys are pressed, reroute the forward/back impulses to the up/down impulse data
 		if (bIsForwardBackKeyPressed)
 		{
-			CameraUserImpulseData->MoveUpDownImpulse = CameraUserImpulseData->MoveForwardBackwardImpulse;
+			CameraUserImpulseData->MoveWorldUpDownImpulse = CameraUserImpulseData->MoveForwardBackwardImpulse;
 			CameraUserImpulseData->MoveForwardBackwardImpulse = 0.0f;
 		}
 		else if (!bIsForwardBackKeyPressed && !bIsLeftRightKeyPressed)
@@ -1331,7 +1331,8 @@ void FDisplayClusterLightCardEditorViewportClient::BeginCameraMovement(bool bHas
 			// Zero out all camera movement if WASD or arrow keys are not pressed, to keep the control scheme cleaner
 			CameraUserImpulseData->MoveForwardBackwardImpulse = 0.0f;
 			CameraUserImpulseData->MoveRightLeftImpulse = 0.0f;
-			CameraUserImpulseData->MoveUpDownImpulse = 0.0f;
+			CameraUserImpulseData->MoveWorldUpDownImpulse = 0.0f;
+			CameraUserImpulseData->MoveLocalUpDownImpulse = 0.0f;
 		}
 	}
 	else
@@ -1339,7 +1340,8 @@ void FDisplayClusterLightCardEditorViewportClient::BeginCameraMovement(bool bHas
 		// Always zero out the camera's movement impulses so that the user cannot move the view using any keyboard chords
 		CameraUserImpulseData->MoveForwardBackwardImpulse = 0.0f;
 		CameraUserImpulseData->MoveRightLeftImpulse = 0.0f;
-		CameraUserImpulseData->MoveUpDownImpulse = 0.0f;
+		CameraUserImpulseData->MoveWorldUpDownImpulse = 0.0f;
+		CameraUserImpulseData->MoveLocalUpDownImpulse = 0.0f;
 	}
 }
 
diff --git a/Engine/Source/Editor/UnrealEd/Private/CameraController.cpp b/Engine/Source/Editor/UnrealEd/Private/CameraController.cpp
--- a/Engine/Source/Editor/UnrealEd/Private/CameraController.cpp
+++ b/Engine/Source/Editor/UnrealEd/Private/CameraController.cpp
@@ -52,7 +52,8 @@ void FEditorCameraController::UpdateSimulation(
 		UserImpulseData.RotateRollVelocityModifier != 0.0f ||
 		UserImpulseData.MoveForwardBackwardImpulse != 0.0f ||
 		UserImpulseData.MoveRightLeftImpulse!= 0.0f ||
-		UserImpulseData.MoveUpDownImpulse != 0.0f ||
+		UserImpulseData.MoveWorldUpDownImpulse != 0.0f ||
+		UserImpulseData.MoveLocalUpDownImpulse != 0.0f ||
 		UserImpulseData.ZoomOutInImpulse != 0.0f ||
 		UserImpulseData.RotateYawImpulse != 0.0f ||
 		UserImpulseData.RotatePitchImpulse != 0.0f ||
@@ -139,9 +140,15 @@ void FEditorCameraController::ApplyImpulseDeadZone( const FCameraControllerUserI
 		bOutAnyImpulseData = true;
 	}
 
-	if( FMath::Abs( InUserImpulse.MoveUpDownImpulse ) >= Config.ImpulseDeadZoneAmount )
+	if( FMath::Abs( InUserImpulse.MoveWorldUpDownImpulse ) >= Config.ImpulseDeadZoneAmount )
 	{
-		OutUserImpulse.MoveUpDownImpulse = InUserImpulse.MoveUpDownImpulse;
+		OutUserImpulse.MoveWorldUpDownImpulse = InUserImpulse.MoveWorldUpDownImpulse;
+		bOutAnyImpulseData = true;
+	}
+
+	if( FMath::Abs( InUserImpulse.MoveLocalUpDownImpulse ) >= Config.ImpulseDeadZoneAmount )
+	{
+		OutUserImpulse.MoveLocalUpDownImpulse = InUserImpulse.MoveLocalUpDownImpulse;
 		bOutAnyImpulseData = true;
 	}
 
@@ -203,7 +210,7 @@ void FEditorCameraController::UpdatePosition( const FCameraControllerUserImpulse
 		LocalSpaceImpulse =
 			FVector( UserImpulse.MoveForwardBackwardImpulse,		// Local space forward/back
 					 UserImpulse.MoveRightLeftImpulse,				// Local space right/left
-					 0.0f );										// Local space up/down
+					 UserImpulse.MoveLocalUpDownImpulse );			// Local space up/down
 	}
 
 
@@ -219,7 +226,7 @@ void FEditorCameraController::UpdatePosition( const FCameraControllerUserImpulse
 		WorldSpaceImpulse +=
 			FVector( 0.0f,											// World space forward/back
 					 0.0f,											// World space right/left
-					 UserImpulse.MoveUpDownImpulse );				// World space up/down
+					 UserImpulse.MoveWorldUpDownImpulse );			// World space up/down
 
 		// Cap impulse by normalizing, but only if our magnitude is greater than 1.0
 		//if( WorldSpaceImpulse.SizeSquared() > 1.0f )
diff --git a/Engine/Source/Editor/UnrealEd/Private/EditorViewportClient.cpp b/Engine/Source/Editor/UnrealEd/Private/EditorViewportClient.cpp
--- a/Engine/Source/Editor/UnrealEd/Private/EditorViewportClient.cpp
+++ b/Engine/Source/Editor/UnrealEd/Private/EditorViewportClient.cpp
@@ -1645,11 +1645,11 @@ void FEditorViewportClient::UpdateCameraMovementFromJoystick(const bool bRelativ
 				}
 				else if (Key == EKeys::Gamepad_LeftTriggerAxis)
 				{
-					CameraUserImpulseData->MoveUpDownImpulse -= GetFilteredDelta(UnfilteredDelta, ViewportDeadZoneConstants::STANDARD_DEAD_ZONE, StandardDeadZone) * InConfig.TranslationMultiplier;
+					CameraUserImpulseData->MoveWorldUpDownImpulse -= GetFilteredDelta(UnfilteredDelta, ViewportDeadZoneConstants::STANDARD_DEAD_ZONE, StandardDeadZone) * InConfig.TranslationMultiplier;
 				}
 				else if (Key == EKeys::Gamepad_RightTriggerAxis)
 				{
-					CameraUserImpulseData->MoveUpDownImpulse += GetFilteredDelta(UnfilteredDelta, ViewportDeadZoneConstants::STANDARD_DEAD_ZONE, StandardDeadZone) * InConfig.TranslationMultiplier;
+					CameraUserImpulseData->MoveWorldUpDownImpulse += GetFilteredDelta(UnfilteredDelta, ViewportDeadZoneConstants::STANDARD_DEAD_ZONE, StandardDeadZone) * InConfig.TranslationMultiplier;
 				}
 			}
 		}
@@ -1888,7 +1888,8 @@ void FEditorViewportClient::UpdateCameraMovement( float DeltaTime )
 		// Apply impulse from magnify gesture and reset impulses if we're using WASD keys
 		CameraUserImpulseData->MoveForwardBackwardImpulse = GestureMoveForwardBackwardImpulse;
 		CameraUserImpulseData->MoveRightLeftImpulse = 0.0f;
-		CameraUserImpulseData->MoveUpDownImpulse = 0.0f;
+		CameraUserImpulseData->MoveWorldUpDownImpulse = 0.0f;
+		CameraUserImpulseData->MoveLocalUpDownImpulse = 0.0f;
 		CameraUserImpulseData->ZoomOutInImpulse = 0.0f;
 		CameraUserImpulseData->RotateYawImpulse = 0.0f;
 		CameraUserImpulseData->RotatePitchImpulse = 0.0f;
@@ -1901,8 +1902,10 @@ void FEditorViewportClient::UpdateCameraMovement( float DeltaTime )
 		bool bRightKeyState = false;
 		bool bLeftKeyState = false;
 
-		bool bUpKeyState = false;
-		bool bDownKeyState = false;
+		bool bWorldUpKeyState = false;
+		bool bWorldDownKeyState = false;
+		bool bLocalUpKeyState = false;
+		bool bLocalDownKeyState = false;
 		bool bZoomOutKeyState = false;
 		bool bZoomInKeyState = false;
 
@@ -1919,8 +1922,10 @@ void FEditorViewportClient::UpdateCameraMovement( float DeltaTime )
 			bRightKeyState |= Viewport->KeyState(FViewportNavigationCommands::Get().Right->GetActiveChord(ChordIndex)->Key);
 			bLeftKeyState |= Viewport->KeyState(FViewportNavigationCommands::Get().Left->GetActiveChord(ChordIndex)->Key);
 
-			bUpKeyState |= Viewport->KeyState(FViewportNavigationCommands::Get().Up->GetActiveChord(ChordIndex)->Key);
-			bDownKeyState |= Viewport->KeyState(FViewportNavigationCommands::Get().Down->GetActiveChord(ChordIndex)->Key);
+			bWorldUpKeyState |= Viewport->KeyState(FViewportNavigationCommands::Get().WorldUp->GetActiveChord(ChordIndex)->Key);
+			bWorldDownKeyState |= Viewport->KeyState(FViewportNavigationCommands::Get().WorldDown->GetActiveChord(ChordIndex)->Key);
+			bLocalUpKeyState |= Viewport->KeyState(FViewportNavigationCommands::Get().LocalUp->GetActiveChord(ChordIndex)->Key);
+			bLocalDownKeyState |= Viewport->KeyState(FViewportNavigationCommands::Get().LocalDown->GetActiveChord(ChordIndex)->Key);
 			bZoomOutKeyState |= Viewport->KeyState(FViewportNavigationCommands::Get().FovZoomOut->GetActiveChord(ChordIndex)->Key);
 			bZoomInKeyState |= Viewport->KeyState(FViewportNavigationCommands::Get().FovZoomIn->GetActiveChord(ChordIndex)->Key);
 
@@ -1963,18 +1968,28 @@ void FEditorViewportClient::UpdateCameraMovement( float DeltaTime )
 			CameraUserImpulseData->MoveRightLeftImpulse -= 1.0f;
 		}
 
-		// Up/down
-		if( ( bRemapWASDKeys && bUpKeyState) ||
+		// World up/down
+		if( ( bRemapWASDKeys && bWorldUpKeyState) ||
 			( bUnmodifiedPress && Viewport->KeyState( EKeys::PageUp ) ) ||
 			( bUnmodifiedPress && bUseNumpadCameraControl && ( Viewport->KeyState( EKeys::NumPadNine ) || Viewport->KeyState( EKeys::Add ) ) ) )
 		{
-			CameraUserImpulseData->MoveUpDownImpulse += 1.0f;
+			CameraUserImpulseData->MoveWorldUpDownImpulse += 1.0f;
 		}
-		if( ( bRemapWASDKeys && bDownKeyState) ||
+		if( ( bRemapWASDKeys && bWorldDownKeyState) ||
 			( bUnmodifiedPress && Viewport->KeyState( EKeys::PageDown ) ) ||
 			( bUnmodifiedPress && bUseNumpadCameraControl && ( Viewport->KeyState( EKeys::NumPadSeven ) || Viewport->KeyState( EKeys::Subtract ) ) ) )
 		{
-			CameraUserImpulseData->MoveUpDownImpulse -= 1.0f;
+			CameraUserImpulseData->MoveWorldUpDownImpulse -= 1.0f;
+		}
+
+		// Local up/down
+		if( ( bRemapWASDKeys && bLocalUpKeyState) )
+		{
+			CameraUserImpulseData->MoveLocalUpDownImpulse += 1.0f;
+		}
+		if( ( bRemapWASDKeys && bLocalDownKeyState) )
+		{
+			CameraUserImpulseData->MoveLocalUpDownImpulse -= 1.0f;
 		}
 
 		// Zoom FOV out/in
@@ -2018,7 +2033,7 @@ void FEditorViewportClient::UpdateCameraMovement( float DeltaTime )
 		{
 			FEditorViewportStats::Using(FEditorViewportStats::CAT_PERSPECTIVE_KEYBOARD_WASD);
 		}
-		else if ( CameraUserImpulseData->MoveUpDownImpulse != 0 )
+		else if ( CameraUserImpulseData->MoveWorldUpDownImpulse != 0 || CameraUserImpulseData->MoveLocalUpDownImpulse != 0 )
 		{
 			FEditorViewportStats::Using(FEditorViewportStats::CAT_PERSPECTIVE_KEYBOARD_UP_DOWN);
 		}
@@ -2067,7 +2082,8 @@ void FEditorViewportClient::UpdateCameraMovement( float DeltaTime )
 			(*CameraUserImpulseData).RotateRollVelocityModifier != 0.0f ||
 			(*CameraUserImpulseData).MoveForwardBackwardImpulse != 0.0f ||
 			(*CameraUserImpulseData).MoveRightLeftImpulse != 0.0f ||
-			(*CameraUserImpulseData).MoveUpDownImpulse != 0.0f ||
+			(*CameraUserImpulseData).MoveWorldUpDownImpulse != 0.0f ||
+			(*CameraUserImpulseData).MoveLocalUpDownImpulse != 0.0f ||
 			(*CameraUserImpulseData).ZoomOutInImpulse != 0.0f ||
 			(*CameraUserImpulseData).RotateYawImpulse != 0.0f ||
 			(*CameraUserImpulseData).RotatePitchImpulse != 0.0f ||
@@ -2752,8 +2768,10 @@ bool FEditorViewportClient::IsFlightCameraActive() const
 			|| Viewport->KeyState(FViewportNavigationCommands::Get().Backward->GetActiveChord(ChordIndex)->Key)
 			|| Viewport->KeyState(FViewportNavigationCommands::Get().Left->GetActiveChord(ChordIndex)->Key)
 			|| Viewport->KeyState(FViewportNavigationCommands::Get().Right->GetActiveChord(ChordIndex)->Key)
-			|| Viewport->KeyState(FViewportNavigationCommands::Get().Up->GetActiveChord(ChordIndex)->Key)
-			|| Viewport->KeyState(FViewportNavigationCommands::Get().Down->GetActiveChord(ChordIndex)->Key)
+			|| Viewport->KeyState(FViewportNavigationCommands::Get().WorldUp->GetActiveChord(ChordIndex)->Key)
+			|| Viewport->KeyState(FViewportNavigationCommands::Get().WorldDown->GetActiveChord(ChordIndex)->Key)
+			|| Viewport->KeyState(FViewportNavigationCommands::Get().LocalUp->GetActiveChord(ChordIndex)->Key)
+			|| Viewport->KeyState(FViewportNavigationCommands::Get().LocalDown->GetActiveChord(ChordIndex)->Key)
 			|| Viewport->KeyState(FViewportNavigationCommands::Get().FovZoomIn->GetActiveChord(ChordIndex)->Key)
 			|| Viewport->KeyState(FViewportNavigationCommands::Get().FovZoomOut->GetActiveChord(ChordIndex)->Key)
 			|| Viewport->KeyState(FViewportNavigationCommands::Get().RotateUp->GetActiveChord(ChordIndex)->Key)
@@ -6804,8 +6822,10 @@ void FViewportNavigationCommands::RegisterCommands()
 	UI_COMMAND(Left, "Left", "Moves the camera Left", EUserInterfaceActionType::Button, FInputChord(EKeys::A));
 	UI_COMMAND(Right, "Right", "Moves the camera Right", EUserInterfaceActionType::Button, FInputChord(EKeys::D));
 
-	UI_COMMAND(Up, "Up", "Moves the camera Up", EUserInterfaceActionType::Button, FInputChord(EKeys::E));
-	UI_COMMAND(Down, "Down", "Moves the camera Down", EUserInterfaceActionType::Button, FInputChord(EKeys::Q));
+	UI_COMMAND(WorldUp, "World Up", "Moves the camera Up in world space", EUserInterfaceActionType::Button, FInputChord(EKeys::E));
+	UI_COMMAND(WorldDown, "World Down", "Moves the camera Down in world space", EUserInterfaceActionType::Button, FInputChord(EKeys::Q));
+	UI_COMMAND(LocalUp, "Local Up", "Moves the camera Up in local space", EUserInterfaceActionType::Button, FInputChord(EKeys::R));
+	UI_COMMAND(LocalDown, "Local Down", "Moves the camera Down in local space", EUserInterfaceActionType::Button, FInputChord(EKeys::F));
 
 	UI_COMMAND(FovZoomIn, "FOV Zoom In", "Narrows the camers FOV", EUserInterfaceActionType::Button, FInputChord(EKeys::C));
 	UI_COMMAND(FovZoomOut, "FOV Zoom Out", "Widens the camera FOV", EUserInterfaceActionType::Button, FInputChord(EKeys::Z));
diff --git a/Engine/Source/Editor/UnrealEd/Public/CameraController.h b/Engine/Source/Editor/UnrealEd/Public/CameraController.h
--- a/Engine/Source/Editor/UnrealEd/Public/CameraController.h
+++ b/Engine/Source/Editor/UnrealEd/Public/CameraController.h
@@ -24,8 +24,11 @@ public:
 	/** Scalar user input for moving right (positive) or left (negative) */
 	float MoveRightLeftImpulse;
 
-	/** Scalar user input for moving up (positive) or down (negative) */
-	float MoveUpDownImpulse;
+	/** Scalar user input for moving up (positive) or down (negative) in world space */
+	float MoveWorldUpDownImpulse;
+
+	/** Scalar user input for moving up (positive) or down (negative) in local space */
+	float MoveLocalUpDownImpulse;
 
 	/** Scalar user input for turning right (positive) or left (negative) */
 	float RotateYawImpulse;
@@ -53,7 +56,8 @@ public:
 	FCameraControllerUserImpulseData()
 		: MoveForwardBackwardImpulse( 0.0f ),
 		  MoveRightLeftImpulse( 0.0f ),
-		  MoveUpDownImpulse( 0.0f ),
+		  MoveWorldUpDownImpulse( 0.0f ),
+		  MoveLocalUpDownImpulse( 0.0f ),
 		  RotateYawImpulse( 0.0f ),
 		  RotatePitchImpulse( 0.0f ),
 		  RotateRollImpulse( 0.0f ),
diff --git a/Engine/Source/Editor/UnrealEd/Public/EditorViewportClient.h b/Engine/Source/Editor/UnrealEd/Public/EditorViewportClient.h
--- a/Engine/Source/Editor/UnrealEd/Public/EditorViewportClient.h
+++ b/Engine/Source/Editor/UnrealEd/Public/EditorViewportClient.h
@@ -76,8 +76,10 @@ public:
 	TSharedPtr< FUICommandInfo > Left;
 	TSharedPtr< FUICommandInfo > Right;
 
-	TSharedPtr< FUICommandInfo > Up;
-	TSharedPtr< FUICommandInfo > Down;
+	TSharedPtr< FUICommandInfo > WorldUp;
+	TSharedPtr< FUICommandInfo > WorldDown;
+	TSharedPtr< FUICommandInfo > LocalUp;
+	TSharedPtr< FUICommandInfo > LocalDown;
 
 	TSharedPtr< FUICommandInfo > FovZoomIn;
 	TSharedPtr< FUICommandInfo > FovZoomOut;
