Submit leftover file from 39211766

#rnx

diff --git a/Engine/Source/Editor/UnrealEd/Private/EditorBuildUtils.cpp b/Engine/Source/Editor/UnrealEd/Private/EditorBuildUtils.cpp
--- a/Engine/Source/Editor/UnrealEd/Private/EditorBuildUtils.cpp
+++ b/Engine/Source/Editor/UnrealEd/Private/EditorBuildUtils.cpp
@@ -1723,29 +1723,6 @@ bool FEditorBuildUtils::EditorBuildMaterialTextureStreamingData(UPackage* Packag
 	return bAnyPackagesDirtied;
 }
 
-static bool BuildVirtualTextureComponents(IVirtualTexturingEditorModule* Module, EShadingPath ShadingPath, TArray<URuntimeVirtualTextureComponent*>& Components)
-{
-	if (Components.Num() == 0)
-	{
-		return true;
-	}
-
-	FScopedSlowTask BuildTask(static_cast<float>(Components.Num()), LOCTEXT("VirtualTextureBuild", "Building Virtual Textures"));
-	BuildTask.MakeDialog(true);
-
-	for (URuntimeVirtualTextureComponent* Component : Components)
-	{
-		BuildTask.EnterProgressFrame();
-
-		// Note that Build*() functions return true if the associated Has*() functions return false
-		if (BuildTask.ShouldCancel() || !Module->BuildStreamedMips(ShadingPath, Component))
-		{
-			return false;
-		}
-	}
-	return true;
-}
-
 bool FEditorBuildUtils::EditorBuildVirtualTexture(UWorld* InWorld)
 {
 	if (InWorld == nullptr)
@@ -1773,16 +1750,7 @@ bool FEditorBuildUtils::EditorBuildVirtualTexture(UWorld* InWorld)
 		UWorldPartitionRuntimeVirtualTextureBuilder::LoadRuntimeVirtualTextureActors(WorldPartition, ForEachActorWithLoadingResult);
 	}
 
-	// Gather which component to build : 
-	TArray<URuntimeVirtualTextureComponent*> Components;
-	for (TObjectIterator<URuntimeVirtualTextureComponent> It(RF_ClassDefaultObject, false, EInternalObjectFlags::Garbage); It; ++It)
-	{
-		if (It->GetWorld() == InWorld)
-		{
-			Components.Add(*It);
-		}
-	}
-
+	TArray<URuntimeVirtualTextureComponent*> Components = Module->GatherRuntimeVirtualTextureComponents(InWorld);
 	FBuildAllStreamedMipsResult Result = Module->BuildAllStreamedMips({ .World = InWorld, .Components = Components, .bRestoreFeatureLevelAfterBuilding = true });
 	return Result.bSuccess;
 }
