Enabled Structured Logging for a common localization warning

diff --git a/Engine/Source/Editor/UnrealEd/Private/Commandlets/GatherTextFromAssetsCommandlet.cpp b/Engine/Source/Editor/UnrealEd/Private/Commandlets/GatherTextFromAssetsCommandlet.cpp
--- a/Engine/Source/Editor/UnrealEd/Private/Commandlets/GatherTextFromAssetsCommandlet.cpp
+++ b/Engine/Source/Editor/UnrealEd/Private/Commandlets/GatherTextFromAssetsCommandlet.cpp
@@ -1,52 +1,57 @@
 // Copyright Epic Games, Inc. All Rights Reserved.
 
 #include "Commandlets/GatherTextFromAssetsCommandlet.h"
-#include "UObject/Class.h"
+
+#include "AssetRegistry/ARFilter.h"
+#include "AssetRegistry/AssetData.h"
+#include "AssetRegistry/AssetRegistryModule.h"
+#include "Async/ParallelFor.h"
+#include "CollectionManagerModule.h"
+#include "Commandlets/GatherTextFromSourceCommandlet.h"
+#include "DistanceFieldAtlas.h"
+#include "Editor.h"
+#include "EditorWorldUtils.h"
+#include "Engine/Level.h"
+#include "Engine/World.h"
+#include "Framework/Commands/Commands.h"
 #include "HAL/FileManager.h"
+#include "ICollectionContainer.h"
+#include "ICollectionManager.h"
 #include "Internationalization/TextPackageNamespaceUtil.h"
 #include "Logging/StructuredLog.h"
+#include "MeshCardBuild.h"
+#include "MeshCardRepresentation.h"
 #include "Misc/CommandLine.h"
 #include "Misc/ConfigCacheIni.h"
-#include "Misc/Paths.h"
+#include "Misc/FeedbackContext.h"
+#include "Misc/FileHelper.h"
 #include "Misc/OutputDeviceHelper.h"
+#include "Misc/PackageName.h"
+#include "Misc/Paths.h"
 #include "Misc/ScopedSlowTask.h"
-#include "Misc/FeedbackContext.h"
-#include "UObject/GCObjectScopeGuard.h"
-#include "UObject/EditorObjectVersion.h"
-#include "UObject/FortniteMainBranchObjectVersion.h"
-#include "UObject/UObjectIterator.h"
 #include "Modules/ModuleManager.h"
-#include "Serialization/PropertyLocalizationDataGathering.h"
-#include "Misc/PackageName.h"
-#include "Misc/FileHelper.h"
-#include "UObject/PackageFileSummary.h"
-#include "Framework/Commands/Commands.h"
-#include "Commandlets/GatherTextFromSourceCommandlet.h"
-#include "AssetRegistry/AssetData.h"
-#include "Sound/DialogueWave.h"
-#include "AssetRegistry/ARFilter.h"
-#include "AssetRegistry/AssetRegistryModule.h"
-#include "WorldPartition/WorldPartition.h"
-#include "WorldPartition/WorldPartitionHelpers.h"
-#include "WorldPartition/WorldPartitionActorDescUtils.h"
-#include "Engine/Level.h"
-#include "Engine/World.h"
-#include "Editor.h"
-#include "EditorWorldUtils.h"
 #include "PackageHelperFunctions.h"
+#include "ProfilingDebugging/ScopedTimers.h"
+#include "Serialization/PropertyLocalizationDataGathering.h"
 #include "ShaderCompiler.h"
-#include "DistanceFieldAtlas.h"
-#include "MeshCardBuild.h"
-#include "MeshCardRepresentation.h"
+#include "Sound/DialogueWave.h"
 #include "Templates/UniquePtr.h"
-#include "CollectionManagerModule.h"
-#include "ICollectionContainer.h"
-#include "ICollectionManager.h"
+#include "UObject/Class.h"
+#include "UObject/EditorObjectVersion.h"
+#include "UObject/FortniteMainBranchObjectVersion.h"
+#include "UObject/GCObjectScopeGuard.h"
+#include "UObject/PackageFileSummary.h"
+#include "UObject/UObjectIterator.h"
+#include "WorldPartition/WorldPartition.h"
 #include "WorldPartition/WorldPartitionActorDescInstance.h"
-#include "Async/ParallelFor.h"
-#include "ProfilingDebugging/ScopedTimers.h"
+#include "WorldPartition/WorldPartitionActorDescUtils.h"
+#include "WorldPartition/WorldPartitionHelpers.h"
 
 DEFINE_LOG_CATEGORY_STATIC(LogGatherTextFromAssetsCommandlet, Log, All);
+namespace GatherTextFromAssetsCommandlet
+{
+	static constexpr int32 LocalizationLogIdentifier = 304;
+}
 
 /** Special feedback context used to stop the commandlet to reporting failure due to a package load error */
 class FLoadPackageLogOutputRedirector : public FFeedbackContext
@@ -1136,7 +1141,7 @@
 				// if this package's localization ID is a duplicate
 				if (const FName* ExistingLongPackageName = AssignedPackageLocalizationIds.Find(PackagePendingGather.PackageLocalizationId))
 				{
-					UE_LOG(LogGatherTextFromAssetsCommandlet, Warning, TEXT("Package '%s' and '%s' have the same localization ID (%s). Please reset one of these (Asset Localization -> Reset Localization ID) to avoid conflicts."), *PackageNameStr, *ExistingLongPackageName->ToString(), *PackagePendingGather.PackageLocalizationId);
+					UE_LOGFMT(LogGatherTextFromAssetsCommandlet, Warning, "Package '{file}' and '{conflict}' have the same localization ID ({locID}). Please reset one of these (Asset Localization -> Reset Localization ID) to avoid conflicts.", ("id", GatherTextFromAssetsCommandlet::LocalizationLogIdentifier), ("file", *PackageNameStr), ("conflict", *ExistingLongPackageName->ToString()), ("locID", *PackagePendingGather.PackageLocalizationId));
 					if (bFixPackageLocalizationIdConflict) 
 					{
 						bThisPackageHasLocIdConflictToFix = true;
diff --git a/Engine/Source/Editor/UnrealEd/Private/Commandlets/GatherTextFromSourceCommandlet.cpp b/Engine/Source/Editor/UnrealEd/Private/Commandlets/GatherTextFromSourceCommandlet.cpp
--- a/Engine/Source/Editor/UnrealEd/Private/Commandlets/GatherTextFromSourceCommandlet.cpp
+++ b/Engine/Source/Editor/UnrealEd/Private/Commandlets/GatherTextFromSourceCommandlet.cpp
@@ -1,21 +1,26 @@
 // Copyright Epic Games, Inc. All Rights Reserved.
 
 #include "Commandlets/GatherTextFromSourceCommandlet.h"
+
+#include "Algo/Transform.h"
 #include "Algo/Unique.h"
 #include "Async/ParallelFor.h"
 #include "HAL/FileManager.h"
 #include "HAL/PlatformMath.h"
+#include "Internationalization/InternationalizationMetadata.h"
+#include "Internationalization/TextNamespaceUtil.h"
+#include "Logging/StructuredLog.h"
 #include "Misc/AsciiSet.h"
+#include "Misc/ExpressionParserTypes.h"
 #include "Misc/FileHelper.h"
 #include "Misc/Paths.h"
-#include "Misc/ExpressionParserTypes.h"
-#include "Algo/Transform.h"
-#include "Internationalization/InternationalizationMetadata.h"
-#include "Internationalization/TextNamespaceUtil.h"
 #include "ProfilingDebugging/ScopedTimers.h"
-#include "Logging/StructuredLog.h"
 
 DEFINE_LOG_CATEGORY_STATIC(LogGatherTextFromSourceCommandlet, Log, All);
+namespace GatherTextFromSourceCommandlet
+{
+	static constexpr int32 LocalizationLogIdentifier = 304;
+}
 
 //////////////////////////////////////////////////////////////////////////
 //GatherTextFromSourceCommandlet
@@ -42,8 +47,6 @@ const FString UGatherTextFromSourceCommandlet::MacroString_NSLOCTEXT(TEXT("NSLOC
 const FString UGatherTextFromSourceCommandlet::MacroString_UI_COMMAND(TEXT("UI_COMMAND"));
 const FString UGatherTextFromSourceCommandlet::MacroString_UI_COMMAND_EXT(TEXT("UI_COMMAND_EXT"));
 
-constexpr int LogIssueIdentifier = 304;
-
 // Nested macro statistics, to track any supported features
 // See https://gcc.gnu.org/onlinedocs/gcc-3.0.1/cpp_3.html
 struct FParsedNestedMacroStats
@@ -1712,7 +1715,7 @@ void UGatherTextFromSourceCommandlet::FSourceFileParseContext::RemoveDefine(cons
 		// #undef LOC_DEFINE_REGION
 		if (!ExcludedRegion)
 		{
-			UE_LOGFMT(LogGatherTextFromSourceCommandlet, Warning, "{file}({line}): Found an '#undef LOC_DEFINE_REGION' without a corresponding '#define LOC_DEFINE_REGION'", ("id", LogIssueIdentifier), ("file", *Filename), ("line", LineNumber));
+			UE_LOGFMT(LogGatherTextFromSourceCommandlet, Warning, "{file}({line}): Found an '#undef LOC_DEFINE_REGION' without a corresponding '#define LOC_DEFINE_REGION'", ("id", GatherTextFromSourceCommandlet::LocalizationLogIdentifier), ("file", *Filename), ("line", LineNumber));
 		}
 		else
 		{
@@ -1727,7 +1730,7 @@ void UGatherTextFromSourceCommandlet::FSourceFileParseContext::RemoveDefine(cons
 			// #undef LOCTEXT_NAMESPACE
 			if (WithinNamespaceDefineLineNumber == INDEX_NONE)
 			{
-				UE_LOGFMT(LogGatherTextFromSourceCommandlet, Warning, "{file}({line}): Found an '#undef LOCTEXT_NAMESPACE' without a corresponding '#define LOCTEXT_NAMESPACE'", ("id", LogIssueIdentifier), ("file", *Filename), ("line", LineNumber));
+				UE_LOGFMT(LogGatherTextFromSourceCommandlet, Warning, "{file}({line}): Found an '#undef LOCTEXT_NAMESPACE' without a corresponding '#define LOCTEXT_NAMESPACE'", ("id", GatherTextFromSourceCommandlet::LocalizationLogIdentifier), ("file", *Filename), ("line", LineNumber));
 			}
 			else
 			{
