Improved GatherTextFromAssetsCommandlet handling of external actors injected via plugins

External actors injected via an External Data Layer or Content Bundle are now only gathered if their associated GameFeatureData asset is included in the gather. For monolithic gathers this should continue to work as-is, but if using modular gathers (where you may need to gather external actors for a plugin separately from their world) you may now need to provide a WorldCollectionFilters entry to control which world actors are gathered (by default all discovered actors from the GFD are gathered).

diff --git a/Engine/Source/Developer/Localization/Private/LocalizationConfigurationScript.cpp b/Engine/Source/Developer/Localization/Private/LocalizationConfigurationScript.cpp
--- a/Engine/Source/Developer/Localization/Private/LocalizationConfigurationScript.cpp
+++ b/Engine/Source/Developer/Localization/Private/LocalizationConfigurationScript.cpp
@@ -389,6 +389,11 @@ namespace LocalizationConfigurationScript
 				ConfigSection.Add( TEXT("CollectionFilters"), CollectionName.ToString() );
 			}
 
+			for (const auto& CollectionName : Target->Settings.GatherFromPackages.WorldCollections)
+			{
+				ConfigSection.Add( TEXT("WorldCollectionFilters"), CollectionName.ToString() );
+			}
+
 			// Class filters
 			for (const FSoftClassPath& ExcludeClass : Target->Settings.GatherFromPackages.ExcludeClasses)
 			{
diff --git a/Engine/Source/Developer/Localization/Public/LocalizationTargetTypes.h b/Engine/Source/Developer/Localization/Public/LocalizationTargetTypes.h
--- a/Engine/Source/Developer/Localization/Public/LocalizationTargetTypes.h
+++ b/Engine/Source/Developer/Localization/Public/LocalizationTargetTypes.h
@@ -188,6 +188,13 @@ struct FGatherTextFromPackagesConfiguration
 	UPROPERTY(config, EditAnywhere, Category = "Filter")
 	TArray<FName> Collections;
 
+	/*
+	 * Worlds (including their associated external actors) in these collections may be processed for gathering. The collections may contain more than just worlds, and will be filtered internally.
+	 * @note This is typically not needed unless setting up modular localization targets that have to gather external actors separately from their world (eg, when using External Data Layers).
+	 */
+	UPROPERTY(config, EditAnywhere, Category = "Filter", AdvancedDisplay)
+	TArray<FName> WorldCollections;
+
 	/* Classes that should be excluded from gathering. */
 	UPROPERTY(config, EditAnywhere, Category = "Filter", meta=(MetaClass="/Script/CoreUObject.Object", AllowAbstract="True"))
 	TArray<FSoftClassPath> ExcludeClasses;
diff --git a/Engine/Source/Editor/UnrealEd/Classes/Commandlets/GatherTextFromAssetsCommandlet.h b/Engine/Source/Editor/UnrealEd/Classes/Commandlets/GatherTextFromAssetsCommandlet.h
--- a/Engine/Source/Editor/UnrealEd/Classes/Commandlets/GatherTextFromAssetsCommandlet.h
+++ b/Engine/Source/Editor/UnrealEd/Classes/Commandlets/GatherTextFromAssetsCommandlet.h
@@ -56,7 +56,7 @@ private:
 	bool PerformFirstPassFilter(TArray<FAssetData>& OutAssetDataArray) const;
 	void ApplyFirstPassFilter(const FARFilter& InFilter, TArray<FAssetData>& InOutAssetDataArray) const;
 	bool BuildFirstPassFilter(FARFilter& InOutFilter) const;
-	bool BuildCollectionFilter(FARFilter& InOutFilter) const;
+	bool BuildCollectionFilter(FARFilter& InOutFilter, const TArray<FString>& Collections) const;
 	bool BuildPackagePathsFilter(FARFilter& InOutFilter) const;
 	bool BuildExcludeDerivedClassesFilter(FARFilter& InOutFilter) const;
 	bool PerformExcludeExactClassesFilter(TArray<FAssetData>& InOutAssetDataArray) const;
@@ -65,7 +65,7 @@ private:
 	void FilterAssetsBasedOnIncludeExcludePaths(TArray<FAssetData>& InOutAssetDataArray) const;
 	
 	bool DiscoverExternalActors(TArray<FAssetData>& InOutAssetDataArray);
-	void RemoveExistingExternalActors(TArray<FAssetData>& InOutAssetDataArray, TSet<FName>& OutPartitionedWorldPackageNames) const;
+	void RemoveExistingExternalActors(TArray<FAssetData>& InOutAssetDataArray, const TSet<FName>* WorldPackageFilter, TSet<FName>& OutPartitionedWorldPackageNames, TSet<FName>& OutGameFeatureDataPackageNames) const;
 
 	TSet<FName> GetPackageNamesToGather(const TArray<FAssetData>& InAssetDataArray) const;
 	void PopulatePackagesPendingGather(TSet<FName> PackageNamesToGather);
@@ -110,6 +110,7 @@ private:
 	TArray<FString> ModulesToPreload;
 	TArray<FString> IncludePathFilters;
 	TArray<FString> CollectionFilters;
+	TArray<FString> WorldCollectionFilters;
 	TArray<FString> ExcludePathFilters;
 	TArray<FString> PackageFileNameFilters;
 	TArray<FString> ExcludeClassNames;
diff --git a/Engine/Source/Editor/UnrealEd/Private/Commandlets/GatherTextFromAssetsCommandlet.cpp b/Engine/Source/Editor/UnrealEd/Private/Commandlets/GatherTextFromAssetsCommandlet.cpp
--- a/Engine/Source/Editor/UnrealEd/Private/Commandlets/GatherTextFromAssetsCommandlet.cpp
+++ b/Engine/Source/Editor/UnrealEd/Private/Commandlets/GatherTextFromAssetsCommandlet.cpp
@@ -25,6 +25,7 @@
 #include "Misc/ConfigCacheIni.h"
 #include "Misc/FeedbackContext.h"
 #include "Misc/FileHelper.h"
+#include "Misc/Optional.h"
 #include "Misc/OutputDeviceHelper.h"
 #include "Misc/PackageName.h"
 #include "Misc/Paths.h"
@@ -537,7 +538,7 @@ bool UGatherTextFromAssetsCommandlet::ShouldRunInPreview(const TArray<FString>&
 bool UGatherTextFromAssetsCommandlet::BuildFirstPassFilter(FARFilter& InOutFilter) const
 {
 	// Filter object paths to only those in any of the specified collections.
-	if (!BuildCollectionFilter(InOutFilter))
+	if (!BuildCollectionFilter(InOutFilter, CollectionFilters))
 	{
 		return false;
 	}
@@ -564,20 +565,20 @@ bool UGatherTextFromAssetsCommandlet::BuildFirstPassFilter(FARFilter& InOutFilte
 }
 
 /** Builds a filter based on the specified collections to be used for gathering.*/
-bool UGatherTextFromAssetsCommandlet::BuildCollectionFilter(FARFilter& InOutFilter) const
+bool UGatherTextFromAssetsCommandlet::BuildCollectionFilter(FARFilter& InOutFilter, const TArray<FString>& Collections) const
 {
 	bool bHasFailedToGetACollection = false;
 	FCollectionManagerModule& CollectionManagerModule = FCollectionManagerModule::GetModule();
 	ICollectionManager& CollectionManager = CollectionManagerModule.Get();
-	for (const FString& CollectionFilter : CollectionFilters)
+	for (const FString& Collection : Collections)
 	{
 		TSharedPtr<ICollectionContainer> CollectionContainer;
 		FName CollectionName;
 		ECollectionShareType::Type ShareType = ECollectionShareType::CST_All;
-		if (!CollectionManager.TryParseCollectionPath(CollectionFilter, &CollectionContainer, &CollectionName, &ShareType) ||
+		if (!CollectionManager.TryParseCollectionPath(Collection, &CollectionContainer, &CollectionName, &ShareType) ||
 			!CollectionContainer->GetObjectsInCollection(CollectionName, ShareType, InOutFilter.SoftObjectPaths, ECollectionRecursionFlags::SelfAndChildren))
 		{
-			UE_LOG(LogGatherTextFromAssetsCommandlet, Error, TEXT("Failed get objects in specified collection: %s"), *CollectionFilter);
+			UE_LOG(LogGatherTextFromAssetsCommandlet, Error, TEXT("Failed get objects in specified collection: %s"), *Collection);
 			bHasFailedToGetACollection = true;
 		}
 	}
@@ -772,15 +773,32 @@ void UGatherTextFromAssetsCommandlet::FilterAssetsBasedOnIncludeExcludePaths(TAr
 }
 
 /** Remove any external actors that currently exist in InOutAssetDataArray. OutPartitionedWorldPackageNames is populated with the package paths of partitioned worlds.*/
-void UGatherTextFromAssetsCommandlet::RemoveExistingExternalActors(TArray<FAssetData>& InOutAssetDataArray, TSet<FName>& OutPartitionedWorldPackageNames) const
+void UGatherTextFromAssetsCommandlet::RemoveExistingExternalActors(TArray<FAssetData>& InOutAssetDataArray, const TSet<FName>* WorldPackageFilter, TSet<FName>& OutPartitionedWorldPackageNames, TSet<FName>& OutGameFeatureDataPackageNames) const
 {
-	InOutAssetDataArray.RemoveAll([&OutPartitionedWorldPackageNames](const FAssetData& AssetData)
+	auto IsGameFeatureData = [](const UClass* AssetClass)
+	{
+		static const FTopLevelAssetPath GameFeatureDataClassPath = FTopLevelAssetPath("/Script/GameFeatures", "GameFeatureData"); // cannot include UGameFeatureData as it is inside a plugin
+		for (const UClass* Class = AssetClass; Class; Class = Class->GetSuperClass())
+		{
+			if (Class->GetClassPathName() == GameFeatureDataClassPath)
+			{
+				return true;
+			}
+		}
+		return false;
+	};
+
+	InOutAssetDataArray.RemoveAll([WorldPackageFilter, &OutPartitionedWorldPackageNames, &OutGameFeatureDataPackageNames, &IsGameFeatureData](const FAssetData& AssetData)
 		{
 			const FNameBuilder PackageNameStr(AssetData.PackageName);
 
 			if (AssetData.AssetClassPath == UWorld::StaticClass()->GetClassPathName())
 			{
-				if (ULevel::GetIsLevelPartitionedFromAsset(AssetData))
+				if (WorldPackageFilter && !WorldPackageFilter->Contains(AssetData.PackageName))
+				{
+					return true;
+				}
+				else if (ULevel::GetIsLevelPartitionedFromAsset(AssetData))
 				{
 					OutPartitionedWorldPackageNames.Add(AssetData.PackageName);
 				}
@@ -791,6 +809,10 @@ void UGatherTextFromAssetsCommandlet::RemoveExistingExternalActors(TArray<FAsset
 				// It is possible for an external actor to be directly specified for gather in the configs but have their world not pass the gather criteria.
 				return true;
 			}
+			else if (IsGameFeatureData(AssetData.GetClass()))
+			{
+				OutGameFeatureDataPackageNames.Add(AssetData.PackageName);
+			}
 
 			return false;
 		});
@@ -802,11 +824,34 @@ bool UGatherTextFromAssetsCommandlet::DiscoverExternalActors(TArray<FAssetData>&
 	UE_LOG(LogGatherTextFromAssetsCommandlet, Display, TEXT("Discovering external actors to gather..."));
 	const double DiscoveringExternalActorsStartTime = FPlatformTime::Seconds();
 
-	TSet<FName> PartitionedWorldPackageNames;
-	RemoveExistingExternalActors(InOutAssetDataArray, PartitionedWorldPackageNames);
-
 	IAssetRegistry& AssetRegistry = IAssetRegistry::GetChecked();
 
+	// If we have WorldCollectionFilters specified, then use those to filter the external actors we gather
+	TOptional<TSet<FName>> WorldPackageFilter;
+	if (WorldCollectionFilters.Num() > 0)
+	{
+		FARFilter Filter;
+		if (!BuildCollectionFilter(Filter, WorldCollectionFilters))
+		{
+			return false;
+		}
+		Filter.ClassPaths.Add(UWorld::StaticClass()->GetClassPathName());
+		Filter.bIncludeOnlyOnDiskAssets = true;
+		Filter.WithoutPackageFlags = PKG_Cooked;
+
+		TArray<FAssetData> FilteredWorldAssets;
+		AssetRegistry.GetAssets(Filter, FilteredWorldAssets);
+
+		TSet<FName> FilteredWorldPackages;
+		FilteredWorldPackages.Reserve(FilteredWorldAssets.Num());
+		Algo::Transform(FilteredWorldAssets, FilteredWorldPackages, &FAssetData::PackageName);
+		WorldPackageFilter.Emplace(MoveTemp(FilteredWorldPackages));
+	}
+
+	TSet<FName> PartitionedWorldPackageNames;
+	TSet<FName> GameFeatureDataPackageNames;
+	RemoveExistingExternalActors(InOutAssetDataArray, WorldPackageFilter.GetPtrOrNull(), PartitionedWorldPackageNames, GameFeatureDataPackageNames);
+
 	// Append the actors that are directly known by each world (by looking for their external actors under the expected path)
 	if (PartitionedWorldPackageNames.Num() > 0)
 	{
@@ -827,25 +872,14 @@ bool UGatherTextFromAssetsCommandlet::DiscoverExternalActors(TArray<FAssetData>&
 
 	// Append additional actors for each world
 	// These are actors added via ExternalDataLayers or ContentBundles, and will be referenced by the GameFeatureData assets that add them (via a GameFeatureAction)
-	if (PartitionedWorldPackageNames.Num() > 0)
+	if (GameFeatureDataPackageNames.Num() > 0)
 	{
-		TArray<FAssetData> GameFeatureDataAssets;
-		{
-			FARFilter Filter;
-			Filter.ClassPaths.Add(FTopLevelAssetPath("/Script/GameFeatures", "GameFeatureData")); // cannot include UGameFeatureData as it is inside a plugin
-			Filter.bRecursiveClasses = true;
-			Filter.bIncludeOnlyOnDiskAssets = true;
-			Filter.WithoutPackageFlags = PKG_Cooked;
-
-			AssetRegistry.GetAssets(Filter, GameFeatureDataAssets);
-		}
-
 		TArray<FAssetData> GameFeatureDataActorDependencies;
 		{
 			FARFilter Filter;
-			for (const FAssetData& GameFeatureDataAsset : GameFeatureDataAssets)
+			for (const FName GameFeatureDataPackageName : GameFeatureDataPackageNames)
 			{
-				AssetRegistry.GetDependencies(GameFeatureDataAsset.PackageName, Filter.PackageNames);
+				AssetRegistry.GetDependencies(GameFeatureDataPackageName, Filter.PackageNames);
 			}
 			Filter.ClassPaths.Add(AActor::StaticClass()->GetClassPathName());
 			Filter.bRecursiveClasses = true;
@@ -858,16 +892,26 @@ bool UGatherTextFromAssetsCommandlet::DiscoverExternalActors(TArray<FAssetData>&
 			}
 		}
 
+		// External actors may be filtered in two ways;
+		//  1. If WorldCollectionFilters were provided, then we only include actors related to those worlds
+		//  2. If PartitionedWorldPackageNames was populated (meaning there are worlds in this gather), then we only include actors related to those worlds
+		// If neither of the above is true then we include all actors related to the GFDs in this gather, as we assume this is a plugin hosting external actors for worlds owned by another localization target
+		const bool bHasWorldFilter = WorldPackageFilter.IsSet() || PartitionedWorldPackageNames.Num() > 0;
+		const TSet<FName>& GameFeatureDataActorsWorldPackageFilter = WorldPackageFilter.Get(PartitionedWorldPackageNames);
 		for (FAssetData& GameFeatureDataActorDependency : GameFeatureDataActorDependencies)
 		{
 			const FName OptionalOuterPathName = GameFeatureDataActorDependency.GetOptionalOuterPathName();
 			if (!OptionalOuterPathName.IsNone())
 			{
 				const FName OptionalOuterPackageName = FName(FPackageName::ObjectPathToPackageName(FNameBuilder(OptionalOuterPathName).ToView()));
-				if (PartitionedWorldPackageNames.Contains(OptionalOuterPackageName))
+				if (!bHasWorldFilter || GameFeatureDataActorsWorldPackageFilter.Contains(OptionalOuterPackageName))
 				{
 					InOutAssetDataArray.Add(MoveTemp(GameFeatureDataActorDependency));
 				}
+				else
+				{
+					UE_LOG(LogGatherTextFromAssetsCommandlet, VeryVerbose, TEXT("Skipping external actor package (%s) as its associated world package (%s) is not relevant to this gather."), *GameFeatureDataActorDependency.GetSoftObjectPath().ToString(), *OptionalOuterPackageName.ToString());
+				}
 			}
 		}
 	}
@@ -941,7 +985,7 @@ bool UGatherTextFromAssetsCommandlet::PerformExcludeExactClassesFilter(TArray<FA
 	if (!ExcludeExactClassesFilter.IsEmpty())
 	{
 		ApplyExcludeExactClassesFilter(ExcludeExactClassesFilter, InOutAssetDataArray);
-		UE_LOG(LogGatherTextFromAssetsCommandlet, Display, TEXT("Filtering otu derived classes took %.2f seconds."), FPlatformTime::Seconds() - ExcludeDerivedClassesStartTime);
+		UE_LOG(LogGatherTextFromAssetsCommandlet, Display, TEXT("Filtering out derived classes took %.2f seconds."), FPlatformTime::Seconds() - ExcludeDerivedClassesStartTime);
 	}
 	return true;
 }
@@ -1748,22 +1792,32 @@ bool UGatherTextFromAssetsCommandlet::ConfigureFromScript(const FString& GatherT
 	}
 
 	// Collections
-	GetStringArrayFromConfig(*SectionName, TEXT("CollectionFilters"), CollectionFilters, GatherTextConfigPath);
-	for (const FString& CollectionFilter : CollectionFilters)
 	{
-		FCollectionManagerModule& CollectionManagerModule = FCollectionManagerModule::GetModule();
-		ICollectionManager& CollectionManager = CollectionManagerModule.Get();
-
-		TSharedPtr<ICollectionContainer> CollectionContainer;
-		FName CollectionName;
-		ECollectionShareType::Type ShareType = ECollectionShareType::CST_All;
-		const bool bDoesCollectionExist = CollectionManager.TryParseCollectionPath(CollectionFilter, &CollectionContainer, &CollectionName, &ShareType) &&
-			CollectionContainer->CollectionExists(CollectionName, ShareType);
-		if (!bDoesCollectionExist)
+		auto GetAndValidateCollections = [this, &bHasFatalError, &SectionName, &GatherTextConfigPath](const TCHAR* KeyName, TArray<FString>& OutCollections)
 		{
-			UE_LOG(LogGatherTextFromAssetsCommandlet, Error, TEXT("Failed to find collection \"%s\", collection does not exist."), *CollectionFilter);
-			bHasFatalError = true;
-		}
+			GetStringArrayFromConfig(*SectionName, KeyName, OutCollections, GatherTextConfigPath);
+			if (OutCollections.Num() > 0)
+			{
+				FCollectionManagerModule& CollectionManagerModule = FCollectionManagerModule::GetModule();
+				ICollectionManager& CollectionManager = CollectionManagerModule.Get();
+				for (const FString& Collection : OutCollections)
+				{
+					TSharedPtr<ICollectionContainer> CollectionContainer;
+					FName CollectionName;
+					ECollectionShareType::Type ShareType = ECollectionShareType::CST_All;
+					const bool bDoesCollectionExist = CollectionManager.TryParseCollectionPath(Collection, &CollectionContainer, &CollectionName, &ShareType) &&
+						CollectionContainer->CollectionExists(CollectionName, ShareType);
+					if (!bDoesCollectionExist)
+					{
+						UE_LOG(LogGatherTextFromAssetsCommandlet, Error, TEXT("Failed to find collection \"%s\", collection does not exist."), *Collection);
+						bHasFatalError = true;
+					}
+				}
+			}
+		};
+
+		GetAndValidateCollections(TEXT("CollectionFilters"), CollectionFilters);
+		GetAndValidateCollections(TEXT("WorldCollectionFilters"), WorldCollectionFilters);
 	}
 
 	// ExcludePathFilters
diff --git a/Engine/Source/Programs/AutomationTool/Localization/LocalizationConfigFileBuilder.cs b/Engine/Source/Programs/AutomationTool/Localization/LocalizationConfigFileBuilder.cs
--- a/Engine/Source/Programs/AutomationTool/Localization/LocalizationConfigFileBuilder.cs
+++ b/Engine/Source/Programs/AutomationTool/Localization/LocalizationConfigFileBuilder.cs
@@ -53,6 +53,7 @@ namespace EpicGames.Localization
 			configSection.AddValues("ExcludePathFilters", gatherTextParams.ExcludePathFilters.ToArray());
 			configSection.AddValues("PackageFileNameFilters", gatherTextParams.PackageFileNameFilters.ToArray());
 			configSection.AddValues("CollectionFilters", gatherTextParams.CollectionFilters.ToArray());
+			configSection.AddValues("WorldCollectionFilters", gatherTextParams.WorldCollectionFilters.ToArray());
 
 			configSection.AddValue("ShouldExcludeDerivedClasses", gatherTextParams.bShouldExcludeDerivedClasses);
 			configSection.AddValue("ShouldGatherFromEditorOnlyData", gatherTextParams.bShouldGatherFromEditorOnlyData);
diff --git a/Engine/Source/Programs/AutomationTool/Localization/LocalizationConfigFileParams.cs b/Engine/Source/Programs/AutomationTool/Localization/LocalizationConfigFileParams.cs
--- a/Engine/Source/Programs/AutomationTool/Localization/LocalizationConfigFileParams.cs
+++ b/Engine/Source/Programs/AutomationTool/Localization/LocalizationConfigFileParams.cs
@@ -73,6 +73,7 @@ namespace EpicGames.Localization
 		public List<string> ExcludePathFilters { get; set; } = new List<string>();
 		public List<string> PackageFileNameFilters { get; set; } = new List<string>();
 		public List<string> CollectionFilters { get; set; } = new List<string>();
+		public List<string> WorldCollectionFilters { get; set; } = new List<string>();
 		public bool bShouldExcludeDerivedClasses { get; set; } = false;
 		public bool bShouldGatherFromEditorOnlyData { get; set; } = false;
 		public bool bSkipGatherCache { get; set; } = false;
@@ -90,6 +91,7 @@ namespace EpicGames.Localization
 			gatherTextFromAssets.PackageFileNameFilters= new List<string>() { "*.umap", "*.uasset" };
 			//@TODOLocalization: Some licensees might not have a need for an Audit_InCook file...should remove this from defaults after config file generation with templates is implemented
 			gatherTextFromAssets.CollectionFilters = new List<string>() { "Audit_InCook" };
+			gatherTextFromAssets.WorldCollectionFilters = new List<string>() { "Audit_InCook" };
 
 			gatherTextFromAssets.bShouldExcludeDerivedClasses = false;
 			gatherTextFromAssets.bShouldGatherFromEditorOnlyData = false;
