MPCook: Add display diagnostic in the CookWorker log when it has received continuous SE_WOULDBLOCK failures from Socket->Send for too long of a period.
#rnx

diff --git a/Engine/Source/Editor/UnrealEd/Private/Cooker/CookWorkerClient.cpp b/Engine/Source/Editor/UnrealEd/Private/Cooker/CookWorkerClient.cpp
--- a/Engine/Source/Editor/UnrealEd/Private/Cooker/CookWorkerClient.cpp
+++ b/Engine/Source/Editor/UnrealEd/Private/Cooker/CookWorkerClient.cpp
@@ -462,6 +462,7 @@ void FCookWorkerClient::CreateServerSocket(const FDirectorConnectionInfo& Connec
 	FWorkerConnectMessage ConnectMessage;
 	ConnectMessage.RemoteIndex = ConnectInfo.RemoteIndex;
 	EConnectionStatus Status = TryWritePacket(ServerSocket, SendBuffer, MarshalToCompactBinaryTCP(ConnectMessage));
+	UpdateSocketSendDiagnostics(Status);
 	if (Status == EConnectionStatus::Incomplete)
 	{
 		SendToState(EConnectStatus::PollWriteConnectMessage);
@@ -483,6 +484,7 @@ void FCookWorkerClient::PollWriteConnectMessage()
 	using namespace CompactBinaryTCP;
 
 	EConnectionStatus Status = TryFlushBuffer(ServerSocket, SendBuffer);
+	UpdateSocketSendDiagnostics(Status);
 	if (Status == EConnectionStatus::Incomplete)
 	{
 		if (FPlatformTime::Seconds() - ConnectStartTimeSeconds > CookWorkerClient::WaitForConnectReplyTimeout &&
@@ -510,6 +512,7 @@ void FCookWorkerClient::PollReceiveConfigMessage()
 	using namespace UE::CompactBinaryTCP;
 	TArray<FMarshalledMessage> Messages;
 	EConnectionStatus SocketStatus = TryReadPacket(ServerSocket, ReceiveBuffer, Messages);
+	UpdateSocketSendDiagnostics(SocketStatus);
 	if (SocketStatus != EConnectionStatus::Okay && SocketStatus != EConnectionStatus::Incomplete)
 	{
 		UE_LOG(LogCook, Error, TEXT("CookWorker initialization failure: failed to read from socket."));
@@ -625,6 +628,7 @@ void FCookWorkerClient::LogConnected()
 void FCookWorkerClient::PumpSendMessages()
 {
 	UE::CompactBinaryTCP::EConnectionStatus Status = UE::CompactBinaryTCP::TryFlushBuffer(ServerSocket, SendBuffer);
+	UpdateSocketSendDiagnostics(Status);
 	if (Status == UE::CompactBinaryTCP::EConnectionStatus::Failed)
 	{
 		UE_LOG(LogCook, Error,
@@ -688,6 +692,7 @@ void FCookWorkerClient::PumpReceiveMessages()
 		}
 		Messages.Reset();
 		EConnectionStatus SocketStatus = TryReadPacket(ServerSocket, ReceiveBuffer, Messages);
+		UpdateSocketSendDiagnostics(SocketStatus);
 		if (SocketStatus != EConnectionStatus::Okay && SocketStatus != EConnectionStatus::Incomplete)
 		{
 			UE_LOG(LogCook, Error,
@@ -874,8 +879,9 @@ void FCookWorkerClient::PumpDisconnect(FTickStackData& StackData)
 
 void FCookWorkerClient::SendMessage(const IMPCollectorMessage& Message)
 {
-	UE::CompactBinaryTCP::TryWritePacket(ServerSocket, SendBuffer,
-		MarshalToCompactBinaryTCP(Message));
+	using namespace UE::CompactBinaryTCP;
+	EConnectionStatus Status = TryWritePacket(ServerSocket, SendBuffer, MarshalToCompactBinaryTCP(Message));
+	UpdateSocketSendDiagnostics(Status);
 }
 
 void FCookWorkerClient::SendToState(EConnectStatus TargetStatus)
@@ -899,6 +905,35 @@ void FCookWorkerClient::LogInvalidMessage(const TCHAR* MessageTypeName)
 		MessageTypeName);
 }
 
+void FCookWorkerClient::UpdateSocketSendDiagnostics(UE::CompactBinaryTCP::EConnectionStatus Status)
+{
+	using namespace UE::CompactBinaryTCP;
+	if (Status != EConnectionStatus::Incomplete)
+	{
+		LastTimeOfCompleteSocketStatusSeconds = 0.0;
+		LastTimeOfWarningOfSocketStatusSeconds = 0.0;
+		return;
+	}
+
+	if (LastTimeOfCompleteSocketStatusSeconds <= 0.0)
+	{
+		LastTimeOfCompleteSocketStatusSeconds = FPlatformTime::Seconds();
+		LastTimeOfWarningOfSocketStatusSeconds = LastTimeOfCompleteSocketStatusSeconds;
+	}
+	else
+	{
+		constexpr double WarningTimePeriod = 60.;
+		double CurrentTime = FPlatformTime::Seconds();
+		if (CurrentTime - LastTimeOfWarningOfSocketStatusSeconds >= WarningTimePeriod)
+		{
+			UE_LOG(LogCook, Display,
+				TEXT("CookWorkerClient has been unable to send messages to the CookDirector for the past %.1f seconds. Continuing cooking locally and attempting to send..."),
+				(float)(CurrentTime - LastTimeOfCompleteSocketStatusSeconds));
+			LastTimeOfWarningOfSocketStatusSeconds = CurrentTime;
+		}
+	}
+}
+
 void FCookWorkerClient::AssignPackages(FAssignPackagesMessage& Message)
 {
 	if (!Message.ExistenceInfos.IsEmpty())
@@ -1075,7 +1110,9 @@ void FCookWorkerClient::TickCollectors(FTickStackData& StackData, bool bFlush, I
 
 		if (!MarshalledMessages.IsEmpty())
 		{
-			UE::CompactBinaryTCP::TryWritePacket(ServerSocket, SendBuffer, MoveTemp(MarshalledMessages));
+			UE::CompactBinaryTCP::EConnectionStatus Status
+				= UE::CompactBinaryTCP::TryWritePacket(ServerSocket, SendBuffer, MoveTemp(MarshalledMessages));
+			UpdateSocketSendDiagnostics(Status);
 		}
 	}
 
diff --git a/Engine/Source/Editor/UnrealEd/Private/Cooker/CookWorkerClient.h b/Engine/Source/Editor/UnrealEd/Private/Cooker/CookWorkerClient.h
--- a/Engine/Source/Editor/UnrealEd/Private/Cooker/CookWorkerClient.h
+++ b/Engine/Source/Editor/UnrealEd/Private/Cooker/CookWorkerClient.h
@@ -122,13 +122,14 @@ private:
 	/** Send this into the given state. Update any state-dependent variables. */
 	void SendToState(EConnectStatus TargetStatus);
 	void LogInvalidMessage(const TCHAR* MessageTypeName);
+	void UpdateSocketSendDiagnostics(UE::CompactBinaryTCP::EConnectionStatus Status);
 	/** Send packages assigned from the server into the request state. */
 	void AssignPackages(FAssignPackagesMessage& Message);
 	/** Tick the registered collectors, or the single given collector if non-null. */
 	void TickCollectors(FTickStackData& StackData, bool bFlush, IMPCollector* SingleCollector = nullptr);
 	/** Helper for ReportDemote/ReportPromote: Collect IMPCollectors and asynchronously add the message to pending. */
 	void ReportPackageMessage(FName PackageName, TUniquePtr<FPackageRemoteResult>&& ResultOwner);
-	
+
 	void HandleAbortPackagesMessage(FMPCollectorClientMessageContext& Context, bool bReadSuccessful,
 		FAbortPackagesMessage&& Message);
 	void HandleRetractionMessage(FMPCollectorClientMessageContext& Context, bool bReadSuccessful,
@@ -170,6 +171,8 @@ private:
 	FSocket* ServerSocket = nullptr;
 	double ConnectStartTimeSeconds = 0.;
 	double NextTickCollectorsTimeSeconds = 0.;
+	double LastTimeOfCompleteSocketStatusSeconds = 0.;
+	double LastTimeOfWarningOfSocketStatusSeconds = 0.;
 	EConnectStatus ConnectStatus = EConnectStatus::Uninitialized;
 	ECookMode::Type DirectorCookMode = ECookMode::CookByTheBook;
 	bool bHasRunFinished = false;
