Cooker: SkipOnlyEditorOnly: Unsolicted packages for instanced package loads: Detect instanced package loads in OnEndLoadPackage hook, and use the loadpath for instanced packages to look up whether packages loaded during the load of an instanced package are unsolicited, or are expected due to the loadpath's hard imports. Discard the editor-only imports based on the AssetRegistry hard imports, the same way we do for editor-only imports in normal package loads.
#rnx
#incrementalcook

diff --git a/Engine/Source/Editor/UnrealEd/Private/CookOnTheFlyServer.cpp b/Engine/Source/Editor/UnrealEd/Private/CookOnTheFlyServer.cpp
--- a/Engine/Source/Editor/UnrealEd/Private/CookOnTheFlyServer.cpp
+++ b/Engine/Source/Editor/UnrealEd/Private/CookOnTheFlyServer.cpp
@@ -4671,21 +4671,52 @@ void UCookOnTheFlyServer::ProcessUnsolicitedPackages(TArray<FName>* OutDiscovere
 			}
 		};
 
-	TMap<FName, FInstigator> NewPackages = PackageTracker->GetNewPackages();
-
-	for (TPair<FName, FInstigator>& PackageWithInstigator : NewPackages)
+	TArray< FPackageStreamEvent> PackageStream = PackageTracker->GetPackageStream();
+	for (FPackageStreamEvent& PackageStreamEvent : PackageStream)
 	{
-		FName PackageName = PackageWithInstigator.Key;
-		FInstigator& Instigator = PackageWithInstigator.Value;
+		FName PackageName = PackageStreamEvent.PackageName;
+		FInstigator& Instigator = PackageStreamEvent.Instigator;
+		if (PackageStreamEvent.EventType == EPackageStreamEvent::InstancedPackageEndLoad)
+		{
+			continue;
+		}
+		check(PackageStreamEvent.EventType == EPackageStreamEvent::PackageLoad);
 
 		FPackageData* PackageData = PackageDatas->TryAddPackageDataByPackageName(PackageName);
 		if (!PackageData)
 		{
-			continue; // Getting the PackageData will fail if e.g. it is a script package
+			// Some types of packages are non-cookable, and we ignore the loadevent of them: scriptpackages,
+			// instanced load packages.
+			continue;
 		}
 		if (PackageData->IsGenerated())
 		{
-			continue; // Generated pacakges are queued separately (with the correct instigator) in QueueGeneratedPackages
+			continue; // Generated packages are queued separately (with the correct instigator) in QueueGeneratedPackages
+		}
+
+		// Packages loaded by an instanced package referencer need to redirect the Instigator to be the Instigator of
+		// the instanced package load, merged with the Instigator category based on the whether the loaded package is
+		// a recorded dependency of the instanced package's LoadPath package.
+		if (!Instigator.Referencer.IsNone())
+		{
+			TRefCountPtr<const FPackageStreamInstancedPackage> Instance
+				= PackageTracker->FindInstancedPackage(Instigator.Referencer);
+			if (Instance)
+			{
+				Instigator.Referencer = Instance->Instigator.Referencer;
+				if (Instigator.Category == EInstigator::Unsolicited)
+				{
+					using namespace UE::AssetRegistry;
+					const EDependencyProperty* DependencyProperty = Instance->Dependencies.Find(PackageName);
+					if (DependencyProperty)
+					{
+						Instigator.Category = EnumHasAnyFlags(*DependencyProperty, EDependencyProperty::Game)
+							? EInstigator::HardDependency : EInstigator::EditorOnlyLoad;
+					}
+				}
+				Instigator.Category = FPackageTracker::MergeReferenceCategories(
+					Instance->Instigator.Category, Instigator.Category);
+			}
 		}
 
 		if (Instigator.Category == EInstigator::EditorOnlyLoad)
@@ -11442,7 +11473,7 @@ void UCookOnTheFlyServer::GenerateInitialRequests(FBeginCookContext& BeginContex
 	{
 		// Clear the list of startup packages currently held by the packagetracker so that we don't see them when we
 		// ProcessUnsolicitedPackages to find hard references used by the first requested package we load.
-		(void)PackageTracker->GetNewPackages();
+		(void)PackageTracker->GetPackageStream();
 	}
 	if (!CookByTheBookOptions->bSkipSoftReferences && !EnumHasAnyFlags(CookOptions, ECookByTheBookOptions::NoStartupPackages))
 	{
diff --git a/Engine/Source/Editor/UnrealEd/Private/Cooker/PackageTracker.cpp b/Engine/Source/Editor/UnrealEd/Private/Cooker/PackageTracker.cpp
--- a/Engine/Source/Editor/UnrealEd/Private/Cooker/PackageTracker.cpp
+++ b/Engine/Source/Editor/UnrealEd/Private/Cooker/PackageTracker.cpp
@@ -8,6 +8,7 @@
 #include "CookProfiling.h"
 #include "Misc/ScopeRWLock.h"
 #include "UObject/Package.h"
+#include "UObject/UObjectGlobals.h"
 #include "UObject/UObjectIterator.h"
 
 namespace UE::Cook
@@ -47,10 +48,10 @@ void FThreadSafeUnsolicitedPackagesList::Empty()
 	CookedPackages.Empty();
 }
 
-
 FPackageTracker::FPackageTracker(UCookOnTheFlyServer& InCOTFS)
 	:COTFS(InCOTFS)
 {
+	ActiveInstances = new FPackageStreamInstancedPackageContainer(); // RefCounted
 }
 
 FPackageTracker::~FPackageTracker()
@@ -60,12 +61,15 @@ FPackageTracker::~FPackageTracker()
 
 void FPackageTracker::Unsubscribe()
 {
-	if (bSubscribed)
+	if (!bSubscribed)
 	{
-		bSubscribed = false;
-		GUObjectArray.RemoveUObjectDeleteListener(this);
-		GUObjectArray.RemoveUObjectCreateListener(this);
+		return;
 	}
+
+	bSubscribed = false;
+	GUObjectArray.RemoveUObjectDeleteListener(this);
+	GUObjectArray.RemoveUObjectCreateListener(this);
+	FCoreUObjectDelegates::OnEndLoadPackage.RemoveAll(this);
 }
 
 void FPackageTracker::InitializeTracking(TSet<FName>& OutStartupPackages)
@@ -87,14 +91,17 @@ void FPackageTracker::InitializeTracking(TSet<FName>& OutStartupPackages)
 		}
 	}
 
-	NewPackages.Reserve(LoadedPackages.Num());
+	TMap<FName, FInstigator> MapOfNewPackages;
+	MapOfNewPackages.Reserve(LoadedPackages.Num());
 	for (UPackage* Package : LoadedPackages)
 	{
-		NewPackages.Add(Package->GetFName(), FInstigator(EInstigator::StartupPackage));
+		MapOfNewPackages.Add(Package->GetFName(), FInstigator(EInstigator::StartupPackage));
 	}
 
 	GUObjectArray.AddUObjectDeleteListener(this);
 	GUObjectArray.AddUObjectCreateListener(this);
+	FCoreUObjectDelegates::OnEndLoadPackage.AddRaw(this, &FPackageTracker::OnEndLoadPackage);
+
 	bSubscribed = true;
 
 	TArray<TPair<FName, ECookLoadType>> StartupPackageLoadTypes;
@@ -107,7 +114,7 @@ void FPackageTracker::InitializeTracking(TSet<FName>& OutStartupPackages)
 		case ECookLoadType::EditorOnly:
 		{
 			OutStartupPackages.Remove(Pair.Key);
-			FInstigator* Instigator = NewPackages.Find(Pair.Key);
+			FInstigator* Instigator = MapOfNewPackages.Find(Pair.Key);
 			if (Instigator)
 			{
 				*Instigator = FInstigator(EInstigator::EditorOnlyLoad);
@@ -116,7 +123,7 @@ void FPackageTracker::InitializeTracking(TSet<FName>& OutStartupPackages)
 		}
 		case ECookLoadType::UsedInGame:
 		{
-			FInstigator* Instigator = NewPackages.Find(Pair.Key);
+			FInstigator* Instigator = MapOfNewPackages.Find(Pair.Key);
 			if (Instigator)
 			{
 				*Instigator = FInstigator(EInstigator::StartupPackageCookLoadScope);
@@ -128,15 +135,21 @@ void FPackageTracker::InitializeTracking(TSet<FName>& OutStartupPackages)
 		}
 	}
 	bTrackingInitialized = true;
+
+	PackageStream.Reserve(MapOfNewPackages.Num());
+	for (TPair<FName, FInstigator>& Pair : MapOfNewPackages)
+	{
+		PackageStream.Add({ Pair.Key, MoveTemp(Pair.Value), EPackageStreamEvent::PackageLoad });
+	}
 }
 
-TMap<FName, FInstigator> FPackageTracker::GetNewPackages()
+TArray<FPackageStreamEvent> FPackageTracker::GetPackageStream()
 {
 	check(bTrackingInitialized);
 
 	FWriteScopeLock ScopeLock(Lock);
-	TMap<FName, FInstigator> Result = MoveTemp(NewPackages);
-	NewPackages.Reset();
+	TArray<FPackageStreamEvent> Result = MoveTemp(PackageStream);
+	PackageStream.Reset();
 	return Result;
 }
 
@@ -144,56 +157,63 @@ void FPackageTracker::NotifyUObjectCreated(const class UObjectBase* Object, int3
 {
 	if (Object->GetClass() == UPackage::StaticClass())
 	{
-		auto Package = const_cast<UPackage*>(static_cast<const UPackage*>(Object));
+		UPackage* Package = const_cast<UPackage*>(static_cast<const UPackage*>(Object));
+		check(Package->GetOuter() == nullptr); // Subpackages are no longer possible in Unreal5 or later.
+		OnCreatePackage(Package->GetFName());
+		LoadedPackages.Add(Package);
+	}
+}
 
-		if (Package->GetOuter() == nullptr)
-		{
-			LLM_SCOPE_BYTAG(Cooker);
-#if ENABLE_COOK_STATS
-			++DetailedCookStats::NumDetectedLoads;
-#endif
-			FName PackageName = Package->GetFName();
+FInstigator FPackageTracker::GetPackageCreationInstigator() const
+{
 #if UE_WITH_PACKAGE_ACCESS_TRACKING
-			PackageAccessTracking_Private::FTrackedData* AccumulatedScopeData =
-				PackageAccessTracking_Private::FPackageAccessRefScope::GetCurrentThreadAccumulatedData();
-			FName ReferencerName(AccumulatedScopeData ? AccumulatedScopeData->PackageName : NAME_None);
+	PackageAccessTracking_Private::FTrackedData* AccumulatedScopeData =
+		PackageAccessTracking_Private::FPackageAccessRefScope::GetCurrentThreadAccumulatedData();
+	FName ReferencerName(AccumulatedScopeData ? AccumulatedScopeData->PackageName : NAME_None);
 #else
-			FName ReferencerName(NAME_None);
+	FName ReferencerName(NAME_None);
 #endif
-			EInstigator InstigatorType;
-			switch (FCookLoadScope::GetCurrentValue())
-			{
-			case ECookLoadType::EditorOnly:
-				InstigatorType = EInstigator::EditorOnlyLoad;
-				break;
-			case ECookLoadType::UsedInGame:
-				InstigatorType = EInstigator::SaveTimeSoftDependency;
-				break;
-			default:
-				InstigatorType = EInstigator::Unsolicited;
-				break;
-			}
-			FInstigator Instigator(InstigatorType, ReferencerName);
-			if (InstigatorType == EInstigator::Unsolicited && COTFS.bHiddenDependenciesDebug)
-			{
-				COTFS.OnDiscoveredPackageDebug(PackageName, Instigator);
-			}
+	EInstigator InstigatorType;
+	switch (FCookLoadScope::GetCurrentValue())
+	{
+	case ECookLoadType::EditorOnly:
+		InstigatorType = EInstigator::EditorOnlyLoad;
+		break;
+	case ECookLoadType::UsedInGame:
+		InstigatorType = EInstigator::SaveTimeSoftDependency;
+		break;
+	default:
+		InstigatorType = EInstigator::Unsolicited;
+		break;
+	}
+	return FInstigator(InstigatorType, ReferencerName);
+}
 
-			FWriteScopeLock ScopeLock(Lock);
-			if (ExpectedNeverLoadPackages.Contains(PackageName))
-			{
-				UE_LOG(LogCook, Verbose, TEXT("SoftGC PoorPerformance: Reloaded package %s."),
-					*WriteToString<256>(PackageName));
-			}
+void FPackageTracker::OnCreatePackage(FName PackageName)
+{
+	LLM_SCOPE_BYTAG(Cooker);
+#if ENABLE_COOK_STATS
+	++DetailedCookStats::NumDetectedLoads;
+#endif
+	FInstigator Instigator = GetPackageCreationInstigator();
+	if (Instigator.Category == EInstigator::Unsolicited && COTFS.bHiddenDependenciesDebug)
+	{
+		COTFS.OnDiscoveredPackageDebug(PackageName, Instigator);
+	}
 
-			LoadedPackages.Add(Package);
-			// We store packages by name rather than by pointer, because they might have their name changed. When
-			// external actors are moved out of their external package, we rename the package to <PackageName>_Trash.
-			// We want to report a load dependency on the package as it was originally loaded; we don't want to report
-			// the renamed packagename if it gets renamed after load.
-			NewPackages.Add(PackageName, MoveTemp(Instigator));
-		}
+	FWriteScopeLock ScopeLock(Lock);
+	if (ExpectedNeverLoadPackages.Contains(PackageName))
+	{
+		UE_LOG(LogCook, Verbose, TEXT("SoftGC PoorPerformance: Reloaded package %s."),
+			*WriteToString<256>(PackageName));
 	}
+
+	// We store packages by name rather than by pointer, because they might have their name changed. When
+	// external actors are moved out of their external package, we rename the package to <PackageName>_Trash.
+	// We want to report a load dependency on the package as it was originally loaded; we don't want to report
+	// the renamed packagename if it gets renamed after load.
+	ActivePackageInstigators.FindOrAdd(PackageName, Instigator);
+	PackageStream.Add({ PackageName, MoveTemp(Instigator), EPackageStreamEvent::PackageLoad });
 }
 
 void FPackageTracker::NotifyUObjectDeleted(const class UObjectBase* Object, int32 Index)
@@ -217,4 +237,182 @@ void FPackageTracker::RemapTargetPlatforms(const TMap<ITargetPlatform*, ITargetP
 	RemapMapKeys(PlatformSpecificNeverCookPackages, Remap);
 }
 
+void FPackageTracker::OnEndLoadPackage(const FEndLoadPackageContext& Context)
+{
+	using namespace UE::AssetRegistry;
+
+	// OnEndLoadPackage is the hook we use to test for whether a load is instanced; the CreatePackage hook is too early
+	// (LoadPath is unknown in the case of instanced loads), and the OnSyncLoadPackage and OnAsyncLoadPackage hooks are
+	// called too frequently - they are called even for packages that have already loaded, before LoadPackage checks
+	// for whether they can early exit. OnEndLoadPackage has the information we need and is only called for packages
+	// when they transition from unloaded to loaded.
+	// 
+	// Use this hook to record instanced loads in our lookup map; ProcessUnsolicitedPackages will respond to the
+	// creation event for the packages by looking them up in the map and using the information we provide about their
+	// AssetRegistry dependencies non-instanced referencers.
+	TArray<TRefCountPtr<FPackageStreamInstancedPackage>, TInlineAllocator<10>> LoadedInstances;
+	for (UPackage* Package : Context.LoadedPackages)
+	{
+		FName PackageName = Package->GetFName();
+		FName LoadedName = Package->GetLoadedPath().GetPackageFName();
+		if (PackageName == LoadedName || LoadedName.IsNone())
+		{
+			continue;
+		}
+
+		FInstigator Instigator(EInstigator::Unsolicited);
+		{
+			FWriteScopeLock ScopeLock(Lock);
+			FInstigator* ActiveInstigator = ActivePackageInstigators.Find(PackageName);
+			if (ActiveInstigator)
+			{
+				Instigator = *ActiveInstigator;
+			}
+		}
+
+		TArray<FAssetDependency> PackageDependencies;
+		COTFS.AssetRegistry->GetDependencies(FAssetIdentifier(LoadedName), PackageDependencies,
+			EDependencyCategory::Package, EDependencyQuery::Hard);
+
+		FWriteScopeLock ActiveInstancesScopeLock(ActiveInstances->Lock);
+		FPackageStreamInstancedPackage*& Existing = ActiveInstances->Map.FindOrAdd(PackageName, nullptr);
+		if (Existing)
+		{
+			if (Existing->LoadedName != LoadedName)
+			{
+				UE_LOG(LogCook, Error,
+					TEXT("OnBeginLoadPackage was called twice for the same package with two different LoadedPaths. Ignoring the second call. PackageName: %s. LoadedPath1: %s. LoadedPath2: %s."),
+					*PackageName.ToString(), *Existing->LoadedName.ToString(), *LoadedName.ToString());
+				continue;
+			}
+			LoadedInstances.Emplace(Existing);
+			continue;
+		}
+
+		TRefCountPtr<FPackageStreamInstancedPackage> InstancedPackage = new FPackageStreamInstancedPackage(*ActiveInstances);
+		Existing = InstancedPackage.GetReference();
+
+		InstancedPackage->PackageName = PackageName;
+		InstancedPackage->LoadedName = LoadedName;
+		InstancedPackage->Instigator = MoveTemp(Instigator);
+		InstancedPackage->Dependencies.Reserve(PackageDependencies.Num());
+		for (FAssetDependency& PackageDependency : PackageDependencies)
+		{
+			InstancedPackage->Dependencies.Add(PackageDependency.AssetId.PackageName, PackageDependency.Properties);
+		}
+
+		LoadedInstances.Emplace(MoveTemp(InstancedPackage));
+	}
+
+	if (LoadedInstances.IsEmpty())
+	{
+		// The usual path through this function is that there were no loaded instances. Clear the
+		// ActivePackageInstigators to fulfill our design of removing that memory when we no longer need it, and then
+		// return without further work.
+		FWriteScopeLock ScopeLock(Lock);
+		ActivePackageInstigators.Reset();
+	}
+	else
+	{
+		// Now that FPackageStreamInstancedPackage have been registered for all of the instanced loads that occurred
+		// during the top-most LoadPackage call, calculate the ancestor non-instanced package referencer for each of
+		// the instanced package loads.
+		{
+			FReadScopeLock ActiveInstancesScopeLock(ActiveInstances->Lock);
+			for (TRefCountPtr<FPackageStreamInstancedPackage>& InstancedPackage : LoadedInstances)
+			{
+				TSet<FPackageStreamInstancedPackage*> Visited;
+				InstancedPackage->FlattenReferencer(ActiveInstancesScopeLock, Visited);
+			}
+		}
+
+		// The FPackageStreamInstancedPackage we created need to remain available until all of the package creation
+		// records created during their load and that we have added into the PackageStream have been processed by
+		// ProcessUnsolicitedPackages. Add end-of-data-lifetime markers for the FPackageStreamInstancedPackages
+		// into the PackageStream, after all of those creation records.
+		FWriteScopeLock ScopeLock(Lock);
+		for (TRefCountPtr<FPackageStreamInstancedPackage>& Instance : LoadedInstances)
+		{
+			PackageStream.Add({ Instance->PackageName, FInstigator(), EPackageStreamEvent::InstancedPackageEndLoad,
+				TRefCountPtr<const FPackageStreamInstancedPackage>(Instance.GetReference()) });
+		}
+
+		// Clear ActivePackageInstigators to fulfill our design of removing that memory when we no longer need it.
+		ActivePackageInstigators.Reset();
+	}
+}
+
+TRefCountPtr<const FPackageStreamInstancedPackage> FPackageTracker::FindInstancedPackage(FName PackageName)
+{
+	FWriteScopeLock ActiveInstancesScopeLock(ActiveInstances->Lock);
+	FPackageStreamInstancedPackage** Existing = ActiveInstances->Map.Find(PackageName);
+	return TRefCountPtr<const FPackageStreamInstancedPackage>(Existing ? *Existing : nullptr);
+}
+
+EInstigator FPackageTracker::MergeReferenceCategories(EInstigator Parent, EInstigator Child)
+{
+	// EditorOnly -> 0, Unsolicited -> 1, UsedInGame -> 2. Return Min(Child, Parent)
+	switch (Parent)
+	{
+	case EInstigator::EditorOnlyLoad: [[fallthrough]];
+	case EInstigator::HardEditorOnlyDependency:
+		return Parent;
+	case EInstigator::Unsolicited:
+		switch (Child)
+		{
+		case EInstigator::EditorOnlyLoad: [[fallthrough]];
+		case EInstigator::HardEditorOnlyDependency:
+			return Child;
+		default:
+			// Child is Unsolicited or UsedInGame, 
+			return Parent;
+		}
+	default:
+		return Child;
+	}
+}
+
+FPackageStreamInstancedPackage::FPackageStreamInstancedPackage(FPackageStreamInstancedPackageContainer& InContainer)
+	: Container(&InContainer)
+{
+}
+
+FPackageStreamInstancedPackage::~FPackageStreamInstancedPackage()
+{
+	{
+		FWriteScopeLock ActiveInstancesScopeLock(Container->Lock);
+		Container->Map.Remove(PackageName);
+	}
+	Container.SafeRelease();
+}
+
+void FPackageStreamInstancedPackage::FlattenReferencer(FReadScopeLock& CalledInsideActiveInstancesLock,
+	TSet<FPackageStreamInstancedPackage*>& Visited)
+{
+	bool bAlreadyVisited;
+	Visited.Add(this, &bAlreadyVisited);
+	if (bAlreadyVisited)
+	{
+		UE_LOG(LogCook, Error, TEXT("Cycle detected in InstancedPackage referencers. PackageName == %s"),
+			*PackageName.ToString());
+		return;
+	}
+	if (Instigator.Referencer.IsNone())
+	{
+		return;
+	}
+	FPackageStreamInstancedPackage** ReferencerPtr = Container->Map.Find(Instigator.Referencer);
+	if (!ReferencerPtr)
+	{
+		return;
+	}
+
+	FPackageStreamInstancedPackage* Referencer = *ReferencerPtr;
+	Referencer->FlattenReferencer(CalledInsideActiveInstancesLock, Visited);
+
+	Instigator.Referencer = Referencer->Instigator.Referencer;
+	Instigator.Category = FPackageTracker::MergeReferenceCategories(
+		Referencer->Instigator.Category, Instigator.Category);
+}
+
 } // namespace UE::Cook
diff --git a/Engine/Source/Editor/UnrealEd/Private/Cooker/PackageTracker.h b/Engine/Source/Editor/UnrealEd/Private/Cooker/PackageTracker.h
--- a/Engine/Source/Editor/UnrealEd/Private/Cooker/PackageTracker.h
+++ b/Engine/Source/Editor/UnrealEd/Private/Cooker/PackageTracker.h
@@ -15,6 +15,9 @@
 #include "UObject/NameTypes.h"
 #include "UObject/UObjectArray.h"
 
+#include <atomic>
+
+class FReadScopeLock;
 class ITargetPlatform;
 class UPackage;
 namespace UE::Cook { struct FInstigator; }
@@ -24,7 +27,7 @@ namespace UE::Cook { struct FRecompileShaderRequest; }
 namespace UE::Cook
 {
 
-struct FPackageDatas;
+struct FPackageStreamInstancedPackage;
 
 template<typename Type>
 struct FThreadSafeQueue
@@ -170,6 +173,73 @@ private:
 	TArray<FFilePlatformRequest>	CookedPackages;
 };
 
+/** Container for name -> data for FPackageStreamInstancedPackage held by the PackageTracker. */
+struct FPackageStreamInstancedPackageContainer : public FThreadSafeRefCountedObject
+{
+	FRWLock Lock;
+	TMap<FName, FPackageStreamInstancedPackage*> Map;
+};
+
+/**
+ * Data about an instanced package load: a package was created with a non-existing name and another package was
+ * loaded into it. We need to handle the recording of packages imported by instanced package loads in a special manner
+ * since their AssetRegistry dependencies will not be found by querying the AssetRegistry for their name.
+ */
+struct FPackageStreamInstancedPackage : public FThreadSafeRefCountedObject
+{
+public:
+	// FPackageStreamInstancedPackages are only createable by FPackageTracker, and they
+	// are registered with PackageTracker in a map lookup by name to raw pointer. They
+	// remove themselves from the map when they are destructed.
+	~FPackageStreamInstancedPackage();
+
+	FName PackageName;
+	FName LoadedName;
+	FInstigator Instigator;
+	TMap<FName, UE::AssetRegistry::EDependencyProperty> Dependencies;
+	
+private:
+	FPackageStreamInstancedPackage(FPackageStreamInstancedPackageContainer& InContainer);
+	/**
+	 * Set the Referencer to the first ancestor in the referencer chain that is non-instanced. Recursively called
+	 * on the parent referencers in between this and the ancestor.
+	 */
+	void FlattenReferencer(FReadScopeLock& CalledInsideActiveInstancesLock,
+		TSet<FPackageStreamInstancedPackage*>& Visited);
+
+private:
+	TRefCountPtr<FPackageStreamInstancedPackageContainer> Container;
+
+	friend FPackageTracker;
+};
+
+/**
+ * The different types of events that can occur about the loads of packages, as collected by the PackageTracker and
+ * reported to the cooker via GetPackageStream.
+ */
+enum class EPackageStreamEvent : uint8
+{
+	PackageLoad,
+	InstancedPackageEndLoad,
+};
+
+/**
+ * Data about an event in the in-order PackageStream collected by the PackageTracker to communicate events about the
+ * load of packages to the cooker.
+ */
+struct FPackageStreamEvent
+{
+	FName PackageName;
+	FInstigator Instigator;
+	EPackageStreamEvent EventType;
+
+	/**
+	 * Only used by EPackageStreamEvent::InstancedPackageEndLoad, to keep the instanced package referenced until
+	 * the event about it has been processed.
+	 */
+	TRefCountPtr<const FPackageStreamInstancedPackage> InstancedPackage;
+};
+
 struct FPackageTracker : public FUObjectArray::FUObjectCreateListener, public FUObjectArray::FUObjectDeleteListener
 {
 public:
@@ -179,12 +249,15 @@ public:
 	void InitializeTracking(TSet<FName>& OutStartupPackages);
 
 	/** Returns all packages that have been loaded since the last time GetNewPackages was called */
-	TMap<FName, FInstigator> GetNewPackages();
+	TArray<FPackageStreamEvent> GetPackageStream();
+	TRefCountPtr<const FPackageStreamInstancedPackage> FindInstancedPackage(FName PackageName);
 
 	virtual void NotifyUObjectCreated(const class UObjectBase* Object, int32 Index) override;
 	virtual void NotifyUObjectDeleted(const class UObjectBase* Object, int32 Index) override;
 	virtual void OnUObjectArrayShutdown() override;
 
+	void OnEndLoadPackage(const FEndLoadPackageContext& Context);
+
 	/** Swap all ITargetPlatform* stored on this instance according to the mapping in @param Remap. */
 	void RemapTargetPlatforms(const TMap<ITargetPlatform*, ITargetPlatform*>& Remap);
 
@@ -226,23 +299,62 @@ public:
 
 	virtual SIZE_T GetAllocatedSize() const override
 	{
-		return LoadedPackages.GetAllocatedSize() + ExpectedNeverLoadPackages.GetAllocatedSize() + NewPackages.GetAllocatedSize();
+		FReadScopeLock ActiveInstancesScopeLock(ActiveInstances->Lock);
+		return LoadedPackages.GetAllocatedSize() + ExpectedNeverLoadPackages.GetAllocatedSize()
+			+ PackageStream.GetAllocatedSize()
+			+ sizeof(*ActiveInstances)
+			+ ActiveInstances->Map.GetAllocatedSize()
+			+ ActiveInstances->Map.Num()*sizeof(FPackageStreamInstancedPackage);
 	}
 
+	/**
+	 * When package A loads B loads C, and we don't want to tell the cooker about B (because it is e.g. an
+	 * insetanced package), we need to calculate the reason that A loaded C by merging the reason A loaded B with
+	 * the reason B loaded C. The PackageTracker sets EInstigator types to represent those reasons, and this function
+	 * does the merging of a parent and child EInstigators to get the merged reason that the GrandParent loads the
+	 * Child.
+	 */
+	static EInstigator MergeReferenceCategories(EInstigator Parent, EInstigator Child);
+
 private:
 	void Unsubscribe();
+	void OnCreatePackage(FName LoadedPackageName);
+	FInstigator GetPackageCreationInstigator() const;
 
-	// Protects data for thread-safety
+	/** Thread synchronizer around data we read/write during hooks from package creation threads. */
 	FRWLock Lock;
 
-	// This is a complete list of currently loaded UPackages
+	// Begin Data Guarded by Lock
+
+	/** All packages currently loaded by the engine. */
 	TFastPointerSet<UPackage*> LoadedPackages;
+	/** List of packages from COTFS that we expect will never be loaded again; log a warning if we see them. */
 	TSet<FName> ExpectedNeverLoadPackages;
-
-	// This list contains the UPackages loaded since last call to GetNewPackages
-	TMap<FName, FInstigator> NewPackages;
+	/* Events about UPackages created since the last call to GetPackageStream, plus related data such as the
+	 * end-of-data-lifetime marker for the data we record about instanced package loads.
+	 */
+	TArray<FPackageStreamEvent> PackageStream;
+	/**
+	 * List of the Instigators we record for each created package; we keep a record of these instigators for use in our
+	 * OnEndLoadPackage hook, and remove them when that hook is complete.
+	 */
+	TMap<FName, FInstigator> ActivePackageInstigators;
+
+	// End Data Guarded by Lock
+
+	// Begin Data ReadOnlyWhileSubscribed
+	// 
+	// This data is read-only while we are subscribed to the hooks that can be called from package creation threads,
+	// and is writable only from scheduler thread, and only when we are not subscribed.
+	/**
+	 * Our pointer to the container used for Instanced packages. The pointer is read only; the container has its own
+	 * internal lock.
+	 */
+	TRefCountPtr<FPackageStreamInstancedPackageContainer> ActiveInstances;
 	bool bTrackingInitialized = false;
 	bool bSubscribed = false;
+
+	// End Data ReadOnlyWhileSubscribed
 };
 
 } // namespace UE::Cook
