Add the name of the current coordinate space to the tooltip of the coordinate space menu.

diff --git a/Engine/Source/Editor/UnrealEd/Private/ViewportToolbar/UnrealEdViewportToolbar.cpp b/Engine/Source/Editor/UnrealEd/Private/ViewportToolbar/UnrealEdViewportToolbar.cpp
--- a/Engine/Source/Editor/UnrealEd/Private/ViewportToolbar/UnrealEdViewportToolbar.cpp
+++ b/Engine/Source/Editor/UnrealEd/Private/ViewportToolbar/UnrealEdViewportToolbar.cpp
@@ -700,6 +700,23 @@ FSlateIcon GetIconFromCoordSystem(ECoordSystem InCoordSystem)
 	}
 }
 
+FText GetNameForCoordSystem(ECoordSystem InCoordSystem)
+{
+	switch (InCoordSystem)
+	{
+	case COORD_World:
+		return LOCTEXT("COORD_World", "World Space");
+	case COORD_Local:
+		return LOCTEXT("COORD_Local", "Local Space");
+	case COORD_Parent:
+		return LOCTEXT("COORD_Parent", "Parent Space");
+	case COORD_Explicit:
+		return LOCTEXT("COORD_Explicit", "Explicit Space");
+	default:
+		return FText::GetEmpty();
+	}
+}
+
 FToolMenuEntry CreateViewportToolbarTransformsSection()
 {
 	return CreateTransformsSubmenu();
@@ -824,8 +841,16 @@ FToolMenuEntry CreateTransformsSubmenu()
 										// Have the Tooltip show the current hotkey(s) for cycling coord spaces
 										{
 											CoordinateSystemSubmenu.ToolTip = TAttribute<FText>::CreateLambda(
-												[]()
+												[WeakViewport = Viewport.ToWeakPtr()]()
 												{
+													ECoordSystem CoordSystem = ECoordSystem::COORD_Local;
+													if (const TSharedPtr<SEditorViewport> EditorViewport = WeakViewport.Pin())
+													{
+														CoordSystem = EditorViewport->GetViewportClient()->GetWidgetCoordSystemSpace();
+													}
+													
+													FText CoordSystemText = GetNameForCoordSystem(CoordSystem);
+												
 													const FInputChord& PrimaryChord =
 														*FEditorViewportCommands::Get().CycleTransformGizmoCoordSystem->GetActiveChord(
 															EMultipleKeyBindingIndex::Primary
@@ -867,14 +892,18 @@ FToolMenuEntry CreateTransformsSubmenu()
 
 														return FText::Format(
 															LOCTEXT(
-																"CoordinateSystemTooltipSingleChord", "Select between coordinate systems. \n{0} to cycle between them."
+																"CoordinateSystemTooltipSingleChord", "{0} Coordinates\n Click or press {1} to cycle between coordinate systems."
 															),
+															CoordSystemText,
 															ChordText
 														);
 													}
 
-													return LOCTEXT(
-														"CoordinateSystemTooltipNoChords", "Select between coordinate systems"
+													return FText::Format(
+														LOCTEXT(
+															"CoordinateSystemTooltipNoChords", "{0} Coordinates\n Click to cycle between coordinate systems."
+														),
+														CoordSystemText
 													);
 												}
 											);
