Fixed typo

#rnx

diff --git a/Engine/Source/Editor/UnrealEd/Public/Settings/BlueprintEditorProjectSettings.h b/Engine/Source/Editor/UnrealEd/Public/Settings/BlueprintEditorProjectSettings.h
--- a/Engine/Source/Editor/UnrealEd/Public/Settings/BlueprintEditorProjectSettings.h
+++ b/Engine/Source/Editor/UnrealEd/Public/Settings/BlueprintEditorProjectSettings.h
@@ -73,11 +73,11 @@ public:
 	TArray<FName> DisabledCompilerMessages;
 
 	/**
-	 * List of deprecated UProperties/UFunctions to supress warning messages for - useful for source changes
+	 * List of deprecated UProperties/UFunctions to suppress warning messages for - useful for source changes
 	 * that would otherwise cause content warnings
 	 * The easiest way to populate this list is using the context menu on nodes with deprecated references
 	 */
-	UPROPERTY(EditAnywhere, config, Category = Blueprints, DisplayName = "Deprecated Symbols to Supress")
+	UPROPERTY(EditAnywhere, config, Category = Blueprints, DisplayName = "Deprecated Symbols to Suppress")
 	TArray<FString> SuppressedDeprecationMessages;
 
 	/**
