[LUF] Top/Down orthogonal view no longer shows incorrect Y axis direction in UE when not using LUF.

diff --git a/Engine/Source/Editor/UnrealEd/Private/EditorViewportClient.cpp b/Engine/Source/Editor/UnrealEd/Private/EditorViewportClient.cpp
--- a/Engine/Source/Editor/UnrealEd/Private/EditorViewportClient.cpp
+++ b/Engine/Source/Editor/UnrealEd/Private/EditorViewportClient.cpp
@@ -3707,10 +3707,10 @@ void FEditorViewportClient::DrawAxes(FViewport* InViewport, FCanvas* Canvas, con
 		TextItem.Draw( Canvas, FVector2D(TextCenter.X - 0.5*XL, TextCenter.Y - 0.5*YL) );
 	}
 
-	const bool bIsLeft = (InAxis & EAxisList::Left) == EAxisList::Left;
-	if ((InAxis & EAxisList::Y) == EAxisList::Y || bIsLeft)
+	if ((InAxis & EAxisList::Y) == EAxisList::Y || (InAxis & EAxisList::Left) == EAxisList::Left)
 	{
-		if (bIsLeft)
+		// If we are using the LUF coordinate system, we need to reflect the Y axis
+		if (AxisDisplayInfo::GetAxisDisplayCoordinateSystem() == EAxisList::LeftUpForward)
 		{
 			AxisVec = AxisSize * ViewTM.InverseTransformVector(FVector(0, -1, 0));
 			TextVec = (AxisSize + TextDistance) * ViewTM.InverseTransformVector(FVector(0, -1, 0));
