[Gizmos] New Transform Gizmo - Preference

- Creates new `UTransformGizmoEditorSettings` class stored in `EditorInteractiveFramework` module

- Removes local variables tracking new TRS on/off and Gizmo Settings from `GizmoManagerLocals` namespace in `EditorInteractiveGizmoManager.cpp`

- Switches to use the new variables in the Settings Class

- Deprecates `ULevelEditorViewportSettings::TransformWidgetSizeAdjustment` in favor of `UTransformGizmoEditorSettings::TransformGizmoSize`

- Maps new, non-config properties to the following found in `ULevelEditorViewportSettings`: `bAllowArcballRotate`, `bAllowScreenRotate`, `bAllowEditWidgetAxisDisplay`, `bAllowTranslateRotateZWidget`. This allows to avoid a circular dependency `EditorITF Module` <-> `LevelEditor` Module, and saves us from having config properties which will be eliminated in the future, since they are part of legacy gizmo settings

- Removes Experimental `GizmoFramework` plugin, which currently only contains settings which would be a duplicate of some of the ones introduced in this CL by `UTransformGizmoEditorSettings`

diff --git a/Engine/Config/BaseEngine.ini b/Engine/Config/BaseEngine.ini
--- a/Engine/Config/BaseEngine.ini
+++ b/Engine/Config/BaseEngine.ini
@@ -1361,6 +1361,7 @@ ClassName=/Script/OnlineSubsystemUtils.OnlineEngineInterfaceImpl
 ; Fixing the output of tags for blueprints to be reset.
 +FunctionRedirects=(OldName="/Script/GameplayTags.GameplayTagAssetInterface.GetOwnedGameplayTags", NewName="BP_GetOwnedGameplayTags")
 +PropertyRedirects=(OldName="GameplayTagAssetInterface.BP_GetOwnedGameplayTags.TagContainer", NewName="ReturnValue")
++PropertyRedirects=(OldName="/Script/UnrealEd.LevelEditorViewportSettings.TransformWidgetSizeAdjustment",NewName="/Script/UnrealEd.LevelEditorViewportSettings.TransformWidgetSizeAdjustment_DEPRECATED")
 
 [CoreUObject.Metadata]
 ; Note: UnrealHeaderTool should be rerun after making changes to MetadataRedirects to catch any keys specified in class headers
diff --git a/Engine/Plugins/Experimental/GizmoFramework/GizmoFramework.uplugin b/Engine/Plugins/Experimental/GizmoFramework/GizmoFramework.uplugin
deleted file mode 100644
diff --git a/Engine/Plugins/Experimental/GizmoFramework/Source/GizmoSettings/GizmoSettings.Build.cs b/Engine/Plugins/Experimental/GizmoFramework/Source/GizmoSettings/GizmoSettings.Build.cs
deleted file mode 100644
diff --git a/Engine/Plugins/Experimental/GizmoFramework/Source/GizmoSettings/Private/GizmoSettings.cpp b/Engine/Plugins/Experimental/GizmoFramework/Source/GizmoSettings/Private/GizmoSettings.cpp
deleted file mode 100644
diff --git a/Engine/Plugins/Experimental/GizmoFramework/Source/GizmoSettings/Private/GizmoSettingsModule.cpp b/Engine/Plugins/Experimental/GizmoFramework/Source/GizmoSettings/Private/GizmoSettingsModule.cpp
deleted file mode 100644
diff --git a/Engine/Plugins/Experimental/GizmoFramework/Source/GizmoSettings/Public/GizmoSettings.h b/Engine/Plugins/Experimental/GizmoFramework/Source/GizmoSettings/Public/GizmoSettings.h
deleted file mode 100644
diff --git a/Engine/Plugins/Experimental/GizmoFramework/Source/GizmoSettings/Public/GizmoSettingsModule.h b/Engine/Plugins/Experimental/GizmoFramework/Source/GizmoSettings/Public/GizmoSettingsModule.h
deleted file mode 100644
diff --git a/Engine/Source/Editor/Experimental/EditorInteractiveToolsFramework/EditorInteractiveToolsFramework.Build.cs b/Engine/Source/Editor/Experimental/EditorInteractiveToolsFramework/EditorInteractiveToolsFramework.Build.cs
--- a/Engine/Source/Editor/Experimental/EditorInteractiveToolsFramework/EditorInteractiveToolsFramework.Build.cs
+++ b/Engine/Source/Editor/Experimental/EditorInteractiveToolsFramework/EditorInteractiveToolsFramework.Build.cs
@@ -24,6 +24,7 @@ public class EditorInteractiveToolsFramework : ModuleRules
 			new string[]
 			{
 				"CoreUObject",
+				"DeveloperSettings",
 				"Engine",
 				"RenderCore",
 				"Slate",
diff --git a/Engine/Source/Editor/Experimental/EditorInteractiveToolsFramework/Private/EditorInteractiveGizmoManager.cpp b/Engine/Source/Editor/Experimental/EditorInteractiveToolsFramework/Private/EditorInteractiveGizmoManager.cpp
--- a/Engine/Source/Editor/Experimental/EditorInteractiveToolsFramework/Private/EditorInteractiveGizmoManager.cpp
+++ b/Engine/Source/Editor/Experimental/EditorInteractiveToolsFramework/Private/EditorInteractiveGizmoManager.cpp
@@ -22,6 +22,7 @@
 #include "Snapping/EditorSnappingManager.h"
 #include "Templates/Casts.h"
 #include "ToolContextInterfaces.h"
+#include "TransformGizmoEditorSettings.h"
 #include "EditorGizmos/TransformGizmo.h"
 
 class FCanvas;
@@ -30,7 +31,6 @@ class FCanvas;
 
 namespace GizmoManagerLocals
 {
-	static bool UseNewTRSWidget = false;
 	static TOptional<FGizmosParameters> OptDefaultParameters;
 	static UEditorInteractiveGizmoManager::FOnUsesNewTRSGizmosChanged OnUsesNewTRSGizmosChanged;
 	static UEditorInteractiveGizmoManager::FOnGizmosParametersChanged OnGizmosParametersChanged;
@@ -38,15 +38,22 @@ namespace GizmoManagerLocals
 
 bool UEditorInteractiveGizmoManager::UsesNewTRSGizmos()
 {
-	return GizmoManagerLocals::UseNewTRSWidget;
+	if (const UTransformGizmoEditorSettings* TransformGizmoEditorSettings = GetDefault<UTransformGizmoEditorSettings>())
+	{
+		return TransformGizmoEditorSettings->UsesNewTRSGizmo();
+	}
+
+	return false;
 }
 
 void UEditorInteractiveGizmoManager::SetUsesNewTRSGizmos(const bool bUseNewTRSGizmos)
 {
-	if (bUseNewTRSGizmos != GizmoManagerLocals::UseNewTRSWidget)
+	if (UTransformGizmoEditorSettings* const TransformGizmoEditorSettings = GetMutableDefault<UTransformGizmoEditorSettings>())
 	{
-		GizmoManagerLocals::UseNewTRSWidget = bUseNewTRSGizmos;
-		GizmoManagerLocals::OnUsesNewTRSGizmosChanged.Broadcast(GizmoManagerLocals::UseNewTRSWidget);
+		if (bUseNewTRSGizmos != TransformGizmoEditorSettings->UsesNewTRSGizmo())
+		{
+			TransformGizmoEditorSettings->SetUsesLegacyGizmo(!bUseNewTRSGizmos);
+		}
 	}
 }
 
@@ -57,8 +64,10 @@ UEditorInteractiveGizmoManager::FOnUsesNewTRSGizmosChanged& UEditorInteractiveGi
 
 void UEditorInteractiveGizmoManager::SetGizmosParameters(const FGizmosParameters& InParameters)
 {
-	GizmoManagerLocals::OptDefaultParameters = InParameters;
-	GizmoManagerLocals::OnGizmosParametersChanged.Broadcast(InParameters);
+	if (UTransformGizmoEditorSettings* const TransformGizmoEditorSettings = GetMutableDefault<UTransformGizmoEditorSettings>())
+	{
+		TransformGizmoEditorSettings->SetGizmosParameters(InParameters);
+	}
 }
 
 UEditorInteractiveGizmoManager::FOnGizmosParametersChanged& UEditorInteractiveGizmoManager::OnGizmosParametersChangedDelegate()
@@ -68,6 +77,11 @@ UEditorInteractiveGizmoManager::FOnGizmosParametersChanged& UEditorInteractiveGi
 
 const TOptional<FGizmosParameters>& UEditorInteractiveGizmoManager::GetDefaultGizmosParameters()
 {
+	if (const UTransformGizmoEditorSettings* TransformGizmoEditorSettings = GetDefault<UTransformGizmoEditorSettings>())
+	{
+		GizmoManagerLocals::OptDefaultParameters = TransformGizmoEditorSettings->GizmosParameters;
+	}\
+
 	static const TOptional<FGizmosParameters> Invalid;
 	return GizmoManagerLocals::OptDefaultParameters.IsSet() ? GizmoManagerLocals::OptDefaultParameters : Invalid;
 }
diff --git a/Engine/Source/Editor/Experimental/EditorInteractiveToolsFramework/Private/TransformGizmoEditorSettings.cpp b/Engine/Source/Editor/Experimental/EditorInteractiveToolsFramework/Private/TransformGizmoEditorSettings.cpp
new file mode 100644
--- /dev/null
+++ b/Engine/Source/Editor/Experimental/EditorInteractiveToolsFramework/Private/TransformGizmoEditorSettings.cpp
@@ -0,0 +1,188 @@
+﻿// Copyright Epic Games, Inc. All Rights Reserved.
+
+#include "TransformGizmoEditorSettings.h"
+#include "Editor.h"
+#include "EditorInteractiveGizmoManager.h"
+#include "Settings/LevelEditorViewportSettings.h"
+
+UTransformGizmoEditorSettings::UTransformGizmoEditorSettings(const FObjectInitializer& ObjectInitializer)
+	: Super(ObjectInitializer)
+{
+	if (ULevelEditorViewportSettings* const ViewportSettings = GetMutableDefault<ULevelEditorViewportSettings>())
+	{
+		bEnableArcballRotate = ViewportSettings->bAllowArcballRotate;
+		bEnableScreenRotate = ViewportSettings->bAllowScreenRotate;
+		bEnableAxisDrawing = ViewportSettings->bAllowEditWidgetAxisDisplay;
+		bEnableCombinedTranslateRotate = ViewportSettings->bAllowTranslateRotateZWidget;
+
+		ViewportSettings->OnSettingChanged().AddUObject(this, &UTransformGizmoEditorSettings::OnLegacySettingChanged);
+	}
+}
+
+#if WITH_EDITOR
+void UTransformGizmoEditorSettings::PostEditChangeChainProperty(FPropertyChangedChainEvent& InPropertyChangedEvent)
+{
+	Super::PostEditChangeChainProperty(InPropertyChangedEvent);
+
+	if (InPropertyChangedEvent.PropertyChain.GetHead())
+	{
+		if (const FProperty* const Property = InPropertyChangedEvent.PropertyChain.GetHead()->GetValue())
+		{
+			const FName ChangedPropertyName = Property->GetFName();
+			static const FName GizmosParametersName = GET_MEMBER_NAME_CHECKED(UTransformGizmoEditorSettings, GizmosParameters);
+
+			if (GizmosParametersName == ChangedPropertyName)
+			{
+				BroadcastGizmosParametersChange();
+			}
+		}
+	}
+}
+
+void UTransformGizmoEditorSettings::PostEditChangeProperty(FPropertyChangedEvent& InPropertyChangedEvent)
+{
+	Super::PostEditChangeProperty(InPropertyChangedEvent);
+
+	static const FName UseLegacyGizmoName = GET_MEMBER_NAME_CHECKED(UTransformGizmoEditorSettings, bEnableLegacyGizmo);
+	static const FName EnableArcballRotateName = GET_MEMBER_NAME_CHECKED(UTransformGizmoEditorSettings, bEnableArcballRotate);
+	static const FName EnableScreenRotateName = GET_MEMBER_NAME_CHECKED(UTransformGizmoEditorSettings, bEnableScreenRotate);
+	static const FName EnableAxisDrawingName = GET_MEMBER_NAME_CHECKED(UTransformGizmoEditorSettings, bEnableAxisDrawing);
+	static const FName EnableCombinedTranslateRotateName = GET_MEMBER_NAME_CHECKED(UTransformGizmoEditorSettings, bEnableCombinedTranslateRotate);
+
+	if (!InPropertyChangedEvent.Property)
+	{
+		return;
+	}
+
+	const FName ChangedPropertyName = InPropertyChangedEvent.Property->GetFName();
+
+	if (ChangedPropertyName == UseLegacyGizmoName)
+	{
+		BroadcastNewTRSGizmoChange();
+	}
+	else if (ChangedPropertyName == EnableArcballRotateName)
+	{
+		if (ULevelEditorViewportSettings* const ViewportSettings = GetMutableDefault<ULevelEditorViewportSettings>())
+		{
+			ViewportSettings->bAllowArcballRotate = bEnableArcballRotate;
+		}
+	}
+	else if (ChangedPropertyName == EnableScreenRotateName)
+	{
+		if (ULevelEditorViewportSettings* const ViewportSettings = GetMutableDefault<ULevelEditorViewportSettings>())
+		{
+			ViewportSettings->bAllowScreenRotate = bEnableScreenRotate;
+		}
+	}
+	else if (ChangedPropertyName == EnableAxisDrawingName)
+	{
+		if (ULevelEditorViewportSettings* const ViewportSettings = GetMutableDefault<ULevelEditorViewportSettings>())
+		{
+			ViewportSettings->bAllowEditWidgetAxisDisplay = bEnableAxisDrawing;
+		}
+	}
+	else if (ChangedPropertyName == EnableCombinedTranslateRotateName)
+	{
+		if (ULevelEditorViewportSettings* const ViewportSettings = GetMutableDefault<ULevelEditorViewportSettings>())
+		{
+			ViewportSettings->bAllowTranslateRotateZWidget = bEnableCombinedTranslateRotate;
+		}
+	}
+}
+#endif // WITH_EDITOR
+
+void UTransformGizmoEditorSettings::SetUsesLegacyGizmo(bool bInUseLegacyGizmo)
+{
+	if (bEnableLegacyGizmo != bInUseLegacyGizmo)
+	{
+		bEnableLegacyGizmo = bInUseLegacyGizmo;
+		SaveConfig();
+
+		BroadcastNewTRSGizmoChange();
+
+		if (GEditor)
+		{
+			GEditor->RedrawAllViewports();
+		}
+	}
+}
+
+bool UTransformGizmoEditorSettings::UsesLegacyGizmo() const
+{
+	return bEnableLegacyGizmo;
+}
+
+bool UTransformGizmoEditorSettings::UsesNewTRSGizmo() const
+{
+	return !UsesLegacyGizmo();
+}
+
+void UTransformGizmoEditorSettings::SetGizmosParameters(const FGizmosParameters& InGizmosParameters)
+{
+	GizmosParameters = InGizmosParameters;
+	SaveConfig();
+
+	BroadcastGizmosParametersChange();
+}
+
+void UTransformGizmoEditorSettings::SetTransformGizmoSize(float InTransformGizmoSize)
+{
+	if (TransformGizmoSize != InTransformGizmoSize)
+	{
+		TransformGizmoSize = InTransformGizmoSize;
+		SaveConfig();
+
+		if (GEditor)
+		{
+			GEditor->RedrawAllViewports();
+		}
+	}
+}
+
+void UTransformGizmoEditorSettings::BroadcastNewTRSGizmoChange() const
+{
+	UEditorInteractiveGizmoManager::OnUsesNewTRSGizmosChangedDelegate().Broadcast(UsesNewTRSGizmo());
+}
+
+void UTransformGizmoEditorSettings::BroadcastGizmosParametersChange() const
+{
+	UEditorInteractiveGizmoManager::OnGizmosParametersChangedDelegate().Broadcast(GizmosParameters);
+}
+
+void UTransformGizmoEditorSettings::OnLegacySettingChanged(FName InPropertyName)
+{
+	// Retrieving names from ULevelEditorViewportSettings
+	static const FName AllowArcballRotateName = GET_MEMBER_NAME_CHECKED(ULevelEditorViewportSettings, bAllowArcballRotate);
+	static const FName AllowScreenRotateName = GET_MEMBER_NAME_CHECKED(ULevelEditorViewportSettings, bAllowScreenRotate);
+	static const FName AllowEditWidgetAxisDisplayName = GET_MEMBER_NAME_CHECKED(ULevelEditorViewportSettings, bAllowEditWidgetAxisDisplay);
+	static const FName AllowTranslateRotateZWidgetName = GET_MEMBER_NAME_CHECKED(ULevelEditorViewportSettings, bAllowTranslateRotateZWidget);
+
+	if (InPropertyName == AllowArcballRotateName)
+	{
+		if (const ULevelEditorViewportSettings* const ViewportSettings = GetDefault<ULevelEditorViewportSettings>())
+		{
+			bEnableArcballRotate = ViewportSettings->bAllowArcballRotate;
+		}
+	}
+	else if (InPropertyName == AllowScreenRotateName)
+	{
+		if (const ULevelEditorViewportSettings* const ViewportSettings = GetDefault<ULevelEditorViewportSettings>())
+		{
+			bEnableScreenRotate = ViewportSettings->bAllowScreenRotate;
+		}
+	}
+	else if (InPropertyName == AllowEditWidgetAxisDisplayName)
+	{
+		if (const ULevelEditorViewportSettings* const ViewportSettings = GetDefault<ULevelEditorViewportSettings>())
+		{
+			bEnableAxisDrawing = ViewportSettings->bAllowEditWidgetAxisDisplay;
+		}
+	}
+	else if (InPropertyName == AllowTranslateRotateZWidgetName)
+	{
+		if (const ULevelEditorViewportSettings* const ViewportSettings = GetDefault<ULevelEditorViewportSettings>())
+		{
+			bEnableCombinedTranslateRotate = ViewportSettings->bAllowTranslateRotateZWidget;
+		}
+	}
+}
diff --git a/Engine/Source/Editor/Experimental/EditorInteractiveToolsFramework/Public/EditorGizmos/TransformGizmo.h b/Engine/Source/Editor/Experimental/EditorInteractiveToolsFramework/Public/EditorGizmos/TransformGizmo.h
--- a/Engine/Source/Editor/Experimental/EditorInteractiveToolsFramework/Public/EditorGizmos/TransformGizmo.h
+++ b/Engine/Source/Editor/Experimental/EditorInteractiveToolsFramework/Public/EditorGizmos/TransformGizmo.h
@@ -112,16 +112,22 @@ struct FGizmosParameters
 {
 	GENERATED_BODY()
 
-	/** Rotate mode used when manipulating axis rotation handles. */
-	UPROPERTY(EditAnywhere, Category = NewTRSGizmo)
+	/** Determines how dragging the rotate gizmo affects the selected objects */
+	UPROPERTY(EditAnywhere, Category = "Transform Gizmo")
 	TEnumAsByte<EAxisRotateMode::Type> RotateMode = EAxisRotateMode::Arc;
-	
-	/** Property used to define whether the Y axis is indirectly controlled via CTRL+MMB or CTRL+RMB (and Z via CTRL+RMB or CTRL+LMB+RMB). */
-	UPROPERTY(EditAnywhere, Category = NewTRSGizmo)
+
+	/**
+	 * When enabled, Ctrl+MMB activates the Y axis and Ctrl+RMB activates the Z axis.
+	 * When disabled, Ctrl+RMB activates the Y axis and Ctrl+LMB+RMB activates the Z axis.
+	 */
+	UPROPERTY(EditAnywhere, Category = "Transform Gizmo")
 	bool bCtrlMiddleDoesY = true;
 
-	/** Enable/Disable the explicit coordinate system. */
-	UPROPERTY(EditAnywhere, Category = NewTRSGizmo)
+	/**
+	 * When enabled, the list of coordinate spaces in the viewport toolbar offers rig space: a coordinate system that is
+	 * similar to parent space but uses gimbal rotations.
+	 */
+	UPROPERTY(EditAnywhere, Category = "Transform Gizmo")
 	bool bEnableExplicit = false;
 };
 
diff --git a/Engine/Source/Editor/Experimental/EditorInteractiveToolsFramework/Public/EditorInteractiveGizmoManager.h b/Engine/Source/Editor/Experimental/EditorInteractiveToolsFramework/Public/EditorInteractiveGizmoManager.h
--- a/Engine/Source/Editor/Experimental/EditorInteractiveToolsFramework/Public/EditorInteractiveGizmoManager.h
+++ b/Engine/Source/Editor/Experimental/EditorInteractiveToolsFramework/Public/EditorInteractiveGizmoManager.h
@@ -10,6 +10,7 @@
 #include "UObject/ObjectMacros.h"
 #include "UObject/ObjectPtr.h"
 #include "UObject/UObjectGlobals.h"
+#include "EditorGizmos/TransformGizmo.h"
 
 #include "EditorInteractiveGizmoManager.generated.h"
 
@@ -26,7 +27,6 @@ class UObject;
 class UTypedElementSelectionSet;
 struct FToolBuilderState;
 class UTransformGizmo;
-struct FGizmosParameters;
 
 USTRUCT()
 struct FActiveEditorGizmo
@@ -198,6 +198,9 @@ protected:
 	UPROPERTY()
 	TArray<FActiveEditorGizmo> ActiveEditorGizmos;
 
+	UPROPERTY()
+	TOptional<FGizmosParameters> DefaultGizmoParameters;
+
 	/** If false, only search gizmo builders in current gizmo manager. If true, also search gizmo subsystem */
 	bool bSearchLocalBuildersOnly = false;
 
diff --git a/Engine/Source/Editor/Experimental/EditorInteractiveToolsFramework/Public/TransformGizmoEditorSettings.h b/Engine/Source/Editor/Experimental/EditorInteractiveToolsFramework/Public/TransformGizmoEditorSettings.h
new file mode 100644
--- /dev/null
+++ b/Engine/Source/Editor/Experimental/EditorInteractiveToolsFramework/Public/TransformGizmoEditorSettings.h
@@ -0,0 +1,78 @@
+﻿// Copyright Epic Games, Inc. All Rights Reserved.
+
+#pragma once
+
+#include "Engine/DeveloperSettings.h"
+#include "EditorGizmos/TransformGizmo.h"
+
+#include "TransformGizmoEditorSettings.generated.h"
+
+UCLASS(BlueprintType, Config = EditorPerProjectUserSettings, meta = (DisplayName = "Transform Gizmo"))
+class EDITORINTERACTIVETOOLSFRAMEWORK_API UTransformGizmoEditorSettings : public UDeveloperSettings
+{
+	GENERATED_BODY()
+
+public:
+	UTransformGizmoEditorSettings(const FObjectInitializer& ObjectInitializer);
+
+#if WITH_EDITOR
+	virtual void PostEditChangeChainProperty(FPropertyChangedChainEvent& InPropertyChangedEvent) override;
+	virtual void PostEditChangeProperty(FPropertyChangedEvent& InPropertyChangedEvent) override;
+#endif
+
+	UPROPERTY(Config, EditAnywhere, Category = "Transform Gizmo", meta = (ShowOnlyInnerProperties))
+	FGizmosParameters GizmosParameters;
+
+	UPROPERTY(Config, EditAnywhere, Category = "Transform Gizmo", DisplayName = "Transform Gizmo Size", meta = (ClampMin = "-10.0", ClampMax = "150.0"))
+	float TransformGizmoSize = 0.0f;
+	//~ End Transform Gizmo Category
+
+	//~ Begin Legacy Gizmo Category
+
+	/** If true, legacy gizmos will be used. */
+	UPROPERTY(Config, EditAnywhere, Category = "Legacy Gizmo", DisplayName = "Enable Legacy Gizmo")
+	bool bEnableLegacyGizmo = true;
+
+	/**
+	 * Allow arcball rotation with rotate widget
+	 * (updates the setting with the same name found in Level Editor Viewport Settings)
+	 */
+	UPROPERTY(EditAnywhere, Category = "Legacy Gizmo", DisplayName = "Enable Arcball Rotate", meta = (EditCondition = "bEnableLegacyGizmo"))
+	bool bEnableArcballRotate;
+
+	/**
+	 * Allow screen rotation with rotate widget
+	 * (updates the setting with the same name found in Level Editor Viewport Settings)
+	 */
+	UPROPERTY(EditAnywhere, Category = "Legacy Gizmo", DisplayName = "Enable Screen Rotate", meta = (EditCondition = "bEnableLegacyGizmo"))
+	bool bEnableScreenRotate;
+
+	/**
+	 * If true, the Edit widget of a transform will display the axis
+	 * (updates the setting with the same name found in Level Editor Viewport Settings)
+	 */
+	UPROPERTY(EditAnywhere, Category = "Legacy Gizmo", DisplayName = "Enable Axis drawing for Transform Edit Widget", meta = (EditCondition = "bEnableLegacyGizmo"))
+	bool bEnableAxisDrawing;
+
+	/**
+	 * Allow translate/rotate widget
+	 * (updates the setting with the same name found in Level Editor Viewport Settings)
+	 */
+	UPROPERTY(EditAnywhere, Category = "Legacy Gizmo", DisplayName = "Enable Combined Translate/Rotate Widget", meta = (EditCondition = "bEnableLegacyGizmo"))
+	bool bEnableCombinedTranslateRotate;
+	//~ End Legacy Gizmo Category
+
+public:
+	void SetUsesLegacyGizmo(bool bInUseLegacyGizmo);
+	bool UsesLegacyGizmo() const;
+	bool UsesNewTRSGizmo() const;
+
+	void SetGizmosParameters(const FGizmosParameters& InGizmosParameters);
+	void SetTransformGizmoSize(float InTransformGizmoSize);
+
+private:
+	void BroadcastNewTRSGizmoChange() const;
+	void BroadcastGizmosParametersChange() const;
+
+	void OnLegacySettingChanged(FName InPropertyName);
+};
diff --git a/Engine/Source/Editor/LevelEditor/Private/LevelEditorActions.cpp b/Engine/Source/Editor/LevelEditor/Private/LevelEditorActions.cpp
--- a/Engine/Source/Editor/LevelEditor/Private/LevelEditorActions.cpp
+++ b/Engine/Source/Editor/LevelEditor/Private/LevelEditorActions.cpp
@@ -3831,6 +3831,7 @@ void FLevelEditorActionCallbacks::ToggleAllowArcballRotation()
 	if (ULevelEditorViewportSettings* ViewportSettings = GetMutableDefault<ULevelEditorViewportSettings>())
 	{
 		ViewportSettings->bAllowArcballRotate = !ViewportSettings->bAllowArcballRotate;
+		ViewportSettings->OnSettingChanged().Broadcast(GET_MEMBER_NAME_CHECKED(ULevelEditorViewportSettings, bAllowArcballRotate));
 	}
 }
 
@@ -3849,6 +3850,7 @@ void FLevelEditorActionCallbacks::ToggleAllowScreenspaceRotation()
 	if (ULevelEditorViewportSettings* ViewportSettings = GetMutableDefault<ULevelEditorViewportSettings>())
 	{
 		ViewportSettings->bAllowScreenRotate = !ViewportSettings->bAllowScreenRotate;
+		ViewportSettings->OnSettingChanged().Broadcast(GET_MEMBER_NAME_CHECKED(ULevelEditorViewportSettings, bAllowScreenRotate));
 	}
 }
 
diff --git a/Engine/Source/Editor/UnrealEd/Classes/Settings/LevelEditorViewportSettings.h b/Engine/Source/Editor/UnrealEd/Classes/Settings/LevelEditorViewportSettings.h
--- a/Engine/Source/Editor/UnrealEd/Classes/Settings/LevelEditorViewportSettings.h
+++ b/Engine/Source/Editor/UnrealEd/Classes/Settings/LevelEditorViewportSettings.h
@@ -571,8 +571,9 @@ public:
 	float BillboardScale;
 
 	/** The size adjustment to apply to the translate/rotate/scale widgets (in Unreal units). */
-	UPROPERTY(EditAnywhere, config, Category=LookAndFeel, AdvancedDisplay, meta=(ClampMin="-10",ClampMax="150") )
-	int32 TransformWidgetSizeAdjustment;
+	UE_DEPRECATED(5.6, "TransformWidgetSizeAdjustment. See UTransformGizmoEditorSettings::TransformGizmoSize")
+	UPROPERTY(config)
+	int32 TransformWidgetSizeAdjustment_DEPRECATED;
 
 	/** When enabled, engine stats that are enabled in level viewports are preserved between editor sessions */
 	UPROPERTY(EditAnywhere, config, AdvancedDisplay, Category = LookAndFeel)
diff --git a/Engine/Source/Editor/UnrealEd/Private/EditorViewportClient.cpp b/Engine/Source/Editor/UnrealEd/Private/EditorViewportClient.cpp
--- a/Engine/Source/Editor/UnrealEd/Private/EditorViewportClient.cpp
+++ b/Engine/Source/Editor/UnrealEd/Private/EditorViewportClient.cpp
@@ -81,6 +81,7 @@
 #include "GameFramework/ActorPrimitiveColorHandler.h"
 #include "LandscapeRender.h"
 #include "LandscapeSettings.h"
+#include "TransformGizmoEditorSettings.h"
 #include "Misc/AxisDisplayInfo.h"
 
 #define LOCTEXT_NAMESPACE "EditorViewportClient"
@@ -5476,9 +5477,10 @@ void FEditorViewportClient::SetCurrentWidgetAxis(EAxisList::Type InAxis)
 
 void FEditorViewportClient::AdjustTransformWidgetSize(const int32 SizeDelta)
 {
-	 ULevelEditorViewportSettings &ViewportSettings = *GetMutableDefault<ULevelEditorViewportSettings>();
-	 ViewportSettings.TransformWidgetSizeAdjustment = FMath::Clamp(ViewportSettings.TransformWidgetSizeAdjustment + SizeDelta, -10, 150);
-	 ViewportSettings.PostEditChange();
+	 if (UTransformGizmoEditorSettings* TransformEditorSettings = GetMutableDefault<UTransformGizmoEditorSettings>())
+	 {
+		 TransformEditorSettings->SetTransformGizmoSize(FMath::Clamp(TransformEditorSettings->TransformGizmoSize + SizeDelta, -10.0f, 150.0f));
+	 }
 }
 
 float FEditorViewportClient::GetNearClipPlane() const
diff --git a/Engine/Source/Editor/UnrealEd/Private/Settings/SettingsClasses.cpp b/Engine/Source/Editor/UnrealEd/Private/Settings/SettingsClasses.cpp
--- a/Engine/Source/Editor/UnrealEd/Private/Settings/SettingsClasses.cpp
+++ b/Engine/Source/Editor/UnrealEd/Private/Settings/SettingsClasses.cpp
@@ -42,6 +42,7 @@
 #include "DrawDebugHelpers.h"
 #include "LevelEditorActions.h"
 #include "ToolMenus.h"
+#include "TransformGizmoEditorSettings.h"
 #include "Subsystems/AssetEditorSubsystem.h"
 
 #define LOCTEXT_NAMESPACE "SettingsClasses"
@@ -954,7 +955,6 @@ ULevelEditorViewportSettings::ULevelEditorViewportSettings( const FObjectInitial
 	MinimumOrthographicZoom = 250.0f;
 	bLevelStreamingVolumePrevis = false;
 	BillboardScale = 1.0f;
-	TransformWidgetSizeAdjustment = 0.0f;
 	MeasuringToolUnits = MeasureUnits_Centimeters;
 	bAllowArcballRotate = false;
 	bAllowScreenRotate = false;
diff --git a/Engine/Source/Editor/UnrealEd/Private/UnrealWidget.cpp b/Engine/Source/Editor/UnrealEd/Private/UnrealWidget.cpp
--- a/Engine/Source/Editor/UnrealEd/Private/UnrealWidget.cpp
+++ b/Engine/Source/Editor/UnrealEd/Private/UnrealWidget.cpp
@@ -15,6 +15,7 @@
 #include "SceneView.h"
 #include "Settings/LevelEditorViewportSettings.h"
 #include "SnappingUtils.h"
+#include "TransformGizmoEditorSettings.h"
 
 IMPLEMENT_HIT_PROXY(HWidgetAxis, HHitProxy);
 
@@ -252,7 +253,7 @@ void FWidget::ConvertMouseToAxis_Rotate(FVector2D TangentDir, FVector2D DragDir,
 			if (Distance < CircleRadius)
 			{
 				const float ScaleInScreen = InView->WorldToScreen(InViewportClient->GetWidgetLocation()).W * (4.0f / InView->UnscaledViewRect.Width() / InView->ViewMatrices.GetProjectionMatrix().M[0][0]);
-				const float SphereRadius = CircleRadius * ScaleInScreen  + GetDefault<ULevelEditorViewportSettings>()->TransformWidgetSizeAdjustment;
+				const float SphereRadius = CircleRadius * ScaleInScreen  + GetDefault<UTransformGizmoEditorSettings>()->TransformGizmoSize;
 				FVector OldLocation, NewLocation;
 				FMath::SphereDistToLine(InViewportClient->GetWidgetLocation(), SphereRadius, MouseViewportRay.GetOrigin(), OldCameraToPixelDir, OldLocation);
 				FMath::SphereDistToLine(InViewportClient->GetWidgetLocation(), SphereRadius, MouseViewportRay.GetOrigin(), CameraToPixelDir, NewLocation);
@@ -268,7 +269,7 @@ void FWidget::ConvertMouseToAxis_Rotate(FVector2D TangentDir, FVector2D DragDir,
 				const float Scale = ScreenLocation.W *
 					(4.0f / InView->UnscaledViewRect.Width() / InView->ViewMatrices.GetProjectionMatrix().M[0][0]);
 				const float Radius = (CircleRadius * Scale) +
-					GetDefault<ULevelEditorViewportSettings>()->TransformWidgetSizeAdjustment;
+					GetDefault<UTransformGizmoEditorSettings>()->TransformGizmoSize;
 				const float LengthOfAdjacent = Length - Radius;
 				RotationAngle = FMath::Acos(FVector::DotProduct(OldCameraToPixelDir, CameraToPixelDir));
 				const float OppositeSize = FMath::Tan(RotationAngle) * LengthOfAdjacent;
diff --git a/Engine/Source/Editor/UnrealEd/Private/UnrealWidgetRender.cpp b/Engine/Source/Editor/UnrealEd/Private/UnrealWidgetRender.cpp
--- a/Engine/Source/Editor/UnrealEd/Private/UnrealWidgetRender.cpp
+++ b/Engine/Source/Editor/UnrealEd/Private/UnrealWidgetRender.cpp
@@ -14,6 +14,7 @@
 #include "DynamicMeshBuilder.h"
 #include "CanvasTypes.h"
 #include "EditorInteractiveGizmoManager.h"
+#include "TransformGizmoEditorSettings.h"
 #include "Misc/AxisDisplayInfo.h"
 
 /*
@@ -288,7 +289,7 @@ void FWidget::Render_Axis(const FSceneView* View, FPrimitiveDrawInterface* PDI,
 		const bool bDisabled = EditorModeTools ? (EditorModeTools->IsDefaultModeActive() && GEditor->HasLockedActors() ) : false;
 		PDI->SetHitProxy( new HWidgetAxis( InAxis, bDisabled) );
 
-		const float AxisLength = AXIS_LENGTH + GetDefault<ULevelEditorViewportSettings>()->TransformWidgetSizeAdjustment - (AxisLengthOffset * 2);
+		const float AxisLength = AXIS_LENGTH + GetDefault<UTransformGizmoEditorSettings>()->TransformGizmoSize - (AxisLengthOffset * 2);
 		const float HalfHeight = AxisLength/2.0f;
 		const float CylinderRadius = 1.2f;
 		const FVector Offset(0, 0, HalfHeight + AxisLengthOffset);
@@ -702,13 +703,13 @@ void FWidget::Render_Rotate( const FSceneView* View,FPrimitiveDrawInterface* PDI
 			
 		}
 
-		if (DrawAxis&EAxisList::XYZ  && (!bDragging || CurrentAxis == EAxisList::XYZ) && GetDefault<ULevelEditorViewportSettings>()->bAllowArcballRotate)
+		if (DrawAxis&EAxisList::XYZ  && (!bDragging || CurrentAxis == EAxisList::XYZ) && GetDefault<UTransformGizmoEditorSettings>()->bEnableArcballRotate)
 		{
 			FVector Center = InLocation;
 			FRotator Orientation = FRotator::ZeroRotator;
 			const float ArcballSize = 1.0f; //always constant, otherwise can cause hit testing issues
 
-			const float InnerDistance = (INNER_AXIS_CIRCLE_RADIUS * ArcballSize * Scale) + GetDefault<ULevelEditorViewportSettings>()->TransformWidgetSizeAdjustment;
+			const float InnerDistance = (INNER_AXIS_CIRCLE_RADIUS * ArcballSize * Scale) + GetDefault<UTransformGizmoEditorSettings>()->TransformGizmoSize;
 			FVector Radii(InnerDistance, InnerDistance, InnerDistance);
 			const bool bDisabled = IsWidgetDisabled();
 			FColor ArcColor = ArcBallColor;
@@ -730,8 +731,8 @@ void FWidget::Render_Rotate( const FSceneView* View,FPrimitiveDrawInterface* PDI
 			const FVector Axis0 = View->GetViewUp();
 			const FVector Axis1 = View->GetViewRight();
 			
-			const float OuterRadius = (OUTER_AXIS_CIRCLE_RADIUS * 1.25f * Scale) + GetDefault<ULevelEditorViewportSettings>()->TransformWidgetSizeAdjustment;
-			const float InnerRadius = ((OUTER_AXIS_CIRCLE_RADIUS -1) * 1.25f * Scale) + GetDefault<ULevelEditorViewportSettings>()->TransformWidgetSizeAdjustment;
+			const float OuterRadius = (OUTER_AXIS_CIRCLE_RADIUS * 1.25f * Scale) + GetDefault<UTransformGizmoEditorSettings>()->TransformGizmoSize;
+			const float InnerRadius = ((OUTER_AXIS_CIRCLE_RADIUS -1) * 1.25f * Scale) + GetDefault<UTransformGizmoEditorSettings>()->TransformGizmoSize;
 
 			FColor ArcColor = (CurrentAxis==EAxisList::Screen) ? CurrentColor : ScreenAxisColor.ToFColor(true);
 
@@ -892,7 +893,7 @@ void FWidget::Render_TranslateRotateZ( const FSceneView* View, FPrimitiveDrawInt
 
 		const bool bDisabled = IsWidgetDisabled();
 
-		const float ScaledRadius = (TRANSLATE_ROTATE_AXIS_CIRCLE_RADIUS * Space.UniformScale) + GetDefault<ULevelEditorViewportSettings>()->TransformWidgetSizeAdjustment;
+		const float ScaledRadius = (TRANSLATE_ROTATE_AXIS_CIRCLE_RADIUS * Space.UniformScale) + GetDefault<UTransformGizmoEditorSettings>()->TransformGizmoSize;
 
 		//ZRotation
 		if( DrawAxis&EAxisList::ZRotation && (Space.bIsPerspective || Space.bIsLocalSpace || View->ViewMatrices.GetViewMatrix().M[0][2] != -1.f) )
@@ -927,8 +928,8 @@ void FWidget::Render_TranslateRotateZ( const FSceneView* View, FPrimitiveDrawInt
 			{
 				// Add more sides to the circle if we've been scaled up to keep the circle looking circular
 				// An extra side for every 5 extra unreal units seems to produce a nice result
-				const int32 CircleSides = (GetDefault<ULevelEditorViewportSettings>()->TransformWidgetSizeAdjustment > 0) 
-					? AXIS_CIRCLE_SIDES + (GetDefault<ULevelEditorViewportSettings>()->TransformWidgetSizeAdjustment / 5)
+				const int32 CircleSides = (GetDefault<UTransformGizmoEditorSettings>()->TransformGizmoSize > 0) 
+					? AXIS_CIRCLE_SIDES + (GetDefault<UTransformGizmoEditorSettings>()->TransformGizmoSize / 5)
 					: AXIS_CIRCLE_SIDES;
 
 				PDI->SetHitProxy( new HWidgetAxis(EAxisList::XY, bDisabled) );
@@ -1010,9 +1011,9 @@ void FWidget::Render_2D(const FSceneView* View, FPrimitiveDrawInterface* PDI, FE
 	const bool bDisabled = IsWidgetDisabled();
 
 	// Radius
-	const float ScaledRadius = (TWOD_AXIS_CIRCLE_RADIUS * Space.UniformScale) + GetDefault<ULevelEditorViewportSettings>()->TransformWidgetSizeAdjustment;
-	const int32 CircleSides = (GetDefault<ULevelEditorViewportSettings>()->TransformWidgetSizeAdjustment > 0)
-		? AXIS_CIRCLE_SIDES + (GetDefault<ULevelEditorViewportSettings>()->TransformWidgetSizeAdjustment / 5)
+	const float ScaledRadius = (TWOD_AXIS_CIRCLE_RADIUS * Space.UniformScale) + GetDefault<UTransformGizmoEditorSettings>()->TransformGizmoSize;
+	const int32 CircleSides = (GetDefault<UTransformGizmoEditorSettings>()->TransformGizmoSize > 0)
+		? AXIS_CIRCLE_SIDES + (GetDefault<UTransformGizmoEditorSettings>()->TransformGizmoSize / 5)
 		: AXIS_CIRCLE_SIDES;
 
 	// Draw the grabbers
@@ -1089,7 +1090,7 @@ void FWidget::Render_2D(const FSceneView* View, FPrimitiveDrawInterface* PDI, FE
  */
 void FWidget::CacheRotationHUDText(const FSceneView* View, FPrimitiveDrawInterface* PDI, const FVector& InLocation, const FVector& Axis0, const FVector& Axis1, const float AngleOfChange, const float InScale)
 {
-	const float TextDistance = (ROTATION_TEXT_RADIUS * InScale) + GetDefault<ULevelEditorViewportSettings>()->TransformWidgetSizeAdjustment;
+	const float TextDistance = (ROTATION_TEXT_RADIUS * InScale) + GetDefault<UTransformGizmoEditorSettings>()->TransformGizmoSize;
 
 	FVector AxisVectors[4] = { Axis0, Axis1, -Axis0, -Axis1};
 
@@ -1256,8 +1257,8 @@ void FWidget::DrawRotationArc(const FSceneView* View, FPrimitiveDrawInterface* P
  */
 void FWidget::DrawPartialRotationArc(const FSceneView* View, FPrimitiveDrawInterface* PDI, EAxisList::Type InAxis, const FVector& InLocation, const FVector& Axis0, const FVector& Axis1, const float InStartAngle, const float InEndAngle, const FColor& InColor, const float InScale, const FVector& InDirectionToWidget )
 {
-	const float InnerRadius = (INNER_AXIS_CIRCLE_RADIUS * InScale) + GetDefault<ULevelEditorViewportSettings>()->TransformWidgetSizeAdjustment;
-	const float OuterRadius = (OUTER_AXIS_CIRCLE_RADIUS * InScale) + GetDefault<ULevelEditorViewportSettings>()->TransformWidgetSizeAdjustment;
+	const float InnerRadius = (INNER_AXIS_CIRCLE_RADIUS * InScale) + GetDefault<UTransformGizmoEditorSettings>()->TransformGizmoSize;
+	const float OuterRadius = (OUTER_AXIS_CIRCLE_RADIUS * InScale) + GetDefault<UTransformGizmoEditorSettings>()->TransformGizmoSize;
 
 	bool bIsPerspective = ( View->ViewMatrices.GetProjectionMatrix().M[3][3] < 1.0f );
 	PDI->SetHitProxy( new HWidgetAxis( InAxis ) );
@@ -1297,8 +1298,8 @@ void FWidget::DrawThickArc (const FThickArcParams& InParams, const FVector& Axis
 
 	// Add more sides to the circle if we've been scaled up to keep the circle looking circular
 	// An extra side for every 5 extra unreal units seems to produce a nice result
-	const int32 CircleSides = (GetDefault<ULevelEditorViewportSettings>()->TransformWidgetSizeAdjustment > 0) 
-		? AXIS_CIRCLE_SIDES + (GetDefault<ULevelEditorViewportSettings>()->TransformWidgetSizeAdjustment / 5)
+	const int32 CircleSides = (GetDefault<UTransformGizmoEditorSettings>()->TransformGizmoSize > 0) 
+		? AXIS_CIRCLE_SIDES + (GetDefault<UTransformGizmoEditorSettings>()->TransformGizmoSize / 5)
 		: AXIS_CIRCLE_SIDES;
 	const int32 NumPoints = FMath::TruncToInt(CircleSides * (InEndAngle-InStartAngle)/(PI/2)) + 1;
 
@@ -1384,8 +1385,8 @@ void FWidget::DrawThickArc (const FThickArcParams& InParams, const FVector& Axis
  */
 void FWidget::DrawSnapMarker(FPrimitiveDrawInterface* PDI, const FVector& InLocation, const FVector& Axis0, const FVector& Axis1, const FColor& InColor, const float InScale, const float InWidthPercent, const float InPercentSize)
 {
-	const float InnerDistance = (INNER_AXIS_CIRCLE_RADIUS * InScale) + GetDefault<ULevelEditorViewportSettings>()->TransformWidgetSizeAdjustment;
-	const float OuterDistance = (OUTER_AXIS_CIRCLE_RADIUS * InScale) + GetDefault<ULevelEditorViewportSettings>()->TransformWidgetSizeAdjustment;
+	const float InnerDistance = (INNER_AXIS_CIRCLE_RADIUS * InScale) + GetDefault<UTransformGizmoEditorSettings>()->TransformGizmoSize;
+	const float OuterDistance = (OUTER_AXIS_CIRCLE_RADIUS * InScale) + GetDefault<UTransformGizmoEditorSettings>()->TransformGizmoSize;
 	const float MaxMarkerHeight = OuterDistance - InnerDistance;
 	const float MarkerWidth = MaxMarkerHeight*InWidthPercent;
 	const float MarkerHeight = MaxMarkerHeight*InPercentSize;
@@ -1443,8 +1444,8 @@ void FWidget::DrawSnapMarker(FPrimitiveDrawInterface* PDI, const FVector& InLoca
 void FWidget::DrawStartStopMarker(FPrimitiveDrawInterface* PDI, const FVector& InLocation, const FVector& Axis0, const FVector& Axis1, const float InAngle, const FColor& InColor, const float InScale)
 {
 	const float ArrowHeightPercent = .8f;
-	const float InnerDistance = (INNER_AXIS_CIRCLE_RADIUS * InScale) + GetDefault<ULevelEditorViewportSettings>()->TransformWidgetSizeAdjustment;
-	const float OuterDistance = (OUTER_AXIS_CIRCLE_RADIUS * InScale) + GetDefault<ULevelEditorViewportSettings>()->TransformWidgetSizeAdjustment;
+	const float InnerDistance = (INNER_AXIS_CIRCLE_RADIUS * InScale) + GetDefault<UTransformGizmoEditorSettings>()->TransformGizmoSize;
+	const float OuterDistance = (OUTER_AXIS_CIRCLE_RADIUS * InScale) + GetDefault<UTransformGizmoEditorSettings>()->TransformGizmoSize;
 	const float RingHeight = OuterDistance - InnerDistance;
 	const float ArrowHeight = RingHeight*ArrowHeightPercent;
 	const float ThirtyDegrees = PI / 6.0f;
diff --git a/Engine/Source/Editor/UnrealEd/Private/ViewportToolbar/UnrealEdViewportToolbar.cpp b/Engine/Source/Editor/UnrealEd/Private/ViewportToolbar/UnrealEdViewportToolbar.cpp
--- a/Engine/Source/Editor/UnrealEd/Private/ViewportToolbar/UnrealEdViewportToolbar.cpp
+++ b/Engine/Source/Editor/UnrealEd/Private/ViewportToolbar/UnrealEdViewportToolbar.cpp
@@ -29,6 +29,7 @@
 #include "ToolMenuSection.h"
 #include "ToolMenus.h"
 #include "ToolMenusEditor.h"
+#include "TransformGizmoEditorSettings.h"
 #include "Tests/ToolMenusTestUtilities.h"
 #include "ViewportToolbar/UnrealEdViewportToolbarContext.h"
 #include "Widgets/Input/SNumericEntryBox.h"
@@ -969,29 +970,28 @@ FToolMenuEntry CreateTransformsSubmenu()
 						.FillWidth(0.9f)
 						[
 							SNew(SSpinBox<int32>)
-								.MinValue(-10)
-								.MaxValue(150)
+								.MinValue(-10.0f)
+								.MaxValue(150.0f)
 								.ToolTipText_Lambda(
 									[]() -> FText
 									{
 										return FText::AsNumber(
-											GetDefault<ULevelEditorViewportSettings>()->TransformWidgetSizeAdjustment
+											GetDefault<UTransformGizmoEditorSettings>()->TransformGizmoSize
 										);
 									}
 								)
 								.Value_Lambda(
 									[]() -> float
 									{
-										return GetDefault<ULevelEditorViewportSettings>()->TransformWidgetSizeAdjustment;
+										return GetDefault<UTransformGizmoEditorSettings>()->TransformGizmoSize;
 									}
 								)
 								.OnValueChanged_Lambda(
 									[](float InValue)
 									{
-										ULevelEditorViewportSettings* ViewportSettings =
-											GetMutableDefault<ULevelEditorViewportSettings>();
-										ViewportSettings->TransformWidgetSizeAdjustment = InValue;
-										ViewportSettings->PostEditChange();
+										UTransformGizmoEditorSettings* ViewportSettings =
+											GetMutableDefault<UTransformGizmoEditorSettings>();
+										ViewportSettings->SetTransformGizmoSize(InValue);
 									}
 								)
 						]
