Cooker: Restore handling of packages with outers; we still need to support old content that contains nested packages.
#rnx
#incrementalcook

diff --git a/Engine/Source/Editor/UnrealEd/Private/Cooker/PackageTracker.cpp b/Engine/Source/Editor/UnrealEd/Private/Cooker/PackageTracker.cpp
--- a/Engine/Source/Editor/UnrealEd/Private/Cooker/PackageTracker.cpp
+++ b/Engine/Source/Editor/UnrealEd/Private/Cooker/PackageTracker.cpp
@@ -158,9 +158,11 @@ void FPackageTracker::NotifyUObjectCreated(const class UObjectBase* Object, int3
 	if (Object->GetClass() == UPackage::StaticClass())
 	{
 		UPackage* Package = const_cast<UPackage*>(static_cast<const UPackage*>(Object));
-		check(Package->GetOuter() == nullptr); // Subpackages are no longer possible in Unreal5 or later.
-		OnCreatePackage(Package->GetFName());
-		LoadedPackages.Add(Package);
+		if (Package->GetOuter() == nullptr) // Nested packages are no longer created, but can still exist in old data.
+		{
+			OnCreatePackage(Package->GetFName());
+			LoadedPackages.Add(Package);
+		}
 	}
 }
 
