[FBX Importer Legacy] Resetting Existing SkeletalMesh in place to be re-used, avoiding re-import crash.

#rnx

diff --git a/Engine/Source/Editor/UnrealEd/Private/Fbx/FbxSkeletalMeshImport.cpp b/Engine/Source/Editor/UnrealEd/Private/Fbx/FbxSkeletalMeshImport.cpp
--- a/Engine/Source/Editor/UnrealEd/Private/Fbx/FbxSkeletalMeshImport.cpp
+++ b/Engine/Source/Editor/UnrealEd/Private/Fbx/FbxSkeletalMeshImport.cpp
@@ -1720,6 +1720,13 @@ USkeletalMesh* UnFbx::FFbxImporter::ImportSkeletalMesh(FImportSkeletalMeshArgs &
 		// When we are not re-importing we want to create the mesh here to be sure there is no material
 		// or texture that will be create with the same name
 		SkeletalMesh = NewObject<USkeletalMesh>(ImportSkeletalMeshArgs.InParent, ImportSkeletalMeshArgs.Name, ImportSkeletalMeshArgs.Flags);
+
+		//This should not happen
+		if (!ensure(SkeletalMesh))
+		{
+			return nullptr;
+		}
+
 		CreatedObjects.Add(SkeletalMesh);
 	}
 	else
@@ -1857,11 +1864,35 @@ USkeletalMesh* UnFbx::FFbxImporter::ImportSkeletalMesh(FImportSkeletalMeshArgs &
 
 	if (SkeletalMesh == nullptr)
 	{
-		//Backup the PostEditchangeStackCounter before overriding the skeletalmesh
-		int32 PostEditChangeStackCounter = ExistingSkelMesh != nullptr ? ExistingSkelMesh->GetPostEditChangeStackCounter() : 0;
-		// Create the new mesh after saving the old data, since it will replace the old skeletalmesh
-		// This should happen only when doing a re-import. Otherwise the skeletal mesh is create before.
-		SkeletalMesh = NewObject<USkeletalMesh>(ImportSkeletalMeshArgs.InParent, ImportSkeletalMeshArgs.Name, ImportSkeletalMeshArgs.Flags);
+		if (!ensure(ExistingSkelMesh != nullptr))
+		{
+			return nullptr;
+		}
+
+		SkeletalMesh = ExistingSkelMesh;
+
+		int32 PostEditChangeStackCounter = SkeletalMesh->GetPostEditChangeStackCounter();
+		SkeletalMesh->SetPostEditChangeStackCounter(0);
+
+		for (int32 LODIndex = 0, LODCount = SkeletalMesh->GetLODNum(); LODIndex < LODCount; ++LODIndex)
+		{
+			SkeletalMesh->ClearMeshDescriptionAndBulkData(LODIndex);
+		}
+
+		FSkeletalMeshModel* ImportedResource = SkeletalMesh->GetImportedModel();
+		ImportedResource->EmptyOriginalReductionSourceMeshData();
+		ImportedResource->LODModels.Empty();
+
+		SkeletalMesh->ResetLODInfo();
+		SkeletalMesh->GetMaterials().Empty();
+		SkeletalMesh->GetRefSkeleton().Empty();
+		SkeletalMesh->SetSkeleton(nullptr);
+		SkeletalMesh->SetPhysicsAsset(nullptr);
+		SkeletalMesh->UnregisterAllMorphTarget();
+
+		SkeletalMesh->ReleaseResources();
+
+		SkeletalMesh->PostEditChange();
 		SkeletalMesh->SetPostEditChangeStackCounter(PostEditChangeStackCounter);
 	}
 
