Add copyright boilerplate

diff --git a/Engine/Source/Editor/UnrealEd/Private/SEditorViewportGridPanel.cpp b/Engine/Source/Editor/UnrealEd/Private/SEditorViewportGridPanel.cpp
--- a/Engine/Source/Editor/UnrealEd/Private/SEditorViewportGridPanel.cpp
+++ b/Engine/Source/Editor/UnrealEd/Private/SEditorViewportGridPanel.cpp
@@ -1,3 +1,5 @@
+// Copyright Epic Games, Inc. All Rights Reserved.
+
 #include "SEditorViewportGridPanel.h"
 #include "Framework/Application/SlateApplication.h"
 
diff --git a/Engine/Source/Editor/UnrealEd/Private/SEditorViewportGridPanel.h b/Engine/Source/Editor/UnrealEd/Private/SEditorViewportGridPanel.h
--- a/Engine/Source/Editor/UnrealEd/Private/SEditorViewportGridPanel.h
+++ b/Engine/Source/Editor/UnrealEd/Private/SEditorViewportGridPanel.h
@@ -1,3 +1,5 @@
+// Copyright Epic Games, Inc. All Rights Reserved.
+
 #pragma once
 
 #include "Widgets/Layout/SGridPanel.h"
