IncrementalCook: Fix some discovered runtime dependencies not being stored in a package, which was causing some discovered runtime dependencies to not be cooked in incremental cook.
Details:
1. Change ProcessUnsolicitedPackages to decide whether packages are expected based immediately on AssetRegistry dependencies rather than needing to wait for PumpRequests.
2. Call ProcessUnsolicitedPackages after SavePackage and before storing the CookDependencies in the oplog so that we can store the runtime dependencies discovered during save in the CookDependencies.
3. UE::Cook::FPackageData now has a list of DiscoveredDependencies FPackageData* rather than having an FName map on CookOnTheFlyServer, to make the discovered dependencies faster to access, and the discovered dependencies are recorded per-platform.
4. Packages loaded by the cooker get a UE_TRACK_REFERENCING_PACKAGE_SCOPED so that we know that they are not unsolicited.
#rnx
#incrementalcook

diff --git a/Engine/Source/Editor/UnrealEd/Classes/CookOnTheSide/CookOnTheFlyServer.h b/Engine/Source/Editor/UnrealEd/Classes/CookOnTheSide/CookOnTheFlyServer.h
--- a/Engine/Source/Editor/UnrealEd/Classes/CookOnTheSide/CookOnTheFlyServer.h
+++ b/Engine/Source/Editor/UnrealEd/Classes/CookOnTheSide/CookOnTheFlyServer.h
@@ -174,7 +174,6 @@ namespace UE::Cook
 {
 	class FAssetRegistryMPCollector;
 	class FBuildDefinitions;
-	class FCachedDependencies;
 	class FCookDirector;
 	class FCookGCDiagnosticContext;
 	class FCookSandbox;
@@ -1524,11 +1523,6 @@ private:
 	 * Read/Write only within HiddenDependenciesLock.
 	 */
 	TSet<FName> HiddenDependenciesClassPathFilterList;
-	/**
-	 * A cache of TSet<FName> of dependencies for each Package, guarded by HiddenDependenciesLock.
-	 * Read/Write only within HiddenDependenciesLock.
-	 */
-	TUniquePtr<UE::Cook::FCachedDependencies> CachedDependencies;
 
 	/** Registration handle for TObjectPtr's OnRead delegate. */
 #if UE_WITH_OBJECT_HANDLE_TRACKING
@@ -1537,8 +1531,6 @@ private:
 	/** Data tracking the package currently having calls made (Load/Save/Other) from COTFS. Used for diagnostics. */
 	FActivePackageData ActivePackageData;
 
-	/** Package -> Package dependencies that were hidden and discovered when a package loaded or ran system-specific code. */
-	TMap<FName, TArray<FName>> DiscoveredDependencies;
 	/** Classes (and all subclasses) that were listed as the only classes that should be cooked in the filter settings. */
 	TSet<FName> CookFilterIncludedClasses;
 	/** Classes (and all subclasses) that were listed as the only asset classes that should be cooked in the filter settings. */
diff --git a/Engine/Source/Editor/UnrealEd/Private/CookOnTheFlyServer.cpp b/Engine/Source/Editor/UnrealEd/Private/CookOnTheFlyServer.cpp
--- a/Engine/Source/Editor/UnrealEd/Private/CookOnTheFlyServer.cpp
+++ b/Engine/Source/Editor/UnrealEd/Private/CookOnTheFlyServer.cpp
@@ -246,6 +246,7 @@ UCookOnTheFlyServer* UCookOnTheFlyServer::ActiveCOTFS = nullptr;
 UCookOnTheFlyServer::FOnCookByTheBookStarted UCookOnTheFlyServer::CookByTheBookStartedEvent;
 UCookOnTheFlyServer::FOnCookByTheBookFinished UCookOnTheFlyServer::CookByTheBookFinishedEvent;
 
+static FName ScriptPackageNameEngine(TEXT("/Script/Engine"));
 
 namespace UE::Cook
 {
@@ -2714,7 +2715,7 @@ void UCookOnTheFlyServer::PumpRequests(UE::Cook::FTickStackData& StackData, int3
 			if (GetCookPhase() == ECookPhase::Cook)
 			{
 				TUniquePtr<FRequestCluster> Cluster = MakeUnique<FRequestCluster>(*this, DiscoveryQueue);
-				if (Cluster->NumPackageDatas() > 0)
+				if (Cluster->NeedsProcessing())
 				{
 					RequestClusters.Add(MoveTemp(Cluster));
 				}
@@ -2740,7 +2741,7 @@ void UCookOnTheFlyServer::PumpRequests(UE::Cook::FTickStackData& StackData, int3
 			{
 				TUniquePtr<FRequestCluster> Cluster = MakeUnique<FRequestCluster>(*this,
 					FRequestCluster::BuildDependencyQueue, BuildDependencyDiscoveryQueue);
-				if (Cluster->NumPackageDatas() > 0)
+				if (Cluster->NeedsProcessing())
 				{
 					RequestClusters.Add(MoveTemp(Cluster));
 				}
@@ -3438,7 +3439,7 @@ void UCookOnTheFlyServer::QueueDiscoveredPackage(UE::Cook::FPackageData& Package
 	else
 	{
 		DiscoveredPlatforms = ReachablePlatforms.GetPlatforms(*this, &Instigator,
-			TConstArrayView<const ITargetPlatform*>(), DiscoveredReachability, &BufferPlatforms);
+			TConstArrayView<const ITargetPlatform*>(), DiscoveredReachability, BufferPlatforms);
 	}
 	if (Instigator.Category != EInstigator::ForceExplorableSaveTimeSoftDependency &&
 		PackageData.HasReachablePlatforms(DiscoveredReachability, DiscoveredPlatforms))
@@ -4559,6 +4560,10 @@ bool UCookOnTheFlyServer::LoadPackageForCooking(UE::Cook::FPackageData& PackageD
 #if ENABLE_COOK_STATS
 			++DetailedCookStats::NumRequestedLoads;
 #endif
+			// Declare the package as being referenced by the ScriptPackageNameEngine; we look for this in
+			// ProcessUnsolicitedPackages so we can avoid adding an error for the package being loaded without a
+			// referencer.
+			UE_TRACK_REFERENCING_PACKAGE_SCOPED(ScriptPackageNameEngine, PackageAccessTrackingOps::NAME_Load);
 			LoadedPackage = LoadPackage(nullptr, *FileName, LOAD_None);
 		}
 		if (IsValid(LoadedPackage) && LoadedPackage->IsFullyLoaded())
@@ -4682,6 +4687,11 @@ void UCookOnTheFlyServer::ProcessUnsolicitedPackages(TArray<FName>* OutDiscovere
 		}
 		check(PackageStreamEvent.EventType == EPackageStreamEvent::PackageLoad);
 
+		if (Instigator.Referencer == ScriptPackageNameEngine)
+		{
+			// A load by the cooker, expected, no need to queue it as discovered.
+			continue;
+		}
 		FPackageData* PackageData = PackageDatas->TryAddPackageDataByPackageName(PackageName);
 		if (!PackageData)
 		{
@@ -4694,11 +4704,12 @@ void UCookOnTheFlyServer::ProcessUnsolicitedPackages(TArray<FName>* OutDiscovere
 			continue; // Generated packages are queued separately (with the correct instigator) in QueueGeneratedPackages
 		}
 
-		// Packages loaded by an instanced package referencer need to redirect the Instigator to be the Instigator of
-		// the instanced package load, merged with the Instigator category based on the whether the loaded package is
-		// a recorded dependency of the instanced package's LoadPath package.
+		FPackageData* Referencer = nullptr;
 		if (!Instigator.Referencer.IsNone())
 		{
+			// Packages loaded by an instanced package referencer need to redirect the Instigator to be the Instigator of
+			// the instanced package load, merged with the Instigator category based on the whether the loaded package is
+			// a recorded dependency of the instanced package's LoadPath package.
 			TRefCountPtr<const FPackageStreamInstancedPackage> Instance
 				= PackageTracker->FindInstancedPackage(Instigator.Referencer);
 			if (Instance)
@@ -4717,6 +4728,7 @@ void UCookOnTheFlyServer::ProcessUnsolicitedPackages(TArray<FName>* OutDiscovere
 				Instigator.Category = FPackageTracker::MergeReferenceCategories(
 					Instance->Instigator.Category, Instigator.Category);
 			}
+			Referencer = PackageDatas->FindPackageDataByPackageName(Instigator.Referencer);
 		}
 
 		if (Instigator.Category == EInstigator::EditorOnlyLoad)
@@ -4734,52 +4746,91 @@ void UCookOnTheFlyServer::ProcessUnsolicitedPackages(TArray<FName>* OutDiscovere
 		AddToOutDiscovered(PackageData, Instigator);
 		TStringBuilder<256> PackageNameStr(InPlace, PackageData->GetPackageName());
 
-		if (bSkipOnlyEditorOnly &&
-			(Instigator.Category == EInstigator::ForceExplorableSaveTimeSoftDependency ||
-				(PackageTracker->NeverCookPackageList.Contains(PackageData->GetPackageName()) &&
-					(INDEX_NONE != UE::String::FindFirst(PackageNameStr,
-						ULevel::GetExternalActorsFolderName(), ESearchCase::IgnoreCase) ||
-						INDEX_NONE != UE::String::FindFirst(PackageNameStr,
-							FPackagePath::GetExternalObjectsFolderName(), ESearchCase::IgnoreCase)))))
-		{
-			// ONLYEDITORONLY_TODO: WorldPartition should mark up these loads with EInstigator::ForceExplorableSaveTimeSoftDependency rather than
-			// needing to use a naming convention. We should also mark them once, during Save, rather than marking them
-			// and reexploring them every time they are loaded.
-			Instigator.Category = EInstigator::ForceExplorableSaveTimeSoftDependency;
+		if (!Referencer)
+		{
+			// Package loads after cook startup that were requested outside of a package operation are a bug unless
+			// marked from code as EInstigator::EditorOnlyLoad using an
+			// FCookLoadScope Scope(ECookLoadType::EditorOnly). Report it as a bug and add it to the cook.
+			// TODO: Incremental cook will not include these packages; add it?
+			if (Instigator.Category != EInstigator::StartupPackage &&
+				Instigator.Category != EInstigator::StartupPackageCookLoadScope &&
+				Instigator.Category != EInstigator::StartupSoftObjectPath)
+			{
+				FString Message = FString::Printf(
+					TEXT("Unsolicited package %s was loaded outside of any cook operation on a referencer package; we don't know why it was loaded. ")
+					TEXT("Adding it to the current cook, but it will possibly not be found in future incremental cooks."),
+					*PackageNameStr);
+				if (bHiddenDependenciesDebug)
+				{
+					UE_LOG(LogCook, Error, TEXT("%s"), *Message);
+				}
+				else
+				{
+					UE_LOG(LogCook, Display, TEXT("%s"), *Message);
+				}
+			}
 			QueueDiscoveredPackage(*PackageData, FInstigator(Instigator), EDiscoveredPlatformSet::CopyFromInstigator);
+			continue;
 		}
-		else if (Instigator.Category != EInstigator::Unsolicited)
+
+		if (bSkipOnlyEditorOnly
+			&& Instigator.Category != EInstigator::ForceExplorableSaveTimeSoftDependency
+			&& PackageTracker->NeverCookPackageList.Contains(PackageData->GetPackageName())
+			&& (INDEX_NONE != UE::String::FindFirst(PackageNameStr,
+				ULevel::GetExternalActorsFolderName(), ESearchCase::IgnoreCase)
+				|| INDEX_NONE != UE::String::FindFirst(PackageNameStr,
+					FPackagePath::GetExternalObjectsFolderName(), ESearchCase::IgnoreCase)))
 		{
-			// A load undeclared in AssetRegistry dependencies, but one that has been marked up as a known issue and needed at runtime.
-			// Queue it for discovery and do not log a warning about it.
-			QueueDiscoveredPackage(*PackageData, FInstigator(Instigator), EDiscoveredPlatformSet::CopyFromInstigator);
+			// ONLYEDITORONLY_TODO: WorldPartition should mark these loads as ForceExplorableSaveTimeSoftDependency
+			// rather than needing to use a naming convention. We should also mark them once, during Save, rather than
+			// marking them and reexploring them every time they are loaded. 
+			Instigator.Category = EInstigator::ForceExplorableSaveTimeSoftDependency;
 		}
-		else if (PackageData->FindOrAddPlatformData(CookerLoadingPlatformKey).IsReachable(EReachability::Runtime))
+
+		// If we have already reported this discovery from the same referencer, at the same or lower edge priority,
+		// then there is no need to report it again so early exit to skip some unnecessary work.
+		TMap<FPackageData*, EInstigator>* DiscoveredDependencies = Referencer->GetDiscoveredDependencies(
+			nullptr /* PlatformAgnosticTargetPlatform */);
+		if (DiscoveredDependencies)
 		{
-			// This load was expected so we do not need to add a hidden dependency for it.
-			// If we are using legacy WhatGetsCookedRules, queue the package for cooking because it was loaded.
-			if (!bSkipOnlyEditorOnly)
+			EInstigator* ExistingInstigator = DiscoveredDependencies->Find(PackageData);
+			if (ExistingInstigator)
 			{
-				QueueDiscoveredPackage(*PackageData, FInstigator(Instigator), EDiscoveredPlatformSet::CopyFromInstigator);
+				if (Instigator.Category != EInstigator::ForceExplorableSaveTimeSoftDependency 
+					|| *ExistingInstigator == EInstigator::ForceExplorableSaveTimeSoftDependency)
+				{
+					continue;
+				}
 			}
 		}
-		else
+
+		if (Instigator.Category == EInstigator::Unsolicited)
 		{
-			// Possibly an unexpected load, but it might be an expected load after we finish processing some other unexpected loads that were
-			// discovered earlier during the same LoadPackage call. Queue it for discovery and if it is still unreachable when its turn comes
-			// we will add the hidden dependency for it.
-			// as reachable in all of its instigator's reachable platforms
-			QueueDiscoveredPackage(*PackageData, FInstigator(Instigator), EDiscoveredPlatformSet::CopyFromInstigator);
+			// If it comes from an import that was declared to the assetregistry, then it is expected, and
+			// CookRequestCluster already handled or will handle adding it if necessary, and we can ignore it here.
+			if (AssetRegistry->ContainsDependency(Instigator.Referencer, PackageName,
+				UE::AssetRegistry::EDependencyCategory::Package, UE::AssetRegistry::EDependencyQuery::Hard))
+			{
+				continue;
+			}
 
-			if (IsDebugRecordUnsolicited() && !Instigator.Referencer.IsNone())
+			// Otherwise, it is a bug; unsolicited loads are not allowed. Report it, and add it as a
+			// discovered dependency.
+			FString Message = FString::Printf(
+				TEXT("Unsolicited package %s was loaded by package %s. This is a hidden dependency, causes poor cook performance, and it might be a bug to add it to the cook. Declare this package in the AssetRegistry dependencies of the loading package, or mark up its load with FCookLoadScope to specify whether it is runtime or editoronly."),
+				*PackageNameStr, *WriteToString<256>(Instigator.Referencer));
+			if (bHiddenDependenciesDebug)
 			{
-				FPackageData* InstigatorPackage = PackageDatas->TryAddPackageDataByPackageName(Instigator.Referencer);
-				if (InstigatorPackage && InstigatorPackage != PackageData)
-				{
-					InstigatorPackage->CreateOrGetUnsolicited().FindOrAdd(PackageData, Instigator.Category);
-				}
+				UE_LOG(LogCook, Error, TEXT("%s"), *Message);
+			}
+			else
+			{
+				UE_LOG(LogCook, Display, TEXT("%s"), *Message);
 			}
 		}
+
+		Referencer->AddDiscoveredDependency(EDiscoveredPlatformSet::CopyFromInstigator, PackageData, Instigator.Category);
+		QueueDiscoveredPackage(*PackageData, MoveTemp(Instigator), EDiscoveredPlatformSet::CopyFromInstigator);
 	}
 }
 
@@ -5769,7 +5820,6 @@ void UCookOnTheFlyServer::Initialize( ECookMode::Type DesiredCookMode, ECookInit
 	CookByTheBookOptions = MakeUnique<UE::Cook::FCookByTheBookOptions>();
 	CookOnTheFlyOptions = MakeUnique<UE::Cook::FCookOnTheFlyOptions>();
 	AssetRegistry = IAssetRegistry::Get();
-	CachedDependencies = MakeUnique<UE::Cook::FCachedDependencies>();
 	GCDiagnosticContext = MakeUnique<UE::Cook::FCookGCDiagnosticContext>();
 	StallDetector = MakeUnique<UE::Cook::FStallDetector>();
 
@@ -5835,10 +5885,13 @@ void UCookOnTheFlyServer::Initialize( ECookMode::Type DesiredCookMode, ECookInit
 	PollNextTimeIdleSeconds = MAX_flt;
 
 	CurrentAsyncCacheForType = MaxAsyncCacheForType;
-	TSet<FName> StartupPackages;
-	PackageTracker->InitializeTracking(StartupPackages);
-	if (IsCookByTheBookMode())
+	// CookCommandlet CookWorker and CookByTheBook do not initialize startup packages until BlockOnAssetRegistry,
+	// because systems that subscribe to the AssetRegistry's OnFilesLoaded can load further packages at that time.
+	// But for the CookOnTheFlyServer in the editor, or CookOnTheFly this is the only opportunity it has.
+	if (IsCookingInEditor() || (!IsCookByTheBookMode() && !IsCookWorkerMode()))
 	{
+		TSet<FName> StartupPackages;
+		PackageTracker->InitializeTracking(StartupPackages);
 		CookByTheBookOptions->StartupPackages = MoveTemp(StartupPackages);
 	}
 
@@ -5888,9 +5941,6 @@ void UCookOnTheFlyServer::Initialize( ECookMode::Type DesiredCookMode, ECookInit
 #endif
 }
 
-
-
-
 void UCookOnTheFlyServer::InitializeAtFirstSession()
 {
 	UE::EditorDomain::UtilsCookInitialize();
@@ -7946,6 +7996,15 @@ void UCookOnTheFlyServer::BlockOnAssetRegistry(TConstArrayView<FString> Commandl
 #if ENABLE_LOW_LEVEL_MEM_TRACKER
 	FLowLevelMemTracker::Get().UpdateStatsPerFrame();
 #endif
+
+	// CookCommandlet CookWorker and CookByTheBook do not initialize startup packages until BlockOnAssetRegistry,
+	// because systems that subscribe to the AssetRegistry's OnFilesLoaded can load further packages at that time.
+	if (!IsCookingInEditor() && (IsCookByTheBookMode() || IsCookWorkerMode()))
+	{
+		TSet<FName> StartupPackages;
+		PackageTracker->InitializeTracking(StartupPackages);
+		CookByTheBookOptions->StartupPackages = MoveTemp(StartupPackages);
+	}
 }
 
 void UCookOnTheFlyServer::RefreshPlatformAssetRegistries(const TArrayView<const ITargetPlatform* const>& TargetPlatforms)
@@ -11372,15 +11431,25 @@ FBeginCookContext UCookOnTheFlyServer::CreateCookWorkerContext()
 void UCookOnTheFlyServer::GenerateInitialRequests(FBeginCookContext& BeginContext)
 {
 	TArray<ITargetPlatform*>& TargetPlatforms = BeginContext.TargetPlatforms;
-	TSet<FName> StartupSoftObjectPackages;
+	TMap<FName, FName> StartupSoftObjectPackageReferencers;
 	if (!CookByTheBookOptions->bSkipSoftReferences)
 	{
+		TSet<FName> PackagesFromRedirectCollector;
 		// Get the list of soft references, for both empty package and all startup packages
-		GRedirectCollector.ProcessSoftObjectPathPackageList(NAME_None, false, StartupSoftObjectPackages);
+		GRedirectCollector.ProcessSoftObjectPathPackageList(NAME_None, false, PackagesFromRedirectCollector);
+		for (FName PackageName : PackagesFromRedirectCollector)
+		{
+			StartupSoftObjectPackageReferencers.Add(PackageName, NAME_None);
+		}
 
-		for (const FName& StartupPackage : CookByTheBookOptions->StartupPackages)
+		for (FName StartupPackage : CookByTheBookOptions->StartupPackages)
 		{
-			GRedirectCollector.ProcessSoftObjectPathPackageList(StartupPackage, false, StartupSoftObjectPackages);
+			PackagesFromRedirectCollector.Empty();
+			GRedirectCollector.ProcessSoftObjectPathPackageList(StartupPackage, false, PackagesFromRedirectCollector);
+			for (FName PackageName : PackagesFromRedirectCollector)
+			{
+				StartupSoftObjectPackageReferencers.Add(PackageName, StartupPackage);
+			}
 		}
 	}
 	GRedirectCollector.OnStartupPackageLoadComplete();
@@ -11393,13 +11462,13 @@ void UCookOnTheFlyServer::GenerateInitialRequests(FBeginCookContext& BeginContex
 	{
 		for (FName AssetName : GameDefaultSet.Value)
 		{
-			StartupSoftObjectPackages.Remove(AssetName);
+			StartupSoftObjectPackageReferencers.Remove(AssetName);
 		}
 	}
 	// Strip out missing packages from SoftObjectPaths collected from startup packages.
-	for (TSet<FName>::TIterator It = StartupSoftObjectPackages.CreateIterator(); It; ++It)
+	for (TMap<FName,FName>::TIterator It = StartupSoftObjectPackageReferencers.CreateIterator(); It; ++It)
 	{
-		const FName PackageName = *It;
+		const FName PackageName = It.Key();
 		if (PackageName.IsNone() || PackageDatas->GetFileNameByPackageName(PackageName).IsNone())
 		{
 			It.RemoveCurrent();
@@ -11477,20 +11546,20 @@ void UCookOnTheFlyServer::GenerateInitialRequests(FBeginCookContext& BeginContex
 	}
 	if (!CookByTheBookOptions->bSkipSoftReferences && !EnumHasAnyFlags(CookOptions, ECookByTheBookOptions::NoStartupPackages))
 	{
-		for (FName SoftObjectPackage : StartupSoftObjectPackages)
+		for (TPair<FName,FName> SoftObjectReferencerPair : StartupSoftObjectPackageReferencers)
 		{
 			TMap<FSoftObjectPath, FSoftObjectPath> RedirectedPaths;
 
 			// If this is a redirector, extract destination from asset registry
-			if (ContainsRedirector(SoftObjectPackage, RedirectedPaths))
+			if (ContainsRedirector(SoftObjectReferencerPair.Key, RedirectedPaths))
 			{
 				for (TPair<FSoftObjectPath, FSoftObjectPath>& RedirectedPath : RedirectedPaths)
 				{
 					GRedirectCollector.AddAssetPathRedirection(RedirectedPath.Key, RedirectedPath.Value);
 				}
 			}
-			AddFileToCook(FilesInPath, FilesInPathInstigators, SoftObjectPackage.ToString(),
-				UE::Cook::EInstigator::StartupSoftObjectPath);
+			AddFileToCook(FilesInPath, FilesInPathInstigators, SoftObjectReferencerPair.Key.ToString(),
+				UE::Cook::FInstigator(UE::Cook::EInstigator::StartupSoftObjectPath, SoftObjectReferencerPair.Value));
 		}
 	}
 
@@ -12664,95 +12733,6 @@ EPackageWriterResult UCookOnTheFlyServer::SavePackageBeginCacheForCookedPlatform
 	return EPackageWriterResult::Success;
 }
 
-namespace UE::Cook
-{
-
-/**
- * Keep a map from PackageName to TSet<FName> package dependencies.
- * Prune old entries in the map when memory limit is reached.
- */
-class FCachedDependencies
-{
-public:
-	const TSet<FName>& FindOrAdd(FName PackageName);
-
-private:
-	struct FEntry
-	{
-		TSet<FName> Dependencies;
-		FName PackageName;
-	};
-
-	void CalculateDependencies(FEntry& Entry);
-
-	TRingBuffer<TUniquePtr<FEntry>> FIFOEntries;
-	TMap<FName, FEntry*> ExistingPackages;
-	TArray<FName> ScratchDependenciesArray;
-	int32 MaxNum = 1024;
-	int32 MaxEntryDependenciesKeepSize = 1024;
-};
-
-const TSet<FName>& FCachedDependencies::FindOrAdd(FName PackageName)
-{
-	FEntry* Entry;
-	FEntry*& EntryInExistingPackages = ExistingPackages.FindOrAdd(PackageName);
-	if (!EntryInExistingPackages)
-	{
-		TUniquePtr<FEntry> AllocatedEntry;
-		if (FIFOEntries.Num() >= MaxNum)
-		{
-			AllocatedEntry = FIFOEntries.PopFrontValue();
-			// If the PackageNames were equal, then setting the new name followed by removal of the old name would not work.
-			// They can't be equal because PackageName did not exist in ExistingPackages so it cannot exist in FIFOEntries.
-			check(AllocatedEntry->PackageName != PackageName);
-			EntryInExistingPackages = AllocatedEntry.Get();
-			ExistingPackages.Remove(AllocatedEntry->PackageName);
-			// After calling Remove, EntryInExistingPackages is possibly invalidated so we cannot access it
-		}
-		else
-		{
-			AllocatedEntry.Reset(new FEntry());
-			EntryInExistingPackages = AllocatedEntry.Get();
-		}
-		Entry = AllocatedEntry.Get();
-		FIFOEntries.Add(MoveTemp(AllocatedEntry));
-
-		Entry->PackageName = PackageName;
-		CalculateDependencies(*Entry);
-	}
-	else
-	{
-		Entry = EntryInExistingPackages;
-	}
-	return Entry->Dependencies;
-}
-
-void FCachedDependencies::CalculateDependencies(FEntry& Entry)
-{
-	// Get dependencies into the array scratch
-	ScratchDependenciesArray.Reset();
-	IAssetRegistry::GetChecked().GetDependencies(Entry.PackageName, ScratchDependenciesArray,
-		UE::AssetRegistry::EDependencyCategory::Package);
-
-	// Shrink the old Dependencies if it is over the keep size and we don't need the space, and reserve the new size
-	if (Entry.Dependencies.GetMaxIndex() > MaxEntryDependenciesKeepSize &&
-		ScratchDependenciesArray.Num() < MaxEntryDependenciesKeepSize / 2)
-	{
-		Entry.Dependencies.Empty();
-	}
-
-	// Copy the array into the Dependencies set
-	Entry.Dependencies.Append(ScratchDependenciesArray);
-
-	// Shrink the no-longer-needed array scratch if is over the keep size
-	if (ScratchDependenciesArray.Num() > MaxEntryDependenciesKeepSize)
-	{
-		ScratchDependenciesArray.Empty();
-	}
-}
-
-} // namespace UE::Cook
-
 void UCookOnTheFlyServer::OnDiscoveredPackageDebug(FName PackageName, const UE::Cook::FInstigator& Instigator)
 {
 	using namespace UE::Cook;
@@ -12884,20 +12864,11 @@ void UCookOnTheFlyServer::OnObjectHandleReadDebug(const TArrayView<const UObject
 		}
 		DependencyPackageNames.Add(DependencyPackageName);
 	}
-	if (DependencyPackageNames.IsEmpty())
-	{
-		return;
-	}
-
-	{
-		LLM_SCOPE_BYTAG(Cooker);
-		FScopeLock HiddenDependenciesScopeLock(&HiddenDependenciesLock);
-		const TSet<FName>& ExpectedDependencyPackageNames = CachedDependencies->FindOrAdd(ReferencerPackageName);
-		DependencyPackageNames.RemoveAllSwap([&ExpectedDependencyPackageNames](FName DependencyPackageName)
-			{
-				return ExpectedDependencyPackageNames.Contains(DependencyPackageName);
-			});
-	}
+	DependencyPackageNames.RemoveAllSwap([this, ReferencerPackageName](FName DependencyPackageName)
+		{
+			return AssetRegistry->ContainsDependency(ReferencerPackageName, DependencyPackageName,
+				UE::AssetRegistry::EDependencyCategory::Package);
+		});
 	if (DependencyPackageNames.IsEmpty())
 	{
 		return;
@@ -12978,32 +12949,14 @@ void UCookOnTheFlyServer::ReportHiddenDependency(FName Referencer, FName Depende
 		}
 	}
 
-	FString PrimaryAssetClassPath;
-	if (!Referencer.IsNone())
-	{
-		TArray<FAssetData> AssetDatas;
-		AssetRegistry->GetAssetsByPackageName(Referencer, AssetDatas, true /* bIncludeOnlyDiskAssets */);
-		TArray<const FAssetData*, TInlineAllocator<1>> AssetDataPtrs;
-		AssetDataPtrs.Reserve(AssetDatas.Num());
-		for (FAssetData& AssetData : AssetDatas)
-		{
-			AssetDataPtrs.Add(&AssetData);
-		}
-		const FAssetData* PrimaryAssetData = UE::AssetRegistry::GetMostImportantAsset(AssetDataPtrs);
-		if (PrimaryAssetData && PrimaryAssetData->AssetClassPath.IsValid())
-		{
-			PrimaryAssetClassPath = FString::Printf(TEXT(", PrimaryAssetClass=%s"),
-				*PrimaryAssetData->AssetClassPath.ToString());
-		}
-	}
 	FPackageData* ReferencerPackageData = PackageDatas->TryAddPackageDataByFileName(Referencer);
 	FPackageData* DependencyPackageData = PackageDatas->TryAddPackageDataByFileName(Dependency);
 	if (!ReferencerPackageData || !DependencyPackageData)
 	{
 		return;
 	}
-	TMap<FPackageData*, EInstigator>& Unsolicited = ReferencerPackageData->CreateOrGetUnsolicited();
-	Unsolicited.Add(DependencyPackageData, EInstigator::Unsolicited);
+	ReferencerPackageData->AddDiscoveredDependency(EDiscoveredPlatformSet::CopyFromInstigator, DependencyPackageData,
+		EInstigator::Unsolicited);
 }
 
 static
diff --git a/Engine/Source/Editor/UnrealEd/Private/Cooker/CookDiagnostics.cpp b/Engine/Source/Editor/UnrealEd/Private/Cooker/CookDiagnostics.cpp
--- a/Engine/Source/Editor/UnrealEd/Private/Cooker/CookDiagnostics.cpp
+++ b/Engine/Source/Editor/UnrealEd/Private/Cooker/CookDiagnostics.cpp
@@ -101,7 +101,7 @@ public:
 };
 
 void FDiagnostics::AnalyzeHiddenDependencies(UCookOnTheFlyServer& COTFS, FPackageData& PackageData,
-	TMap<FPackageData*, EInstigator>&& UnsolicitedForPackage, TMap<FPackageData*, EInstigator>& SaveReferences,
+	TMap<FPackageData*, EInstigator>* DiscoveredDependencies, TMap<FPackageData*, EInstigator>& SaveReferences,
 	TConstArrayView<const ITargetPlatform*> ReachablePlatforms, bool bOnlyEditorOnlyDebug,
 	bool bHiddenDependenciesDebug)
 {
@@ -119,9 +119,13 @@ void FDiagnostics::AnalyzeHiddenDependencies(UCookOnTheFlyServer& COTFS, FPackag
 	// This function investigates each unsolicited package that was not added and writes information about
 	// why it was not added. Each of these needs to be investigated, and either fixed somehow so that it is
 	// referenced from the cooked SavePackage, or marked as an expected correct difference using FCookLoadScope.
-	for (const TPair<FPackageData*, EInstigator>& UnsolicitedPair : UnsolicitedForPackage)
+	if (!DiscoveredDependencies)
 	{
-		FPackageData* Unsolicited = UnsolicitedPair.Key;
+		return;
+	}
+	for (const TPair<const FPackageData*, EInstigator>& UnsolicitedPair : *DiscoveredDependencies)
+	{
+		const FPackageData* Unsolicited = UnsolicitedPair.Key;
 		EInstigator Instigator = UnsolicitedPair.Value;
 		if (!bHiddenDependenciesDebug && SaveReferences.Contains(Unsolicited))
 		{
diff --git a/Engine/Source/Editor/UnrealEd/Private/Cooker/CookDiagnostics.h b/Engine/Source/Editor/UnrealEd/Private/Cooker/CookDiagnostics.h
--- a/Engine/Source/Editor/UnrealEd/Private/Cooker/CookDiagnostics.h
+++ b/Engine/Source/Editor/UnrealEd/Private/Cooker/CookDiagnostics.h
@@ -23,7 +23,7 @@ public:
 	 * with diagnostics for unexpected differences.
 	 */
 	static void AnalyzeHiddenDependencies(UCookOnTheFlyServer& COTFS, FPackageData& PackageData,
-		TMap<FPackageData*, EInstigator>&& UnsolicitedForPackage, TMap<FPackageData*, EInstigator>& SaveReferences,
+		TMap<FPackageData*, EInstigator>* DiscoveredDependencies, TMap<FPackageData*, EInstigator>& SaveReferences,
 		TConstArrayView<const ITargetPlatform*> ReachablePlatforms, bool bOnlyEditorOnlyDebug,
 		bool bHiddenDependenciesDebug);
 
diff --git a/Engine/Source/Editor/UnrealEd/Private/Cooker/CookPackageData.cpp b/Engine/Source/Editor/UnrealEd/Private/Cooker/CookPackageData.cpp
--- a/Engine/Source/Editor/UnrealEd/Private/Cooker/CookPackageData.cpp
+++ b/Engine/Source/Editor/UnrealEd/Private/Cooker/CookPackageData.cpp
@@ -150,6 +150,7 @@ void FPackageData::ClearReferences()
 	{
 		PackagePreloader->Shutdown(); // Clears references to any other preloaders
 	}
+	ClearDiscoveredDependencies();
 }
 
 int32 FPackageData::GetPlatformsNeedingCommitNum(EReachability Reachability) const
@@ -206,26 +207,9 @@ bool FPackageData::AreAllReachablePlatformsVisitedByCluster(EReachability InReac
 	return true;
 }
 
-void FPackageData::GetReachablePlatformsForInstigator(EReachability InReachability, UCookOnTheFlyServer& COTFS,
-	FName InInstigator, TArray<const ITargetPlatform*>& Platforms)
+const TArray<const ITargetPlatform*>& FPackageData::GetSessionPlatformsInternal(UCookOnTheFlyServer& COTFS)
 {
-	return GetReachablePlatformsForInstigator(InReachability, COTFS,
-		COTFS.PackageDatas->TryAddPackageDataByPackageName(InInstigator), Platforms);
-}
-
-void FPackageData::GetReachablePlatformsForInstigator(EReachability InReachability, UCookOnTheFlyServer& COTFS,
-	UE::Cook::FPackageData* InInstigator, TArray<const ITargetPlatform*>& Platforms)
-{
-	if (InInstigator)
-	{
-		InInstigator->GetReachablePlatforms(InReachability, Platforms);
-	}
-	else
-	{
-		const TArray<const ITargetPlatform*>& SessionPlatforms = COTFS.PlatformManager->GetSessionPlatforms();
-		Platforms.Reset(SessionPlatforms.Num() + 1);
-		Platforms.Append(SessionPlatforms);
-	}
+	return COTFS.PlatformManager->GetSessionPlatforms();
 }
 
 void FPackageData::AddReachablePlatforms(FRequestCluster& RequestCluster, EReachability InReachability,
@@ -290,7 +274,7 @@ void FPackageData::QueueAsDiscoveredInternal(FPackageData& PackageData, FInstiga
 		TArray<const ITargetPlatform*, TInlineAllocator<ExpectedMaxNumPlatforms>> BufferPlatforms;
 		TConstArrayView<const ITargetPlatform*> PlatformArray =
 			ReachablePlatforms.GetPlatforms(COTFS, &InInstigator, TConstArrayView<const ITargetPlatform*>(),
-				EReachability::Build, &BufferPlatforms);
+				EReachability::Build, BufferPlatforms);
 		bool bHasNewPlatforms = !PackageData.HasAllCommittedPlatforms(PlatformArray);
 		if (bHasNewPlatforms)
 		{
@@ -1492,6 +1476,10 @@ void FPackageData::ClearCookedPlatformData()
 void FPackageData::OnRemoveSessionPlatform(const ITargetPlatform* Platform)
 {
 	PlatformDatas.Remove(Platform);
+	if (DiscoveredDependencies)
+	{
+		DiscoveredDependencies->Remove(Platform);
+	}
 }
 
 bool FPackageData::HasReferencedObjects() const
@@ -1508,6 +1496,17 @@ void FPackageData::RemapTargetPlatforms(const TMap<ITargetPlatform*, ITargetPlat
 	{
 		ITargetPlatform* NewKey = Remap[ExistingPair.Key];
 		NewPlatformDatas.FindOrAdd(NewKey) = MoveTemp(ExistingPair.Value);
+
+		if (DiscoveredDependencies)
+		{
+			TMap<FPackageData*, EInstigator>* OldValue = DiscoveredDependencies->Find(ExistingPair.Key);
+			if (OldValue)
+			{
+				TMap<FPackageData*, EInstigator> MovedValue = MoveTemp(*OldValue);
+				DiscoveredDependencies->Remove(ExistingPair.Key);
+				DiscoveredDependencies->Add(NewKey, MoveTemp(MovedValue));
+			}
+		}
 	}
 
 	// The save state (and maybe more in the future) by contract can depend on the order of the request platforms
@@ -1633,6 +1632,26 @@ TRefCountPtr<FGenerationHelper> FPackageData::GetOrFindParentGenerationHelper()
 	return ParentGenerationHelper;
 }
 
+TRefCountPtr<FGenerationHelper> FPackageData::GetOrFindParentGenerationHelperNoCache()
+{
+	if (ParentGenerationHelper)
+	{
+		return ParentGenerationHelper;
+	}
+	if (!IsGenerated())
+	{
+		return nullptr;
+	}
+
+	FPackageData* OwnerPackageData = PackageDatas.FindPackageDataByPackageName(GetParentGenerator());
+	if (!OwnerPackageData)
+	{
+		return nullptr;
+	}
+
+	return OwnerPackageData->GetGenerationHelper();
+}
+
 TRefCountPtr<FGenerationHelper> FPackageData::TryCreateValidParentGenerationHelper()
 {
 	if (ParentGenerationHelper)
@@ -1767,32 +1786,72 @@ FConstructPackageData FPackageData::CreateConstructData()
 	return ConstructData;
 }
 
-TMap<FPackageData*, EInstigator>& FPackageData::CreateOrGetUnsolicited()
+void FPackageData::AddDiscoveredDependency(const FDiscoveredPlatformSet& Platforms, FPackageData* Dependency,
+	EInstigator Category)
 {
-	if (!Unsolicited)
+	TConstArrayView<const ITargetPlatform*> PlatformArray;
+	TArray<const ITargetPlatform*, TInlineAllocator<ExpectedMaxNumPlatforms>> BufferPlatforms;
+
+	if (Platforms.GetSource() == EDiscoveredPlatformSet::CopyFromInstigator)
+	{
+		BufferPlatforms.Add(nullptr); // PlatformAgnostic platform
+		PlatformArray = BufferPlatforms;
+	}
+	else
+	{
+		UCookOnTheFlyServer& COTFS = PackageDatas.GetCookOnTheFlyServer();
+		PlatformArray = Platforms.GetPlatforms(COTFS, nullptr,
+			TConstArrayView<const ITargetPlatform*>(), EReachability::Runtime, BufferPlatforms);
+		if (PlatformArray.Num() == COTFS.PlatformManager->GetSessionPlatforms().Num())
+		{
+			BufferPlatforms.SetNum(1, EAllowShrinking::No);
+			BufferPlatforms[0] = nullptr; // PlatformAgnostic platform
+			PlatformArray = BufferPlatforms;
+		}
+	}
+
+	if (!DiscoveredDependencies)
+	{
+		DiscoveredDependencies = MakeUnique<TMap<const ITargetPlatform*, TMap<FPackageData*, EInstigator>>>();
+	}
+
+	for (const ITargetPlatform* TargetPlatform : PlatformArray)
 	{
-		Unsolicited = MakeUnique<TMap<FPackageData*, EInstigator>>();
+		TMap<FPackageData*, EInstigator>& PlatformDependencies = DiscoveredDependencies->FindOrAdd(TargetPlatform);
+		EInstigator& ExistingEdgeType = PlatformDependencies.FindOrAdd(Dependency, Category);
+
+		// Overwrite the previous edge type with the new edge type if the new edge type is higher priority.
+		if (Category == EInstigator::ForceExplorableSaveTimeSoftDependency)
+		{
+			ExistingEdgeType = Category;
+		}
 	}
-	return *Unsolicited;
 }
 
-TMap<FPackageData*, EInstigator> FPackageData::DetachUnsolicited()
+void FPackageData::ClearDiscoveredDependencies()
+{
+	DiscoveredDependencies.Reset();
+}
+
+TMap<FPackageData*, EInstigator>& FPackageData::CreateOrGetDiscoveredDependencies(const ITargetPlatform* TargetPlatform)
 {
-	TMap<FPackageData*, EInstigator> Result;
-	if (Unsolicited)
+	if (!DiscoveredDependencies)
 	{
-		Result = MoveTemp(*Unsolicited);
-		Unsolicited.Reset();
+		DiscoveredDependencies = MakeUnique<TMap<const ITargetPlatform*, TMap<FPackageData*, EInstigator>>>();
 	}
-	return Result;
+	return DiscoveredDependencies->FindOrAdd(TargetPlatform);
 }
 
-void FPackageData::ClearUnsolicited()
+TMap<FPackageData*, EInstigator>* FPackageData::GetDiscoveredDependencies(const ITargetPlatform* TargetPlatform)
 {
-	Unsolicited.Reset();
+	if (!DiscoveredDependencies)
+	{
+		return nullptr;
+	}
+	return DiscoveredDependencies->Find(TargetPlatform);
 }
 
-}
+} // namespace UE::Cook
 
 FCbWriter& operator<<(FCbWriter& Writer, const UE::Cook::FConstructPackageData& PackageData)
 {
diff --git a/Engine/Source/Editor/UnrealEd/Private/Cooker/CookPackageData.h b/Engine/Source/Editor/UnrealEd/Private/Cooker/CookPackageData.h
--- a/Engine/Source/Editor/UnrealEd/Private/Cooker/CookPackageData.h
+++ b/Engine/Source/Editor/UnrealEd/Private/Cooker/CookPackageData.h
@@ -712,6 +712,12 @@ public:
 	 * and return it, otherwise return null.
 	 */
 	TRefCountPtr<FGenerationHelper> GetOrFindParentGenerationHelper();
+	/**
+	 * Return the ParentGenerator's GenerationHelper if the pointer to it has already been set on this, otherwise
+	 * look for it on the ParentGenerator by calling GetGenerationHelper, return the found generator or null.
+	 * Does not cache the found generator on this if found.
+	 */
+	TRefCountPtr<FGenerationHelper> GetOrFindParentGenerationHelperNoCache();
 	/**
 	 * Return the ParentGenerator's GenerationHelper if the pointer to it has already been set on this, otherwise
 	 * try to find or create it by calling TryCreateValidGenerationHelper on the ParentGenerator.
@@ -764,24 +770,29 @@ public:
 	/** Marshall this PackageData to a ConstructData that is used later or on a remote machine to reconstruct it. */
 	FConstructPackageData CreateConstructData();
 
-	/** Only used When IsDebugRecordUnsolicited: storage for the list of unsolicited packages for this package */
-	TMap<FPackageData*, EInstigator>& CreateOrGetUnsolicited();
-	TMap<FPackageData*, EInstigator> DetachUnsolicited();
-	void ClearUnsolicited();
+	/** Storage for dependencies discovered for the package during cook that are not reported by the AssetRegistry. */
+	void AddDiscoveredDependency(const FDiscoveredPlatformSet& Platforms, FPackageData* Dependency,
+		EInstigator Category);
+	void ClearDiscoveredDependencies();
+	TMap<FPackageData*, EInstigator>& CreateOrGetDiscoveredDependencies(const ITargetPlatform* TargetPlatform);
+	TMap<FPackageData*, EInstigator>* GetDiscoveredDependencies(const ITargetPlatform* TargetPlatform);
 
 	/**
 	 * Return the platforms for which the given Package has been marked reachable.
 	 * If the package does not exist, return the COTFS's list of Session platforms
 	 */
+	template <typename ArrayType>
 	static void GetReachablePlatformsForInstigator(EReachability InReachability, UCookOnTheFlyServer& COTFS,
-		FPackageData* InInstigator, TArray<const ITargetPlatform*>& Platforms);
+		FPackageData* InInstigator, ArrayType& Platforms);
+	template <typename ArrayType>
 	static void GetReachablePlatformsForInstigator(EReachability InReachability, UCookOnTheFlyServer& COTFS,
-		FName InInstigator, TArray<const ITargetPlatform*>& Platforms);
+		FName InInstigator, ArrayType& Platforms);
 private:
 	friend struct UE::Cook::FPackageDatas;
 
 	void UpdateContainerUrgency(EUrgency OldUrgency, EUrgency NewUrgency);
 	void SetInstigatorInternal(EReachability InReachability, FInstigator&& InInstigator);
+	static const TArray<const ITargetPlatform*>& GetSessionPlatformsInternal(UCookOnTheFlyServer& COTFS);
 	static void AddReachablePlatformsInternal(FPackageData& PackageData, EReachability InReachability,
 		TConstArrayView<const ITargetPlatform*> Platforms, FInstigator&& InInstigator);
 	static void QueueAsDiscoveredInternal(FPackageData& PackageData, FInstigator&& InInstigator,
@@ -844,7 +855,7 @@ private:
 
 	TArray<FCachedObjectInOuter> CachedObjectsInOuter;
 	FCompletionCallback CompletionCallback;
-	TUniquePtr<TMap<FPackageData*, EInstigator>> Unsolicited;
+	TUniquePtr<TMap<const ITargetPlatform*, TMap<FPackageData*, EInstigator>>> DiscoveredDependencies;
 	TUniquePtr<FBuildResultDependenciesMap> LoadDependencies;
 	FName PackageName;
 	FName FileName;
@@ -1683,6 +1694,30 @@ inline void FPackageData::GetReachablePlatforms(EReachability InReachability, Ar
 	}
 }
 
+template <typename ArrayType>
+void FPackageData::GetReachablePlatformsForInstigator(EReachability InReachability, UCookOnTheFlyServer& COTFS,
+	FName InInstigator, ArrayType& Platforms)
+{
+	return GetReachablePlatformsForInstigator(InReachability, COTFS,
+		COTFS.PackageDatas->TryAddPackageDataByPackageName(InInstigator), Platforms);
+}
+
+template <typename ArrayType>
+void FPackageData::GetReachablePlatformsForInstigator(EReachability InReachability, UCookOnTheFlyServer& COTFS,
+	UE::Cook::FPackageData* InInstigator, ArrayType& Platforms)
+{
+	if (InInstigator)
+	{
+		InInstigator->GetReachablePlatforms(InReachability, Platforms);
+	}
+	else
+	{
+		const TArray<const ITargetPlatform*>& SessionPlatforms = GetSessionPlatformsInternal(COTFS);
+		Platforms.Reset(SessionPlatforms.Num() + 1);
+		Platforms.Append(SessionPlatforms);
+	}
+}
+
 inline EUrgency FPackageData::GetUrgency() const
 {
 	return static_cast<EUrgency>(Urgency);
diff --git a/Engine/Source/Editor/UnrealEd/Private/Cooker/CookRequestCluster.cpp b/Engine/Source/Editor/UnrealEd/Private/Cooker/CookRequestCluster.cpp
--- a/Engine/Source/Editor/UnrealEd/Private/Cooker/CookRequestCluster.cpp
+++ b/Engine/Source/Editor/UnrealEd/Private/Cooker/CookRequestCluster.cpp
@@ -200,6 +200,7 @@ FRequestCluster::FRequestCluster(UCookOnTheFlyServer& InCOTFS, TRingBuffer<FDisc
 	: FRequestCluster(InCOTFS, EReachability::Runtime)
 {
 	TArray<const ITargetPlatform*, TInlineAllocator<ExpectedMaxNumPlatforms>> BufferPlatforms;
+	TArray<const ITargetPlatform*, TInlineAllocator<ExpectedMaxNumPlatforms>> ImmediateAddPlatforms;
 	if (!COTFS.bSkipOnlyEditorOnly)
 	{
 		BufferPlatforms = COTFS.PlatformManager->GetSessionPlatforms();
@@ -209,131 +210,120 @@ FRequestCluster::FRequestCluster(UCookOnTheFlyServer& InCOTFS, TRingBuffer<FDisc
 	while (!DiscoveryQueue.IsEmpty())
 	{
 		FDiscoveryQueueElement* Discovery = &DiscoveryQueue.First();
+		ON_SCOPE_EXIT
+		{
+			DiscoveryQueue.PopFrontNoCheck();
+		};
 		FPackageData& PackageData = *Discovery->PackageData;
 
 		TConstArrayView<const ITargetPlatform*> NewReachablePlatforms;
 		if (COTFS.bSkipOnlyEditorOnly)
 		{
 			NewReachablePlatforms = Discovery->ReachablePlatforms.GetPlatforms(COTFS, &Discovery->Instigator,
-				TConstArrayView<const ITargetPlatform*>(), EReachability::Runtime, &BufferPlatforms);
+				TConstArrayView<const ITargetPlatform*>(), EReachability::Runtime, BufferPlatforms);
 		}
 		else
 		{
 			NewReachablePlatforms = BufferPlatforms;
 		}
-		if (Discovery->Instigator.Category == EInstigator::ForceExplorableSaveTimeSoftDependency)
+
+		FPackageData* Referencer = COTFS.PackageDatas->FindPackageDataByPackageName(Discovery->Instigator.Referencer);
+		if (Referencer)
+		{
+			// The discovery may have come from a cookworker; add it again in case it was not already added.
+			Referencer->AddDiscoveredDependency(Discovery->ReachablePlatforms, &PackageData, Discovery->Instigator.Category);
+		}
+
+		// Create a list of immediate add platforms: platforms for which the referencer is already reachable so we know
+		// we can add the Discovery to the cook now. For Platforms on which the 
+		// If its referencer is not reachable for a given platform, then take no further action for that
+		// platform now; if the referencer becomes reachable later, we will then add the target to the cook
+		// when we visit the referencer and traverse the DiscoveredDependency edge we just added.
+		if (!Referencer ||
+			Discovery->Urgency > PackageData.GetUrgency())
+		{
+			// In the no-referencer case, add all discovered reachability platforms immediately.
+			// And for urgent requests, also add them immediately. Urgency greater than normal can only be set from
+			// referencers that were already reachable; we currently rely on this so that we don't have to make
+			// another message type for them and can process them here in the DiscoveryQueue. If the discovery
+			// carries a raise in urgency, then add it to the cook even if the referencer is not yet reachable.
+			ImmediateAddPlatforms = NewReachablePlatforms;
+		}
+		else
 		{
-			// This package was possibly previously marked as not explorable, but now it is marked as explorable.
-			// One example of this is externalactor packages - they are by default not cookable and not explorable
-			// (see comment in FRequestCluster::IsRequestCookable). But once WorldPartition loads them, we need to mark
-			// them as explored so that their imports are marked as expected and all of their soft dependencies
-			// are included.
+			ImmediateAddPlatforms.Reset();
 			for (const ITargetPlatform* TargetPlatform : NewReachablePlatforms)
 			{
-				if (TargetPlatform != CookerLoadingPlatformKey)
+				FPackagePlatformData& PlatformData = Referencer->FindOrAddPlatformData(TargetPlatform);
+				if (PlatformData.IsReachable(EReachability::Runtime) && PlatformData.IsExplorable())
 				{
-					PackageData.FindOrAddPlatformData(TargetPlatform).MarkAsExplorable();
+					ImmediateAddPlatforms.Add(TargetPlatform);
 				}
 			}
 		}
 
-		if (PackageData.HasReachablePlatforms(EReachability::Runtime, NewReachablePlatforms))
+		// Remove platforms that are already reachable and explorable from the ImmediateAddPlatforms.
+		// Also handle the ForceExplorableSaveTimeSoftDependency flag to mark the platform explorable.
+		for (TArray<const ITargetPlatform*, TInlineAllocator<ExpectedMaxNumPlatforms>>::TIterator
+			Iter(ImmediateAddPlatforms); Iter; ++Iter)
 		{
-			// If there are no new reachable platforms, add it to the cluster for cooking if it needs
-			// it, otherwise let it remain where it is
-			EUrgency Urgency = Discovery->Urgency;
-			DiscoveryQueue.PopFrontValue();
-			Discovery = nullptr;
-			if (PackageData.IsInProgress())
+			const ITargetPlatform* TargetPlatform = *Iter;
+			FPackagePlatformData& PlatformData = PackageData.FindOrAddPlatformData(TargetPlatform);
+
+			if (Discovery->Instigator.Category == EInstigator::ForceExplorableSaveTimeSoftDependency &&
+				!PlatformData.IsExplorable())
 			{
-				PackageData.RaiseUrgency(Urgency, ESendFlags::QueueAddAndRemove);
+				PlatformData.MarkAsExplorable(); // Clears reachability so IsReachable below returns false.
 			}
-			else if (PackageData.GetPlatformsNeedingCommitNum(EReachability::Runtime) > 0 || !PackageData.AreAllReachablePlatformsVisitedByCluster(EReachability::Runtime))
+			if (PlatformData.IsReachable(EReachability::Runtime))
 			{
-				PackageData.SendToState(EPackageState::Request, ESendFlags::QueueRemove,
-					EStateChangeReason::RequestCluster);
-				PackageData.RaiseUrgency(Urgency, ESendFlags::QueueNone);
-				FVertexData& Vertex = FindOrAddVertex(PackageData);
-				SetOwnedByCluster(Vertex, true /* bOwnedByCluster */, false /* bNeedsStateChange */);
+				Iter.RemoveCurrentSwap();
 			}
-			continue;
 		}
 
-		// Startup packages and Generated packages are expected discovery types and do not need to add hidden
-		// dependencies.
-		if (Discovery->Instigator.Category != EInstigator::StartupPackage &&
-			Discovery->Instigator.Category != EInstigator::StartupPackageCookLoadScope &&
-			Discovery->Instigator.Category != EInstigator::GeneratedPackage)
+		bool bAddToCluster = !ImmediateAddPlatforms.IsEmpty();
+		// Handle the edge case that all of the Addable platforms are already reachable, and are not yet committed,
+		// but the package is not in progress.
+		// TODO: Is this edge case possible? How can it occur?
+		if (!bAddToCluster && !PackageData.IsInProgress()
+			&& (PackageData.GetPlatformsNeedingCommitNum(EReachability::Runtime) > 0
+				|| !PackageData.AreAllReachablePlatformsVisitedByCluster(EReachability::Runtime)))
 		{
-			// For unsolicited packages, we need to check load-reachability to decide whether the load was expected.
-			bool bExpectedDiscoveryType = false;
-			if (Discovery->Instigator.Category == EInstigator::Unsolicited)
-			{
-				bExpectedDiscoveryType = PackageData.FindOrAddPlatformData(CookerLoadingPlatformKey)
-					.IsReachable(EReachability::Runtime);
-				if (bExpectedDiscoveryType && COTFS.bSkipOnlyEditorOnly)
-				{
-					// In SkipOnlyEditorOnly mode, expected-load unsolicited packages are skipped; merely loading
-					// a package is not sufficient to add it to the cook. So take no action on this package.
-					DiscoveryQueue.PopFrontValue();
-					Discovery = nullptr;
-					continue;
-				}
-			}
-			else
-			{
-				// For other instigator types, the discovery is either expected or unexpected depending on type.
-				// Adding packages to the cook should happen only for a few types of instigators, from external
-				// package requests, or during cluster exploration. If not expected, add a diagnostic message.
-				bExpectedDiscoveryType = Discovery->Instigator.Category == EInstigator::SaveTimeHardDependency ||
-					Discovery->Instigator.Category == EInstigator::SaveTimeSoftDependency ||
-					Discovery->Instigator.Category == EInstigator::ForceExplorableSaveTimeSoftDependency;
-			}
+			bAddToCluster = true;
+		}
 
-			// If there are other discovered packages we have already added to this cluster, then defer this one
-			// until we have explored those; add this one to the next cluster. Exploring those earlier discoveries
-			// might add this one through cluster exploration and not require a hidden dependency.
-			if (!ClusterPackages.IsEmpty())
-			{
-				break;
-			}
+		if (bAddToCluster)
+		{
+			// Add the new reachable platforms
+			PackageData.AddReachablePlatforms(*this, EReachability::Runtime, ImmediateAddPlatforms,
+				MoveTemp(Discovery->Instigator));
 
-			if (Discovery->ReachablePlatforms.GetSource() == EDiscoveredPlatformSet::CopyFromInstigator)
+			// Send it to the Request state if it's not already there, remove it from its old container
+			// and add it to this cluster.
+			FVertexData& Vertex = FindOrAddVertex(PackageData);
+			if (!Vertex.IsOwnedByCluster())
 			{
-				// Add it as a hidden dependency so that future platforms discovered as reachable in
-				// the instigator will also be marked as reachable in the dependency.
-				if (COTFS.bSkipOnlyEditorOnly)
-				{
-					FPackageData* InstigatorPackageData = Discovery->Instigator.Referencer.IsNone() ? nullptr
-						: COTFS.PackageDatas->TryAddPackageDataByPackageName(Discovery->Instigator.Referencer);
-					if (InstigatorPackageData)
-					{
-						COTFS.DiscoveredDependencies.FindOrAdd(InstigatorPackageData->GetPackageName())
-							.Add(PackageData.GetPackageName());
-					}
-				}
+				// QueueRemove in SendToState does not know how to handle Packages assigned to a cluster
+				// in construction, so we must pass in QueueRemove if and only if it's not in this cluster
+				PackageData.SendToState(EPackageState::Request, ESendFlags::QueueRemove, EStateChangeReason::RequestCluster);
+				SetOwnedByCluster(Vertex, true /* bOwnedByCluster */, false /* bNeedsStateChange */);
 			}
-
-			if (!bExpectedDiscoveryType)
+			else
 			{
-				COTFS.OnDiscoveredPackageDebug(PackageData.GetPackageName(), Discovery->Instigator);
+				// If it is in this cluster, it should have already been put into the Request state.
+				ensure(PackageData.GetState() == EPackageState::Request);
 			}
 		}
-		// Add the new reachable platforms
-		PackageData.AddReachablePlatforms(*this, EReachability::Runtime, NewReachablePlatforms,
-			MoveTemp(Discovery->Instigator));
 
-		// Pop it off the list; note that this invalidates the pointers we had into the DiscoveryQueueElement
-		FDiscoveryQueueElement PoppedDiscovery = DiscoveryQueue.PopFrontValue();
-		Discovery = &PoppedDiscovery;
-		NewReachablePlatforms = TConstArrayView<const ITargetPlatform*>();
-
-		// Send it to the Request state if it's not already there, remove it from its old container
-		// and add it to this cluster.
-		PackageData.SendToState(EPackageState::Request, ESendFlags::QueueRemove, EStateChangeReason::RequestCluster);
-		PackageData.RaiseUrgency(Discovery->Urgency, ESendFlags::QueueNone);
-		FVertexData& Vertex = FindOrAddVertex(PackageData);
-		SetOwnedByCluster(Vertex, true /* bOwnedByCluster */, false /* bNeedsStateChange */);
+		// If urgency was specified and the package is now (or was already) in progress, raise the urgency
+		if (Discovery->Urgency > PackageData.GetUrgency() && PackageData.IsInProgress())
+		{
+			PackageData.RaiseUrgency(Discovery->Urgency,
+				// Raising urgency depending on state will need to remove and readd it, but don't allow that
+				// if we added it to this cluster because RaiseUrgency doesn't know how to handle adding and
+				// removing from the cluster.
+				bAddToCluster ? ESendFlags::QueueNone : ESendFlags::QueueAddAndRemove);
+		}
 	}
 }
 
@@ -1502,8 +1492,6 @@ void FRequestCluster::FGraphSearch::FExploreEdgesContext::Initialize(FVertexData
 	LocalNumFetchPlatforms = Cluster.GetNumFetchPlatforms();
 	bFetchAnyTargetPlatform = false;
 
-	DiscoveredDependencies = Cluster.COTFS.DiscoveredDependencies.Find(PackageName);
-
 	GraphSearch.PendingTransitiveBuildDependencyVertices.Remove(Vertex);
 }
 
@@ -1567,7 +1555,8 @@ bool FRequestCluster::FGraphSearch::FExploreEdgesContext::TryCalculateIterativel
 	}
 	else
 	{
-		GenerationHelper = PackageData->GetOrFindParentGenerationHelper();
+		GenerationHelper = Vertex->IsOwnedByCluster() ? PackageData->GetOrFindParentGenerationHelper()
+			: PackageData->GetOrFindParentGenerationHelperNoCache();
 		if (GenerationHelper)
 		{
 			ParentPackageData = &GenerationHelper->GetOwner();
@@ -1810,10 +1799,6 @@ void FRequestCluster::FGraphSearch::FExploreEdgesContext::CalculatePackageDataDe
 		HardDependenciesSet.Append(HardGameDependencies);
 		HardDependenciesSet.Append(HardEditorDependencies);
 	}
-	if (DiscoveredDependencies)
-	{
-		HardDependenciesSet.Append(*DiscoveredDependencies);
-	}
 	if (Cluster.bAllowSoftDependencies)
 	{
 		// bSkipOnlyEditorOnly is always true for soft dependencies; skip editoronly soft dependencies
@@ -1892,12 +1877,36 @@ void FRequestCluster::FGraphSearch::FExploreEdgesContext::CalculateDependenciesA
 			AddPlatformDependencyRange(HardGameDependencies, PlatformIndex, EInstigator::HardDependency);
 			AddPlatformDependencyRange(HardEditorDependencies, PlatformIndex, EInstigator::HardEditorOnlyDependency);
 			AddPlatformDependencyRange(SoftGameDependencies, PlatformIndex, EInstigator::SoftDependency);
+			ProcessPlatformDiscoveredDependencies(PlatformIndex, TargetPlatform);
 			ProcessPlatformAttachments(PlatformIndex, TargetPlatform, FetchPlatformData, PackagePlatformData,
 				QueryPlatformData.CookAttachments, true /* bExploreRuntimeDependencies  */);
 		}
-		if (DiscoveredDependencies)
+	}
+}
+
+void FRequestCluster::FGraphSearch::FExploreEdgesContext::ProcessPlatformDiscoveredDependencies(int32 PlatformIndex,
+	const ITargetPlatform* TargetPlatform)
+{
+	// nullptr in PackageData->GetDiscoveredDependencies means PlatformAgnostic; this function is for
+	// the platformspecific, so TargetPlatform must not be null.
+	check(TargetPlatform != nullptr);
+	const TMap<FPackageData*, EInstigator>* PlatformDependencies
+		= PackageData->GetDiscoveredDependencies(TargetPlatform);
+	const TMap<FPackageData*, EInstigator>* AgnosticDependencies
+		= PackageData->GetDiscoveredDependencies(nullptr);
+	for (const TMap<FPackageData*, EInstigator>* DependenciesMap : { AgnosticDependencies, PlatformDependencies })
+	{
+		if (DependenciesMap)
 		{
-			AddPlatformDependencyRange(*DiscoveredDependencies, PlatformIndex, EInstigator::HardDependency);
+			for (const TPair<FPackageData*, EInstigator>& PackagePair : *DependenciesMap)
+			{
+				// DiscoveredDependencies are always treated as Soft, but might also have the
+				// ForceExplorableSaveTimeSoftDependency property, which sets the explorable property on the target
+				// if it doesn't already have it.
+				EInstigator EdgeType = PackagePair.Value == EInstigator::ForceExplorableSaveTimeSoftDependency ?
+					EInstigator::ForceExplorableSaveTimeSoftDependency : EInstigator::SoftDependency;
+				AddPlatformDependency(PackagePair.Key->GetPackageName(), PlatformIndex, EdgeType);
+			}
 		}
 	}
 }
@@ -1918,6 +1927,7 @@ void FRequestCluster::FGraphSearch::FExploreEdgesContext::QueueVisitsOfDependenc
 		FName DependencyName = PlatformDependencyPair.Key;
 		TBitArray<>& HasRuntimePlatformByIndex = PlatformDependencyPair.Value.HasRuntimePlatformByIndex;
 		TBitArray<>& HasBuildPlatformByIndex = PlatformDependencyPair.Value.HasBuildPlatformByIndex;
+		TBitArray<>& ForceExplorableByIndex = PlatformDependencyPair.Value.ForceExplorableByIndex;
 		EInstigator InstigatorType = PlatformDependencyPair.Value.InstigatorType;
 		EInstigator BuildInstigatorType = PlatformDependencyPair.Value.BuildInstigatorType;
 
@@ -1950,6 +1960,22 @@ void FRequestCluster::FGraphSearch::FExploreEdgesContext::QueueVisitsOfDependenc
 			if (HasRuntimePlatformByIndex[PlatformIndex] && ClusterReachability == EReachability::Runtime)
 			{
 				// RuntimeDependencies only cause edges and reachability in EReachability::Runtime clusters.
+
+				if (TargetPlatform != CookerLoadingPlatformKey && ForceExplorableByIndex[PlatformIndex])
+				{
+					// This package was possibly previously marked as not explorable, but now the referencer wants to
+					// mark it as explorable. One example of this is externalactor packages - they are by default not
+					// cookable and not explorable (see comment in FRequestCluster::IsRequestCookable). But once
+					// WorldPartition loads them, we need to mark them as explored so that their imports are marked as
+					// expected and all of their soft dependencies are included.
+					if (!PlatformData.IsExplorableOverride())
+					{
+						// MarkAsExplorable calls ResetReachable, so we need to process this before testing IsReachable
+						// below.
+						PlatformData.MarkAsExplorable();
+					}
+				}
+
 				if (PlatformIndex == CookerLoadingPlatformIndex)
 				{
 					if (!Edges)
@@ -2034,6 +2060,7 @@ void FRequestCluster::FGraphSearch::FExploreEdgesContext::AddPlatformDependency(
 	{
 		PlatformDependencyBits.HasRuntimePlatformByIndex.Init(false, LocalNumFetchPlatforms);
 		PlatformDependencyBits.HasBuildPlatformByIndex.Init(false, LocalNumFetchPlatforms);
+		PlatformDependencyBits.ForceExplorableByIndex.Init(false, LocalNumFetchPlatforms);
 		PlatformDependencyBits.InstigatorType = EInstigator::InvalidCategory;
 		PlatformDependencyBits.BuildInstigatorType = EInstigator::InvalidCategory;
 	}
@@ -2056,12 +2083,17 @@ void FRequestCluster::FGraphSearch::FExploreEdgesContext::AddPlatformDependency(
 				PlatformDependencyBits.InstigatorType = InstigatorType;
 			}
 			break;
-		case EInstigator::SoftDependency:
+		case EInstigator::SoftDependency: [[fallthrough]];
+		case EInstigator::ForceExplorableSaveTimeSoftDependency:
 			if (PlatformDependencyBits.InstigatorType != EInstigator::HardDependency
 				&& PlatformDependencyBits.InstigatorType != EInstigator::HardEditorOnlyDependency)
 			{
 				PlatformDependencyBits.InstigatorType = InstigatorType;
 			}
+			if (InstigatorType == EInstigator::ForceExplorableSaveTimeSoftDependency)
+			{
+				PlatformDependencyBits.ForceExplorableByIndex[PlatformIndex] = true;
+			}
 			break;
 		case EInstigator::InvalidCategory:
 			// Caller indicated they do not want to set the InstigatorType
@@ -2236,7 +2268,8 @@ void FRequestCluster::FGraphSearch::FExploreEdgesContext::ProcessPlatformAttachm
 	{
 		FQueryPlatformData& PlatformAgnosticQueryData = Vertex->GetPlatformData()[PlatformAgnosticPlatformIndex];
 
-		if (PlatformAgnosticQueryData.bSchedulerThreadFetchCompleted)
+		if (PlatformAgnosticQueryData.bSchedulerThreadFetchCompleted
+			&& PlatformAgnosticQueryData.CookAttachments.Artifacts.IsValid())
 		{
 			const FAssetPackageData* OverrideAssetPackageData = nullptr;
 			TRefCountPtr<FGenerationHelper> GenerationHelper;
@@ -2246,7 +2279,8 @@ void FRequestCluster::FGraphSearch::FExploreEdgesContext::ProcessPlatformAttachm
 			}
 			else
 			{
-				GenerationHelper = PackageData->GetOrFindParentGenerationHelper();
+				GenerationHelper = Vertex->IsOwnedByCluster() ? PackageData->GetOrFindParentGenerationHelper()
+					: PackageData->GetOrFindParentGenerationHelperNoCache();
 			}
 			if (PlatformAgnosticQueryData.CookAttachments.Artifacts.HasKeyMatch(nullptr, GenerationHelper.GetReference()))
 			{
@@ -2582,7 +2616,7 @@ bool FRequestCluster::IsIncrementalCook() const
 	return bAllowIterativeResults && COTFS.bHybridIterativeEnabled;
 }
 
-void FRequestCluster::IsRequestCookable(const ITargetPlatform* Platform, FPackageData& PackageData,
+void FRequestCluster::IsRequestCookable(const ITargetPlatform* Platform, const FPackageData& PackageData,
 	UCookOnTheFlyServer& COTFS, ESuppressCookReason& OutReason, bool& bOutCookable, bool& bOutExplorable)
 {
 	FString LocalDLCPath;
@@ -2595,14 +2629,14 @@ void FRequestCluster::IsRequestCookable(const ITargetPlatform* Platform, FPackag
 	IsRequestCookable(Platform, PackageData, COTFS, LocalDLCPath, OutReason, bOutCookable, bOutExplorable);
 }
 
-void FRequestCluster::IsRequestCookable(const ITargetPlatform* Platform, FPackageData& PackageData,
+void FRequestCluster::IsRequestCookable(const ITargetPlatform* Platform, const FPackageData& PackageData,
 	ESuppressCookReason& OutReason, bool& bOutCookable, bool& bOutExplorable)
 {
 	return IsRequestCookable(Platform, PackageData, COTFS,
 		DLCPath, OutReason, bOutCookable, bOutExplorable);
 }
 
-void FRequestCluster::IsRequestCookable(const ITargetPlatform* Platform, FPackageData& PackageData,
+void FRequestCluster::IsRequestCookable(const ITargetPlatform* Platform, const FPackageData& PackageData,
 	UCookOnTheFlyServer& InCOTFS, FStringView InDLCPath, ESuppressCookReason& OutReason, bool& bOutCookable,
 	bool& bOutExplorable)
 {
@@ -2623,7 +2657,7 @@ void FRequestCluster::IsRequestCookable(const ITargetPlatform* Platform, FPackag
 		return;
 	}
 
-	FPackagePlatformData* PlatformData = PackageData.FindPlatformData(Platform);
+	const FPackagePlatformData* PlatformData = PackageData.FindPlatformData(Platform);
 	bool bExplorableOverride = PlatformData ? PlatformData->IsExplorableOverride() : false;
 	ON_SCOPE_EXIT
 	{
diff --git a/Engine/Source/Editor/UnrealEd/Private/Cooker/CookRequestCluster.h b/Engine/Source/Editor/UnrealEd/Private/Cooker/CookRequestCluster.h
--- a/Engine/Source/Editor/UnrealEd/Private/Cooker/CookRequestCluster.h
+++ b/Engine/Source/Editor/UnrealEd/Private/Cooker/CookRequestCluster.h
@@ -68,6 +68,9 @@ public:
 	 */
 	void Process(const FCookerTimer& CookerTimer, bool& bOutComplete);
 
+	/** Return whether the cluster found work to do after construction and needs to be processed. */
+	bool NeedsProcessing() const;
+
 	/** PackageData container interface: return the number of PackageDatas owned by this container. */
 	int32 NumPackageDatas() const;
 	/** PackageData container interface: remove the PackageData from this container. */
@@ -96,7 +99,7 @@ public:
 
 	static TConstArrayView<FName> GetLocalizationReferences(FName PackageName, UCookOnTheFlyServer& InCOTFS);
 	static TArray<FName> GetAssetManagerReferences(FName PackageName);
-	static void IsRequestCookable(const ITargetPlatform* TargetPlatform, FPackageData& PackageData,
+	static void IsRequestCookable(const ITargetPlatform* TargetPlatform, const FPackageData& PackageData,
 		UCookOnTheFlyServer& COTFS, ESuppressCookReason& OutReason, bool& bOutCookable, bool& bOutExplorable);
 
 private:
@@ -182,7 +185,7 @@ private:
 		 */
 		bool IsOwnedByCluster() const;
 		void SetOwnedByCluster(bool bOwned);
-		
+
 		/**
 		 * Whether the package has already been pulled into the cluster once. If there is access from multiple
 		 * clusters this can be true even if GetOwnedByCluster is false.
@@ -361,6 +364,7 @@ private:
 		{
 			TBitArray<> HasRuntimePlatformByIndex;
 			TBitArray<> HasBuildPlatformByIndex;
+			TBitArray<> ForceExplorableByIndex;
 			EInstigator InstigatorType = EInstigator::InvalidCategory;
 			EInstigator BuildInstigatorType = EInstigator::InvalidCategory;
 		};
@@ -389,6 +393,7 @@ private:
 			void ProcessPlatformAttachments(int32 PlatformIndex, const ITargetPlatform* TargetPlatform,
 				FFetchPlatformData& FetchPlatformData, FPackagePlatformData& PackagePlatformData,
 				FIncrementalCookAttachments& PlatformAttachments, bool bExploreDependencies);
+			void ProcessPlatformDiscoveredDependencies(int32 PlatformIndex, const ITargetPlatform* TargetPlatform);
 
 			void SetIsIterativelyUnmodified(int32 PlatformIndex, bool bIterativelyUnmodified,
 				FPackagePlatformData& PackagePlatformData);
@@ -496,9 +501,9 @@ private:
 	void PumpExploration(const FCookerTimer& CookerTimer, bool& bOutComplete);
 	void StartAsync(const FCookerTimer& CookerTimer, bool& bOutComplete);
 	bool IsIncrementalCook() const;
-	void IsRequestCookable(const ITargetPlatform* TargetPlatform, FPackageData& PackageData,
+	void IsRequestCookable(const ITargetPlatform* TargetPlatform, const FPackageData& PackageData,
 		ESuppressCookReason& OutReason, bool& bOutCookable, bool& bOutExplorable);
-	static void IsRequestCookable(const ITargetPlatform* TargetPlatform, FPackageData& PackageData,
+	static void IsRequestCookable(const ITargetPlatform* TargetPlatform, const FPackageData& PackageData,
 		UCookOnTheFlyServer& InCOTFS, FStringView InDLCPath, ESuppressCookReason& OutReason, bool& bOutCookable,
 		bool& bOutExplorable);
 	/**
@@ -680,6 +685,11 @@ inline bool FRequestCluster::FQueryPlatformData::CompareExchangeAsyncQueryStatus
 	);
 }
 
+inline bool FRequestCluster::NeedsProcessing() const
+{
+	return !ClusterPackages.IsEmpty() || !FilePlatformRequests.IsEmpty();
+}
+
 inline int32 FRequestCluster::NumPackageDatas() const
 {
 	return NumOwned;
diff --git a/Engine/Source/Editor/UnrealEd/Private/Cooker/CookSavePackage.cpp b/Engine/Source/Editor/UnrealEd/Private/Cooker/CookSavePackage.cpp
--- a/Engine/Source/Editor/UnrealEd/Private/Cooker/CookSavePackage.cpp
+++ b/Engine/Source/Editor/UnrealEd/Private/Cooker/CookSavePackage.cpp
@@ -461,6 +461,12 @@ void FSaveCookedPackageContext::FinishPlatform()
 	// Commit the saved bytes and the iterative cook data to the PackageWriter
 	if (bPlatformSetupSuccessful)
 	{
+		// ProcessUnsolicitedPackages so that we record discovereddependencies for any hidden dependency
+		// packages loaded during Save,BeginCacheForCookedPlatformData, or Generator functions. These dependencies
+		// are added to the runtime dependencies stored in the oplog, so we need to know them now.
+		COTFS.ProcessUnsolicitedPackages();
+
+		// Collect dependencies from all sources and record them.
 		CalculatePlatformAgnosticRuntimeDependencies();
 		CalculatePlatformRuntimeDependencies();
 		// Note CalculateCookDependencies mutates this->SaveResult; it moves BuildResultDependencies out of it
@@ -612,7 +618,6 @@ void FSaveCookedPackageContext::FinishPackage()
 	{
 		if (COTFS.IsDebugRecordUnsolicited())
 		{
-			COTFS.ProcessUnsolicitedPackages();
 			TMap<FPackageData*, EInstigator> AllPlatformDependenciesBuffer;
 			TMap<FPackageData*, EInstigator>* AllPlatformDependencies = &PlatformDependencies[0].RuntimeDependencies;
 			if (PlatformsForPackage.Num() > 1)
@@ -624,7 +629,8 @@ void FSaveCookedPackageContext::FinishPackage()
 				}
 			}
 
-			FDiagnostics::AnalyzeHiddenDependencies(COTFS, PackageData, PackageData.DetachUnsolicited(),
+			FDiagnostics::AnalyzeHiddenDependencies(COTFS, PackageData,
+				PackageData.GetDiscoveredDependencies(nullptr /* PlatformAgnostic TargetPlatform*/),
 				*AllPlatformDependencies, PlatformsForPackage, COTFS.bOnlyEditorOnlyDebug, COTFS.bHiddenDependenciesDebug);
 		}
 
@@ -702,10 +708,21 @@ void FSaveCookedPackageContext::CalculatePlatformAgnosticRuntimeDependencies()
 
 			if (COTFS.IsDebugRecordUnsolicited())
 			{
-				PackageData.CreateOrGetUnsolicited().Add(SoftObjectPackageData, EInstigator::SoftDependency);
+				PackageData.AddDiscoveredDependency(EDiscoveredPlatformSet::CopyFromInstigator,	SoftObjectPackageData,
+					EInstigator::Unspecified);
 			}
 		}
 	}
+
+	// Add discovered dependencies
+	TMap<FPackageData*, EInstigator>* DiscoveredDependencies = PackageData.GetDiscoveredDependencies(nullptr /* PlatformAgnostic */);
+	if (DiscoveredDependencies)
+	{
+		for (TPair<FPackageData*, EInstigator>& Pair : *DiscoveredDependencies)
+		{
+			AddDependency(PlatformAgnosticDependencies, Pair.Key, false /* bHard */);
+		}
+	}
 }
 
 void FSaveCookedPackageContext::CalculatePlatformRuntimeDependencies()
@@ -734,6 +751,16 @@ void FSaveCookedPackageContext::CalculatePlatformRuntimeDependencies()
 			}
 		}
 	}
+
+	// Add discovered dependencies
+	TMap<FPackageData*, EInstigator>* DiscoveredDependencies = PackageData.GetDiscoveredDependencies(TargetPlatform);
+	if (DiscoveredDependencies)
+	{
+		for (TPair<FPackageData*, EInstigator>& Pair : *DiscoveredDependencies)
+		{
+			AddDependency(CurrentDependencies, Pair.Key, false /* bHard */);
+		}
+	}
 }
 
 TArray<FName> FSaveCookedPackageContext::GetPlatformRuntimeDependencies() const
diff --git a/Engine/Source/Editor/UnrealEd/Private/Cooker/CookTypes.cpp b/Engine/Source/Editor/UnrealEd/Private/Cooker/CookTypes.cpp
--- a/Engine/Source/Editor/UnrealEd/Private/Cooker/CookTypes.cpp
+++ b/Engine/Source/Editor/UnrealEd/Private/Cooker/CookTypes.cpp
@@ -582,7 +582,7 @@ void FDiscoveredPlatformSet::ConvertToBitfield(TConstArrayView<const ITargetPlat
 
 TConstArrayView<const ITargetPlatform*> FDiscoveredPlatformSet::GetPlatforms(UCookOnTheFlyServer& COTFS,
 	FInstigator* Instigator, TConstArrayView<const ITargetPlatform*> OrderedPlatforms, EReachability Reachability,
-	TArray<const ITargetPlatform*, TInlineAllocator<ExpectedMaxNumPlatforms>>* OutBuffer)
+	TArray<const ITargetPlatform*, TInlineAllocator<ExpectedMaxNumPlatforms>>& OutBuffer) const
 {
 	switch (Source)
 	{
@@ -590,27 +590,27 @@ TConstArrayView<const ITargetPlatform*> FDiscoveredPlatformSet::GetPlatforms(UCo
 		return Platforms;
 	case EDiscoveredPlatformSet::EmbeddedBitField:
 	{
-		check(OutBuffer);
-		OutBuffer->Reset();
+		OutBuffer.Reset();
 		int32 NumPlatforms = OrderedPlatformBits.Num();
 		check(NumPlatforms == OrderedPlatforms.Num());
 		for (int32 Index = 0; Index < NumPlatforms; ++Index)
 		{
 			if (OrderedPlatformBits[Index])
 			{
-				OutBuffer->Add(OrderedPlatforms[Index]);
+				OutBuffer.Add(OrderedPlatforms[Index]);
 			}
 		}
-		return *OutBuffer;
+		return OutBuffer;
 	}
 	case EDiscoveredPlatformSet::CopyFromInstigator:
-		Platforms.Reset();
+		OutBuffer.Reset();
 		check(Instigator);
-		FPackageData::GetReachablePlatformsForInstigator(Reachability, COTFS, Instigator->Referencer, Platforms);
-		return Platforms;
+		FPackageData::GetReachablePlatformsForInstigator(Reachability, COTFS, Instigator->Referencer, OutBuffer);
+		return OutBuffer;
 	default:
 		checkNoEntry();
-		return Platforms;
+		OutBuffer.Reset();
+		return OutBuffer;
 	}
 }
 
diff --git a/Engine/Source/Editor/UnrealEd/Private/Cooker/CookTypes.h b/Engine/Source/Editor/UnrealEd/Private/Cooker/CookTypes.h
--- a/Engine/Source/Editor/UnrealEd/Private/Cooker/CookTypes.h
+++ b/Engine/Source/Editor/UnrealEd/Private/Cooker/CookTypes.h
@@ -643,7 +643,7 @@ struct FDiscoveredPlatformSet
 	void OnPlatformAddedToSession(const ITargetPlatform* Platform);
 	TConstArrayView<const ITargetPlatform*> GetPlatforms(UCookOnTheFlyServer& COTFS,
 		FInstigator* Instigator, TConstArrayView<const ITargetPlatform*> OrderedPlatforms, EReachability Reachability,
-		TArray<const ITargetPlatform*, TInlineAllocator<ExpectedMaxNumPlatforms>>* OutBuffer);
+		TArray<const ITargetPlatform*, TInlineAllocator<ExpectedMaxNumPlatforms>>& OutBuffer) const;
 	/** If the current type is EmbeddedBitField, change it to EmbeddedList. */
 	void ConvertFromBitfield(TConstArrayView<const ITargetPlatform*> OrderedPlatforms);
 	/**
diff --git a/Engine/Source/Editor/UnrealEd/Private/Cooker/CookWorkerClient.cpp b/Engine/Source/Editor/UnrealEd/Private/Cooker/CookWorkerClient.cpp
--- a/Engine/Source/Editor/UnrealEd/Private/Cooker/CookWorkerClient.cpp
+++ b/Engine/Source/Editor/UnrealEd/Private/Cooker/CookWorkerClient.cpp
@@ -977,7 +977,7 @@ void FCookWorkerClient::AssignPackages(FAssignPackagesMessage& Message)
 			EReachability Reachability = AssignData.Reachability;
 			TConstArrayView<const ITargetPlatform*> NeedCommitPlatforms =
 				AssignData.NeedCommitPlatforms.GetPlatforms(COTFS, nullptr, OrderedSessionPlatforms,
-					Reachability, &NeedCommitPlatformsBuffer);
+					Reachability, NeedCommitPlatformsBuffer);
 			if (PackageData.IsInProgress())
 			{
 				// If already in progress but there are new platforms requested, demote the package back to Load
diff --git a/Engine/Source/Editor/UnrealEd/Private/Cooker/CookWorkerServer.cpp b/Engine/Source/Editor/UnrealEd/Private/Cooker/CookWorkerServer.cpp
--- a/Engine/Source/Editor/UnrealEd/Private/Cooker/CookWorkerServer.cpp
+++ b/Engine/Source/Editor/UnrealEd/Private/Cooker/CookWorkerServer.cpp
@@ -1269,7 +1269,7 @@ void FCookWorkerServer::QueueDiscoveredPackage(FDiscoveredPackageReplication&& D
 	else
 	{
 		DiscoveredPlatforms = Platforms.GetPlatforms(COTFS, &Instigator, OrderedSessionAndSpecialPlatforms,
-			DiscoveredReachability, &BufferPlatforms);
+			DiscoveredReachability, BufferPlatforms);
 	}
 
 	if (Instigator.Category != EInstigator::ForceExplorableSaveTimeSoftDependency &&
