Cooker: Fix clang compile error.
#rnx
#incrementalcook

diff --git a/Engine/Source/Editor/UnrealEd/Private/Cooker/CookDiagnostics.cpp b/Engine/Source/Editor/UnrealEd/Private/Cooker/CookDiagnostics.cpp
--- a/Engine/Source/Editor/UnrealEd/Private/Cooker/CookDiagnostics.cpp
+++ b/Engine/Source/Editor/UnrealEd/Private/Cooker/CookDiagnostics.cpp
@@ -123,7 +123,7 @@ void FDiagnostics::AnalyzeHiddenDependencies(UCookOnTheFlyServer& COTFS, FPackag
 	{
 		return;
 	}
-	for (const TPair<const FPackageData*, EInstigator>& UnsolicitedPair : *DiscoveredDependencies)
+	for (const TPair<FPackageData*, EInstigator>& UnsolicitedPair : *DiscoveredDependencies)
 	{
 		const FPackageData* Unsolicited = UnsolicitedPair.Key;
 		EInstigator Instigator = UnsolicitedPair.Value;
