Fix MouseDeltaTracker tools reliant on screenspace delta to translation to perform correctly with LeftUpForward display coordinate system

Examples of this are in the level editor, ctrl-drag/shift-drag the widget left axis.

diff --git a/Engine/Source/Editor/UnrealEd/Private/MouseDeltaTracker.cpp b/Engine/Source/Editor/UnrealEd/Private/MouseDeltaTracker.cpp
--- a/Engine/Source/Editor/UnrealEd/Private/MouseDeltaTracker.cpp
+++ b/Engine/Source/Editor/UnrealEd/Private/MouseDeltaTracker.cpp
@@ -11,6 +11,7 @@
 #include "SnappingUtils.h"
 #include "UnrealWidget.h"
 #include "EditorDragTools/EditorDragToolBehaviorTarget.h"
+#include "Misc/AxisDisplayInfo.h"
 
 #define LOCTEXT_NAMESPACE "MouseDeltaTracker"
 
@@ -639,6 +640,11 @@ void FMouseDeltaTracker::ConvertMovementDeltaToDragRot(FSceneView* InView, FEdit
 	{
 		InViewportClient->ConvertMovementToDragRot( InOutDragDelta, OutDrag, OutRotation );
 	}
+
+	if (AxisDisplayInfo::GetAxisDisplayCoordinateSystem() == EAxisList::LeftUpForward)
+	{
+		OutDrag.Y = -OutDrag.Y;
+	}
 }
 
 /**
