[DiffUtils] udpated IdenticalHelper to support null values in property comparison.
This fixes crashes while adding/removing empty elements in arrays in the diff'ed assets while the DiffTool is active

diff --git a/Engine/Source/Editor/UnrealEd/Private/DiffUtils.cpp b/Engine/Source/Editor/UnrealEd/Private/DiffUtils.cpp
--- a/Engine/Source/Editor/UnrealEd/Private/DiffUtils.cpp
+++ b/Engine/Source/Editor/UnrealEd/Private/DiffUtils.cpp
@@ -821,6 +821,15 @@ static void IdenticalHelper(const FProperty* AProperty, const FProperty* BProper
 		return;
 	}
 
+	if (AValue == nullptr || BValue == nullptr)
+	{
+		if (AValue != BValue)
+		{
+			DifferingSubProperties.Push(DiffParameters.RootPath);
+		}
+		return;
+	}
+
 	if (DiffParameters.ShouldIgnorePropertyPredicate
 		&& DiffParameters.ShouldIgnorePropertyPredicate(*AProperty))
 	{
