Mark GameInstance as garbage when shutting down PIE and add any living GameInstance objects to the ref trace.

#tests Reproduced a known leak in PIE

diff --git a/Engine/Source/Editor/UnrealEd/Private/PlayLevel.cpp b/Engine/Source/Editor/UnrealEd/Private/PlayLevel.cpp
--- a/Engine/Source/Editor/UnrealEd/Private/PlayLevel.cpp
+++ b/Engine/Source/Editor/UnrealEd/Private/PlayLevel.cpp
@@ -472,6 +472,7 @@ void UEditorEngine::EndPlayMap()
 			Object->MarkAsGarbage();
 		};
 		ForEachObjectWithOuter(*It, MarkObjectPendingKill, true, RF_NoFlags, EInternalObjectFlags::Garbage);
+		It->MarkAsGarbage();
 	}
 
 	// Flush any render commands and released accessed UTextures and materials to give them a chance to be collected.
@@ -513,6 +514,14 @@ void UEditorEngine::EndPlayMap()
 		}
 	}
 
+	for (TObjectIterator<UGameInstance> It; It; ++It)
+	{
+		if (It->HasAnyInternalFlags(EInternalObjectFlags::Garbage))
+		{
+			LeakedObjectsSet.Add(*It);
+		}
+	}
+
 	TArray<UObject*> LeakedObjects = LeakedObjectsSet.Array();
 	TArray<FString> Paths;
 	if (LeakedObjects.Num() > 0)
