Change viewport toolbar perspective camera icon to just the camera.

diff --git a/Engine/Source/Editor/UnrealEd/Private/ViewportToolbar/UnrealEdViewportToolbar.cpp b/Engine/Source/Editor/UnrealEd/Private/ViewportToolbar/UnrealEdViewportToolbar.cpp
--- a/Engine/Source/Editor/UnrealEd/Private/ViewportToolbar/UnrealEdViewportToolbar.cpp
+++ b/Engine/Source/Editor/UnrealEd/Private/ViewportToolbar/UnrealEdViewportToolbar.cpp
@@ -2653,7 +2653,9 @@ FText GetCameraSubmenuLabelFromViewportType(const ELevelViewportType ViewportTyp
 
 FName GetCameraSubmenuIconFNameFromViewportType(const ELevelViewportType ViewportType)
 {
-	static FName PerspectiveIcon("EditorViewport.Perspective");
+	// Use the raw camera icon rather than the perspective icon
+	// so that in the default state, the camera menu is easily recognizable.
+	static FName PerspectiveIcon("ClassIcon.CameraComponent");
 	static FName TopIcon("EditorViewport.Top");
 	static FName LeftIcon("EditorViewport.Left");
 	static FName FrontIcon("EditorViewport.Front");
