Updated code to use the FPlatformSplash::SetSplashText overload that takes FText

#rnx

diff --git a/Engine/Source/Editor/UnrealEd/Private/FeedbackContextEditor.cpp b/Engine/Source/Editor/UnrealEd/Private/FeedbackContextEditor.cpp
--- a/Engine/Source/Editor/UnrealEd/Private/FeedbackContextEditor.cpp
+++ b/Engine/Source/Editor/UnrealEd/Private/FeedbackContextEditor.cpp
@@ -533,7 +533,7 @@ void FFeedbackContextEditor::StartSlowTask( const FText& Task, bool bShowCancelB
 				TickSlate(SlowTaskWindow.Pin());
 			}
 
-			FPlatformSplash::SetSplashText(SplashTextType::StartupProgress, *Task.ToString());
+			FPlatformSplash::SetSplashText(SplashTextType::StartupProgress, Task);
 		}
 	}
 }
@@ -662,11 +662,11 @@ void FFeedbackContextEditor::ProgressReported( const float TotalProgressInterp,
 			}
 
 			DisplayProgress = int(TotalProgressInterp * 100.f);
-			DisplayMessage = FText::FromString(FString::Printf(TEXT("%i%% - %s"), DisplayProgress, *NewDisplayMessage));
+			DisplayMessage = FText::Format(NSLOCTEXT("FeedbackContextProgress", "StartupProgress", "{0} - {1}"), FText::AsPercent(TotalProgressInterp), FText::AsCultureInvariant(NewDisplayMessage));
 		}
 
 		FPlatformSplash::SetProgress(DisplayProgress);
-		FPlatformSplash::SetSplashText(SplashTextType::StartupProgress, *DisplayMessage.ToString());
+		FPlatformSplash::SetSplashText(SplashTextType::StartupProgress, DisplayMessage);
 	}
 }
 
