World Partition - EDL: Add hook to allow overriding the EDL injection logic
* We are already always injecting in cook
* Added hook to always inject during HLOD builds too

diff --git a/Engine/Source/Editor/UnrealEd/Private/WorldPartition/WorldPartitionHLODsBuilder.cpp b/Engine/Source/Editor/UnrealEd/Private/WorldPartition/WorldPartitionHLODsBuilder.cpp
--- a/Engine/Source/Editor/UnrealEd/Private/WorldPartition/WorldPartitionHLODsBuilder.cpp
+++ b/Engine/Source/Editor/UnrealEd/Private/WorldPartition/WorldPartitionHLODsBuilder.cpp
@@ -23,6 +23,7 @@
 #include "DerivedDataCacheInterface.h"
 
 #include "WorldPartition/DataLayer/DataLayerInstance.h"
+#include "WorldPartition/DataLayer/ExternalDataLayerEngineSubsystem.h"
 #include "WorldPartition/ActorDescContainerInstance.h"
 #include "WorldPartition/WorldPartitionActorDescInstance.h"
 #include "WorldPartition/WorldPartition.h"
@@ -49,38 +50,49 @@ UWorldPartitionHLODsBuilder::UWorldPartitionHLODsBuilder(const FObjectInitialize
 	, BuilderIdx(INDEX_NONE)
 	, BuilderCount(INDEX_NONE)
 {
-	BuildOptions  = HasParam("SetupHLODs") ? EHLODBuildStep::HLOD_Setup : EHLODBuildStep::None;
-	BuildOptions |= HasParam("BuildHLODs") ? EHLODBuildStep::HLOD_Build : EHLODBuildStep::None;
-	BuildOptions |= HasParam("RebuildHLODs") ? EHLODBuildStep::HLOD_Build : EHLODBuildStep::None;
-	BuildOptions |= HasParam("DeleteHLODs") ? EHLODBuildStep::HLOD_Delete : EHLODBuildStep::None;
-	BuildOptions |= HasParam("FinalizeHLODs") ? EHLODBuildStep::HLOD_Finalize : EHLODBuildStep::None;
-	BuildOptions |= HasParam("DumpStats") ? EHLODBuildStep::HLOD_Stats : EHLODBuildStep::None;
-
-	bResumeBuild = GetParamValue("ResumeBuild=", ResumeBuildIndex);
-
-	bDistributedBuild = HasParam("DistributedBuild");
-	bForceBuild = HasParam("RebuildHLODs");
-	bReportOnly = HasParam("ReportOnly");
-
-	GetParamValue("BuildManifest=", BuildManifest);
-	GetParamValue("BuilderIdx=", BuilderIdx);
-	GetParamValue("BuilderCount=", BuilderCount);
-	GetParamValue("BuildHLODLayer=", HLODLayerToBuild);
-	GetParamValue("BuildSingleHLOD=", HLODActorToBuild);
-
-	if (!HLODActorToBuild.IsNone() || !HLODLayerToBuild.IsNone())
+	if (!IsTemplate())
 	{
-		BuildOptions |= EHLODBuildStep::HLOD_Build;
-		bForceBuild = bForceBuild || !HLODActorToBuild.IsNone();
-	}
-	
-	// Default behavior without any option is to setup + build
-	if (BuildOptions == EHLODBuildStep::None)
-	{
-		BuildOptions = EHLODBuildStep::HLOD_Setup | EHLODBuildStep::HLOD_Build;
+		BuildOptions = HasParam("SetupHLODs") ? EHLODBuildStep::HLOD_Setup : EHLODBuildStep::None;
+		BuildOptions |= HasParam("BuildHLODs") ? EHLODBuildStep::HLOD_Build : EHLODBuildStep::None;
+		BuildOptions |= HasParam("RebuildHLODs") ? EHLODBuildStep::HLOD_Build : EHLODBuildStep::None;
+		BuildOptions |= HasParam("DeleteHLODs") ? EHLODBuildStep::HLOD_Delete : EHLODBuildStep::None;
+		BuildOptions |= HasParam("FinalizeHLODs") ? EHLODBuildStep::HLOD_Finalize : EHLODBuildStep::None;
+		BuildOptions |= HasParam("DumpStats") ? EHLODBuildStep::HLOD_Stats : EHLODBuildStep::None;
+
+		bResumeBuild = GetParamValue("ResumeBuild=", ResumeBuildIndex);
+
+		bDistributedBuild = HasParam("DistributedBuild");
+		bForceBuild = HasParam("RebuildHLODs");
+		bReportOnly = HasParam("ReportOnly");
+
+		GetParamValue("BuildManifest=", BuildManifest);
+		GetParamValue("BuilderIdx=", BuilderIdx);
+		GetParamValue("BuilderCount=", BuilderCount);
+		GetParamValue("BuildHLODLayer=", HLODLayerToBuild);
+		GetParamValue("BuildSingleHLOD=", HLODActorToBuild);
+
+		if (!HLODActorToBuild.IsNone() || !HLODLayerToBuild.IsNone())
+		{
+			BuildOptions |= EHLODBuildStep::HLOD_Build;
+			bForceBuild = bForceBuild || !HLODActorToBuild.IsNone();
+		}
+
+		// Default behavior without any option is to setup + build
+		if (BuildOptions == EHLODBuildStep::None)
+		{
+			BuildOptions = EHLODBuildStep::HLOD_Setup | EHLODBuildStep::HLOD_Build;
+		}
+
+		UExternalDataLayerEngineSubsystem::Get().OnExternalDataLayerOverrideInjection.AddUObject(this, &UWorldPartitionHLODsBuilder::AllowExternalDataLayerInjection);
 	}
 }
 
+void UWorldPartitionHLODsBuilder::AllowExternalDataLayerInjection(const UWorld* InWorld, const UExternalDataLayerAsset* InExternalDataLayerAsset, bool& bOutAllowInjection)
+{
+	// Always allow EDL injections during HLOD builds
+	bOutAllowInjection = true;
+}
+
 bool UWorldPartitionHLODsBuilder::RequiresCommandletRendering() const
 {
 	// Commandlet requires rendering only for building HLODs
diff --git a/Engine/Source/Editor/UnrealEd/Public/WorldPartition/WorldPartitionHLODsBuilder.h b/Engine/Source/Editor/UnrealEd/Public/WorldPartition/WorldPartitionHLODsBuilder.h
--- a/Engine/Source/Editor/UnrealEd/Public/WorldPartition/WorldPartitionHLODsBuilder.h
+++ b/Engine/Source/Editor/UnrealEd/Public/WorldPartition/WorldPartitionHLODsBuilder.h
@@ -7,6 +7,8 @@
 #include "WorldPartition/WorldPartitionBuilderHelpers.h"
 #include "WorldPartitionHLODsBuilder.generated.h"
 
+class UExternalDataLayerAsset;
+
 enum class EHLODBuildStep : uint8
 {
 	None		= 0,
@@ -57,6 +59,8 @@ protected:
 
 	UNREALED_API bool AddBuildProducts(const TArray<FString>& BuildProducts) const;
 
+	void AllowExternalDataLayerInjection(const UWorld* InWorld, const UExternalDataLayerAsset* InExternalDataLayerAsset, bool& bOutAllowInjection);
+
 private:
 	UWorld* World;
 	class UWorldPartition* WorldPartition;
diff --git a/Engine/Source/Runtime/Engine/Private/WorldPartition/DataLayer/ExternalDataLayerEngineSubsystem.cpp b/Engine/Source/Runtime/Engine/Private/WorldPartition/DataLayer/ExternalDataLayerEngineSubsystem.cpp
--- a/Engine/Source/Runtime/Engine/Private/WorldPartition/DataLayer/ExternalDataLayerEngineSubsystem.cpp
+++ b/Engine/Source/Runtime/Engine/Private/WorldPartition/DataLayer/ExternalDataLayerEngineSubsystem.cpp
@@ -165,6 +165,17 @@ bool UExternalDataLayerEngineSubsystem::CanWorldInjectExternalDataLayerAsset(con
 #endif
 }
 
+bool UExternalDataLayerEngineSubsystem::CanInjectOverride(const UWorld* InWorld, const UExternalDataLayerAsset* InExternalDataLayerAsset, bool& bOutOverrideValue) const
+{
+	if (OnExternalDataLayerOverrideInjection.IsBound())
+	{
+		OnExternalDataLayerOverrideInjection.Broadcast(InWorld, InExternalDataLayerAsset, bOutOverrideValue);
+		return true;
+	}
+
+	return false;
+}
+
 void UExternalDataLayerEngineSubsystem::RegisterExternalDataLayerAsset(const UExternalDataLayerAsset* InExternalDataLayerAsset, const UObject* InClient)
 {
 	check(IsValid(InExternalDataLayerAsset));
diff --git a/Engine/Source/Runtime/Engine/Private/WorldPartition/DataLayer/ExternalDataLayerInjectionPolicy.cpp b/Engine/Source/Runtime/Engine/Private/WorldPartition/DataLayer/ExternalDataLayerInjectionPolicy.cpp
--- a/Engine/Source/Runtime/Engine/Private/WorldPartition/DataLayer/ExternalDataLayerInjectionPolicy.cpp
+++ b/Engine/Source/Runtime/Engine/Private/WorldPartition/DataLayer/ExternalDataLayerInjectionPolicy.cpp
@@ -9,7 +9,11 @@
 #if WITH_EDITOR
 bool UExternalDataLayerInjectionPolicy::CanInject(const UWorld* InWorld, const UExternalDataLayerAsset* InExternalDataLayerAsset, const UObject* InClient, FText* OutFailureReason) const
 {
-	if (IsRunningCookCommandlet())
+	bool bOverrideCanInject = false;
+	const bool bOverrideInjectionPolicy = UExternalDataLayerEngineSubsystem::Get().CanInjectOverride(InWorld, InExternalDataLayerAsset, bOverrideCanInject);
+	
+	// If running cook commandlet or forced by custom injection logic
+	if (IsRunningCookCommandlet() || (bOverrideInjectionPolicy && bOverrideCanInject))
 	{
 		if (!UExternalDataLayerEngineSubsystem::Get().IsExternalDataLayerAssetRegistered(InExternalDataLayerAsset, InClient))
 		{
@@ -20,6 +24,15 @@ bool UExternalDataLayerInjectionPolicy::CanInject(const UWorld* InWorld, const U
 			return false;
 		}
 	}
+	// If rejected by override
+	else if ((bOverrideInjectionPolicy && !bOverrideCanInject))
+	{
+		if (OutFailureReason)
+		{
+			*OutFailureReason = FText::Format(LOCTEXT("CantInjectOverride", "External Data Layer Asset {0} rejected by custom injection logic"), FText::FromString(InExternalDataLayerAsset->GetName()));
+		}
+		return false;
+	}
 	// Only allow injection if active
 	else if (!UExternalDataLayerEngineSubsystem::Get().IsExternalDataLayerAssetActive(InExternalDataLayerAsset, InClient))
 	{
@@ -29,6 +42,7 @@ bool UExternalDataLayerInjectionPolicy::CanInject(const UWorld* InWorld, const U
 		}
 		return false;
 	}
+
 	return true;
 }
 #endif
diff --git a/Engine/Source/Runtime/Engine/Public/WorldPartition/DataLayer/ExternalDataLayerEngineSubsystem.h b/Engine/Source/Runtime/Engine/Public/WorldPartition/DataLayer/ExternalDataLayerEngineSubsystem.h
--- a/Engine/Source/Runtime/Engine/Public/WorldPartition/DataLayer/ExternalDataLayerEngineSubsystem.h
+++ b/Engine/Source/Runtime/Engine/Public/WorldPartition/DataLayer/ExternalDataLayerEngineSubsystem.h
@@ -53,6 +53,10 @@ public:
 	DECLARE_MULTICAST_DELEGATE_ThreeParams(FExternalDataLayerAssetRegistrationStateChangedEventDelegate, const UExternalDataLayerAsset* ExternalDataLayerAsset, EExternalDataLayerRegistrationState OldState, EExternalDataLayerRegistrationState NewState);
 	FExternalDataLayerAssetRegistrationStateChangedEventDelegate OnExternalDataLayerAssetRegistrationStateChanged;
 
+	DECLARE_MULTICAST_DELEGATE_ThreeParams(FExternalDataLayerOverrideInjectionDelegate, const UWorld* InWorld, const UExternalDataLayerAsset* InExternalDataLayerAsset, bool& bOutOverrideValue);
+	FExternalDataLayerOverrideInjectionDelegate OnExternalDataLayerOverrideInjection;
+	ENGINE_API bool CanInjectOverride(const UWorld* InWorld, const UExternalDataLayerAsset* InExternalDataLayerAsset, bool& bOutOverrideValue) const;
+
 private:
 #if WITH_EDITOR
 	void OnAssetsPreDelete(const TArray<UObject*>& Objects);
