IncrementalCook - Content Missing from Cook warnings - save the import and export names from each package (in a compressed format) in the PackageArtifacts, so we can update them in the CookImportsChecker when incrementally skipping a package.
#rnx
#incrementalcook

diff --git a/Engine/Source/Editor/UnrealEd/Private/Cooker/CookImportsChecker.cpp b/Engine/Source/Editor/UnrealEd/Private/Cooker/CookImportsChecker.cpp
--- a/Engine/Source/Editor/UnrealEd/Private/Cooker/CookImportsChecker.cpp
+++ b/Engine/Source/Editor/UnrealEd/Private/Cooker/CookImportsChecker.cpp
@@ -23,35 +23,39 @@
 LLM_DEFINE_TAG(EDLCookChecker);
 
 FEDLCookChecker::FEDLNodeHash::FEDLNodeHash()
+	: ObjectTypeData { nullptr }
+	, NodeHashType(ENodeHashType::Object)
+	, ObjectEvent(EObjectEvent::Create)
 {
 }
 
 FEDLCookChecker::FEDLNodeHash::FEDLNodeHash(const FEDLNodeHash& Other)
+	: ObjectTypeData{ nullptr }
+	, NodeHashType(ENodeHashType::Object)
+	, ObjectEvent(EObjectEvent::Create)
 {
-	bIsNode = Other.bIsNode;
-	if (bIsNode)
-	{
-		Nodes = Other.Nodes;
-	}
-	else
-	{
-		Object = Other.Object;
-	}
-	NodeID = Other.NodeID;
-	ObjectEvent = Other.ObjectEvent;
+	*this = Other;
 }
 
-FEDLCookChecker::FEDLNodeHash::FEDLNodeHash(const TArray<FEDLNodeData>* InNodes, FEDLNodeID InNodeID, EObjectEvent InObjectEvent)
-	: Nodes(InNodes)
-	, NodeID(InNodeID)
-	, bIsNode(true)
+FEDLCookChecker::FEDLNodeHash::FEDLNodeHash(const TArray<FEDLNodeData>* InNodes, FEDLNodeID InNodeID,
+	EObjectEvent InObjectEvent)
+	: NodeTypeData { InNodes, InNodeID }
+	, NodeHashType(ENodeHashType::Node)
+	, ObjectEvent(InObjectEvent)
+{
+}
+
+FEDLCookChecker::FEDLNodeHash::FEDLNodeHash(const TArray<FEDLNodeData>* InNodes, FEDLNodeID InParentNodeID,
+	FName InObjectName, EObjectEvent InObjectEvent)
+	: NameTypeData{ InObjectName, InNodes, InParentNodeID }
+	, NodeHashType(ENodeHashType::NameAndParentNode)
 	, ObjectEvent(InObjectEvent)
 {
 }
 
 FEDLCookChecker::FEDLNodeHash::FEDLNodeHash(TObjectPtr<UObject> InObject, EObjectEvent InObjectEvent)
-	: Object(InObject)
-	, bIsNode(false)
+	: ObjectTypeData{ InObject }
+	, NodeHashType(ENodeHashType::Object)
 	, ObjectEvent(InObjectEvent)
 {
 }
@@ -63,37 +67,48 @@ bool FEDLCookChecker::FEDLNodeHash::operator==(const FEDLNodeHash& Other) const
 		return false;
 	}
 
-	uint32 LocalNodeID;
-	uint32 OtherNodeID;
-	TObjectPtr<const UObject> LocalObject;
-	TObjectPtr<const UObject> OtherObject;
-	FName LocalName = ObjectNameFirst(*this, LocalNodeID, LocalObject);
-	FName OtherName = ObjectNameFirst(Other, OtherNodeID, OtherObject);
+	if (GetName() != Other.GetName())
+	{
+		return false;
+	}
+	FEDLCookChecker::FEDLNodeHash CurrentThis;
+	FEDLCookChecker::FEDLNodeHash CurrentOther;
+	bool bCurrentThisValid = TryGetParent(CurrentThis);
+	bool bOtherThisValid = Other.TryGetParent(CurrentOther);
 
-	do
+	while (bCurrentThisValid && bOtherThisValid)
 	{
-		if (LocalName != OtherName)
+		if (CurrentThis.GetName() != CurrentOther.GetName())
 		{
 			return false;
 		}
-		LocalName = ObjectNameNext(*this, LocalNodeID, LocalObject);
-		OtherName = ObjectNameNext(Other, OtherNodeID, OtherObject);
-	} while (!LocalName.IsNone() && !OtherName.IsNone());
-	return LocalName.IsNone() == OtherName.IsNone();
+		bCurrentThisValid = CurrentThis.TryGetParent(CurrentThis);
+		bOtherThisValid = CurrentOther.TryGetParent(CurrentOther);
+	}
+	return bCurrentThisValid == bOtherThisValid;
 }
 
 FEDLCookChecker::FEDLNodeHash& FEDLCookChecker::FEDLNodeHash::operator=(const FEDLNodeHash& Other)
 {
-	bIsNode = Other.bIsNode;
-	if (bIsNode)
-	{
-		Nodes = Other.Nodes;
-	}
-	else
+	NodeHashType = Other.NodeHashType;
+	switch (NodeHashType)
 	{
-		Object = Other.Object;
+	case ENodeHashType::Node:
+		NodeTypeData.Nodes = Other.NodeTypeData.Nodes;
+		NodeTypeData.NodeID = Other.NodeTypeData.NodeID;
+		break;
+	case ENodeHashType::Object:
+		ObjectTypeData.Object = Other.ObjectTypeData.Object;
+		break;
+	case ENodeHashType::NameAndParentNode:
+		NameTypeData.ObjectName = Other.NameTypeData.ObjectName;
+		NameTypeData.Nodes = Other.NameTypeData.Nodes;
+		NameTypeData.ParentID = Other.NameTypeData.ParentID;
+		break;
+	default:
+		checkNoEntry();
+		break;
 	}
-	NodeID = Other.NodeID;
 	ObjectEvent = Other.ObjectEvent;
 
 	return *this;
@@ -101,54 +116,111 @@ FEDLCookChecker::FEDLNodeHash& FEDLCookChecker::FEDLNodeHash::operator=(const FE
 
 uint32 GetTypeHash(const FEDLCookChecker::FEDLNodeHash& A)
 {
-	uint32 Hash = 0;
+	return FEDLCookChecker::FEDLNodeHash::GetTypeHashInternal(A);
+}
 
-	uint32 LocalNodeID;
-	TObjectPtr<const UObject> LocalObject;
-	FName LocalName = FEDLCookChecker::FEDLNodeHash::ObjectNameFirst(A, LocalNodeID, LocalObject);
-	do
+uint32 FEDLCookChecker::FEDLNodeHash::GetTypeHashInternal(const FEDLCookChecker::FEDLNodeHash& A)
+{
+	auto GetTypeHashFromNodeOuterChain =
+		[](uint32 Hash, const TArray<FEDLCookChecker::FEDLNodeData>& Nodes, FEDLCookChecker::FEDLNodeID ParentNodeID,
+			FName ObjectName)
+		{
+			Hash = HashCombine(Hash, GetTypeHash(ObjectName));
+			while (ParentNodeID != NodeIDInvalid)
+			{
+				const FEDLCookChecker::FEDLNodeData& ParentNode = Nodes[ParentNodeID];
+				Hash = HashCombine(Hash, GetTypeHash(ParentNode.Name));
+				ParentNodeID = ParentNode.ParentID;
+			}
+			return Hash;
+		};
+
+	uint32 Hash = 0;
+	switch (A.NodeHashType)
 	{
-		Hash = HashCombine(Hash, GetTypeHash(LocalName));
-		LocalName = FEDLCookChecker::FEDLNodeHash::ObjectNameNext(A, LocalNodeID, LocalObject);
-	} while (!LocalName.IsNone());
+	case ENodeHashType::Node:
+	{
+		const TArray<FEDLCookChecker::FEDLNodeData>& Nodes = *A.NodeTypeData.Nodes;
+		const FEDLCookChecker::FEDLNodeData& Node = Nodes[A.NodeTypeData.NodeID];
+		Hash = GetTypeHashFromNodeOuterChain(Hash, Nodes, Node.ParentID, Node.Name);
+		break;
+	}
+	case ENodeHashType::Object:
+	{
+		TObjectPtr<const UObject> Object = A.ObjectTypeData.Object;
+		while (Object)
+		{
+			Hash = HashCombine(Hash, GetTypeHash(Object.GetFName()));
+			Object = Object.GetOuter();
+		}
+		break;
+	}
+	case ENodeHashType::NameAndParentNode:
+	{
+		const TArray<FEDLCookChecker::FEDLNodeData>& Nodes = *A.NameTypeData.Nodes;
+		Hash = GetTypeHashFromNodeOuterChain(Hash, Nodes, A.NameTypeData.ParentID, A.NameTypeData.ObjectName);
+		break;
+	}
+	default:
+		checkNoEntry();
+		break;
+	}
 
 	return (Hash << 1) | (uint32)A.ObjectEvent;
 }
 
 FName FEDLCookChecker::FEDLNodeHash::GetName() const
 {
-	if (bIsNode)
-	{
-		return (*Nodes)[NodeID].Name;
-	}
-	else
+	switch (NodeHashType)
 	{
-		return Object.GetFName();
+	case ENodeHashType::Node:
+		return (*NodeTypeData.Nodes)[NodeTypeData.NodeID].Name;
+	case ENodeHashType::Object:
+		return ObjectTypeData.Object.GetFName();
+	case ENodeHashType::NameAndParentNode:
+		return NameTypeData.ObjectName;
+	default:
+		checkNoEntry();
+		return NAME_None;
 	}
 }
 
 bool FEDLCookChecker::FEDLNodeHash::TryGetParent(FEDLCookChecker::FEDLNodeHash& Parent) const
 {
 	EObjectEvent ParentObjectEvent = EObjectEvent::Create; // For purposes of parents, which is used only to get the ObjectPath, we always use the Create version of the node as the parent
-	if (bIsNode)
+	switch (NodeHashType)
+	{
+	case ENodeHashType::Node:
 	{
-		FEDLNodeID ParentID = (*Nodes)[NodeID].ParentID;
+		FEDLNodeID ParentID = (*NodeTypeData.Nodes)[NodeTypeData.NodeID].ParentID;
 		if (ParentID != NodeIDInvalid)
 		{
-			Parent = FEDLNodeHash(Nodes, ParentID, ParentObjectEvent);
+			Parent = FEDLNodeHash(NodeTypeData.Nodes, ParentID, ParentObjectEvent);
 			return true;
 		}
+		return false;
 	}
-	else
+	case ENodeHashType::Object:
 	{
-		TObjectPtr<UObject> ParentObject = Object.GetOuter();
+		TObjectPtr<UObject> ParentObject = ObjectTypeData.Object.GetOuter();
 		if (ParentObject)
 		{
 			Parent = FEDLNodeHash(ParentObject, ParentObjectEvent);
 			return true;
 		}
+		return false;
+	}
+	case ENodeHashType::NameAndParentNode:
+		if (NameTypeData.ParentID != NodeIDInvalid)
+		{
+			Parent = FEDLNodeHash(NameTypeData.Nodes, NameTypeData.ParentID, ParentObjectEvent);
+			return true;
+		}
+		return false;
+	default:
+		checkNoEntry();
+		return false;
 	}
-	return false;
 }
 
 FEDLCookChecker::EObjectEvent FEDLCookChecker::FEDLNodeHash::GetObjectEvent() const
@@ -158,37 +230,19 @@ FEDLCookChecker::EObjectEvent FEDLCookChecker::FEDLNodeHash::GetObjectEvent() co
 
 void FEDLCookChecker::FEDLNodeHash::SetNodes(const TArray<FEDLNodeData>* InNodes)
 {
-	if (bIsNode)
-	{
-		Nodes = InNodes;
-	}
-}
-
-FName FEDLCookChecker::FEDLNodeHash::ObjectNameFirst(const FEDLNodeHash& InNode, uint32& OutNodeID, TObjectPtr<const UObject>& OutObject)
-{
-	if (InNode.bIsNode)
-	{
-		OutNodeID = InNode.NodeID;
-		return (*InNode.Nodes)[OutNodeID].Name;
-	}
-	else
-	{
-		OutObject = InNode.Object;
-		return OutObject.GetFName();
-	}
-}
-
-FName FEDLCookChecker::FEDLNodeHash::ObjectNameNext(const FEDLNodeHash& InNode, uint32& OutNodeID, TObjectPtr<const UObject>& OutObject)
-{
-	if (InNode.bIsNode)
-	{
-		OutNodeID = (*InNode.Nodes)[OutNodeID].ParentID;
-		return OutNodeID != NodeIDInvalid ? (*InNode.Nodes)[OutNodeID].Name : NAME_None;
-	}
-	else
+	switch (NodeHashType)
 	{
-		OutObject = OutObject.GetOuter();
-		return OutObject ? OutObject.GetFName() : NAME_None;
+	case ENodeHashType::Node:
+		NodeTypeData.Nodes = InNodes;
+		break;
+	case ENodeHashType::Object:
+		break;
+	case ENodeHashType::NameAndParentNode:
+		NameTypeData.Nodes = InNodes;
+		break;
+	default:
+		checkNoEntry();
+		break;
 	}
 }
 
@@ -300,18 +354,80 @@ void FEDLCookChecker::AddImport(TObjectPtr<UObject> Import, UPackage* ImportingP
 		{
 			LLM_SCOPE_BYTAG(EDLCookChecker);
 			FEDLNodeID NodeId = FindOrAddNode(FEDLNodeHash(Import, EObjectEvent::Serialize));
-			FEDLNodeData& NodeData = Nodes[NodeId];
-			FName ImportingPackageName = ImportingPackage->GetFName();
-			TArray<FName>& Sorted = NodeData.ImportingPackagesSorted;
-			int32 InsertionIndex = Algo::LowerBound(Sorted, ImportingPackageName, FNameFastLess());
-			if (InsertionIndex == Sorted.Num() || Sorted[InsertionIndex] != ImportingPackageName)
+			RecordImportFromPackage(NodeId, ImportingPackage->GetFName());
+		}
+	}
+}
+
+void FEDLCookChecker::RecordImportFromPackage(FEDLNodeID NodeId, FName ImportingPackageName)
+{
+	FEDLNodeData& NodeData = Nodes[NodeId];
+	TArray<FName>& Sorted = NodeData.ImportingPackagesSorted;
+	int32 InsertionIndex = Algo::LowerBound(Sorted, ImportingPackageName, FNameFastLess());
+	if (InsertionIndex == Sorted.Num() || Sorted[InsertionIndex] != ImportingPackageName)
+	{
+		Sorted.Insert(ImportingPackageName, InsertionIndex);
+	}
+}
+
+template <typename AddNodeType>
+void FEDLCookChecker::AddImportExportNodeList(TConstArrayView<UE::Cook::FImportExportNode> NodeList, AddNodeType&& AddNode)
+{
+	// See the comment in FImportsCheckerData::ObjectListToNodeList, this is the same algorithm. Recursively
+	// calculate and cache EDLNodes for parent nodes.
+	TArray<int32, TInlineAllocator<10>> Stack;
+	TArray<FEDLNodeID> NodeIDForExternalIndex;
+	NodeIDForExternalIndex.SetNumUninitialized(NodeList.Num());
+	for (FEDLNodeID& NodeID : NodeIDForExternalIndex)
+	{
+		NodeID = NodeIDInvalid;
+	}
+	for (int32 ExternalIndex = 0; ExternalIndex < NodeList.Num(); ++ExternalIndex)
+	{
+		check(Stack.IsEmpty());
+		FEDLNodeID ParentNodeID = NodeIDInvalid;
+		int32 CurrentExternalIndex = ExternalIndex;
+		while (CurrentExternalIndex != -1)
+		{
+			FEDLNodeID& CurrentNodeID = NodeIDForExternalIndex[CurrentExternalIndex];
+			if (CurrentNodeID == NodeIDInvalid)
 			{
-				Sorted.Insert(ImportingPackageName, InsertionIndex);
+				const UE::Cook::FImportExportNode& Node = NodeList[CurrentExternalIndex];
+				if (Node.ParentId != -1 && ParentNodeID == NodeIDInvalid)
+				{
+					Stack.Push(CurrentExternalIndex);
+					CurrentExternalIndex = Node.ParentId;
+					continue;
+				}
+
+				CurrentNodeID = FindOrAddNode(FEDLNodeHash(&Nodes, ParentNodeID,
+					Node.ObjectName, EObjectEvent::Serialize));
+				AddNode(Node, CurrentNodeID, ParentNodeID);
 			}
+
+			check(CurrentNodeID != NodeIDInvalid);
+			ParentNodeID = CurrentNodeID;
+			CurrentExternalIndex = !Stack.IsEmpty() ? Stack.Pop(EAllowShrinking::No) : -1;
 		}
 	}
 }
 
+void FEDLCookChecker::AddImports(TConstArrayView<UE::Cook::FImportExportNode> Imports, FName ImportingPackageName)
+{
+	if (!bIsActive)
+	{
+		return;
+	}
+	LLM_SCOPE_BYTAG(EDLCookChecker);
+
+	AddImportExportNodeList(Imports,
+		[ImportingPackageName, this]
+		(const UE::Cook::FImportExportNode& ImportNode, FEDLNodeID NodeID, FEDLNodeID ParentNodeID)
+		{
+			RecordImportFromPackage(NodeID, ImportingPackageName);
+		});
+}
+
 void FEDLCookChecker::AddExport(UObject* Export)
 {
 	if (bIsActive)
@@ -321,10 +437,43 @@ void FEDLCookChecker::AddExport(UObject* Export)
 		Nodes[SerializeID].bIsExport = true;
 		FEDLNodeID CreateID = FindOrAddNode(FEDLNodeHash(Export, EObjectEvent::Create));
 		Nodes[CreateID].bIsExport = true;
-		AddDependency(SerializeID, CreateID); // every export must be created before it can be serialized...these arcs are implicit and not listed in any table.
+
+		// every export must be created before it can be serialized...
+		// these arcs are implicit and not listed in any table.
+		AddDependency(SerializeID, CreateID); 
 	}
 }
 
+void FEDLCookChecker::AddExports(TConstArrayView<UE::Cook::FImportExportNode> Exports)
+{
+	if (!bIsActive)
+	{
+		return;
+	}
+	LLM_SCOPE_BYTAG(EDLCookChecker);
+
+	AddImportExportNodeList(Exports,
+		[this](const UE::Cook::FImportExportNode& ExportNode, FEDLNodeID SerializeID, FEDLNodeID ParentNodeID)
+		{
+			// AddImportExportNodeList added the Serialize node for us, we also need to add the Create node
+			FEDLNodeID CreateID = FindOrAddNode(FEDLNodeHash(&Nodes, ParentNodeID,
+				ExportNode.ObjectName, EObjectEvent::Create));
+
+			Nodes[SerializeID].bIsExport = true;
+			Nodes[CreateID].bIsExport = true;
+
+			// every export must be created before it can be serialized...
+			// these arcs are implicit and not listed in any table.
+			AddDependency(SerializeID, CreateID);
+		});
+}
+
+void FEDLCookChecker::Add(UE::Cook::FImportsCheckerData& ImportsCheckerData, FName PackageName)
+{
+	AddImports(ImportsCheckerData.Imports, PackageName);
+	AddExports(ImportsCheckerData.Exports);
+}
+
 void FEDLCookChecker::AddArc(UObject* DepObject, bool bDepIsSerialize, UObject* Export, bool bExportIsSerialize)
 {
 	if (bIsActive)
@@ -802,3 +951,141 @@ bool FEDLCookChecker::ReadFromCompactBinary(FCbFieldView Field)
 
 FCriticalSection FEDLCookChecker::CookCheckerInstanceCritical;
 TArray<FEDLCookChecker*> FEDLCookChecker::CookCheckerInstances;
+
+namespace UE::Cook
+{
+
+void FImportExportNode::Save(FCbWriter& Writer) const
+{
+	Writer.BeginArray();
+	Writer << ObjectName;
+	Writer << ParentId;
+	Writer.EndArray();
+}
+
+bool FImportExportNode::TryLoad(const FCbFieldView& Field)
+{
+	FCbFieldViewIterator ElementView(Field.CreateViewIterator());
+	if (!LoadFromCompactBinary(ElementView++, ObjectName))
+	{
+		return false;
+	}
+	if (!LoadFromCompactBinary(ElementView++, ParentId))
+	{
+		return false;
+	}
+	return true;
+}
+
+void FImportsCheckerData::Save(FCbWriter& Writer) const
+{
+	Writer.BeginObject();
+	Writer << "Imports" << Imports;
+	Writer << "Exports" << Exports;
+	Writer.EndObject();
+}
+
+bool FImportsCheckerData::TryLoad(const FCbFieldView& Field)
+{
+	bool bImports = false;
+	bool bExports = false;
+	for (FCbFieldViewIterator ElementView(Field.CreateViewIterator()); ElementView; )
+	{
+		const FCbFieldViewIterator Last = ElementView;
+		if (ElementView.GetName().Equals(UTF8TEXTVIEW("Imports")))
+		{
+			bImports = true;
+			if (!LoadFromCompactBinary(ElementView++, Imports))
+			{
+				return false;
+			}
+		}
+		if (ElementView.GetName().Equals(UTF8TEXTVIEW("Exports")))
+		{
+			bExports = true;
+			if (!LoadFromCompactBinary(ElementView++, Exports))
+			{
+				return false;
+			}
+		}
+		if (ElementView == Last)
+		{
+			++ElementView;
+		}
+	}
+	return bImports && bExports;
+}
+
+FImportsCheckerData FImportsCheckerData::FromObjectLists(TConstArrayView<UObject*> Imports,
+	TConstArrayView<UObject*> Exports)
+{
+	FImportsCheckerData Result;
+	TArray<UObject*> FilteredImports;
+	FilteredImports.Reserve(Imports.Num());
+	for (UObject* Import : Imports)
+	{
+		if (!Import->GetOutermost()->HasAnyPackageFlags(PKG_CompiledIn))
+		{
+			FilteredImports.Add(Import);
+		}
+	}
+	Result.Imports = ObjectListToNodeList(FilteredImports);
+	Result.Exports = ObjectListToNodeList(Exports);
+	return Result;
+}
+
+TArray<FImportExportNode> FImportsCheckerData::ObjectListToNodeList(TConstArrayView<UObject*> Objects)
+{
+	TArray<FImportExportNode> Result;
+
+	// Iterate each LeafObject in the array of Objects, and walk up the outer chain of each Object recursively adding
+	// a node for each outer in the outer chain, and then add a node for the object as a child of the outer's node. If
+	// any Outer (or even the LeafObject itself) has already been given a node, use the index of that node from the map
+	// and stop walking up the stack. When walking up the outer chain we keep a stack of objects we are working on
+	// beneath the current object in the outer chain, and when we reach the outermost or an already handled node, we
+	// keep a parentindex variable which we set in previous loop iteration on the outer and use that as the recursive
+	// result for the outer.
+	TArray<UObject*, TInlineAllocator<10>> Stack;
+	TMap<UObject*, int32> Map;
+	Map.Reserve(Objects.Num());
+	for (UObject* LeafObject : Objects)
+	{
+		check(Stack.IsEmpty());
+		int32 ParentIndex = -1;
+		UObject* Current = LeafObject;
+		while (Current)
+		{
+			int32 CurrentIndex = -1;
+			int32* CurrentIndexPtr = Map.Find(Current);
+			if (CurrentIndexPtr)
+			{
+				CurrentIndex = *CurrentIndexPtr;
+			}
+			else
+			{
+				UObject* Outer = Current->GetOuter();
+				if (Outer && ParentIndex == -1)
+				{
+					Stack.Push(Current);
+					Current = Outer;
+					continue;
+				}
+
+				CurrentIndex = Result.Num();
+				FImportExportNode& CurrentNode = Result.Emplace_GetRef();
+				check(Current); // Workaround StaticAnalyzer bug: spurious warning C28182: Dereferencing NULL pointer
+				CurrentNode.ObjectName = Current->GetFName();
+				CurrentNode.ParentId = ParentIndex;
+				Map.Add(Current, CurrentIndex);
+			}
+
+			check(0 <= CurrentIndex && CurrentIndex < Result.Num());
+			ParentIndex = CurrentIndex;
+			Current = !Stack.IsEmpty() ? Stack.Pop(EAllowShrinking::No) : nullptr;
+		}
+	}
+
+	return Result;
+}
+
+}
diff --git a/Engine/Source/Editor/UnrealEd/Private/Cooker/CookImportsChecker.h b/Engine/Source/Editor/UnrealEd/Private/Cooker/CookImportsChecker.h
--- a/Engine/Source/Editor/UnrealEd/Private/Cooker/CookImportsChecker.h
+++ b/Engine/Source/Editor/UnrealEd/Private/Cooker/CookImportsChecker.h
@@ -18,6 +18,66 @@ namespace UE { class FLogRecord; }
 
 struct FEDLCookCheckerThreadState;
 
+namespace UE::Cook
+{
+
+/**
+ * Data describing an object's name in its outer hierarchy. Used so that we can persistently represent all the names in
+ * a tree of object outers in a list that does not duplicate the string data between an object and
+ * its outer.
+ */
+struct FImportExportNode
+{
+public:
+	FName ObjectName;
+	int32 ParentId = -1;
+
+	void Save(FCbWriter& Writer) const;
+	bool TryLoad(const FCbFieldView& Field);
+private:
+	friend FCbWriter& operator<<(FCbWriter& Writer, const FImportExportNode& Node)
+	{
+		Node.Save(Writer);
+		return Writer;
+	}
+	friend bool LoadFromCompactBinary(const FCbFieldView& Field, FImportExportNode& Node)
+	{
+		return Node.TryLoad(Field);
+	}
+};
+
+/**
+ * Data about imports and exports from a package that can be stored in the oplog for incremental cooks for replay into
+ * the CookImportsChecker when a package is skipped.
+ */
+struct FImportsCheckerData
+{
+public:
+	TArray<FImportExportNode> Imports;
+	TArray<FImportExportNode> Exports;
+
+	bool IsEmpty() const;
+	void Save(FCbWriter& Writer) const;
+	bool TryLoad(const FCbFieldView& Field);
+
+	static FImportsCheckerData FromObjectLists(TConstArrayView<UObject*> Imports,
+		TConstArrayView<UObject*> Exports);
+	static TArray<FImportExportNode> ObjectListToNodeList(TConstArrayView<UObject*> Objects);
+
+private:
+	friend FCbWriter& operator<<(FCbWriter& Writer, const FImportsCheckerData& Data)
+	{
+		Data.Save(Writer);
+		return Writer;
+	}
+	friend bool LoadFromCompactBinary(const FCbFieldView& Field, FImportsCheckerData& Data)
+	{
+		return Data.TryLoad(Field);
+	}
+};
+
+} // namespace UE::Cook
+
 /**
  * Helper struct used during cooking to validate EDL dependencies
  */
@@ -28,7 +88,10 @@ struct FEDLCookChecker
 	void Reset();
 
 	void AddImport(TObjectPtr<UObject> Import, UPackage* ImportingPackage);
+	void AddImports(TConstArrayView<UE::Cook::FImportExportNode> Imports, FName ImportingPackageName);
 	void AddExport(UObject* Export);
+	void AddExports(TConstArrayView<UE::Cook::FImportExportNode> Exports);
+	void Add(UE::Cook::FImportsCheckerData& ImportsCheckerData, FName PackageName);
 	void AddArc(UObject* DepObject, bool bDepIsSerialize, UObject* Export, bool bExportIsSerialize);
 	void AddPackageWithUnknownExports(FName LongPackageName);
 
@@ -63,6 +126,7 @@ public: // FEDLNodeHash is public only so that GetTypeHash can be defined
 	{
 		FEDLNodeHash(); // creates an uninitialized node; only use this to provide as an out parameter
 		FEDLNodeHash(const TArray<FEDLNodeData>* InNodes, FEDLNodeID InNodeID, EObjectEvent InObjectEvent);
+		FEDLNodeHash(const TArray<FEDLNodeData>* InNodes, FEDLNodeID InParentNodeID, FName InObjectName, EObjectEvent InObjectEvent);
 		FEDLNodeHash(TObjectPtr<UObject> InObject, EObjectEvent InObjectEvent);
 		FEDLNodeHash(const FEDLNodeHash& Other);
 		bool operator==(const FEDLNodeHash& Other) const;
@@ -75,24 +139,42 @@ public: // FEDLNodeHash is public only so that GetTypeHash can be defined
 		void SetNodes(const TArray<FEDLNodeData>* InNodes);
 
 	private:
-		static FName ObjectNameFirst(const FEDLNodeHash& InNode, uint32& OutNodeID, TObjectPtr<const UObject>& OutObject);
-		static FName ObjectNameNext(const FEDLNodeHash& InNode, uint32& OutNodeID, TObjectPtr<const UObject>& OutObject);
+		static uint32 GetTypeHashInternal(const FEDLNodeHash& A);
 
-		union
+		struct FNodeTypeData
 		{
 			/**
-			 * The array of nodes from the FEDLCookChecker; this is how we lookup the node for the FEDLNodeData.
+			 * The array of nodes from the FEDLCookChecker; this is how we lookup the node for the NodeID.
 			 * Because the FEDLNodeData are elements in an array which can resize and therefore reallocate the nodes, we cannot store the pointer to the node.
-			 * Only used if bIsNode is true.
 			 */
 			const TArray<FEDLNodeData>* Nodes;
-			/** Pointer to the Object we are looking up, if this hash was created during lookup-by-objectpath for an object */
+			/** The identifier for the FEDLNodeData this hash is wrapping. */
+			FEDLNodeID NodeID;
+		};
+		struct FNameTypeData
+		{
+			FName ObjectName;
+			/** The array of nodes from the FEDLCookChecker, the same as used in FNodeTypeData. */
+			const TArray<FEDLNodeData>* Nodes;
+			FEDLNodeID ParentID;
+		};
+		struct FObjectTypeData
+		{
 			TObjectPtr<const UObject> Object;
 		};
-		/** The identifier for the FEDLNodeData this hash is wrapping. Only used if bIsNode is true. */
-		FEDLNodeID NodeID;
-		/** True if this hash is wrapping an FEDLNodeData, false if it is wrapping a UObject. */
-		bool bIsNode;
+		enum class ENodeHashType : uint8
+		{
+			Node,
+			Object,
+			NameAndParentNode,
+		};
+		union
+		{
+			FNodeTypeData NodeTypeData;
+			FObjectTypeData ObjectTypeData;
+			FNameTypeData NameTypeData;
+		};
+		ENodeHashType NodeHashType;
 		EObjectEvent ObjectEvent;
 	};
 
@@ -144,6 +226,9 @@ private:
 	FEDLNodeID FindOrAddNode(const FEDLNodeHash& NodeLookup);
 	FEDLNodeID FindOrAddNode(FEDLNodeData&& NodeData, const FEDLCookChecker& OldOwnerOfNode, FEDLNodeID ParentIDInThis, bool& bNew);
 	FEDLNodeID FindNode(const FEDLNodeHash& NodeHash);
+	template <typename AddNodeType>
+	void AddImportExportNodeList(TConstArrayView<UE::Cook::FImportExportNode> NodeList, AddNodeType&& AddNode);
+	void RecordImportFromPackage(FEDLNodeID NodeId, FName ImportingPackageName);
 	void Merge(FEDLCookChecker&& Other);
 	bool CheckForCyclesInner(TSet<FEDLNodeID>& Visited, TSet<FEDLNodeID>& Stack, const FEDLNodeID& Visit, FEDLNodeID& FailNode);
 	void AddDependency(FEDLNodeID SourceID, FEDLNodeID TargetID);
@@ -181,10 +266,22 @@ struct FEDLCookCheckerThreadState : public TThreadSingleton<FEDLCookCheckerThrea
 	{
 		Checker.AddImport(Import, ImportingPackage);
 	}
+	void AddImports(TConstArrayView<UE::Cook::FImportExportNode> Imports, FName ImportingPackageName)
+	{
+		Checker.AddImports(Imports, ImportingPackageName);
+	}
 	void AddExport(UObject* Export)
 	{
 		Checker.AddExport(Export);
 	}
+	void AddExports(TConstArrayView<UE::Cook::FImportExportNode> Exports)
+	{
+		Checker.AddExports(Exports);
+	}
+	void Add(UE::Cook::FImportsCheckerData& ImportsCheckerData, FName PackageName)
+	{
+		Checker.Add(ImportsCheckerData, PackageName);
+	}
 	void AddArc(UObject* DepObject, bool bDepIsSerialize, UObject* Export, bool bExportIsSerialize)
 	{
 		Checker.AddArc(DepObject, bDepIsSerialize, Export, bExportIsSerialize);
@@ -198,4 +295,20 @@ private:
 	FEDLCookChecker Checker;
 	friend TThreadSingleton<FEDLCookCheckerThreadState>;
 	friend FEDLCookChecker;
-};
\ No newline at end of file
+};
+
+
+///////////////////////////////////////////////////////
+// Inline implementations
+///////////////////////////////////////////////////////
+
+
+namespace UE::Cook
+{
+
+inline bool FImportsCheckerData::IsEmpty() const
+{
+	return Imports.IsEmpty() && Exports.IsEmpty();
+}
+
+} // namespace UE::Cook
diff --git a/Engine/Source/Editor/UnrealEd/Private/Cooker/CookPackageArtifacts.cpp b/Engine/Source/Editor/UnrealEd/Private/Cooker/CookPackageArtifacts.cpp
--- a/Engine/Source/Editor/UnrealEd/Private/Cooker/CookPackageArtifacts.cpp
+++ b/Engine/Source/Editor/UnrealEd/Private/Cooker/CookPackageArtifacts.cpp
@@ -32,7 +32,7 @@ namespace UE::Cook
 {
 
 // Change TargetDomainVersion to a new guid when all TargetDomain keys in an incremental cook need to be invalidated
-FGuid TargetDomainVersion(0xE0DBB4ED, 0x5C8A4914, 0xB83C77D0, 0x4152131F);
+FGuid TargetDomainVersion(0x80B89F71, 0x7C7243B9, 0x9E089834, 0x14F6BA20);
 
 // Bump PackageArtifactsVersion version when the serialization of PackageArtifacts has changed and we want to add
 // backwards compatibility rather than invalidating everything.
@@ -40,6 +40,7 @@ constexpr uint32 PackageArtifactsVersion = 0x00000004;
 
 static const FUtf8StringView PackageArtifactsAttachmentKey = UTF8TEXTVIEW("CookPackageArtifacts");
 static const FUtf8StringView BuildDefinitionsAttachmentKey = UTF8TEXTVIEW("BuildDefinitionsAttachmentKey");
+static const FUtf8StringView ImportsCheckerAttachmentKey = UTF8TEXTVIEW("ImportsCheckerAttachmentKey");
 
 bool FBuildDependencySet::HasKeyMatch(FName PackageName, const ITargetPlatform* TargetPlatform,
 	FGenerationHelper* GenerationHelper)
@@ -996,12 +997,23 @@ void FIncrementalCookAttachments::AppendCommitAttachments(TArray<IPackageWriter:
 		Attachment.Key = BuildDefinitionsAttachmentKey;
 		Attachment.Value = Writer.Save().AsObject();
 	}
+
+	if (!ImportsCheckerData.IsEmpty())
+	{
+		FCbWriter Writer;
+		Writer << ImportsCheckerData;
+		IPackageWriter::FCommitAttachmentInfo& Attachment = OutAttachments.Emplace_GetRef();
+		Attachment.Key = ImportsCheckerAttachmentKey;
+		Attachment.Value = Writer.Save().AsObject();
+	}
 }
 
 FIncrementalCookAttachments FIncrementalCookAttachments::Collect(const UPackage* Package,
 	const ITargetPlatform* TargetPlatform, FBuildResultDependenciesMap&& InResultDependencies,
 	bool bHasSaveResult, TConstArrayView<FName> UntrackedSoftPackageReferences, FGenerationHelper* GenerationHelper,
-	bool bGenerated, TArray<FName>&& RuntimeDependencies)
+	bool bGenerated, TArray<FName>&& RuntimeDependencies,
+	TConstArrayView<UObject*> Imports, TConstArrayView<UObject*> Exports,
+	TConstArrayView<UE::SavePackageUtilities::FPreloadDependency> PreloadDependencies)
 {
 	FIncrementalCookAttachments Result;
 
@@ -1023,6 +1035,7 @@ FIncrementalCookAttachments FIncrementalCookAttachments::Collect(const UPackage*
 	}
 
 	Result.BuildDefinitions = FBuildDefinitionList::Collect(Package, TargetPlatform);
+	Result.ImportsCheckerData = UE::Cook::FImportsCheckerData::FromObjectLists(Imports, Exports);
 
 	return Result;
 }
@@ -1039,7 +1052,8 @@ void FIncrementalCookAttachments::Fetch(TArrayView<FName> PackageNames, const IT
 		int ReceivedAttachmentCount = 0;
 	};
 	TMap<FName, FInProgressResult> InProgressResults;
-	TArray<FUtf8StringView> RequestedAttachments = { PackageArtifactsAttachmentKey, BuildDefinitionsAttachmentKey };
+	TArray<FUtf8StringView> RequestedAttachments = { PackageArtifactsAttachmentKey, BuildDefinitionsAttachmentKey,
+		ImportsCheckerAttachmentKey };
 
 	if (TargetPlatform || GEditorDomainOplog)
 	{
@@ -1070,6 +1084,10 @@ void FIncrementalCookAttachments::Fetch(TArrayView<FName> PackageNames, const IT
 				{
 					LoadFromCompactBinary(MoveTemp(Attachment), InProgressResult.Result.BuildDefinitions);
 				}
+				else if (AttachmentKey == ImportsCheckerAttachmentKey)
+				{
+					LoadFromCompactBinary(Attachment.AsFieldView(), InProgressResult.Result.ImportsCheckerData);
+				}
 
 				if (InProgressResult.ReceivedAttachmentCount == RequestedAttachmentNum)
 				{
diff --git a/Engine/Source/Editor/UnrealEd/Private/Cooker/CookPackageArtifacts.h b/Engine/Source/Editor/UnrealEd/Private/Cooker/CookPackageArtifacts.h
--- a/Engine/Source/Editor/UnrealEd/Private/Cooker/CookPackageArtifacts.h
+++ b/Engine/Source/Editor/UnrealEd/Private/Cooker/CookPackageArtifacts.h
@@ -5,6 +5,7 @@
 #include "AssetRegistry/AssetData.h"
 #include "AssetRegistry/IAssetRegistry.h"
 #include "Cooker/CookDependency.h"
+#include "Cooker/CookImportsChecker.h"
 #include "Containers/Array.h"
 #include "Containers/ArrayView.h"
 #include "Containers/Map.h"
@@ -315,6 +316,7 @@ struct FIncrementalCookAttachments
 {
 	FPackageArtifacts Artifacts;
 	FBuildDefinitionList BuildDefinitions;
+	FImportsCheckerData ImportsCheckerData;
 	IPackageWriter::ECommitStatus CommitStatus = IPackageWriter::ECommitStatus::NotCommitted;
 
 	void Empty();
@@ -326,7 +328,9 @@ struct FIncrementalCookAttachments
 	static FIncrementalCookAttachments Collect(const UPackage* Package,
 		const ITargetPlatform* TargetPlatform, FBuildResultDependenciesMap&& InResultDependencies,
 		bool bHasSaveResult, TConstArrayView<FName> UntrackedSoftPackageReferences, FGenerationHelper* GenerationHelper,
-		bool bGenerated, TArray<FName>&& RuntimeDependencies);
+		bool bGenerated, TArray<FName>&& RuntimeDependencies,
+		TConstArrayView<UObject*> Imports, TConstArrayView<UObject*> Exports,
+		TConstArrayView<UE::SavePackageUtilities::FPreloadDependency> PreloadDependencies);
 };
 
 } // namespace UE::Cook
diff --git a/Engine/Source/Editor/UnrealEd/Private/Cooker/CookRequestCluster.cpp b/Engine/Source/Editor/UnrealEd/Private/Cooker/CookRequestCluster.cpp
--- a/Engine/Source/Editor/UnrealEd/Private/Cooker/CookRequestCluster.cpp
+++ b/Engine/Source/Editor/UnrealEd/Private/Cooker/CookRequestCluster.cpp
@@ -2235,8 +2235,10 @@ void FRequestCluster::FGraphSearch::FExploreEdgesContext::ProcessPlatformAttachm
 						*WriteToString<256>(PackageData->GetPackageName()),
 						*PackageData->GetInstigator(EReachability::Runtime).ToString());
 				}
-				// Declare the package to the EDLCookInfo verification so we don't warn about missing exports from it
-				FEDLCookCheckerThreadState::Get().AddPackageWithUnknownExports(PackageName);
+
+				// Replay the package's saved data into our collectors of packagedata
+				FEDLCookCheckerThreadState::Get().Add(QueryPlatformData.CookAttachments.ImportsCheckerData,
+					PackageData->GetPackageName());
 			}
 			else
 			{
diff --git a/Engine/Source/Editor/UnrealEd/Private/Cooker/CookSavePackage.cpp b/Engine/Source/Editor/UnrealEd/Private/Cooker/CookSavePackage.cpp
--- a/Engine/Source/Editor/UnrealEd/Private/Cooker/CookSavePackage.cpp
+++ b/Engine/Source/Editor/UnrealEd/Private/Cooker/CookSavePackage.cpp
@@ -809,7 +809,10 @@ void FSaveCookedPackageContext::CalculateCookDependencies(FGenerationHelper* Gen
 		FBuildResultDependenciesMap BuildResultDependencies;
 		TArray<FName> PlatformRuntimeDependencies;
 		TConstArrayView<FName> UntrackedSoftPackageReferences;
-		
+		TConstArrayView<UObject*> Imports;
+		TConstArrayView<UObject*> Exports;
+		TConstArrayView<UE::SavePackageUtilities::FPreloadDependency> PreloadDependencies;
+
 		PlatformRuntimeDependencies = GetPlatformRuntimeDependencies();
 		if (ExtraArDependencies)
 		{
@@ -872,6 +875,9 @@ void FSaveCookedPackageContext::CalculateCookDependencies(FGenerationHelper* Gen
 			BuildResultDependencies.Append(MoveTemp(LocalSavePackageResult->BuildResultDependencies));
 			PlatformRuntimeDependencies.Append(MoveTemp(LocalSavePackageResult->SoftPackageReferences));
 			UntrackedSoftPackageReferences = LocalSavePackageResult->UntrackedSoftPackageReferences;
+			Imports = LocalSavePackageResult->Imports;
+			Exports = LocalSavePackageResult->Exports;
+			PreloadDependencies = LocalSavePackageResult->PreloadDependencies;
 		}
 		const FBuildResultDependenciesMap* LoadDependencies = PackageData.GetLoadDependencies();
 		checkf(LoadDependencies,
@@ -881,7 +887,7 @@ void FSaveCookedPackageContext::CalculateCookDependencies(FGenerationHelper* Gen
 
 		PlatformCookAttachments = FIncrementalCookAttachments::Collect(Package, TargetPlatform,
 			MoveTemp(BuildResultDependencies), bHasSavePackageResult, UntrackedSoftPackageReferences, GenerationHelper,
-			bGenerated, MoveTemp(PlatformRuntimeDependencies));
+			bGenerated, MoveTemp(PlatformRuntimeDependencies), Imports, Exports, PreloadDependencies);
 	}
 }
 
diff --git a/Engine/Source/Editor/UnrealEd/Private/EditorDomain/EditorDomainUtils.cpp b/Engine/Source/Editor/UnrealEd/Private/EditorDomain/EditorDomainUtils.cpp
--- a/Engine/Source/Editor/UnrealEd/Private/EditorDomain/EditorDomainUtils.cpp
+++ b/Engine/Source/Editor/UnrealEd/Private/EditorDomain/EditorDomainUtils.cpp
@@ -2026,7 +2026,8 @@ bool TrySavePackage(UPackage* Package)
 			UE::Cook::FIncrementalCookAttachments::Collect(
 				Package, nullptr /* TargetPlatform */, MoveTemp(ExistingResultDependencies),
 				false /* bHasSaveResult */, TConstArrayView<FName>() /* UntrackedSoftPackageReferences */,
-				nullptr /* GenerationHelper */, false /* bGenerated */, TArray<FName>());
+				nullptr /* GenerationHelper */, false /* bGenerated */, TArray<FName>(), TConstArrayView<UObject*>(),
+				TConstArrayView<UObject*>(), TConstArrayView<UE::SavePackageUtilities::FPreloadDependency>());
 		TArray<IPackageWriter::FCommitAttachmentInfo> Attachments;
 		CookAttachments.AppendCommitAttachments(Attachments);
 		if (!Attachments.IsEmpty())
