IncrementalCook: Rename TargetDomainVersion to CookIncrementalVersion for clarity, move it into its own file since it relatively frequently needs to be invalidated, and add it to CookSettings.txt so that it causes a wipe of cook artifacts in addition to wiping cooked packages when it changes. Remove ExecutableHash from the CookSettings when using IncrementalCook since it is not relevant to IncrementalCook.
#rnx

diff --git a/Engine/Source/Editor/UnrealEd/Private/CookOnTheFlyServer.cpp b/Engine/Source/Editor/UnrealEd/Private/CookOnTheFlyServer.cpp
--- a/Engine/Source/Editor/UnrealEd/Private/CookOnTheFlyServer.cpp
+++ b/Engine/Source/Editor/UnrealEd/Private/CookOnTheFlyServer.cpp
@@ -7264,43 +7264,54 @@ TMap<FName, FString> UCookOnTheFlyServer::CalculateCookSettingStrings() const
 		});
 	IPlatformFile& PlatformFile = FPlatformFileManager::Get().GetPlatformFile();
 
-	// Calculate the executable hash by combining the module file hash of every loaded module
-	// TODO: Write the module file hash from UnrealBuildTool into the .modules file and read it
-	// here from the .modules file instead of calculating it on every cook.
-	if (bIterativeCalculateExe)
+	CookSettingStrings.Add(FName(TEXT("Version")), TEXT("21F52B9EDD4D456AB1AF381CA172BD28"));
+	CookSettingStrings.Add(FName(TEXT("CookIncremental")), *LexToString(bHybridIterativeEnabled));
+
+	if (bHybridIterativeEnabled)
+	{
+		// Store the CookIncrementalVersion in the global settings, so that it will cause deletion of all artifacts
+		// when it changes, in addition to invalidating all cooked packages.
+		CookSettingStrings.Add(FName(TEXT("CookIncrementalVersion")), UE::Cook::CookIncrementalVersion.ToString());
+	}
+	else
 	{
-		FString InvalidModule;
-		bool bValid = true;
-		FXxHash64Builder Hasher;
-		TArray<uint8> Buffer;
-		for (FModuleStatus& ModuleStatus : Modules)
+		// Calculate the executable hash by combining the module file hash of every loaded module
+		// TODO: Write the module file hash from UnrealBuildTool into the .modules file and read it
+		// here from the .modules file instead of calculating it on every cook.
+		if (bIterativeCalculateExe)
 		{
-			TUniquePtr<IFileHandle> FileHandle(PlatformFile.OpenRead(*ModuleStatus.FilePath));
-			if (!FileHandle)
+			FString InvalidModule;
+			bool bValid = true;
+			FXxHash64Builder Hasher;
+			TArray<uint8> Buffer;
+			for (FModuleStatus& ModuleStatus : Modules)
 			{
-				InvalidModule = ModuleStatus.FilePath;
-				break;
+				TUniquePtr<IFileHandle> FileHandle(PlatformFile.OpenRead(*ModuleStatus.FilePath));
+				if (!FileHandle)
+				{
+					InvalidModule = ModuleStatus.FilePath;
+					break;
+				}
+				int64 FileSize = FileHandle->Size();
+				Buffer.SetNumUninitialized(FileSize, EAllowShrinking::No);
+				if (!FileHandle->Read(Buffer.GetData(), FileSize))
+				{
+					InvalidModule = ModuleStatus.FilePath;
+					break;
+				}
+				Hasher.Update(Buffer.GetData(), FileSize);
 			}
-			int64 FileSize = FileHandle->Size();
-			Buffer.SetNumUninitialized(FileSize, EAllowShrinking::No);
-			if (!FileHandle->Read(Buffer.GetData(), FileSize))
+			if (InvalidModule.IsEmpty())
 			{
-				InvalidModule = ModuleStatus.FilePath;
-				break;
+				CookSettingStrings.Add(ExecutableHashName, FString(*WriteToString<64>(Hasher.Finalize())));
+			}
+			else
+			{
+				CookSettingStrings.Add(ExecutableHashInvalidModuleName, InvalidModule);
 			}
-			Hasher.Update(Buffer.GetData(), FileSize);
-		}
-		if (InvalidModule.IsEmpty())
-		{
-			CookSettingStrings.Add(ExecutableHashName, FString(*WriteToString<64>(Hasher.Finalize())));
-		}
-		else
-		{
-			CookSettingStrings.Add(ExecutableHashInvalidModuleName, InvalidModule);
 		}
 	}
 
-	CookSettingStrings.Add(FName(TEXT("Version")), TEXT("21F52B9EDD4D456AB1AF381CA172BD28"));
 	if (IsDirectorCookByTheBook())
 	{
 		CookSettingStrings.Add(NAME_CookMode, TEXT("CookByTheBook"));
diff --git a/Engine/Source/Editor/UnrealEd/Private/Cooker/CookIncrementalVersion.cpp b/Engine/Source/Editor/UnrealEd/Private/Cooker/CookIncrementalVersion.cpp
new file mode 100644
--- /dev/null
+++ b/Engine/Source/Editor/UnrealEd/Private/Cooker/CookIncrementalVersion.cpp
@@ -0,0 +1,13 @@
+// Copyright Epic Games, Inc. All Rights Reserved.
+
+#include "CookTypes.h"
+
+#include "Misc/Guid.h"
+
+namespace UE::Cook
+{
+
+// Change CookIncrementalVersion to a new guid when all packages in an incremental cook need to be invalidated
+FGuid CookIncrementalVersion(0x80B89F71, 0x7C7243B9, 0x9E089834, 0x14F6BA20);
+
+}
diff --git a/Engine/Source/Editor/UnrealEd/Private/Cooker/CookPackageArtifacts.cpp b/Engine/Source/Editor/UnrealEd/Private/Cooker/CookPackageArtifacts.cpp
--- a/Engine/Source/Editor/UnrealEd/Private/Cooker/CookPackageArtifacts.cpp
+++ b/Engine/Source/Editor/UnrealEd/Private/Cooker/CookPackageArtifacts.cpp
@@ -31,9 +31,6 @@
 namespace UE::Cook
 {
 
-// Change TargetDomainVersion to a new guid when all TargetDomain keys in an incremental cook need to be invalidated
-FGuid TargetDomainVersion(0x80B89F71, 0x7C7243B9, 0x9E089834, 0x14F6BA20);
-
 // Bump PackageArtifactsVersion version when the serialization of PackageArtifacts has changed and we want to add
 // backwards compatibility rather than invalidating everything.
 constexpr uint32 PackageArtifactsVersion = 0x00000004;
@@ -87,7 +84,7 @@ FBuildDependencySet::ECurrentKeyResult FBuildDependencySet::TryCalculateCurrentK
 		return ECurrentKeyResult::Error;
 	}
 	FBlake3 KeyBuilder;
-	KeyBuilder.Update(&TargetDomainVersion, sizeof(TargetDomainVersion));
+	KeyBuilder.Update(&CookIncrementalVersion, sizeof(CookIncrementalVersion));
 
 	ECurrentKeyResult Result = ECurrentKeyResult::Success;
 	FCookDependencyContext Context(&KeyBuilder,
diff --git a/Engine/Source/Editor/UnrealEd/Private/Cooker/CookTypes.h b/Engine/Source/Editor/UnrealEd/Private/Cooker/CookTypes.h
--- a/Engine/Source/Editor/UnrealEd/Private/Cooker/CookTypes.h
+++ b/Engine/Source/Editor/UnrealEd/Private/Cooker/CookTypes.h
@@ -37,6 +37,7 @@ namespace UE::Cook { struct FCookByTheBookOptions; }
 namespace UE::Cook { struct FCookOnTheFlyOptions; }
 namespace UE::Cook { struct FInitializeConfigSettings; }
 namespace UE::Cook { struct FInstigator; }
+struct FGuid;
 
 FCbWriter& operator<<(FCbWriter& Writer, const UE::Cook::FBeginCookConfigSettings& Value);
 bool LoadFromCompactBinary(FCbFieldView Field, UE::Cook::FBeginCookConfigSettings& OutValue);
@@ -516,6 +517,9 @@ namespace UE::Cook
 
 	TConstArrayView<const TCHAR*> GetCommandLineDelimiterStrs();
 	TConstArrayView<TCHAR> GetCommandLineDelimiterChars();
+
+	extern FGuid CookIncrementalVersion;
+
 } // namespace UE::Cook
 
 bool LexTryParseString(FPlatformMemoryStats::EMemoryPressureStatus& OutValue, FStringView Text);
