IncrementalCook: Move the code to serialize logs into a cook structure out of the MPCook classes and into a separate class. This code will be shared with incremental cook and the cook structure stored in the oplog in an upcoming change.
#rnx

diff --git a/Engine/Source/Editor/UnrealEd/Classes/CookOnTheSide/CookOnTheFlyServer.h b/Engine/Source/Editor/UnrealEd/Classes/CookOnTheSide/CookOnTheFlyServer.h
--- a/Engine/Source/Editor/UnrealEd/Classes/CookOnTheSide/CookOnTheFlyServer.h
+++ b/Engine/Source/Editor/UnrealEd/Classes/CookOnTheSide/CookOnTheFlyServer.h
@@ -180,6 +180,7 @@ namespace UE::Cook
 	class FCookWorkerClient;
 	class FCookWorkerServer;
 	class FDiagnostics;
+	class FLogRecordSerializationContext;
 	class FODSCClientData; 
 	class FPackagePreloader;
 	class FPackageWriterMPCollector;
@@ -1637,6 +1638,7 @@ private:
 	TUniquePtr<FLayeredCookArtifactReader> AllContextArtifactReader;
 	TSharedPtr<FLooseFilesCookArtifactReader> SharedLooseFilesCookArtifactReader;
 	TUniquePtr<UE::Cook::FCookGCDiagnosticContext> GCDiagnosticContext;
+	TUniquePtr<UE::Cook::FLogRecordSerializationContext> LogRecordSerializationContext;
 
 	TArray<UE::Cook::FCookSavePackageContext*> SavePackageContexts;
 	/**
diff --git a/Engine/Source/Editor/UnrealEd/Private/CookOnTheFlyServer.cpp b/Engine/Source/Editor/UnrealEd/Private/CookOnTheFlyServer.cpp
--- a/Engine/Source/Editor/UnrealEd/Private/CookOnTheFlyServer.cpp
+++ b/Engine/Source/Editor/UnrealEd/Private/CookOnTheFlyServer.cpp
@@ -27,6 +27,7 @@
 #include "Cooker/CookGarbageCollect.h"
 #include "Cooker/CookGenerationHelper.h"
 #include "Cooker/CookImportsChecker.h"
+#include "Cooker/CookLogPrivate.h"
 #include "Cooker/CookOnTheFlyServerInterface.h"
 #include "Cooker/CookPackageData.h"
 #include "Cooker/CookPackagePreloader.h"
@@ -170,11 +171,6 @@
 
 #define LOCTEXT_NAMESPACE "Cooker"
 
-DEFINE_LOG_CATEGORY(LogCook);
-DEFINE_LOG_CATEGORY(LogCookStats);
-DEFINE_LOG_CATEGORY_STATIC(LogCookList, Log, All);
-
-FName LogCookName(TEXT("LogCook"));
 LLM_DEFINE_TAG(Cooker);
 
 int32 GCookProgressDisplay = (int32)ECookProgressDisplayMode::RemainingPackages;
@@ -5822,6 +5818,7 @@ void UCookOnTheFlyServer::Initialize( ECookMode::Type DesiredCookMode, ECookInit
 	AssetRegistry = IAssetRegistry::Get();
 	GCDiagnosticContext = MakeUnique<UE::Cook::FCookGCDiagnosticContext>();
 	StallDetector = MakeUnique<UE::Cook::FStallDetector>();
+	LogRecordSerializationContext.Reset(new UE::Cook::FLogRecordSerializationContext());
 
 	if (!IsCookWorkerMode())
 	{
diff --git a/Engine/Source/Editor/UnrealEd/Private/Cooker/CookDirector.cpp b/Engine/Source/Editor/UnrealEd/Private/Cooker/CookDirector.cpp
--- a/Engine/Source/Editor/UnrealEd/Private/Cooker/CookDirector.cpp
+++ b/Engine/Source/Editor/UnrealEd/Private/Cooker/CookDirector.cpp
@@ -147,7 +147,7 @@ FCookDirector::FCookDirector(UCookOnTheFlyServer& InCOTFS, int32 CookProcessCoun
 	}
 	bMultiprocessAvailable = true;
 
-	Register(new FLogMessagesMessageHandler());
+	Register(new FLogMessagesMessageHandler(*COTFS.LogRecordSerializationContext));
 	Register(new TMPCollectorServerMessageCallback<FRetractionResultsMessage>([this]
 	(FMPCollectorServerMessageContext& Context, bool bReadSuccessful, FRetractionResultsMessage&& Message)
 		{
diff --git a/Engine/Source/Editor/UnrealEd/Private/Cooker/CookLog.cpp b/Engine/Source/Editor/UnrealEd/Private/Cooker/CookLog.cpp
new file mode 100644
--- /dev/null
+++ b/Engine/Source/Editor/UnrealEd/Private/Cooker/CookLog.cpp
@@ -0,0 +1,349 @@
+// Copyright Epic Games, Inc. All Rights Reserved.
+
+#include "CookOnTheSide/CookLog.h"
+
+#include "Cooker/CookLogPrivate.h"
+
+#include "Containers/AnsiString.h"
+#include "CoreGlobals.h"
+#include "Logging/StructuredLog.h"
+#include "Logging/StructuredLogFormat.h"
+#include "Misc/AssertionMacros.h"
+#include "Misc/DateTime.h"
+#include "Misc/FeedbackContext.h"
+#include "Misc/OutputDeviceRedirector.h"
+#include "Misc/StringBuilder.h"
+#include "Serialization/CompactBinary.h"
+#include "Serialization/CompactBinaryWriter.h"
+#include "Serialization/CompactBinarySerialization.h"
+
+DEFINE_LOG_CATEGORY(LogCook);
+DEFINE_LOG_CATEGORY(LogCookStats);
+DEFINE_LOG_CATEGORY(LogCookList);
+
+FName LogCookName(TEXT("LogCook"));
+
+namespace UE::Cook
+{
+
+FCbWriter& operator<<(FCbWriter& Writer, const FReplicatedLogData& LogData)
+{
+	// Serializing as an array of unnamed fields and using the quantity of fields
+	// as the discriminator between structured and unstructured log data.
+	Writer.BeginArray();
+	if (LogData.LogDataVariant.IsType<FReplicatedLogData::FUnstructuredLogData>())
+	{
+		const FReplicatedLogData::FUnstructuredLogData& UnstructuredLogData = LogData.LogDataVariant.Get<FReplicatedLogData::FUnstructuredLogData>();
+		Writer << UnstructuredLogData.Category;
+		uint8 Verbosity = static_cast<uint8>(UnstructuredLogData.Verbosity);
+		Writer << Verbosity;
+		Writer << UnstructuredLogData.Message;
+	}
+	else if (LogData.LogDataVariant.IsType<FCbObject>())
+	{
+		Writer << LogData.LogDataVariant.Get<FCbObject>();
+	}
+	else
+	{
+		checkNoEntry();
+	}
+	Writer.EndArray();
+	return Writer;
+}
+
+bool LoadFromCompactBinary(FCbFieldView Field, FReplicatedLogData& OutLogData)
+{
+	bool bOk = true;
+	FCbArrayView ArrayView = Field.AsArrayView();
+	switch (ArrayView.Num())
+	{
+	case 3:
+	{
+		OutLogData.LogDataVariant.Emplace<FReplicatedLogData::FUnstructuredLogData>();
+		FReplicatedLogData::FUnstructuredLogData& UnstructuredLogData = OutLogData.LogDataVariant.Get<FReplicatedLogData::FUnstructuredLogData>();
+		FCbFieldViewIterator It = ArrayView.CreateViewIterator();
+		bOk = LoadFromCompactBinary(*It++, UnstructuredLogData.Category) & bOk;
+		uint8 Verbosity;
+		if (LoadFromCompactBinary(*It++, Verbosity))
+		{
+			UnstructuredLogData.Verbosity = static_cast<ELogVerbosity::Type>(Verbosity);
+		}
+		else
+		{
+			bOk = false;
+			UnstructuredLogData.Verbosity = static_cast<ELogVerbosity::Type>(0);
+		}
+		bOk = LoadFromCompactBinary(*It++, UnstructuredLogData.Message) & bOk;
+		break;
+	}
+	case 1:
+	{
+		OutLogData.LogDataVariant.Emplace<FCbObject>();
+		FCbObject& StructuredLogData = OutLogData.LogDataVariant.Get<FCbObject>();
+		FCbFieldViewIterator It = ArrayView.CreateViewIterator();
+		if (It->IsObject())
+		{
+			StructuredLogData = FCbObject::Clone(It->AsObjectView());
+		}
+		else
+		{
+			bOk = false;
+		}
+		break;
+	}
+	default:
+		bOk = false;
+	}
+	return bOk;
+}
+
+void FLogRecordSerializationContext::Serialize(FReplicatedLogData& OutData, const TCHAR* V,
+	ELogVerbosity::Type Verbosity, const FName& Category)
+{
+	OutData.LogDataVariant.Emplace<FReplicatedLogData::FUnstructuredLogData>();
+	FReplicatedLogData::FUnstructuredLogData& OutVal
+		= OutData.LogDataVariant.Get<FReplicatedLogData::FUnstructuredLogData>();
+	OutVal.Message = V;
+	OutVal.Category = Category;
+	OutVal.Verbosity = Verbosity;
+}
+
+void FLogRecordSerializationContext::Serialize(FReplicatedLogData& OutData, const TCHAR* V,
+	ELogVerbosity::Type Verbosity, const FName& Category, const double Time)
+{
+	Serialize(OutData, V, Verbosity, Category);
+}
+
+void FLogRecordSerializationContext::Serialize(FReplicatedLogData& OutData, const FLogRecord& LogRecord)
+{
+	FCbWriter Writer;
+	Writer.BeginObject();
+	Writer << "S";
+	Serialize(Writer, LogRecord);
+	Writer.EndObject();
+	FCbObject Object = Writer.Save().AsObject();
+	OutData.LogDataVariant.Emplace<FCbObject>(MoveTemp(Object));
+}
+
+FCbWriter& FLogRecordSerializationContext::Serialize(FCbWriter& Writer, const FLogRecord& LogRecord)
+{
+	Writer.BeginArray();
+	Writer << LogRecord.GetCategory();
+	Writer << static_cast<uint8>(LogRecord.GetVerbosity());
+	Writer << LogRecord.GetTime().GetUtcTime();
+	Writer << LogRecord.GetFormat();
+	Writer << LogRecord.GetFields();
+	Writer << LogRecord.GetFile();
+	Writer << LogRecord.GetLine();
+	Writer << LogRecord.GetTextNamespace();
+	Writer << LogRecord.GetTextKey();
+	Writer.EndArray();
+	return Writer;
+}
+
+void FLogRecordSerializationContext::DeserializeAndLog(const FReplicatedLogData& LogData,
+	TFunctionRef<bool(FName, const FString&)> MessagePassesFilter,
+	TFunctionRef<bool(const FString&, FString&)> TryTransformMessage)
+{
+	if (const FReplicatedLogData::FUnstructuredLogData* UnStructuredLogData
+		= LogData.LogDataVariant.TryGet<FReplicatedLogData::FUnstructuredLogData>())
+	{
+		if (!MessagePassesFilter(UnStructuredLogData->Category, UnStructuredLogData->Message))
+		{
+			return;
+		}
+		const FString* SerializedString = &UnStructuredLogData->Message;
+		FString TransformedString;
+		if (TryTransformMessage(*SerializedString, TransformedString))
+		{
+			SerializedString = &TransformedString;
+		}
+
+		FMsg::Logf(__FILE__, __LINE__, UnStructuredLogData->Category, UnStructuredLogData->Verbosity,
+			TEXT("%s"), **SerializedString);
+	}
+	else if (const FCbObject* StructuredLogObject = LogData.LogDataVariant.TryGet<FCbObject>())
+	{
+		FLogRecord LogRecord;
+		if (Deserialize((*StructuredLogObject)["S"], LogRecord, MessagePassesFilter, TryTransformMessage))
+		{
+			FOutputDevice* LogOverride = nullptr;
+			switch (LogRecord.GetVerbosity())
+			{
+			case ELogVerbosity::Error:
+			case ELogVerbosity::Warning:
+			case ELogVerbosity::Display:
+			case ELogVerbosity::SetColor:
+				LogOverride = GWarn;
+				break;
+			default:
+				break;
+			}
+			if (LogOverride)
+			{
+				LogOverride->SerializeRecord(LogRecord);
+			}
+			else
+			{
+				GLog->SerializeRecord(LogRecord);
+			}
+		}
+	}
+	else
+	{
+		checkNoEntry();
+	}
+
+}
+bool FLogRecordSerializationContext::Deserialize(FCbFieldView Field, FLogRecord& OutLogRecord,
+	TFunctionRef<bool(FName, const FString&)> MessagePassesFilter,
+	TFunctionRef<bool(const FString&, FString&)> TryTransformMessage)
+{
+	bool bOk = true;
+	FCbFieldViewIterator It = Field.CreateViewIterator();
+	FName Category;
+	if (LoadFromCompactBinary(*It++, Category))
+	{
+		OutLogRecord.SetCategory(Category);
+	}
+	else
+	{
+		bOk = false;
+	}
+	if (uint8 Verbosity; LoadFromCompactBinary(*It++, Verbosity) && Verbosity < ELogVerbosity::NumVerbosity)
+	{
+		OutLogRecord.SetVerbosity(static_cast<ELogVerbosity::Type>(Verbosity));
+	}
+	else
+	{
+		bOk = false;
+	}
+	if (FDateTime Time; LoadFromCompactBinary(*It++, Time))
+	{
+		OutLogRecord.SetTime(FLogTime::FromUtcTime(Time));
+	}
+	else
+	{
+		bOk = false;
+	}
+	if (FString SerializedString;
+		LoadFromCompactBinary(*It++, SerializedString) && MessagePassesFilter(Category, SerializedString))
+	{
+		FString TransformedString;
+		if (TryTransformMessage(SerializedString, TransformedString))
+		{
+			SerializedString = MoveTemp(TransformedString);
+		}
+		FString& FormatString = StringTable.AddDefaulted_GetRef();
+		FormatString = MoveTemp(SerializedString);
+		OutLogRecord.SetFormat(*FormatString);
+	}
+	else
+	{
+		bOk = false;
+	}
+
+	FCbObject Object(FCbObject::Clone(It->AsObjectView()));
+	OutLogRecord.SetFields(MoveTemp(Object));
+	bOk = !It->HasError() && bOk;
+	It++;
+
+	if (TUtf8StringBuilder<64> FileStringBuilder; LoadFromCompactBinary(*It++, FileStringBuilder))
+	{
+		FAnsiString& FileString = AnsiStringTable.AddDefaulted_GetRef();
+		FileString = FileStringBuilder.ToString();
+		OutLogRecord.SetFile(*FileString);
+	}
+	else
+	{
+		bOk = false;
+	}
+	if (int32 Line; LoadFromCompactBinary(*It++, Line))
+	{
+		OutLogRecord.SetLine(Line);
+	}
+	else
+	{
+		bOk = false;
+	}
+	if (FString TextNamespaceString; LoadFromCompactBinary(*It++, TextNamespaceString))
+	{
+		if (!TextNamespaceString.IsEmpty())
+		{
+			OutLogRecord.SetTextNamespace(*StringTable.Emplace_GetRef(MoveTemp(TextNamespaceString)));
+		}
+		else
+		{
+			OutLogRecord.SetTextNamespace(nullptr);
+		}
+	}
+	else
+	{
+		bOk = false;
+	}
+	bool bHasTextKey = false;
+	if (FString TextKeyString; LoadFromCompactBinary(*It++, TextKeyString))
+	{
+		if (!TextKeyString.IsEmpty())
+		{
+			bHasTextKey = true;
+			OutLogRecord.SetTextKey(*StringTable.Emplace_GetRef(MoveTemp(TextKeyString)));
+		}
+		else
+		{
+			OutLogRecord.SetTextKey(nullptr);
+		}
+	}
+	else
+	{
+		bOk = false;
+	}
+
+	if (bHasTextKey)
+	{
+		FLogTemplate* LogTemplate = CreateLogTemplate(OutLogRecord.GetTextNamespace(), OutLogRecord.GetTextKey(), OutLogRecord.GetFormat());
+		TemplateTable.Add(LogTemplate);
+		OutLogRecord.SetTemplate(LogTemplate);
+	}
+	else
+	{
+		FLogTemplate* LogTemplate = CreateLogTemplate(OutLogRecord.GetFormat());
+		TemplateTable.Add(LogTemplate);
+		OutLogRecord.SetTemplate(LogTemplate);
+	}
+
+	return bOk;
+}
+
+void FLogRecordSerializationContext::ConditionalFlush()
+{
+	// Flush if the tables in the serialization context have exceeded 100 entries
+	const int32 TableSizeToFlushAt = 100;
+	if ((StringTable.Num() > TableSizeToFlushAt)
+		|| (AnsiStringTable.Num() > TableSizeToFlushAt)
+		|| (TemplateTable.Num() > TableSizeToFlushAt))
+	{
+		Flush();
+	}
+}
+
+void FLogRecordSerializationContext::Flush()
+{
+	if (!StringTable.IsEmpty() || !AnsiStringTable.IsEmpty() || !TemplateTable.IsEmpty())
+	{
+		// NOTE: We only call FlushThreadedLogs on GLog even though we might serialize structured logs via GLog or GWarn.
+		// GWarn is an output device, but GLog is a an output redirector, and only the redirector has/needs FlushThreadedLogs.
+		// Output devices are expected to not use any pointer on a structured log record after completion of the SerializeRecord call.
+		GLog->FlushThreadedLogs();
+	}
+	for (FLogTemplate* LogTemplate : TemplateTable)
+	{
+		DestroyLogTemplate(LogTemplate);
+	}
+
+	StringTable.Empty();
+	AnsiStringTable.Empty();
+	TemplateTable.Empty();
+}
+
+} // namespace UE::Cook
\ No newline at end of file
diff --git a/Engine/Source/Editor/UnrealEd/Private/Cooker/CookLogPrivate.h b/Engine/Source/Editor/UnrealEd/Private/Cooker/CookLogPrivate.h
new file mode 100644
--- /dev/null
+++ b/Engine/Source/Editor/UnrealEd/Private/Cooker/CookLogPrivate.h
@@ -0,0 +1,68 @@
+// Copyright Epic Games, Inc. All Rights Reserved.
+
+#pragma once
+
+#include "Containers/StringView.h"
+#include "Containers/AnsiString.h"
+#include "Containers/UnrealString.h"
+#include "CookOnTheSide/CookLog.h"
+#include "Logging/LogMacros.h"
+#include "Logging/LogVerbosity.h"
+#include "Misc/TVariant.h"
+#include "Serialization/CompactBinary.h"
+#include "UObject/NameTypes.h"
+#include "Templates/Function.h"
+
+namespace UE { class FLogTemplate; }
+namespace UE { class FLogRecord; }
+
+DECLARE_LOG_CATEGORY_EXTERN(LogCookList, Log, All);
+
+namespace UE::Cook
+{
+
+/** Stores the data passed into FOutputDevice::Serialize, for replication to the CookDirector. */
+struct FReplicatedLogData
+{
+	struct FUnstructuredLogData
+	{
+		FString Message;
+		FName Category;
+		ELogVerbosity::Type Verbosity;
+	};
+	TVariant<FUnstructuredLogData, FCbObject> LogDataVariant;
+};
+FCbWriter& operator<<(FCbWriter& Writer, const FReplicatedLogData& LogData);
+bool LoadFromCompactBinary(FCbFieldView Field, FReplicatedLogData& OutLogData);
+
+class FLogRecordSerializationContext
+{
+public:
+	~FLogRecordSerializationContext()
+	{
+		Flush();
+	}
+
+	static void Serialize(FReplicatedLogData& OutData, const TCHAR* V, ELogVerbosity::Type Verbosity,
+		const FName& Category);
+	static void Serialize(FReplicatedLogData& OutData, const TCHAR* V, ELogVerbosity::Type Verbosity,
+		const FName& Category, const double Time);
+	static void Serialize(FReplicatedLogData& OutData, const FLogRecord& LogRecord);
+	void DeserializeAndLog(const FReplicatedLogData& LogData,
+		TFunctionRef<bool(FName, const FString&)> MessagePassesFilter,
+		TFunctionRef<bool(const FString&, FString&)> TryTransformMessage);
+
+	void ConditionalFlush();
+	void Flush();
+private:
+	static FCbWriter& Serialize(FCbWriter& Writer, const FLogRecord& LogRecord);
+	bool Deserialize(FCbFieldView Field, FLogRecord& OutLogRecord,
+		TFunctionRef<bool(FName, const FString&)> MessagePassesFilter,
+		TFunctionRef<bool(const FString&, FString&)> TryTransformMessage);
+
+	TArray<FString> StringTable;
+	TArray<FAnsiString> AnsiStringTable;
+	TArray<FLogTemplate*> TemplateTable;
+};
+
+} // namespace UE::Cook
\ No newline at end of file
diff --git a/Engine/Source/Editor/UnrealEd/Private/Cooker/CookWorkerClient.cpp b/Engine/Source/Editor/UnrealEd/Private/Cooker/CookWorkerClient.cpp
--- a/Engine/Source/Editor/UnrealEd/Private/Cooker/CookWorkerClient.cpp
+++ b/Engine/Source/Editor/UnrealEd/Private/Cooker/CookWorkerClient.cpp
@@ -32,7 +32,7 @@ constexpr float WaitForConnectReplyTimeout = 60.f;
 FCookWorkerClient::FCookWorkerClient(UCookOnTheFlyServer& InCOTFS)
 	: COTFS(InCOTFS)
 {
-	LogMessageHandler = new FLogMessagesMessageHandler();
+	LogMessageHandler = new FLogMessagesMessageHandler(*COTFS.LogRecordSerializationContext);
 	LogMessageHandler->InitializeClient();
 	Register(LogMessageHandler);
 	Register(new TMPCollectorClientMessageCallback<FRetractionRequestMessage>([this]
diff --git a/Engine/Source/Editor/UnrealEd/Private/Cooker/CookWorkerServer.cpp b/Engine/Source/Editor/UnrealEd/Private/Cooker/CookWorkerServer.cpp
--- a/Engine/Source/Editor/UnrealEd/Private/Cooker/CookWorkerServer.cpp
+++ b/Engine/Source/Editor/UnrealEd/Private/Cooker/CookWorkerServer.cpp
@@ -4,22 +4,19 @@
 
 #include "Algo/Find.h"
 #include "Commandlets/AssetRegistryGenerator.h"
-#include "Containers/AnsiString.h"
 #include "Cooker/CompactBinaryTCP.h"
 #include "Cooker/CookDirector.h"
 #include "Cooker/CookGenerationHelper.h"
+#include "Cooker/CookLogPrivate.h"
 #include "Cooker/CookPackageData.h"
 #include "Cooker/CookPlatformManager.h"
 #include "HAL/Platform.h"
 #include "HAL/PlatformProcess.h"
 #include "Interfaces/ITargetPlatform.h"
 #include "Interfaces/ITargetPlatformManagerModule.h"
-#include "Logging/StructuredLog.h"
-#include "Logging/StructuredLogFormat.h"
 #include "Math/NumericLimits.h"
 #include "Misc/AssertionMacros.h"
 #include "Misc/Char.h"
-#include "Misc/FeedbackContext.h"
 #include "Misc/FileHelper.h"
 #include "Misc/Optional.h"
 #include "Misc/Parse.h"
@@ -1896,233 +1893,13 @@ bool FGeneratorEventMessage::TryRead(FCbObjectView Object)
 
 FGuid FGeneratorEventMessage::MessageType(TEXT("B6EE94CA70EC4F40B0D2214EDC11ED03"));
 
-FCbWriter& operator<<(FCbWriter& Writer, const FReplicatedLogData& LogData)
-{
-	// Serializing as an array of unnamed fields and using the quantity of fields
-	// as the discriminator between structured and unstructured log data.
-	Writer.BeginArray();
-	if (LogData.LogDataVariant.IsType<FReplicatedLogData::FUnstructuredLogData>())
-	{
-		const FReplicatedLogData::FUnstructuredLogData& UnstructuredLogData = LogData.LogDataVariant.Get<FReplicatedLogData::FUnstructuredLogData>();
-		Writer << UnstructuredLogData.Category;
-		uint8 Verbosity = static_cast<uint8>(UnstructuredLogData.Verbosity);
-		Writer << Verbosity;
-		Writer << UnstructuredLogData.Message;
-	}
-	else if (LogData.LogDataVariant.IsType<FCbObject>())
-	{
-		Writer << LogData.LogDataVariant.Get<FCbObject>();
-	}
-	else
-	{
-		checkNoEntry();
-	}
-	Writer.EndArray();
-	return Writer;
-}
-
-bool LoadFromCompactBinary(FCbFieldView Field, FReplicatedLogData& OutLogData)
-{
-	bool bOk = true;
-	FCbArrayView ArrayView = Field.AsArrayView();
-	switch (ArrayView.Num())
-	{
-	case 3:
-	{
-		OutLogData.LogDataVariant.Emplace<FReplicatedLogData::FUnstructuredLogData>();
-		FReplicatedLogData::FUnstructuredLogData& UnstructuredLogData = OutLogData.LogDataVariant.Get<FReplicatedLogData::FUnstructuredLogData>();
-		FCbFieldViewIterator It = ArrayView.CreateViewIterator();
-		bOk = LoadFromCompactBinary(*It++, UnstructuredLogData.Category) & bOk;
-		uint8 Verbosity;
-		if (LoadFromCompactBinary(*It++, Verbosity))
-		{
-			UnstructuredLogData.Verbosity = static_cast<ELogVerbosity::Type>(Verbosity);
-		}
-		else
-		{
-			bOk = false;
-			UnstructuredLogData.Verbosity = static_cast<ELogVerbosity::Type>(0);
-		}
-		bOk = LoadFromCompactBinary(*It++, UnstructuredLogData.Message) & bOk;
-		break;
-	}
-	case 1:
-	{
-		OutLogData.LogDataVariant.Emplace<FCbObject>();
-		FCbObject& StructuredLogData = OutLogData.LogDataVariant.Get<FCbObject>();
-		FCbFieldViewIterator It = ArrayView.CreateViewIterator();
-		if (It->IsObject())
-		{
-			StructuredLogData = FCbObject::Clone(It->AsObjectView());
-		}
-		else
-		{
-			bOk = false;
-		}
-		break;
-	}
-	default:
-		bOk = false;
-	}
-	return bOk;
-}
-
-FCbWriter& FLogMessagesMessageHandler::FLogRecordSerializationContext::Serialize(FCbWriter& Writer, const FLogRecord& LogRecord)
-{
-	Writer.BeginArray();
-	Writer << LogRecord.GetCategory();
-	Writer << static_cast<uint8>(LogRecord.GetVerbosity());
-	Writer << LogRecord.GetTime().GetUtcTime();
-	Writer << LogRecord.GetFormat();
-	Writer << LogRecord.GetFields();
-	Writer << LogRecord.GetFile();
-	Writer << LogRecord.GetLine();
-	Writer << LogRecord.GetTextNamespace();
-	Writer << LogRecord.GetTextKey();
-	Writer.EndArray();
-	return Writer;
-}
-
-bool FLogMessagesMessageHandler::FLogRecordSerializationContext::Deserialize(FCbFieldView Field, FLogRecord& OutLogRecord, int32 ProfileId)
-{
-	bool bOk = true;
-	FCbFieldViewIterator It = Field.CreateViewIterator();
-	if (FName Category; LoadFromCompactBinary(*It++, Category))
-	{
-		OutLogRecord.SetCategory(Category);
-	}
-	else
-	{
-		bOk = false;
-	}
-	if (uint8 Verbosity; LoadFromCompactBinary(*It++, Verbosity) && Verbosity < ELogVerbosity::NumVerbosity)
-	{
-		OutLogRecord.SetVerbosity(static_cast<ELogVerbosity::Type>(Verbosity));
-	}
-	else
-	{
-		bOk = false;
-	}
-	if (FDateTime Time; LoadFromCompactBinary(*It++, Time))
-	{
-		OutLogRecord.SetTime(FLogTime::FromUtcTime(Time));
-	}
-	else
-	{
-		bOk = false;
-	}
-	if (FString SerializedString; LoadFromCompactBinary(*It++, SerializedString))
-	{
-		FString& FormatString = StringTable.AddDefaulted_GetRef();
-		FormatString = FString::Printf(TEXT("[CookWorker %d]: %s"), ProfileId, *SerializedString);
-		OutLogRecord.SetFormat(*FormatString);
-	}
-	else
-	{
-		bOk = false;
-	}
-
-	FCbObject Object(FCbObject::Clone(It->AsObjectView()));
-	OutLogRecord.SetFields(MoveTemp(Object));
-	bOk = !It->HasError() && bOk;
-	It++;
-
-	if (TUtf8StringBuilder<64> FileStringBuilder; LoadFromCompactBinary(*It++, FileStringBuilder))
-	{
-		FAnsiString& FileString = AnsiStringTable.AddDefaulted_GetRef();
-		FileString = FileStringBuilder.ToString();
-		OutLogRecord.SetFile(*FileString);
-	}
-	else
-	{
-		bOk = false;
-	}
-	if (int32 Line; LoadFromCompactBinary(*It++, Line))
-	{
-		OutLogRecord.SetLine(Line);
-	}
-	else
-	{
-		bOk = false;
-	}
-	if (FString TextNamespaceString; LoadFromCompactBinary(*It++, TextNamespaceString))
-	{
-		if (!TextNamespaceString.IsEmpty())
-		{
-			OutLogRecord.SetTextNamespace(*StringTable.Emplace_GetRef(MoveTemp(TextNamespaceString)));
-		}
-		else
-		{
-			OutLogRecord.SetTextNamespace(nullptr);
-		}
-	}
-	else
-	{
-		bOk = false;
-	}
-	bool bHasTextKey = false;
-	if (FString TextKeyString; LoadFromCompactBinary(*It++, TextKeyString))
-	{
-		if (!TextKeyString.IsEmpty())
-		{
-			bHasTextKey = true;
-			OutLogRecord.SetTextKey(*StringTable.Emplace_GetRef(MoveTemp(TextKeyString)));
-		}
-		else
-		{
-			OutLogRecord.SetTextKey(nullptr);
-		}
-	}
-	else
-	{
-		bOk = false;
-	}
-
-	if (bHasTextKey)
-	{
-		FLogTemplate* LogTemplate = CreateLogTemplate(OutLogRecord.GetTextNamespace(), OutLogRecord.GetTextKey(), OutLogRecord.GetFormat());
-		TemplateTable.Add(LogTemplate);
-		OutLogRecord.SetTemplate(LogTemplate);
-	}
-	else
-	{
-		FLogTemplate* LogTemplate = CreateLogTemplate(OutLogRecord.GetFormat());
-		TemplateTable.Add(LogTemplate);
-		OutLogRecord.SetTemplate(LogTemplate);
-	}
-
-	return bOk;
-}
-
-void FLogMessagesMessageHandler::FLogRecordSerializationContext::ConditionalFlush(int32 TableSize)
-{
-	if ((StringTable.Num() > TableSize) || (AnsiStringTable.Num() > TableSize) || (TemplateTable.Num() > TableSize))
-	{
-		Flush();
-	}
-}
+FGuid FLogMessagesMessageHandler::MessageType(TEXT("DB024D28203D4FBAAAF6AAD7080CF277"));
 
-void FLogMessagesMessageHandler::FLogRecordSerializationContext::Flush()
+FLogMessagesMessageHandler::FLogMessagesMessageHandler(FLogRecordSerializationContext& InLogRecordSerializationContext)
+	: LogRecordSerializationContext(InLogRecordSerializationContext)
 {
-	if (!StringTable.IsEmpty() || !AnsiStringTable.IsEmpty() || !TemplateTable.IsEmpty())
-	{
-		// NOTE: We only call FlushThreadedLogs on GLog even though we might serialize structured logs via GLog or GWarn.
-		// GWarn is an output device, but GLog is a an output redirector, and only the redirector has/needs FlushThreadedLogs.
-		// Output devices are expected to not use any pointer on a structured log record after completion of the SerializeRecord call.
-		GLog->FlushThreadedLogs();
-	}
-	for (FLogTemplate* LogTemplate : TemplateTable)
-	{
-		DestroyLogTemplate(LogTemplate);
-	}
-
-	StringTable.Empty();
-	AnsiStringTable.Empty();
-	TemplateTable.Empty();
 }
 
-FGuid FLogMessagesMessageHandler::MessageType(TEXT("DB024D28203D4FBAAAF6AAD7080CF277"));
-
 FLogMessagesMessageHandler::~FLogMessagesMessageHandler()
 {
 	if (bRegistered && GLog)
@@ -2165,67 +1942,32 @@ void FLogMessagesMessageHandler::ServerReceiveMessage(FMPCollectorServerMessageC
 		return;
 	}
 
-	for (FReplicatedLogData& LogData : Messages)
-	{
-		if (const FReplicatedLogData::FUnstructuredLogData* UnStructuredLogData = LogData.LogDataVariant.TryGet<FReplicatedLogData::FUnstructuredLogData>())
-		{
-			if (UnStructuredLogData->Category == LogCookName && UnStructuredLogData->Message.Contains(HeartbeatCategoryText))
-			{
-				// Do not spam heartbeat messages into the CookDirector log
-				continue;
-			}
-
-			FMsg::Logf(__FILE__, __LINE__, UnStructuredLogData->Category, UnStructuredLogData->Verbosity, TEXT("[CookWorker %d]: %s"),
-			Context.GetProfileId(), *UnStructuredLogData->Message);
-		}
-		else if (const FCbObject* StructuredLogObject = LogData.LogDataVariant.TryGet<FCbObject>())
+	auto MessagePassesFilter = [](FName Category, const FString& Message)
 		{
-			FLogRecord LogRecord;
-			if (LogRecordSerializationContext.Deserialize((*StructuredLogObject)["S"], LogRecord, Context.GetProfileId()))
-			{
-				FOutputDevice* LogOverride = nullptr;
-				switch (LogRecord.GetVerbosity())
-				{
-				case ELogVerbosity::Error:
-				case ELogVerbosity::Warning:
-				case ELogVerbosity::Display:
-				case ELogVerbosity::SetColor:
-					LogOverride = GWarn;
-					break;
-				default:
-					break;
-				}
-				if (LogOverride)
-				{
-					LogOverride->SerializeRecord(LogRecord);
-				}
-				else
-				{
-					GLog->SerializeRecord(LogRecord);
-				}
-			}
-		}
-		else
+			// Do not spam heartbeat messages into the CookDirector log
+			return Category != LogCookName || !Message.Contains(HeartbeatCategoryText);
+		};
+	auto TryTransformMessage = [&Context](const FString& In, FString& Out)
 		{
-			checkNoEntry();
-		}
+			Out = FString::Printf(TEXT("[CookWorker %d]: %s"), Context.GetProfileId(), *In);
+			return true;
+		};
+	for (FReplicatedLogData& LogData : Messages)
+	{
+		LogRecordSerializationContext.DeserializeAndLog(LogData, MessagePassesFilter, TryTransformMessage);
 	}
 
-	// Flush if the tables in the serialization context have exceeded 100 entries
-	const int32 TableSizeToFlushAt = 100;
-	LogRecordSerializationContext.ConditionalFlush(TableSizeToFlushAt);
+	LogRecordSerializationContext.ConditionalFlush();
 }
 
 void FLogMessagesMessageHandler::Serialize(const TCHAR* V, ELogVerbosity::Type Verbosity,
 	const FName& Category)
 {
+	FReplicatedLogData SerializedData;
+	FLogRecordSerializationContext::Serialize(SerializedData, V, Verbosity, Category);
+
 	FScopeLock QueueScopeLock(&QueueLock);
-	FReplicatedLogData& LogData = QueuedLogs.Emplace_GetRef();
-	LogData.LogDataVariant.Emplace<FReplicatedLogData::FUnstructuredLogData>();
-	FReplicatedLogData::FUnstructuredLogData& NewVal = LogData.LogDataVariant.Get<FReplicatedLogData::FUnstructuredLogData>();
-	NewVal.Message = V;
-	NewVal.Category = Category;
-	NewVal.Verbosity = Verbosity;
+	QueuedLogs.Emplace_GetRef() = MoveTemp(SerializedData);
 }
 
 void FLogMessagesMessageHandler::Serialize(const TCHAR* V, ELogVerbosity::Type Verbosity,
@@ -2236,16 +1978,11 @@ void FLogMessagesMessageHandler::Serialize(const TCHAR* V, ELogVerbosity::Type V
 
 void FLogMessagesMessageHandler::SerializeRecord(const UE::FLogRecord& Record)
 {
-	FCbWriter Writer;
-	Writer.BeginObject();
-	Writer << "S";
-	FLogRecordSerializationContext::Serialize(Writer, Record);
-	Writer.EndObject();
-	FCbObject Object = Writer.Save().AsObject();
+	FReplicatedLogData SerializedData;
+	FLogRecordSerializationContext::Serialize(SerializedData, Record);
 
 	FScopeLock QueueScopeLock(&QueueLock);
-	FReplicatedLogData& LogData = QueuedLogs.Emplace_GetRef();
-	LogData.LogDataVariant.Emplace<FCbObject>(MoveTemp(Object));
+	QueuedLogs.Emplace_GetRef() = MoveTemp(SerializedData);
 }
 
 FGuid FHeartbeatMessage::MessageType(TEXT("C08FFAF07BF34DD3A2FFB8A287CDDE83"));
diff --git a/Engine/Source/Editor/UnrealEd/Private/Cooker/CookWorkerServer.h b/Engine/Source/Editor/UnrealEd/Private/Cooker/CookWorkerServer.h
--- a/Engine/Source/Editor/UnrealEd/Private/Cooker/CookWorkerServer.h
+++ b/Engine/Source/Editor/UnrealEd/Private/Cooker/CookWorkerServer.h
@@ -31,8 +31,10 @@ namespace UE::Cook { class FCookDirector; }
 namespace UE::Cook { struct FDiscoveredPackageReplication; }
 namespace UE::Cook { enum class ECookDirectorThread : uint8; };
 namespace UE::Cook { struct FGeneratorEventMessage; }
+namespace UE::Cook { class FLogRecordSerializationContext; }
 namespace UE::Cook { struct FPackageData; }
 namespace UE::Cook { struct FPackageResultsMessage; }
+namespace UE::Cook { struct FReplicatedLogData; }
 namespace UE::Cook { struct FWorkerConnectMessage; }
 
 namespace UE::Cook
@@ -490,20 +492,6 @@ public:
 	static FGuid MessageType;
 };
 
-/** Stores the data passed into FOutputDevice::Serialize, for replication to the CookDirector. */
-struct FReplicatedLogData
-{
-	struct FUnstructuredLogData
-	{
-		FString Message;
-		FName Category;
-		ELogVerbosity::Type Verbosity;
-	};
-	TVariant<FUnstructuredLogData, FCbObject> LogDataVariant;
-};
-FCbWriter& operator<<(FCbWriter& Writer, const FReplicatedLogData& LogData);
-bool LoadFromCompactBinary(FCbFieldView Field, FReplicatedLogData& OutLogData);
-
 /**
  * Send log messages from CookWorkers to the CookDirector, which marks them up with the CookWorkerId and
  * then prints them to its own log.
@@ -511,6 +499,7 @@ bool LoadFromCompactBinary(FCbFieldView Field, FReplicatedLogData& OutLogData);
 class FLogMessagesMessageHandler : public IMPCollector, public FOutputDevice
 {
 public:
+	FLogMessagesMessageHandler(FLogRecordSerializationContext& InLogRecordSerializationContext);
 	~FLogMessagesMessageHandler();
 	void InitializeClient();
 
@@ -529,26 +518,8 @@ public:
 	virtual bool CanBeUsedOnMultipleThreads() const override { return true; }
 
 private:
-	class FLogRecordSerializationContext
-	{
-	public:
-		~FLogRecordSerializationContext()
-		{
-			Flush();
-		}
-
-		static FCbWriter& Serialize(FCbWriter& Writer, const FLogRecord& LogRecord);
-		bool Deserialize(FCbFieldView Field, FLogRecord& OutLogRecord, int32 ProfileId);
-
-		void ConditionalFlush(int32 TableSize);
-		void Flush();
-	private:
-		TArray<FString> StringTable;
-		TArray<FAnsiString> AnsiStringTable;
-		TArray<FLogTemplate*> TemplateTable;
-	};
 
-	FLogRecordSerializationContext LogRecordSerializationContext;
+	FLogRecordSerializationContext& LogRecordSerializationContext;
 
 	FCriticalSection QueueLock;
 	TArray<FReplicatedLogData> QueuedLogs;
