Deleted two commented out lines of code.

#rnx

diff --git a/Engine/Source/Editor/UnrealEd/Private/CookOnTheFlyServer.cpp b/Engine/Source/Editor/UnrealEd/Private/CookOnTheFlyServer.cpp
--- a/Engine/Source/Editor/UnrealEd/Private/CookOnTheFlyServer.cpp
+++ b/Engine/Source/Editor/UnrealEd/Private/CookOnTheFlyServer.cpp
@@ -321,8 +321,6 @@ void LogCookerMessage( const FString& MessageText, EMessageSeverity::Type Severi
 	TSharedRef<FTokenizedMessage> Message = FTokenizedMessage::Create(Severity);
 
 	Message->AddToken( FTextToken::Create( FText::FromString(MessageText) ) );
-	// Message->AddToken(FTextToken::Create(MessageLogTextDetail)); 
-	// Message->AddToken(FDocumentationToken::Create(TEXT("https://docs.unrealengine.com/latest/INT/Platforms/iOS/QuickStart/6/index.html"))); 
 	MessageLog.AddMessage(Message);
 
 	MessageLog.Notify(FText(), EMessageSeverity::Warning, false);
