Package has no linker
#rnx

diff --git a/Engine/Source/Editor/UnrealEd/Private/Commandlets/BaseIteratePackagesCommandlet.cpp b/Engine/Source/Editor/UnrealEd/Private/Commandlets/BaseIteratePackagesCommandlet.cpp
--- a/Engine/Source/Editor/UnrealEd/Private/Commandlets/BaseIteratePackagesCommandlet.cpp
+++ b/Engine/Source/Editor/UnrealEd/Private/Commandlets/BaseIteratePackagesCommandlet.cpp
@@ -614,7 +614,14 @@ void UBaseIteratePackagesCommandlet::LoadAndSaveOnePackage(const FString& Filena
 			}
 
 			FLinkerLoad* Linker = Package->GetLinker();
-			PerformPreloadOperations(Linker, bSavePackage); 
+			if(Linker)
+			{
+				PerformPreloadOperations(Linker, bSavePackage); 
+			}
+			else
+			{
+				UE_LOG(LogIteratePackagesCommandlet, Warning, TEXT("Package %s had no linker. Skipping PreloadOperations"), *PackagePath.GetPackageName());
+			}
 			VerboseMessage(FString::Printf(TEXT("Post PerformPreloadOperations, Resave? %d"), bSavePackage));
 
 			VerboseMessage(TEXT("Post LoadPackage"));
