IncrementalCook: Replaying Logs: Move log-listener functionality from FLogMessagesMessageHandler into FLogRecordSerializationContext and make it also pass the incoming logs to the FPackageData that was active when the logs were sent. Rename the expanded FLogRecordSerializationContext to FLogHandler to indicate its increased responsibilities. When logs come from a CookWorker, skip them on the cook director rather rather than saving them to the CookDirector's package. Add function to receive stored from CookRequestCluster when packages are incrementally skipped.
#rnx
#incrementalcook

diff --git a/Engine/Source/Editor/UnrealEd/Classes/CookOnTheSide/CookOnTheFlyServer.h b/Engine/Source/Editor/UnrealEd/Classes/CookOnTheSide/CookOnTheFlyServer.h
--- a/Engine/Source/Editor/UnrealEd/Classes/CookOnTheSide/CookOnTheFlyServer.h
+++ b/Engine/Source/Editor/UnrealEd/Classes/CookOnTheSide/CookOnTheFlyServer.h
@@ -180,7 +180,7 @@ namespace UE::Cook
 	class FCookWorkerClient;
 	class FCookWorkerServer;
 	class FDiagnostics;
-	class FLogRecordSerializationContext;
+	class FLogHandler;
 	class FODSCClientData; 
 	class FPackagePreloader;
 	class FPackageWriterMPCollector;
@@ -1638,7 +1638,7 @@ private:
 	TUniquePtr<FLayeredCookArtifactReader> AllContextArtifactReader;
 	TSharedPtr<FLooseFilesCookArtifactReader> SharedLooseFilesCookArtifactReader;
 	TUniquePtr<UE::Cook::FCookGCDiagnosticContext> GCDiagnosticContext;
-	TUniquePtr<UE::Cook::FLogRecordSerializationContext> LogRecordSerializationContext;
+	TUniquePtr<UE::Cook::FLogHandler> LogHandler;
 
 	TArray<UE::Cook::FCookSavePackageContext*> SavePackageContexts;
 	/**
@@ -1698,6 +1698,7 @@ private:
 	friend UE::Cook::FDiagnostics;
 	friend UE::Cook::FGenerationHelper;
 	friend UE::Cook::FInitializeConfigSettings;
+	friend UE::Cook::FLogHandler;
 	friend UE::Cook::FPackageData;
 	friend UE::Cook::FPackageDatas;
 	friend UE::Cook::FPackagePreloader;
diff --git a/Engine/Source/Editor/UnrealEd/Private/CookOnTheFlyServer.cpp b/Engine/Source/Editor/UnrealEd/Private/CookOnTheFlyServer.cpp
--- a/Engine/Source/Editor/UnrealEd/Private/CookOnTheFlyServer.cpp
+++ b/Engine/Source/Editor/UnrealEd/Private/CookOnTheFlyServer.cpp
@@ -1608,6 +1608,7 @@ void UCookOnTheFlyServer::TickCookStatus(UE::Cook::FTickStackData& StackData)
 	UpdateDisplay(StackData, false /* bForceDisplay */);
 	ProcessAsyncLoading(false, false, 0.0f);
 	ProcessUnsolicitedPackages();
+	LogHandler->FlushIncrementalCookLogs();
 	PumpExternalRequests(StackData.Timer);
 }
 
@@ -5816,7 +5817,7 @@ void UCookOnTheFlyServer::Initialize( ECookMode::Type DesiredCookMode, ECookInit
 	AssetRegistry = IAssetRegistry::Get();
 	GCDiagnosticContext = MakeUnique<UE::Cook::FCookGCDiagnosticContext>();
 	StallDetector = MakeUnique<UE::Cook::FStallDetector>();
-	LogRecordSerializationContext.Reset(new UE::Cook::FLogRecordSerializationContext());
+	LogHandler.Reset(new UE::Cook::FLogHandler(*this));
 
 	if (!IsCookWorkerMode())
 	{
diff --git a/Engine/Source/Editor/UnrealEd/Private/Cooker/CookDirector.cpp b/Engine/Source/Editor/UnrealEd/Private/Cooker/CookDirector.cpp
--- a/Engine/Source/Editor/UnrealEd/Private/Cooker/CookDirector.cpp
+++ b/Engine/Source/Editor/UnrealEd/Private/Cooker/CookDirector.cpp
@@ -147,7 +147,7 @@ FCookDirector::FCookDirector(UCookOnTheFlyServer& InCOTFS, int32 CookProcessCoun
 	}
 	bMultiprocessAvailable = true;
 
-	Register(new FLogMessagesMessageHandler(*COTFS.LogRecordSerializationContext));
+	Register(new FLogMessagesMessageHandler(*COTFS.LogHandler));
 	Register(new TMPCollectorServerMessageCallback<FRetractionResultsMessage>([this]
 	(FMPCollectorServerMessageContext& Context, bool bReadSuccessful, FRetractionResultsMessage&& Message)
 		{
diff --git a/Engine/Source/Editor/UnrealEd/Private/Cooker/CookLog.cpp b/Engine/Source/Editor/UnrealEd/Private/Cooker/CookLog.cpp
--- a/Engine/Source/Editor/UnrealEd/Private/Cooker/CookLog.cpp
+++ b/Engine/Source/Editor/UnrealEd/Private/Cooker/CookLog.cpp
@@ -5,6 +5,9 @@
 #include "Cooker/CookLogPrivate.h"
 
 #include "Containers/AnsiString.h"
+#include "Cooker/CookPackageData.h"
+#include "Cooker/CookPlatformManager.h"
+#include "Cooker/CookWorkerClient.h"
 #include "CoreGlobals.h"
 #include "Logging/StructuredLog.h"
 #include "Logging/StructuredLogFormat.h"
@@ -12,6 +15,8 @@
 #include "Misc/DateTime.h"
 #include "Misc/FeedbackContext.h"
 #include "Misc/OutputDeviceRedirector.h"
+#include "Misc/PackageAccessTracking.h"
+#include "Misc/ScopeLock.h"
 #include "Misc/StringBuilder.h"
 #include "Serialization/CompactBinary.h"
 #include "Serialization/CompactBinaryWriter.h"
@@ -97,36 +102,153 @@ bool LoadFromCompactBinary(FCbFieldView Field, FReplicatedLogData& OutLogData)
 	return bOk;
 }
 
-void FLogRecordSerializationContext::Serialize(FReplicatedLogData& OutData, const TCHAR* V,
+FLogHandler::FLogHandler(UCookOnTheFlyServer& InCOTFS)
+	: COTFS(InCOTFS)
+{
+	check(!bRegistered);
+	check(GLog);
+	GLog->AddOutputDevice(this);
+	bRegistered = true;
+}
+
+FLogHandler::~FLogHandler()
+{
+	PruneReplay();
+
+	if (bRegistered)
+	{
+		if (GLog)
+		{
+			GLog->RemoveOutputDevice(this);
+		}
+		bRegistered = false;
+	}
+}
+
+void FLogHandler::FlushIncrementalCookLogs()
+{
+	TArray<FQueuedLog> LocalQueuedLogs;
+	{
+		FScopeLock ScopeLock(&QueuedLogsForIncrementalCookLock);
+		LocalQueuedLogs = MoveTemp(QueuedLogsForIncrementalCook);
+	}
+
+	for (FQueuedLog& QueuedLog : LocalQueuedLogs)
+	{
+		RecordLogForIncrementalCookGameThreadPortion(QueuedLog.ActivePackage, MoveTemp(QueuedLog.LogData));
+	}
+}
+
+void FLogHandler::ReplayLogFromIncrementallySkipped(FReplicatedLogData&& LogData)
+{
+	// Replays only come from MarkPackageIterativelySkipped, which happens only on the CookDirector, during
+	// CookRequestCluster traversal. We rely on that, and do not report whether messages from CookWorkers
+	// came from a replay or an active log, we always assume they came from active logs. So we currently forbid
+	// replay on CookWorkers.
+	check(!COTFS.CookWorkerClient);
+
+	auto MessagePassesFilter = [](FName Category, const FString& Message)
+		{
+			return true;
+		};
+	auto TryTransformMessage = [](const FString& In, FString& Out)
+		{
+			return false;
+		};
+	UnMarshalAndLog(LogData, MessagePassesFilter, TryTransformMessage);
+}
+
+void FLogHandler::ReplayLogFromCookWorker(FReplicatedLogData&& LogData, int32 CookWorkerProfileId)
+{
+	auto MessagePassesFilter = [](FName Category, const FString& Message)
+		{
+			// Do not spam heartbeat messages into the CookDirector log
+			return Category != LogCookName || !Message.Contains(HeartbeatCategoryText);
+		};
+	auto TryTransformMessage = [CookWorkerProfileId](const FString& In, FString& Out)
+		{
+			Out = FString::Printf(TEXT("[CookWorker %d]: %s"), CookWorkerProfileId, *In);
+			return true;
+		};
+	UnMarshalAndLog(LogData, MessagePassesFilter, TryTransformMessage);
+}
+
+void FLogHandler::ConditionalPruneReplay()
+{
+	// Flush if the tables in the serialization context have exceeded 100 entries
+	const int32 TableSizeToFlushAt = 100;
+	if ((StringTable.Num() > TableSizeToFlushAt)
+		|| (AnsiStringTable.Num() > TableSizeToFlushAt)
+		|| (TemplateTable.Num() > TableSizeToFlushAt))
+	{
+		PruneReplay();
+	}
+}
+
+void FLogHandler::PruneReplay()
+{
+	// We are going to drop data from our tables that might be pointed to from logs still pending in GLog. So Flush
+	// logs before we prune the tables.
+	if (!StringTable.IsEmpty() || !AnsiStringTable.IsEmpty() || !TemplateTable.IsEmpty())
+	{
+		// NOTE: We only call FlushThreadedLogs on GLog even though we might serialize structured logs via GLog or GWarn.
+		// GWarn is an output device, but GLog is a an output redirector, and only the redirector has/needs FlushThreadedLogs.
+		// Output devices are expected to not use any pointer on a structured log record after completion of the SerializeRecord call.
+		GLog->FlushThreadedLogs();
+	}
+	for (FLogTemplate* LogTemplate : TemplateTable)
+	{
+		DestroyLogTemplate(LogTemplate);
+	}
+
+	StringTable.Empty();
+	AnsiStringTable.Empty();
+	TemplateTable.Empty();
+}
+
+void FLogHandler::Serialize(const TCHAR* V, ELogVerbosity::Type Verbosity,
+	const FName& Category)
+{
+	FReplicatedLogData SerializedData;
+	FStringView FormatString(V);
+	Marshal(SerializedData, V, Verbosity, Category);
+	ReportActiveLog(MoveTemp(SerializedData), FormatString, Verbosity);
+}
+
+void FLogHandler::Serialize(const TCHAR* V, ELogVerbosity::Type Verbosity,
+	const FName& Category, const double Time)
+{
+	Serialize(V, Verbosity, Category);
+}
+
+void FLogHandler::SerializeRecord(const UE::FLogRecord& Record)
+{
+	FReplicatedLogData SerializedData;
+	Marshal(SerializedData, Record);
+	ReportActiveLog(MoveTemp(SerializedData), Record.GetFormat(), Record.GetVerbosity());
+}
+
+void FLogHandler::Flush()
+{
+	PruneReplay();
+}
+
+void FLogHandler::Marshal(FReplicatedLogData& OutData, FStringView Message,
 	ELogVerbosity::Type Verbosity, const FName& Category)
 {
 	OutData.LogDataVariant.Emplace<FReplicatedLogData::FUnstructuredLogData>();
 	FReplicatedLogData::FUnstructuredLogData& OutVal
 		= OutData.LogDataVariant.Get<FReplicatedLogData::FUnstructuredLogData>();
-	OutVal.Message = V;
+	OutVal.Message = Message;
 	OutVal.Category = Category;
 	OutVal.Verbosity = Verbosity;
 }
 
-void FLogRecordSerializationContext::Serialize(FReplicatedLogData& OutData, const TCHAR* V,
-	ELogVerbosity::Type Verbosity, const FName& Category, const double Time)
-{
-	Serialize(OutData, V, Verbosity, Category);
-}
-
-void FLogRecordSerializationContext::Serialize(FReplicatedLogData& OutData, const FLogRecord& LogRecord)
+void FLogHandler::Marshal(FReplicatedLogData& OutData, const FLogRecord& LogRecord)
 {
 	FCbWriter Writer;
 	Writer.BeginObject();
 	Writer << "S";
-	Serialize(Writer, LogRecord);
-	Writer.EndObject();
-	FCbObject Object = Writer.Save().AsObject();
-	OutData.LogDataVariant.Emplace<FCbObject>(MoveTemp(Object));
-}
-
-FCbWriter& FLogRecordSerializationContext::Serialize(FCbWriter& Writer, const FLogRecord& LogRecord)
-{
 	Writer.BeginArray();
 	Writer << LogRecord.GetCategory();
 	Writer << static_cast<uint8>(LogRecord.GetVerbosity());
@@ -138,10 +260,12 @@ FCbWriter& FLogRecordSerializationContext::Serialize(FCbWriter& Writer, const FL
 	Writer << LogRecord.GetTextNamespace();
 	Writer << LogRecord.GetTextKey();
 	Writer.EndArray();
-	return Writer;
+	Writer.EndObject();
+	FCbObject Object = Writer.Save().AsObject();
+	OutData.LogDataVariant.Emplace<FCbObject>(MoveTemp(Object));
 }
 
-void FLogRecordSerializationContext::DeserializeAndLog(const FReplicatedLogData& LogData,
+void FLogHandler::UnMarshalAndLog(const FReplicatedLogData& LogData,
 	TFunctionRef<bool(FName, const FString&)> MessagePassesFilter,
 	TFunctionRef<bool(const FString&, FString&)> TryTransformMessage)
 {
@@ -165,7 +289,7 @@ void FLogRecordSerializationContext::DeserializeAndLog(const FReplicatedLogData&
 	else if (const FCbObject* StructuredLogObject = LogData.LogDataVariant.TryGet<FCbObject>())
 	{
 		FLogRecord LogRecord;
-		if (Deserialize((*StructuredLogObject)["S"], LogRecord, MessagePassesFilter, TryTransformMessage))
+		if (UnMarshal((*StructuredLogObject)["S"], LogRecord, MessagePassesFilter, TryTransformMessage))
 		{
 			FOutputDevice* LogOverride = nullptr;
 			switch (LogRecord.GetVerbosity())
@@ -193,9 +317,9 @@ void FLogRecordSerializationContext::DeserializeAndLog(const FReplicatedLogData&
 	{
 		checkNoEntry();
 	}
-
 }
-bool FLogRecordSerializationContext::Deserialize(FCbFieldView Field, FLogRecord& OutLogRecord,
+
+bool FLogHandler::UnMarshal(FCbFieldView Field, FLogRecord& OutLogRecord,
 	TFunctionRef<bool(FName, const FString&)> MessagePassesFilter,
 	TFunctionRef<bool(const FString&, FString&)> TryTransformMessage)
 {
@@ -315,35 +439,73 @@ bool FLogRecordSerializationContext::Deserialize(FCbFieldView Field, FLogRecord&
 	return bOk;
 }
 
-void FLogRecordSerializationContext::ConditionalFlush()
+void FLogHandler::ReportActiveLog(FReplicatedLogData&& LogData, FStringView FormatMessage,
+	ELogVerbosity::Type Verbosity)
 {
-	// Flush if the tables in the serialization context have exceeded 100 entries
-	const int32 TableSizeToFlushAt = 100;
-	if ((StringTable.Num() > TableSizeToFlushAt)
-		|| (AnsiStringTable.Num() > TableSizeToFlushAt)
-		|| (TemplateTable.Num() > TableSizeToFlushAt))
+	if (COTFS.CookWorkerClient)
+	{
+		COTFS.CookWorkerClient->ReportLogMessage(LogData);
+	}
+	else if (COTFS.CookDirector)
 	{
-		Flush();
+		if (FormatMessage.StartsWith(TEXT("[CookWorker")))
+		{
+			// Do not store logs from CookWorkers; only the CookWorker saving the package needs to store those logs.
+			return;
+		}
 	}
+	RecordLogForIncrementalCook(MoveTemp(LogData), Verbosity);
 }
 
-void FLogRecordSerializationContext::Flush()
+void FLogHandler::RecordLogForIncrementalCook(UE::Cook::FReplicatedLogData&& LogData,
+	ELogVerbosity::Type LogVerbosity)
 {
-	if (!StringTable.IsEmpty() || !AnsiStringTable.IsEmpty() || !TemplateTable.IsEmpty())
+	// Note that this function can be called from any thread. Only threadsafe data only can be accessed.
+	if (LogVerbosity > ELogVerbosity::Warning)
 	{
-		// NOTE: We only call FlushThreadedLogs on GLog even though we might serialize structured logs via GLog or GWarn.
-		// GWarn is an output device, but GLog is a an output redirector, and only the redirector has/needs FlushThreadedLogs.
-		// Output devices are expected to not use any pointer on a structured log record after completion of the SerializeRecord call.
-		GLog->FlushThreadedLogs();
+		// Only warnings and errors are recorded; we don't want to spam display logs and they would waste memory to record.
+		return;
 	}
-	for (FLogTemplate* LogTemplate : TemplateTable)
+	PackageAccessTracking_Private::FTrackedData* AccumulatedScopeData
+		= PackageAccessTracking_Private::FPackageAccessRefScope::GetCurrentThreadAccumulatedData();
+	if (!AccumulatedScopeData)
 	{
-		DestroyLogTemplate(LogTemplate);
+		return;
+	}
+	FName ActivePackage = AccumulatedScopeData->PackageName;
+	if (ActivePackage.IsNone())
+	{
+		return;
 	}
 
-	StringTable.Empty();
-	AnsiStringTable.Empty();
-	TemplateTable.Empty();
+	if (!IsInGameThread())
+	{
+		// The rest of the function requires access to schedulerthread-only data, so queue it.
+		FScopeLock ScopeLock(&QueuedLogsForIncrementalCookLock);
+		QueuedLogsForIncrementalCook.Emplace(ActivePackage, MoveTemp(LogData));
+	}
+	else
+	{
+		RecordLogForIncrementalCookGameThreadPortion(ActivePackage, MoveTemp(LogData));
+	}
+}
+
+void FLogHandler::RecordLogForIncrementalCookGameThreadPortion(FName ActivePackage, FReplicatedLogData&& LogData)
+{
+	FPackageData* PackageData = COTFS.PackageDatas->TryAddPackageDataByPackageName(ActivePackage);
+	if (!PackageData)
+	{
+		return;
+	}
+
+	// We want to avoid wasting memory for packages if they have already saved, which we can do because we will not
+	// have an opportunity to save the data for them anyway so it causes no change in behavior.
+	if (PackageData->HasAllCommittedPlatforms(COTFS.PlatformManager->GetSessionPlatforms()))
+	{
+		return;
+	}
+
+	//PackageData->AddLogMessage(MoveTemp(LogData));
 }
 
 } // namespace UE::Cook
\ No newline at end of file
diff --git a/Engine/Source/Editor/UnrealEd/Private/Cooker/CookLogPrivate.h b/Engine/Source/Editor/UnrealEd/Private/Cooker/CookLogPrivate.h
--- a/Engine/Source/Editor/UnrealEd/Private/Cooker/CookLogPrivate.h
+++ b/Engine/Source/Editor/UnrealEd/Private/Cooker/CookLogPrivate.h
@@ -2,17 +2,19 @@
 
 #pragma once
 
-#include "Containers/StringView.h"
 #include "Containers/AnsiString.h"
+#include "Containers/StringView.h"
 #include "Containers/UnrealString.h"
 #include "CookOnTheSide/CookLog.h"
 #include "Logging/LogMacros.h"
 #include "Logging/LogVerbosity.h"
+#include "Misc/OutputDevice.h"
 #include "Misc/TVariant.h"
 #include "Serialization/CompactBinary.h"
 #include "UObject/NameTypes.h"
 #include "Templates/Function.h"
 
+class UCookOnTheFlyServer;
 namespace UE { class FLogTemplate; }
 namespace UE { class FLogRecord; }
 
@@ -35,34 +37,84 @@ struct FReplicatedLogData
 FCbWriter& operator<<(FCbWriter& Writer, const FReplicatedLogData& LogData);
 bool LoadFromCompactBinary(FCbFieldView Field, FReplicatedLogData& OutLogData);
 
-class FLogRecordSerializationContext
+/**
+ * The Cooker's listener to log messages. It passes the logs onto CookWorkerClient for reporting to CookDirector in
+ * MPCook, and it stores the log messages on the active package for storage in incremental cooks.
+ */
+class FLogHandler : public FOutputDevice
 {
 public:
-	~FLogRecordSerializationContext()
+	FLogHandler(UCookOnTheFlyServer& InCOTFS);
+	~FLogHandler();
+
+	void ReplayLogFromIncrementallySkipped(FReplicatedLogData&& LogData);
+	void ReplayLogFromCookWorker(FReplicatedLogData&& LogData, int32 CookWorkerProfileId);
+	void ConditionalPruneReplay();
+	void FlushIncrementalCookLogs();
+
+	// FOutputDevice
+	virtual void Serialize(const TCHAR* V, ELogVerbosity::Type Verbosity, const FName& Category) override;
+	virtual void Serialize(const TCHAR* V, ELogVerbosity::Type Verbosity, const FName& Category,
+		const double Time) override;
+	virtual void SerializeRecord(const FLogRecord& Record) override;
+	virtual void Flush() override;
+	virtual bool CanBeUsedOnAnyThread() const override;
+	virtual bool CanBeUsedOnMultipleThreads() const override;
+
+private:
+	struct FQueuedLog
 	{
-		Flush();
-	}
+		FName ActivePackage;
+		FReplicatedLogData LogData;
+	};
 
-	static void Serialize(FReplicatedLogData& OutData, const TCHAR* V, ELogVerbosity::Type Verbosity,
+	void Marshal(FReplicatedLogData& OutData, FStringView Message, ELogVerbosity::Type Verbosity,
 		const FName& Category);
-	static void Serialize(FReplicatedLogData& OutData, const TCHAR* V, ELogVerbosity::Type Verbosity,
-		const FName& Category, const double Time);
-	static void Serialize(FReplicatedLogData& OutData, const FLogRecord& LogRecord);
-	void DeserializeAndLog(const FReplicatedLogData& LogData,
+	void Marshal(FReplicatedLogData& OutData, const FLogRecord& LogRecord);
+	void UnMarshalAndLog(const FReplicatedLogData& LogData,
 		TFunctionRef<bool(FName, const FString&)> MessagePassesFilter,
 		TFunctionRef<bool(const FString&, FString&)> TryTransformMessage);
-
-	void ConditionalFlush();
-	void Flush();
-private:
-	static FCbWriter& Serialize(FCbWriter& Writer, const FLogRecord& LogRecord);
-	bool Deserialize(FCbFieldView Field, FLogRecord& OutLogRecord,
+	bool UnMarshal(FCbFieldView Field, FLogRecord& OutLogRecord,
 		TFunctionRef<bool(FName, const FString&)> MessagePassesFilter,
 		TFunctionRef<bool(const FString&, FString&)> TryTransformMessage);
 
+	void ReportActiveLog(FReplicatedLogData&& LogData, FStringView FormatMessage, ELogVerbosity::Type Verbosity);
+	void RecordLogForIncrementalCook(FReplicatedLogData&& LogData, ELogVerbosity::Type Verbosity);
+	void RecordLogForIncrementalCookGameThreadPortion(FName PackageName, FReplicatedLogData&& LogData);
+
+	void PruneReplay();
+
+private:
+	UCookOnTheFlyServer& COTFS;
+	bool bRegistered = false;
+
+	FCriticalSection QueuedLogsForIncrementalCookLock;
+	TArray<FQueuedLog> QueuedLogsForIncrementalCook;
+
+	FCriticalSection TableLock;
 	TArray<FString> StringTable;
 	TArray<FAnsiString> AnsiStringTable;
 	TArray<FLogTemplate*> TemplateTable;
 };
 
+
+constexpr FStringView HeartbeatCategoryText(TEXTVIEW("CookWorkerHeartbeat:"));
+
+} // namespace UE::Cook
+
+// Inline Implementations
+
+namespace UE::Cook
+{
+
+inline bool FLogHandler::CanBeUsedOnAnyThread() const
+{
+	return true;
+}
+
+inline bool FLogHandler::CanBeUsedOnMultipleThreads() const
+{
+	return true;
+}
+
 } // namespace UE::Cook
\ No newline at end of file
diff --git a/Engine/Source/Editor/UnrealEd/Private/Cooker/CookRequestCluster.cpp b/Engine/Source/Editor/UnrealEd/Private/Cooker/CookRequestCluster.cpp
--- a/Engine/Source/Editor/UnrealEd/Private/Cooker/CookRequestCluster.cpp
+++ b/Engine/Source/Editor/UnrealEd/Private/Cooker/CookRequestCluster.cpp
@@ -12,6 +12,7 @@
 #include "Cooker/CookDependency.h"
 #include "Cooker/CookGenerationHelper.h"
 #include "Cooker/CookImportsChecker.h"
+#include "Cooker/CookLogPrivate.h"
 #include "Cooker/CookPackageData.h"
 #include "Cooker/CookPlatformManager.h"
 #include "Cooker/CookProfiling.h"
@@ -747,6 +748,10 @@ void FRequestCluster::PumpExploration(const FCookerTimer& CookerTimer, bool& bOu
 		return;
 	}
 
+	ON_SCOPE_EXIT
+	{
+		COTFS.LogHandler->ConditionalPruneReplay();
+	};
 	if (!GraphSearch.IsStarted())
 	{
 		GraphSearch.Initialize();
diff --git a/Engine/Source/Editor/UnrealEd/Private/Cooker/CookSavePackage.cpp b/Engine/Source/Editor/UnrealEd/Private/Cooker/CookSavePackage.cpp
--- a/Engine/Source/Editor/UnrealEd/Private/Cooker/CookSavePackage.cpp
+++ b/Engine/Source/Editor/UnrealEd/Private/Cooker/CookSavePackage.cpp
@@ -11,6 +11,7 @@
 #include "Cooker/CookEvents.h"
 #include "Cooker/CookGenerationHelper.h"
 #include "Cooker/CookImportsChecker.h"
+#include "Cooker/CookLogPrivate.h"
 #include "Cooker/CookOnTheFlyServerInterface.h"
 #include "Cooker/CookPackageArtifacts.h"
 #include "Cooker/CookPackageData.h"
@@ -465,6 +466,9 @@ void FSaveCookedPackageContext::FinishPlatform()
 		// packages loaded during Save,BeginCacheForCookedPlatformData, or Generator functions. These dependencies
 		// are added to the runtime dependencies stored in the oplog, so we need to know them now.
 		COTFS.ProcessUnsolicitedPackages();
+		// Flush any logs that were logged from other threads into the Package's RecordedCookLogs so we can store
+		// them in the CommitAttachments.
+		COTFS.LogHandler->FlushIncrementalCookLogs();
 
 		// Collect dependencies from all sources and record them.
 		CalculatePlatformAgnosticRuntimeDependencies();
diff --git a/Engine/Source/Editor/UnrealEd/Private/Cooker/CookWorkerClient.cpp b/Engine/Source/Editor/UnrealEd/Private/Cooker/CookWorkerClient.cpp
--- a/Engine/Source/Editor/UnrealEd/Private/Cooker/CookWorkerClient.cpp
+++ b/Engine/Source/Editor/UnrealEd/Private/Cooker/CookWorkerClient.cpp
@@ -6,6 +6,7 @@
 #include "Cooker/CompactBinaryTCP.h"
 #include "Cooker/CookDirector.h"
 #include "Cooker/CookGenerationHelper.h"
+#include "Cooker/CookLogPrivate.h"
 #include "Cooker/CookPackageData.h"
 #include "Cooker/CookPlatformManager.h"
 #include "Cooker/CookTypes.h"
@@ -32,8 +33,7 @@ constexpr float WaitForConnectReplyTimeout = 60.f;
 FCookWorkerClient::FCookWorkerClient(UCookOnTheFlyServer& InCOTFS)
 	: COTFS(InCOTFS)
 {
-	LogMessageHandler = new FLogMessagesMessageHandler(*COTFS.LogRecordSerializationContext);
-	LogMessageHandler->InitializeClient();
+	LogMessageHandler = new FLogMessagesMessageHandler(*COTFS.LogHandler);
 	Register(LogMessageHandler);
 	Register(new TMPCollectorClientMessageCallback<FRetractionRequestMessage>([this]
 	(FMPCollectorClientMessageContext& Context, bool bReadSuccessful, FRetractionRequestMessage&& Message)
@@ -334,6 +334,11 @@ void FCookWorkerClient::ReportDiscoveredPackage(const FPackageData& PackageData,
 	}
 }
 
+void FCookWorkerClient::ReportLogMessage(const FReplicatedLogData& LogData)
+{
+	LogMessageHandler->ClientReportLogMessage(LogData);
+}
+
 void FCookWorkerClient::ReportGeneratorQueuedGeneratedPackages(FGenerationHelper& GenerationHelper)
 {
 	PendingGeneratorEvents.Add(FGeneratorEventMessage(EGeneratorEvent::QueuedGeneratedPackages,
diff --git a/Engine/Source/Editor/UnrealEd/Private/Cooker/CookWorkerClient.h b/Engine/Source/Editor/UnrealEd/Private/Cooker/CookWorkerClient.h
--- a/Engine/Source/Editor/UnrealEd/Private/Cooker/CookWorkerClient.h
+++ b/Engine/Source/Editor/UnrealEd/Private/Cooker/CookWorkerClient.h
@@ -25,6 +25,7 @@ namespace UE::Cook { struct FInitialConfigMessage; }
 namespace UE::Cook { struct FGenerationHelper; }
 namespace UE::Cook { struct FGeneratorEventMessage; }
 namespace UE::Cook { struct FPackageRemoteResult; }
+namespace UE::Cook { struct FReplicatedLogData; }
 namespace UE::Cook { struct FRetractionRequestMessage; }
 
 namespace UE::Cook
@@ -70,6 +71,8 @@ public:
 	/** Queue a message to the server that a package was discovered as needed in the cook. Will be sent during Tick. */
 	void ReportDiscoveredPackage(const FPackageData& PackageData, const FInstigator& Instigator,
 		FDiscoveredPlatformSet&& ReachablePlatforms, FGenerationHelper* ParentGenerationHelper, EUrgency Urgency);
+	void ReportLogMessage(const FReplicatedLogData& LogData);
+
 	/**
 	 * Queue a message to the server that a PackageGenerator queued its generated packages for cooking, and will
 	 * keep itself in memory until it receives an EGeneratorEvent::QueuedGeneratedPackagesFencePassed.
diff --git a/Engine/Source/Editor/UnrealEd/Private/Cooker/CookWorkerServer.cpp b/Engine/Source/Editor/UnrealEd/Private/Cooker/CookWorkerServer.cpp
--- a/Engine/Source/Editor/UnrealEd/Private/Cooker/CookWorkerServer.cpp
+++ b/Engine/Source/Editor/UnrealEd/Private/Cooker/CookWorkerServer.cpp
@@ -1895,23 +1895,15 @@ FGuid FGeneratorEventMessage::MessageType(TEXT("B6EE94CA70EC4F40B0D2214EDC11ED03
 
 FGuid FLogMessagesMessageHandler::MessageType(TEXT("DB024D28203D4FBAAAF6AAD7080CF277"));
 
-FLogMessagesMessageHandler::FLogMessagesMessageHandler(FLogRecordSerializationContext& InLogRecordSerializationContext)
-	: LogRecordSerializationContext(InLogRecordSerializationContext)
+FLogMessagesMessageHandler::FLogMessagesMessageHandler(FLogHandler& InCOTFSLogHandler)
+	: COTFSLogHandler(InCOTFSLogHandler)
 {
 }
 
-FLogMessagesMessageHandler::~FLogMessagesMessageHandler()
+void FLogMessagesMessageHandler::ClientReportLogMessage(const FReplicatedLogData& LogData)
 {
-	if (bRegistered && GLog)
-	{
-		GLog->RemoveOutputDevice(this);
-	}
-}
-void FLogMessagesMessageHandler::InitializeClient()
-{
-	check(!bRegistered);
-	GLog->AddOutputDevice(this);
-	bRegistered = true;
+	FScopeLock QueueScopeLock(&QueueLock);
+	QueuedLogs.Add(LogData);
 }
 
 void FLogMessagesMessageHandler::ClientTick(FMPCollectorClientTickContext& Context)
@@ -1942,47 +1934,12 @@ void FLogMessagesMessageHandler::ServerReceiveMessage(FMPCollectorServerMessageC
 		return;
 	}
 
-	auto MessagePassesFilter = [](FName Category, const FString& Message)
-		{
-			// Do not spam heartbeat messages into the CookDirector log
-			return Category != LogCookName || !Message.Contains(HeartbeatCategoryText);
-		};
-	auto TryTransformMessage = [&Context](const FString& In, FString& Out)
-		{
-			Out = FString::Printf(TEXT("[CookWorker %d]: %s"), Context.GetProfileId(), *In);
-			return true;
-		};
+	int32 CookWorkerProfileId = Context.GetProfileId();
 	for (FReplicatedLogData& LogData : Messages)
 	{
-		LogRecordSerializationContext.DeserializeAndLog(LogData, MessagePassesFilter, TryTransformMessage);
+		COTFSLogHandler.ReplayLogFromCookWorker(MoveTemp(LogData), CookWorkerProfileId);
 	}
-
-	LogRecordSerializationContext.ConditionalFlush();
-}
-
-void FLogMessagesMessageHandler::Serialize(const TCHAR* V, ELogVerbosity::Type Verbosity,
-	const FName& Category)
-{
-	FReplicatedLogData SerializedData;
-	FLogRecordSerializationContext::Serialize(SerializedData, V, Verbosity, Category);
-
-	FScopeLock QueueScopeLock(&QueueLock);
-	QueuedLogs.Emplace_GetRef() = MoveTemp(SerializedData);
-}
-
-void FLogMessagesMessageHandler::Serialize(const TCHAR* V, ELogVerbosity::Type Verbosity,
-	const FName& Category, const double Time)
-{
-	Serialize(V, Verbosity, Category);
-}
-
-void FLogMessagesMessageHandler::SerializeRecord(const UE::FLogRecord& Record)
-{
-	FReplicatedLogData SerializedData;
-	FLogRecordSerializationContext::Serialize(SerializedData, Record);
-
-	FScopeLock QueueScopeLock(&QueueLock);
-	QueuedLogs.Emplace_GetRef() = MoveTemp(SerializedData);
+	COTFSLogHandler.ConditionalPruneReplay();
 }
 
 FGuid FHeartbeatMessage::MessageType(TEXT("C08FFAF07BF34DD3A2FFB8A287CDDE83"));
diff --git a/Engine/Source/Editor/UnrealEd/Private/Cooker/CookWorkerServer.h b/Engine/Source/Editor/UnrealEd/Private/Cooker/CookWorkerServer.h
--- a/Engine/Source/Editor/UnrealEd/Private/Cooker/CookWorkerServer.h
+++ b/Engine/Source/Editor/UnrealEd/Private/Cooker/CookWorkerServer.h
@@ -19,7 +19,6 @@
 #include "IO/IoHash.h"
 #include "Logging/LogVerbosity.h"
 #include "Misc/Guid.h"
-#include "Misc/OutputDevice.h"
 #include "Templates/RefCounting.h"
 
 class FSocket;
@@ -28,10 +27,10 @@ struct FProcHandle;
 namespace UE { class FLogTemplate; }
 namespace UE::CompactBinaryTCP { struct FMarshalledMessage; }
 namespace UE::Cook { class FCookDirector; }
+namespace UE::Cook { class FLogHandler; }
 namespace UE::Cook { struct FDiscoveredPackageReplication; }
 namespace UE::Cook { enum class ECookDirectorThread : uint8; };
 namespace UE::Cook { struct FGeneratorEventMessage; }
-namespace UE::Cook { class FLogRecordSerializationContext; }
 namespace UE::Cook { struct FPackageData; }
 namespace UE::Cook { struct FPackageResultsMessage; }
 namespace UE::Cook { struct FReplicatedLogData; }
@@ -496,12 +495,11 @@ public:
  * Send log messages from CookWorkers to the CookDirector, which marks them up with the CookWorkerId and
  * then prints them to its own log.
  */
-class FLogMessagesMessageHandler : public IMPCollector, public FOutputDevice
+class FLogMessagesMessageHandler : public IMPCollector
 {
 public:
-	FLogMessagesMessageHandler(FLogRecordSerializationContext& InLogRecordSerializationContext);
-	~FLogMessagesMessageHandler();
-	void InitializeClient();
+	FLogMessagesMessageHandler(FLogHandler& InCOTFSLogHandler);
+	void ClientReportLogMessage(const FReplicatedLogData& LogData);
 
 	// IMPCollector
 	virtual FGuid GetMessageType() const override { return MessageType; }
@@ -509,22 +507,12 @@ public:
 	virtual void ClientTick(FMPCollectorClientTickContext& Context) override;
 	virtual void ServerReceiveMessage(FMPCollectorServerMessageContext& Context, FCbObjectView Message) override;
 
-	// FOutputDevice
-	virtual void Serialize(const TCHAR* V, ELogVerbosity::Type Verbosity, const FName& Category) override;
-	virtual void Serialize(const TCHAR* V, ELogVerbosity::Type Verbosity, const FName& Category,
-		const double Time) override;
-	virtual void SerializeRecord(const FLogRecord& Record) override;
-	virtual bool CanBeUsedOnAnyThread() const override { return true; }
-	virtual bool CanBeUsedOnMultipleThreads() const override { return true; }
-
 private:
-
-	FLogRecordSerializationContext& LogRecordSerializationContext;
+	FLogHandler& COTFSLogHandler;
 
 	FCriticalSection QueueLock;
 	TArray<FReplicatedLogData> QueuedLogs;
 	TArray<FReplicatedLogData> QueuedLogsBackBuffer;
-	bool bRegistered = false;
 
 	static FGuid MessageType;
 };
@@ -548,7 +536,6 @@ public:
 	int32 HeartbeatNumber;
 	static FGuid MessageType;
 };
-constexpr FStringView HeartbeatCategoryText(TEXTVIEW("CookWorkerHeartbeat:"));
 
 class FPackageWriterMPCollector : public UE::Cook::IMPCollector
 {
