Build fix for implicit conversion

#rnx

diff --git a/Engine/Source/Editor/UnrealEd/Classes/CookOnTheSide/CookOnTheFlyServer.h b/Engine/Source/Editor/UnrealEd/Classes/CookOnTheSide/CookOnTheFlyServer.h
--- a/Engine/Source/Editor/UnrealEd/Classes/CookOnTheSide/CookOnTheFlyServer.h
+++ b/Engine/Source/Editor/UnrealEd/Classes/CookOnTheSide/CookOnTheFlyServer.h
@@ -374,7 +374,7 @@ private:
 	int64 SoftGCNextAvailablePhysicalTarget = -1;
 	int32 SoftGCStartNumerator = 5;
 	int32 SoftGCDenominator = 10;
-	float SoftGCTimeFractionBudget = 0.05;
+	float SoftGCTimeFractionBudget = 0.05f;
 	float SoftGCMinimumPeriodSeconds = 30.f;
 	TUniquePtr<UE::Cook::FSoftGCHistory> SoftGCHistory;
 	bool bUseSoftGC = false;
