Added construction parameters to the thumbnail scene. FOV for the thumbnail scene can now be customized

diff --git a/Engine/Source/Editor/UnrealEd/Private/ThumbnailHelpers.cpp b/Engine/Source/Editor/UnrealEd/Private/ThumbnailHelpers.cpp
--- a/Engine/Source/Editor/UnrealEd/Private/ThumbnailHelpers.cpp
+++ b/Engine/Source/Editor/UnrealEd/Private/ThumbnailHelpers.cpp
@@ -33,43 +33,55 @@
 ***************************************************************
 */
 
-FThumbnailPreviewScene::FThumbnailPreviewScene()
-	: FPreviewScene( ConstructionValues()
+FThumbnailPreviewScene::FThumbnailPreviewScene(const FConstructionValues& InCVS)
+	: FPreviewScene(FPreviewScene::ConstructionValues()
 						.SetLightRotation( FRotator(304.736, 39.84, 0) )
 						.SetSkyBrightness(1.69f)
 						.SetCreatePhysicsScene(false)
 						.SetTransactional(false))
 {
 	// A background sky sphere
-	UStaticMeshComponent* BackgroundComponent = NewObject<UStaticMeshComponent>();
- 	BackgroundComponent->SetStaticMesh( GUnrealEd->GetThumbnailManager()->EditorSkySphere );
- 	const float SkySphereScale = 2000.0f;
- 	const FTransform BackgroundTransform(FRotator(0,0,0), FVector(0,0,0), FVector(SkySphereScale));
- 	FPreviewScene::AddComponent(BackgroundComponent, BackgroundTransform);
+	if (InCVS.bCreateSkySphere)
+	{
+		UStaticMeshComponent* BackgroundComponent = NewObject<UStaticMeshComponent>();
+ 		BackgroundComponent->SetStaticMesh( GUnrealEd->GetThumbnailManager()->EditorSkySphere );
+ 		const float SkySphereScale = 2000.0f;
+ 		const FTransform BackgroundTransform(FRotator(0,0,0), FVector(0,0,0), FVector(SkySphereScale));
+ 		FPreviewScene::AddComponent(BackgroundComponent, BackgroundTransform);
+	}
 
 	// Adjust the default light
 	DirectionalLight->Intensity = 0.2f;
 
-	// Add additional lights
-	UDirectionalLightComponent* DirectionalLight2 = NewObject<UDirectionalLightComponent>();
-	DirectionalLight2->Intensity = 5.0f;
-	DirectionalLight2->ForwardShadingPriority = 1;
-	AddComponent(DirectionalLight2, FTransform( FRotator(-40,-144.678, 0) ));
-
-	UDirectionalLightComponent* DirectionalLight3 = NewObject<UDirectionalLightComponent>();
-	DirectionalLight3->Intensity = 1.0f;
-	DirectionalLight2->ForwardShadingPriority = 2;
-	AddComponent(DirectionalLight3, FTransform( FRotator(299.235,144.993, 0) ));
+	if (InCVS.bDefaultLightingThumbnailScene)
+	{
+		// Add additional lights
+		UDirectionalLightComponent* DirectionalLight2 = NewObject<UDirectionalLightComponent>();
+		DirectionalLight2->Intensity = InCVS.LightBrightness2;
+		DirectionalLight2->ForwardShadingPriority = 1;
+		AddComponent(DirectionalLight2, FTransform(InCVS.LightRotation2));
+
+		UDirectionalLightComponent* DirectionalLight3 = NewObject<UDirectionalLightComponent>();
+		DirectionalLight3->Intensity = InCVS.LightBrightness3;
+		DirectionalLight2->ForwardShadingPriority = 2;
+		AddComponent(DirectionalLight3, FTransform(InCVS.LightRotation3));
+	}
 
-	SetSkyCubemap(GUnrealEd->GetThumbnailManager()->AmbientCubemap);
+	if (InCVS.bSetSkyCubeMap)
+	{
+		SetSkyCubemap(GUnrealEd->GetThumbnailManager()->AmbientCubemap);
+	}
 
 	// Add an infinite plane
-	const float FloorPlaneScale = 10000.f;
-	const FTransform FloorPlaneTransform(FRotator(-90.f,0,0), FVector::ZeroVector, FVector(FloorPlaneScale));
-	UStaticMeshComponent* FloorPlaneComponent = NewObject<UStaticMeshComponent>();
-	FloorPlaneComponent->SetStaticMesh( GUnrealEd->GetThumbnailManager()->EditorPlane );
-	FloorPlaneComponent->SetMaterial( 0, GUnrealEd->GetThumbnailManager()->FloorPlaneMaterial );
-	FPreviewScene::AddComponent(FloorPlaneComponent, FloorPlaneTransform);
+	if (InCVS.bCreateFloorPlane)
+	{
+		const float FloorPlaneScale = 10000.f;
+		const FTransform FloorPlaneTransform(FRotator(-90.f,0,0), FVector::ZeroVector, FVector(FloorPlaneScale));
+		UStaticMeshComponent* FloorPlaneComponent = NewObject<UStaticMeshComponent>();
+		FloorPlaneComponent->SetStaticMesh( GUnrealEd->GetThumbnailManager()->EditorPlane );
+		FloorPlaneComponent->SetMaterial( 0, GUnrealEd->GetThumbnailManager()->FloorPlaneMaterial );
+		FPreviewScene::AddComponent(FloorPlaneComponent, FloorPlaneTransform);
+	}
 }
 
 PRAGMA_DISABLE_DEPRECATION_WARNINGS
@@ -94,8 +106,8 @@ FSceneView* FThumbnailPreviewScene::CreateView(FSceneViewFamily * ViewFamily, in
 	{
 		return nullptr;
 	}
-		
-	const float FOVDegrees = 30.f;
+
+	const float FOVDegrees = GetFOV();
 	const float HalfFOVRadians = FMath::DegreesToRadians<float>(FOVDegrees) * 0.5f;
 	static_assert((int32)ERHIZBuffer::IsInverted != 0, "Check NearPlane and Projection Matrix");
 	const float NearPlane = 1.0f;
@@ -173,6 +185,12 @@ float FThumbnailPreviewScene::GetBoundsZOffset(const FBoxSphereBounds& Bounds) c
 	return static_cast<float>(Bounds.BoxExtent.Z + 1.0);
 }
 
+float FThumbnailPreviewScene::GetFOV() const
+{
+	// Default FOV value
+	return 30.0f;
+}
+
 /*
 ***************************************************************
   FParticleSystemThumbnailScene
diff --git a/Engine/Source/Editor/UnrealEd/Public/ThumbnailHelpers.h b/Engine/Source/Editor/UnrealEd/Public/ThumbnailHelpers.h
--- a/Engine/Source/Editor/UnrealEd/Public/ThumbnailHelpers.h
+++ b/Engine/Source/Editor/UnrealEd/Public/ThumbnailHelpers.h
@@ -18,8 +18,41 @@ class USceneThumbnailInfo;
 class FThumbnailPreviewScene : public FPreviewScene, public FTickableEditorObject
 {
 public:
+	struct FConstructionValues
+	{
+		FConstructionValues()
+			: bCreateSkySphere(true)
+			, bDefaultLightingThumbnailScene(true)
+			, LightRotation2(FRotator(-40.0, -144.678, 0.0))
+			, LightBrightness2(5.0f)
+			, LightRotation3(FRotator(299.235, 144.993, 0.0))
+			, LightBrightness3(1.0f)
+			, bSetSkyCubeMap(true)
+			, bCreateFloorPlane(true)
+		{
+		}
+
+		uint32 bCreateSkySphere : 1;
+		uint32 bDefaultLightingThumbnailScene : 1;
+		FRotator LightRotation2;
+		float LightBrightness2;
+		FRotator LightRotation3;
+		float LightBrightness3;
+		uint32 bSetSkyCubeMap : 1;
+		uint32 bCreateFloorPlane : 1;
+
+		FConstructionValues& SetCreateSkySphere(const bool bInValue) { bCreateSkySphere = bInValue; return *this; }
+		FConstructionValues& SetDefaultLightingThumbnailScene(const bool bInValue) { bDefaultLightingThumbnailScene = bInValue; return *this; }
+		FConstructionValues& SetLightRotation2(const FRotator& InValue) { LightRotation2 = InValue; return *this; }
+		FConstructionValues& SetLightBrightness2(const float InValue) { LightBrightness2 = InValue; return *this; }
+		FConstructionValues& SetLightRotation3(const FRotator& InValue) { LightRotation3 = InValue; return *this; }
+		FConstructionValues& SetLightBrightness3(const float InValue) { LightBrightness3 = InValue; return *this; }
+		FConstructionValues& SetSkyCubeMap(const bool bInValue) { bSetSkyCubeMap = bInValue; return *this; }
+		FConstructionValues& SetCreateFloorPlane(const bool bInValue) { bCreateFloorPlane = bInValue; return *this; }
+	};
+
 	/** Constructor */
-	UNREALED_API FThumbnailPreviewScene();
+	UNREALED_API explicit FThumbnailPreviewScene(const FConstructionValues& InCVS = FConstructionValues());
 
 	UE_DEPRECATED(5.0, "Use CreateView")
 	UNREALED_API void GetView(FSceneViewFamily* ViewFamily, int32 X, int32 Y, uint32 SizeX, uint32 SizeY) const;
@@ -46,6 +79,9 @@ protected:
 	  * @param OutOrbitZoom  The camera distance from the object.
 	  */
 	virtual void GetViewMatrixParameters(const float InFOVDegrees, FVector& OutOrigin, float& OutOrbitPitch, float& OutOrbitYaw, float& OutOrbitZoom) const = 0;
+
+	/** Returns FOV value to use for the view matrix */
+	UNREALED_API virtual float GetFOV() const;
 };
 
 class FParticleSystemThumbnailScene : public FThumbnailPreviewScene
