IncrementalCook: Replaying Logs: Store logs reported from UE::Cook::FLogHandler on the CookPackageData from PackageAccessTracking when the log was sent, save that list of logs as an attachment when the package is saved, and replay the logs when the package is incrementally skipped.
Change naming convention for the attachment name in incremental cook attachments to be meta.cook.<shortname>, for readability in the zenwebui and for setting zenwebui display for the cook attachments.
#rnx
#incrementalcook

diff --git a/Engine/Source/Editor/UnrealEd/Private/Cooker/CookIncrementalVersion.cpp b/Engine/Source/Editor/UnrealEd/Private/Cooker/CookIncrementalVersion.cpp
--- a/Engine/Source/Editor/UnrealEd/Private/Cooker/CookIncrementalVersion.cpp
+++ b/Engine/Source/Editor/UnrealEd/Private/Cooker/CookIncrementalVersion.cpp
@@ -8,6 +8,6 @@ namespace UE::Cook
 {
 
 // Change CookIncrementalVersion to a new guid when all packages in an incremental cook need to be invalidated
-FGuid CookIncrementalVersion(0x80B89F71, 0x7C7243B9, 0x9E089834, 0x14F6BA20);
+FGuid CookIncrementalVersion( 0x06A2576D, 0x2C654CE5, 0x86A9A40C, 0xB74471B7 );
 
 }
diff --git a/Engine/Source/Editor/UnrealEd/Private/Cooker/CookLog.cpp b/Engine/Source/Editor/UnrealEd/Private/Cooker/CookLog.cpp
--- a/Engine/Source/Editor/UnrealEd/Private/Cooker/CookLog.cpp
+++ b/Engine/Source/Editor/UnrealEd/Private/Cooker/CookLog.cpp
@@ -139,7 +139,7 @@ void FLogHandler::FlushIncrementalCookLogs()
 	}
 }
 
-void FLogHandler::ReplayLogFromIncrementallySkipped(FReplicatedLogData&& LogData)
+void FLogHandler::ReplayLogsFromIncrementallySkipped(TConstArrayView<FReplicatedLogData> LogMessages)
 {
 	// Replays only come from MarkPackageIterativelySkipped, which happens only on the CookDirector, during
 	// CookRequestCluster traversal. We rely on that, and do not report whether messages from CookWorkers
@@ -155,7 +155,10 @@ void FLogHandler::ReplayLogFromIncrementallySkipped(FReplicatedLogData&& LogData
 		{
 			return false;
 		};
-	UnMarshalAndLog(LogData, MessagePassesFilter, TryTransformMessage);
+	for (const FReplicatedLogData& LogMessage : LogMessages)
+	{
+		UnMarshalAndLog(LogMessage, MessagePassesFilter, TryTransformMessage);
+	}
 }
 
 void FLogHandler::ReplayLogFromCookWorker(FReplicatedLogData&& LogData, int32 CookWorkerProfileId)
@@ -505,7 +508,7 @@ void FLogHandler::RecordLogForIncrementalCookGameThreadPortion(FName ActivePacka
 		return;
 	}
 
-	//PackageData->AddLogMessage(MoveTemp(LogData));
+	PackageData->AddLogMessage(MoveTemp(LogData));
 }
 
 } // namespace UE::Cook
\ No newline at end of file
diff --git a/Engine/Source/Editor/UnrealEd/Private/Cooker/CookLogPrivate.h b/Engine/Source/Editor/UnrealEd/Private/Cooker/CookLogPrivate.h
--- a/Engine/Source/Editor/UnrealEd/Private/Cooker/CookLogPrivate.h
+++ b/Engine/Source/Editor/UnrealEd/Private/Cooker/CookLogPrivate.h
@@ -47,7 +47,7 @@ public:
 	FLogHandler(UCookOnTheFlyServer& InCOTFS);
 	~FLogHandler();
 
-	void ReplayLogFromIncrementallySkipped(FReplicatedLogData&& LogData);
+	void ReplayLogsFromIncrementallySkipped(TConstArrayView<FReplicatedLogData> LogMessages);
 	void ReplayLogFromCookWorker(FReplicatedLogData&& LogData, int32 CookWorkerProfileId);
 	void ConditionalPruneReplay();
 	void FlushIncrementalCookLogs();
diff --git a/Engine/Source/Editor/UnrealEd/Private/Cooker/CookPackageArtifacts.cpp b/Engine/Source/Editor/UnrealEd/Private/Cooker/CookPackageArtifacts.cpp
--- a/Engine/Source/Editor/UnrealEd/Private/Cooker/CookPackageArtifacts.cpp
+++ b/Engine/Source/Editor/UnrealEd/Private/Cooker/CookPackageArtifacts.cpp
@@ -35,9 +35,10 @@ namespace UE::Cook
 // backwards compatibility rather than invalidating everything.
 constexpr uint32 PackageArtifactsVersion = 0x00000004;
 
-static const FUtf8StringView PackageArtifactsAttachmentKey = UTF8TEXTVIEW("CookPackageArtifacts");
-static const FUtf8StringView BuildDefinitionsAttachmentKey = UTF8TEXTVIEW("BuildDefinitionsAttachmentKey");
-static const FUtf8StringView ImportsCheckerAttachmentKey = UTF8TEXTVIEW("ImportsCheckerAttachmentKey");
+static const FUtf8StringView PackageArtifactsAttachmentKey = UTF8TEXTVIEW("meta.cook.artifacts");
+static const FUtf8StringView BuildDefinitionsAttachmentKey = UTF8TEXTVIEW("meta.cook.builddefinitions");
+static const FUtf8StringView ImportsCheckerAttachmentKey = UTF8TEXTVIEW("meta.cook.importexport");
+static const FUtf8StringView LogMessagesAttachmentKey = UTF8TEXTVIEW("meta.cook.logs");
 
 bool FBuildDependencySet::HasKeyMatch(FName PackageName, const ITargetPlatform* TargetPlatform,
 	FGenerationHelper* GenerationHelper)
@@ -969,39 +970,66 @@ FCbWriter& operator<<(FCbWriter& Writer, const UE::Cook::FBuildDefinitionList& D
 namespace UE::Cook
 {
 
+/** Wrapper around TArray<FLogReplicationData> so we can serialize is as FCbObject instead of FCbArray. */
+struct FLogMessagesArray
+{
+public:
+	TArray<FReplicatedLogData>& Array;
+
+public:
+	FLogMessagesArray(TArray<FReplicatedLogData>& InLogMessages)
+		: Array(InLogMessages)
+	{
+	}
+
+private:
+	friend bool LoadFromCompactBinary(FCbFieldView FieldView, UE::Cook::FLogMessagesArray& LogMessages)
+	{
+		return LoadFromCompactBinary(FieldView["Logs"], LogMessages.Array);
+	}
+	friend FCbWriter& operator<<(FCbWriter& Writer, const UE::Cook::FLogMessagesArray& LogMessages)
+	{
+		Writer.BeginObject();
+		Writer << "Logs" << const_cast<FLogMessagesArray&>(LogMessages).Array;
+		Writer.EndObject();
+		return Writer;
+	}
+};
+
 void FIncrementalCookAttachments::Empty()
 {
 	Artifacts.Empty();
 	BuildDefinitions.Empty();
 }
 
+template <typename WritableType>
+static void AddAttachment(TArray<IPackageWriter::FCommitAttachmentInfo>& OutAttachments,
+	WritableType&& Writable, FUtf8StringView AttachmentKey)
+{
+	FCbWriter Writer;
+	Writer << Writable;
+	IPackageWriter::FCommitAttachmentInfo& Attachment = OutAttachments.Emplace_GetRef();
+	Attachment.Key = AttachmentKey;
+	Attachment.Value = Writer.Save().AsObject();
+}
+
 void FIncrementalCookAttachments::AppendCommitAttachments(TArray<IPackageWriter::FCommitAttachmentInfo>& OutAttachments)
 {
 	if (Artifacts.IsValid())
 	{
-		FCbWriter Writer;
-		Writer << Artifacts;
-		IPackageWriter::FCommitAttachmentInfo& Attachment = OutAttachments.Emplace_GetRef();
-		Attachment.Key = PackageArtifactsAttachmentKey;
-		Attachment.Value = Writer.Save().AsObject();
+		AddAttachment(OutAttachments, Artifacts, PackageArtifactsAttachmentKey);
 	}
-
 	if (!BuildDefinitions.Definitions.IsEmpty())
 	{
-		FCbWriter Writer;
-		Writer << BuildDefinitions;
-		IPackageWriter::FCommitAttachmentInfo& Attachment = OutAttachments.Emplace_GetRef();
-		Attachment.Key = BuildDefinitionsAttachmentKey;
-		Attachment.Value = Writer.Save().AsObject();
+		AddAttachment(OutAttachments, BuildDefinitions, BuildDefinitionsAttachmentKey);
 	}
-
 	if (!ImportsCheckerData.IsEmpty())
 	{
-		FCbWriter Writer;
-		Writer << ImportsCheckerData;
-		IPackageWriter::FCommitAttachmentInfo& Attachment = OutAttachments.Emplace_GetRef();
-		Attachment.Key = ImportsCheckerAttachmentKey;
-		Attachment.Value = Writer.Save().AsObject();
+		AddAttachment(OutAttachments, ImportsCheckerData, ImportsCheckerAttachmentKey);
+	}
+	if (!LogMessages.IsEmpty())
+	{
+		AddAttachment(OutAttachments, FLogMessagesArray(LogMessages), LogMessagesAttachmentKey);
 	}
 }
 
@@ -1010,7 +1038,8 @@ FIncrementalCookAttachments FIncrementalCookAttachments::Collect(const UPackage*
 	bool bHasSaveResult, TConstArrayView<FName> UntrackedSoftPackageReferences, FGenerationHelper* GenerationHelper,
 	bool bGenerated, TArray<FName>&& RuntimeDependencies,
 	TConstArrayView<UObject*> Imports, TConstArrayView<UObject*> Exports,
-	TConstArrayView<UE::SavePackageUtilities::FPreloadDependency> PreloadDependencies)
+	TConstArrayView<UE::SavePackageUtilities::FPreloadDependency> PreloadDependencies,
+	TConstArrayView<FReplicatedLogData> InLogMessages)
 {
 	FIncrementalCookAttachments Result;
 
@@ -1033,6 +1062,7 @@ FIncrementalCookAttachments FIncrementalCookAttachments::Collect(const UPackage*
 
 	Result.BuildDefinitions = FBuildDefinitionList::Collect(Package, TargetPlatform);
 	Result.ImportsCheckerData = UE::Cook::FImportsCheckerData::FromObjectLists(Imports, Exports);
+	Result.LogMessages = InLogMessages;
 
 	return Result;
 }
@@ -1050,7 +1080,7 @@ void FIncrementalCookAttachments::Fetch(TArrayView<FName> PackageNames, const IT
 	};
 	TMap<FName, FInProgressResult> InProgressResults;
 	TArray<FUtf8StringView> RequestedAttachments = { PackageArtifactsAttachmentKey, BuildDefinitionsAttachmentKey,
-		ImportsCheckerAttachmentKey };
+		ImportsCheckerAttachmentKey, LogMessagesAttachmentKey };
 
 	if (TargetPlatform || GEditorDomainOplog)
 	{
@@ -1085,6 +1115,11 @@ void FIncrementalCookAttachments::Fetch(TArrayView<FName> PackageNames, const IT
 				{
 					LoadFromCompactBinary(Attachment.AsFieldView(), InProgressResult.Result.ImportsCheckerData);
 				}
+				else if (AttachmentKey == LogMessagesAttachmentKey)
+				{
+					FLogMessagesArray LogMessagesArray(InProgressResult.Result.LogMessages);
+					LoadFromCompactBinary(Attachment.AsFieldView(), LogMessagesArray);
+				}
 
 				if (InProgressResult.ReceivedAttachmentCount == RequestedAttachmentNum)
 				{
diff --git a/Engine/Source/Editor/UnrealEd/Private/Cooker/CookPackageArtifacts.h b/Engine/Source/Editor/UnrealEd/Private/Cooker/CookPackageArtifacts.h
--- a/Engine/Source/Editor/UnrealEd/Private/Cooker/CookPackageArtifacts.h
+++ b/Engine/Source/Editor/UnrealEd/Private/Cooker/CookPackageArtifacts.h
@@ -6,6 +6,7 @@
 #include "AssetRegistry/IAssetRegistry.h"
 #include "Cooker/CookDependency.h"
 #include "Cooker/CookImportsChecker.h"
+#include "Cooker/CookLogPrivate.h"
 #include "Containers/Array.h"
 #include "Containers/ArrayView.h"
 #include "Containers/Map.h"
@@ -317,6 +318,7 @@ struct FIncrementalCookAttachments
 	FPackageArtifacts Artifacts;
 	FBuildDefinitionList BuildDefinitions;
 	FImportsCheckerData ImportsCheckerData;
+	TArray<FReplicatedLogData> LogMessages;
 	IPackageWriter::ECommitStatus CommitStatus = IPackageWriter::ECommitStatus::NotCommitted;
 
 	void Empty();
@@ -330,7 +332,8 @@ struct FIncrementalCookAttachments
 		bool bHasSaveResult, TConstArrayView<FName> UntrackedSoftPackageReferences, FGenerationHelper* GenerationHelper,
 		bool bGenerated, TArray<FName>&& RuntimeDependencies,
 		TConstArrayView<UObject*> Imports, TConstArrayView<UObject*> Exports,
-		TConstArrayView<UE::SavePackageUtilities::FPreloadDependency> PreloadDependencies);
+		TConstArrayView<UE::SavePackageUtilities::FPreloadDependency> PreloadDependencies,
+		TConstArrayView<FReplicatedLogData> LogMessages);
 };
 
 } // namespace UE::Cook
diff --git a/Engine/Source/Editor/UnrealEd/Private/Cooker/CookPackageData.cpp b/Engine/Source/Editor/UnrealEd/Private/Cooker/CookPackageData.cpp
--- a/Engine/Source/Editor/UnrealEd/Private/Cooker/CookPackageData.cpp
+++ b/Engine/Source/Editor/UnrealEd/Private/Cooker/CookPackageData.cpp
@@ -14,6 +14,7 @@
 #include "Cooker/BuildResultDependenciesMap.h"
 #include "Cooker/CookDirector.h"
 #include "Cooker/CookGenerationHelper.h"
+#include "Cooker/CookLogPrivate.h"
 #include "Cooker/CookPackagePreloader.h"
 #include "Cooker/CookPlatformManager.h"
 #include "Cooker/CookRequestCluster.h"
@@ -112,6 +113,7 @@ FPackageData::FPackageData(FPackageDatas& PackageDatas, const FName& InPackageNa
 	, bGenerated(0), bKeepReferencedDuringGC(0)
 	, bWasCookedThisSession(0)
 	, DoesGeneratedRequireGeneratorValue(static_cast<uint32>(ICookPackageSplitter::EGeneratedRequiresGenerator::None))
+	, bHasReplayedLogMessages(0)
 {
 	SetState(EPackageState::Idle);
 	SetSaveSubState(ESaveSubState::StartSave);
@@ -364,6 +366,12 @@ void FPackageData::ClearInProgressData(EStateChangeReason StateChangeReason)
 		GenerationHelper->ClearKeepForGeneratorSaveAllPlatforms();
 	}
 	SetParentGenerationHelper(nullptr, StateChangeReason);
+
+	// Clear data that is no longer needed when we have comitted all platforms
+	if (HasAllCommittedPlatforms(PackageDatas.GetCookOnTheFlyServer().PlatformManager->GetSessionPlatforms()))
+	{
+		ClearLogMessages();
+	}
 }
 
 void FPackageData::SetPlatformsCooked(
@@ -476,6 +484,7 @@ void FPackageData::ClearCookResults()
 		PackageDatas.GetMonitor().OnLastCookedPlatformRemoved(*this);
 	}
 	SetSuppressCookReason(ESuppressCookReason::NotSuppressed);
+	bHasReplayedLogMessages = false;
 }
 
 void FPackageData::ClearCookResults(const ITargetPlatform* TargetPlatform)
@@ -500,6 +509,7 @@ void FPackageData::ClearCookResults(const ITargetPlatform* TargetPlatform)
 	{
 		bWasCookedThisSession = false;
 		PackageDatas.GetMonitor().OnLastCookedPlatformRemoved(*this);
+		bHasReplayedLogMessages = false;
 	}
 }
 
@@ -1851,6 +1861,29 @@ TMap<FPackageData*, EInstigator>* FPackageData::GetDiscoveredDependencies(const
 	return DiscoveredDependencies->Find(TargetPlatform);
 }
 
+void FPackageData::AddLogMessage(FReplicatedLogData&& LogData)
+{
+	if (!LogMessages)
+	{
+		LogMessages.Reset(new TArray<FReplicatedLogData>());
+	}
+	LogMessages->Add(MoveTemp(LogData));
+}
+
+TConstArrayView<FReplicatedLogData> FPackageData::GetLogMessages() const
+{
+	if (!LogMessages)
+	{
+		return TConstArrayView<FReplicatedLogData>();
+	}
+	return *LogMessages;
+}
+
+void FPackageData::ClearLogMessages()
+{
+	LogMessages.Reset();
+}
+
 } // namespace UE::Cook
 
 FCbWriter& operator<<(FCbWriter& Writer, const UE::Cook::FConstructPackageData& PackageData)
diff --git a/Engine/Source/Editor/UnrealEd/Private/Cooker/CookPackageData.h b/Engine/Source/Editor/UnrealEd/Private/Cooker/CookPackageData.h
--- a/Engine/Source/Editor/UnrealEd/Private/Cooker/CookPackageData.h
+++ b/Engine/Source/Editor/UnrealEd/Private/Cooker/CookPackageData.h
@@ -41,6 +41,7 @@ namespace UE::Cook { class FRequestCluster; }
 namespace UE::Cook { struct FBuildResultDependenciesMap; }
 namespace UE::Cook { struct FConstructPackageData; }
 namespace UE::Cook { struct FDiscoveredPlatformSet; }
+namespace UE::Cook { struct FReplicatedLogData; }
 
 FCbWriter& operator<<(FCbWriter& Writer, const UE::Cook::FConstructPackageData& PackageData);
 bool LoadFromCompactBinary(FCbFieldView Field, UE::Cook::FConstructPackageData& PackageData);
@@ -777,6 +778,13 @@ public:
 	TMap<FPackageData*, EInstigator>& CreateOrGetDiscoveredDependencies(const ITargetPlatform* TargetPlatform);
 	TMap<FPackageData*, EInstigator>* GetDiscoveredDependencies(const ITargetPlatform* TargetPlatform);
 
+	/** Storage for warnings and errors about the package discovered during cook; these are saved into the oplog. */
+	void AddLogMessage(FReplicatedLogData&& LogData);
+	TConstArrayView<FReplicatedLogData> GetLogMessages() const;
+	void ClearLogMessages();
+	bool HasReplayedLogMessages() const;
+	void SetHasReplayedLogMessages(bool bValue);
+
 	/**
 	 * Return the platforms for which the given Package has been marked reachable.
 	 * If the package does not exist, return the COTFS's list of Session platforms
@@ -857,6 +865,7 @@ private:
 	FCompletionCallback CompletionCallback;
 	TUniquePtr<TMap<const ITargetPlatform*, TMap<FPackageData*, EInstigator>>> DiscoveredDependencies;
 	TUniquePtr<FBuildResultDependenciesMap> LoadDependencies;
+	TUniquePtr<TArray<FReplicatedLogData>> LogMessages;
 	FName PackageName;
 	FName FileName;
 	FName ParentGenerator;
@@ -889,6 +898,7 @@ private:
 	uint32 bWasCookedThisSession : 1;
 	static_assert(static_cast<uint32>(ICookPackageSplitter::EGeneratedRequiresGenerator::Count) <= 4, "We are storing Enum value in 2 bits");
 	uint32 DoesGeneratedRequireGeneratorValue : 2;
+	uint32 bHasReplayedLogMessages : 1;
 };
 
 /**
@@ -1895,6 +1905,16 @@ inline bool FPackageData::GetWasCookedThisSession() const
 	return static_cast<bool>(bWasCookedThisSession);
 }
 
+inline bool FPackageData::HasReplayedLogMessages() const
+{
+	return static_cast<bool>(bHasReplayedLogMessages);
+}
+
+inline void FPackageData::SetHasReplayedLogMessages(bool Value)
+{
+	bHasReplayedLogMessages = Value != 0;
+}
+
 inline FWorkerId FPackageData::GetWorkerAssignment() const
 {
 	return WorkerAssignment;
diff --git a/Engine/Source/Editor/UnrealEd/Private/Cooker/CookRequestCluster.cpp b/Engine/Source/Editor/UnrealEd/Private/Cooker/CookRequestCluster.cpp
--- a/Engine/Source/Editor/UnrealEd/Private/Cooker/CookRequestCluster.cpp
+++ b/Engine/Source/Editor/UnrealEd/Private/Cooker/CookRequestCluster.cpp
@@ -2244,6 +2244,12 @@ void FRequestCluster::FGraphSearch::FExploreEdgesContext::ProcessPlatformAttachm
 				// Replay the package's saved data into our collectors of packagedata
 				FEDLCookCheckerThreadState::Get().Add(QueryPlatformData.CookAttachments.ImportsCheckerData,
 					PackageData->GetPackageName());
+				TConstArrayView<FReplicatedLogData> LogMessages = QueryPlatformData.CookAttachments.LogMessages;
+				if (!LogMessages.IsEmpty() && !PackageData->HasReplayedLogMessages())
+				{
+					Cluster.COTFS.LogHandler->ReplayLogsFromIncrementallySkipped(LogMessages);
+					PackageData->SetHasReplayedLogMessages(true);
+				}
 			}
 			else
 			{
diff --git a/Engine/Source/Editor/UnrealEd/Private/Cooker/CookSavePackage.cpp b/Engine/Source/Editor/UnrealEd/Private/Cooker/CookSavePackage.cpp
--- a/Engine/Source/Editor/UnrealEd/Private/Cooker/CookSavePackage.cpp
+++ b/Engine/Source/Editor/UnrealEd/Private/Cooker/CookSavePackage.cpp
@@ -891,7 +891,8 @@ void FSaveCookedPackageContext::CalculateCookDependencies(FGenerationHelper* Gen
 
 		PlatformCookAttachments = FIncrementalCookAttachments::Collect(Package, TargetPlatform,
 			MoveTemp(BuildResultDependencies), bHasSavePackageResult, UntrackedSoftPackageReferences, GenerationHelper,
-			bGenerated, MoveTemp(PlatformRuntimeDependencies), Imports, Exports, PreloadDependencies);
+			bGenerated, MoveTemp(PlatformRuntimeDependencies), Imports, Exports, PreloadDependencies,
+			PackageData.GetLogMessages());
 	}
 }
 
diff --git a/Engine/Source/Editor/UnrealEd/Private/EditorDomain/EditorDomainUtils.cpp b/Engine/Source/Editor/UnrealEd/Private/EditorDomain/EditorDomainUtils.cpp
--- a/Engine/Source/Editor/UnrealEd/Private/EditorDomain/EditorDomainUtils.cpp
+++ b/Engine/Source/Editor/UnrealEd/Private/EditorDomain/EditorDomainUtils.cpp
@@ -2027,7 +2027,8 @@ bool TrySavePackage(UPackage* Package)
 				Package, nullptr /* TargetPlatform */, MoveTemp(ExistingResultDependencies),
 				false /* bHasSaveResult */, TConstArrayView<FName>() /* UntrackedSoftPackageReferences */,
 				nullptr /* GenerationHelper */, false /* bGenerated */, TArray<FName>(), TConstArrayView<UObject*>(),
-				TConstArrayView<UObject*>(), TConstArrayView<UE::SavePackageUtilities::FPreloadDependency>());
+				TConstArrayView<UObject*>(), TConstArrayView<UE::SavePackageUtilities::FPreloadDependency>(),
+				TConstArrayView<UE::Cook::FReplicatedLogData>());
 		TArray<IPackageWriter::FCommitAttachmentInfo> Attachments;
 		CookAttachments.AppendCommitAttachments(Attachments);
 		if (!Attachments.IsEmpty())
