PR #12939: Add filter UI to SDeleteAssetsDialog to allow usage of FARFilters and other filters (e.g. writeable, checked out, gameplay tags)

diff --git a/Engine/Source/Editor/UnrealEd/Private/Dialogs/SDeleteAssetsDialog.cpp b/Engine/Source/Editor/UnrealEd/Private/Dialogs/SDeleteAssetsDialog.cpp
--- a/Engine/Source/Editor/UnrealEd/Private/Dialogs/SDeleteAssetsDialog.cpp
+++ b/Engine/Source/Editor/UnrealEd/Private/Dialogs/SDeleteAssetsDialog.cpp
@@ -801,6 +801,7 @@ TSharedRef<SWidget> SDeleteAssetsDialog::MakeConsolidationAssetPicker()
 	AssetPickerConfig.bAllowDragging = false;
 	AssetPickerConfig.bCanShowClasses = false;
 	AssetPickerConfig.SelectionMode = ESelectionMode::Single;
+	AssetPickerConfig.bAddFilterUI = true;
 
 	FContentBrowserModule& ContentBrowserModule = FModuleManager::Get().LoadModuleChecked<FContentBrowserModule>( TEXT( "ContentBrowser" ) );
 
