[UE] Auto-generate BP class path name if asset tag is missing during component type registration.

diff --git a/Engine/Source/Editor/UnrealEd/Private/ComponentTypeRegistry.cpp b/Engine/Source/Editor/UnrealEd/Private/ComponentTypeRegistry.cpp
--- a/Engine/Source/Editor/UnrealEd/Private/ComponentTypeRegistry.cpp
+++ b/Engine/Source/Editor/UnrealEd/Private/ComponentTypeRegistry.cpp
@@ -614,7 +614,8 @@ void FComponentTypeRegistryData::ForceRefreshComponentList()
 		AssetRegistry.GetAssetsByClass(UBlueprint::StaticClass()->GetClassPathName(), Assets, true);
 		for (FAssetData& BPAsset : Assets)
 		{
-			FTopLevelAssetPath ClassPath(FEditorClassUtils::GetClassPathNameFromAssetTag(BPAsset));
+			constexpr bool bGenerateClassPathIfMissing = true;
+			FTopLevelAssetPath ClassPath(FEditorClassUtils::GetClassPathNameFromAsset(BPAsset, bGenerateClassPathIfMissing));
 			FSoftObjectPath SoftClassPath(ClassPath);
 			if (SoftClassPath.FixupCoreRedirects() && SoftClassPath.GetSubPathString().IsEmpty())
 			{
