Load suggested Visual Studio components from Windows_SDK.json

#rnx

diff --git a/Engine/Config/Windows/Windows_SDK.json b/Engine/Config/Windows/Windows_SDK.json
--- a/Engine/Config/Windows/Windows_SDK.json
+++ b/Engine/Config/Windows/Windows_SDK.json
@@ -46,7 +46,6 @@
 	"//04.2": "Microsoft.VisualStudio.Component.VC.14.38.17.8.ATL -- Match version of MSVC",
 	"VisualStudioSuggestedComponents": [
 		"Microsoft.Net.Component.4.6.2.TargetingPack",
-		"Microsoft.VisualStudio.Component.Unreal.Workspace",
 		"Microsoft.VisualStudio.Component.VC.14.38.17.8.x86.x64",
 		"Microsoft.VisualStudio.Component.VC.14.38.17.8.ATL",
 		"Microsoft.VisualStudio.Component.VC.Tools.x86.x64",
@@ -55,7 +54,9 @@
 		"Microsoft.VisualStudio.Workload.CoreEditor",
 		"Microsoft.VisualStudio.Workload.ManagedDesktop",
 		"Microsoft.VisualStudio.Workload.NativeDesktop",
-		"Microsoft.VisualStudio.Workload.NativeGame"
+		"Microsoft.VisualStudio.Workload.NativeGame",
+		"Component.Unreal.Ide",
+		"Component.Unreal.Debugger"
 	],
 
 	"VisualStudioSuggestedLinuxComponents": [
@@ -69,6 +70,6 @@
 
 	"//05.1": "X64 -- 0x601 = Win7",
 	"//05.2": "Arm64 -- 0xA00 = Win10",
-	"MinimumWindowsX64TargetVersion": "0x601", // Win7
-	"MinimumWindowsArm64TargetVersion": "0xA00" // Win10
+	"MinimumWindowsX64TargetVersion": "0x601",
+	"MinimumWindowsArm64TargetVersion": "0xA00"
 }
diff --git a/Engine/Source/Editor/UnrealEd/Private/SourceCodeNavigation.cpp b/Engine/Source/Editor/UnrealEd/Private/SourceCodeNavigation.cpp
--- a/Engine/Source/Editor/UnrealEd/Private/SourceCodeNavigation.cpp
+++ b/Engine/Source/Editor/UnrealEd/Private/SourceCodeNavigation.cpp
@@ -18,11 +18,14 @@
 #include "Misc/ConfigCacheIni.h"
 #include "Misc/DataDrivenPlatformInfoRegistry.h"
 #include "Misc/EngineBuildSettings.h"
+#include "Misc/FileHelper.h"
 #include "Misc/PackageName.h"
 #include "Misc/Paths.h"
 #include "Misc/ScopeLock.h"
 #include "Modules/ModuleManager.h"
 #include "Preferences/UnrealEdOptions.h"
+#include "Serialization/JsonReader.h"
+#include "Serialization/JsonSerializer.h"
 #include "Styling/AppStyle.h"
 #include "TickableEditorObject.h"
 #include "UnrealEdGlobals.h"
@@ -2178,8 +2181,39 @@ FString FSourceCodeNavigationImpl::GetSuggestedIDEInstallerFileName()
 void FSourceCodeNavigationImpl::LaunchIDEInstaller(const FString& Filepath)
 {
 #if PLATFORM_WINDOWS
-	auto Params = TEXT("--productId \"Microsoft.VisualStudio.Product.Community\" --add \"Microsoft.VisualStudio.Workload.NativeDesktop\" --add \"Microsoft.VisualStudio.Workload.NativeGame\" --add \"Component.Unreal\" --add \"Microsoft.VisualStudio.Component.Windows10SDK.18362\" --campaign \"EpicGames_UE5\"");
-	FPlatformProcess::ExecElevatedProcess(*Filepath, Params, nullptr);
+	TStringBuilder<512> Params;	
+
+	// Try and load suggested components from the Windows SDK configuration file
+	const FString WindowsSDKFile = FPaths::RootDir() / TEXT("Engine/Config/Windows/Windows_SDK.json");
+	if (FPaths::FileExists(WindowsSDKFile))
+	{
+		FString FileContents;
+		if (FFileHelper::LoadFileToString(FileContents, *WindowsSDKFile))
+		{
+			TSharedRef<TJsonReader<>> JsonReader = TJsonReaderFactory<>::Create(FileContents);
+			TSharedPtr<FJsonObject> RootObject;
+			if (FJsonSerializer::Deserialize(JsonReader, RootObject) && RootObject.IsValid())
+			{
+				const TArray<TSharedPtr<FJsonValue>> SuggestedComponentsArray = RootObject->GetArrayField(TEXT("VisualStudioSuggestedComponents"));
+				for (const TSharedPtr<FJsonValue>& SuggestedComponents : SuggestedComponentsArray)
+				{
+					FString Value = SuggestedComponents->AsString();
+					if (!Value.IsEmpty())
+					{
+						Params.Appendf(TEXT("--add \"%s\" "), *Value);
+					}
+				}
+			}
+		}
+	}
+	
+	if (Params.Len() == 0)
+	{
+		UE_LOG(LogSelectionDetails, Warning, TEXT("LaunchIDEInstaller: Failed to load suggested Visual Studio components from [%s]"), *WindowsSDKFile);
+	}
+
+	Params.Append(TEXT("--productId \"Microsoft.VisualStudio.Product.Community\" --campaign \"EpicGames_UE5\""));
+	FPlatformProcess::ExecElevatedProcess(*Filepath, Params.GetData(), nullptr);
 #endif
 }
 
