Added CreateOrEmptyCollection util to UCollectionManagerScriptingSubsystem

diff --git a/Engine/Source/Editor/UnrealEd/Private/Subsystems/CollectionManagerScriptingSubsystem.cpp b/Engine/Source/Editor/UnrealEd/Private/Subsystems/CollectionManagerScriptingSubsystem.cpp
--- a/Engine/Source/Editor/UnrealEd/Private/Subsystems/CollectionManagerScriptingSubsystem.cpp
+++ b/Engine/Source/Editor/UnrealEd/Private/Subsystems/CollectionManagerScriptingSubsystem.cpp
@@ -108,7 +108,8 @@ bool UCollectionManagerScriptingSubsystem::CreateCollection(const FCollectionScr
 		return false;
 	}
 
-	if (const ECollectionShareType::Type CastShareType = CollectionManagerScriptingSubsystemUtil::ECollectionScriptingShareType_To_ECollectionShareType(ShareType); CollectionContainer->CollectionExists(Collection, CastShareType))
+	const ECollectionShareType::Type CastShareType = CollectionManagerScriptingSubsystemUtil::ECollectionScriptingShareType_To_ECollectionShareType(ShareType);
+	if (CollectionContainer->CollectionExists(Collection, CastShareType))
 	{
 		CollectionManagerScriptingSubsystemUtil::LogLastCollectionManagerError(FText::Format(CollectionManagerScriptingSubsystemUtil::FConstants::Get().CollectionExistsInContainer, { CollectionContainer->GetCollectionSource()->GetTitle(), FText::FromName(Collection), CastShareType }));
 		return false;
@@ -121,7 +122,7 @@ bool UCollectionManagerScriptingSubsystem::CreateCollection(const FCollectionScr
 		return false;
 	}
 
-	if (!CollectionContainer->CreateCollection(Collection, CollectionManagerScriptingSubsystemUtil::ECollectionScriptingShareType_To_ECollectionShareType(ShareType), ECollectionStorageMode::Static, &Error))
+	if (!CollectionContainer->CreateCollection(Collection, CastShareType, ECollectionStorageMode::Static, &Error))
 	{
 		CollectionManagerScriptingSubsystemUtil::LogLastCollectionManagerError(Error);
 		return false;
@@ -134,6 +135,48 @@ bool UCollectionManagerScriptingSubsystem::CreateCollection(const FCollectionScr
 	return true;
 }
 
+bool UCollectionManagerScriptingSubsystem::CreateOrEmptyCollection(const FCollectionScriptingContainerSource Container, const FName Collection, const ECollectionScriptingShareType ShareType, FCollectionScriptingRef& OutNewOrEmptyCollection)
+{
+	const TSharedPtr<ICollectionContainer> CollectionContainer = FCollectionManagerModule::GetModule().Get().FindCollectionContainer(Container.Name);
+	if (!CollectionContainer.IsValid())
+	{
+		CollectionManagerScriptingSubsystemUtil::LogLastCollectionManagerError(FText::Format(CollectionManagerScriptingSubsystemUtil::FConstants::Get().ContainerNotFound, FText::FromName(Container.Name)));
+		return false;
+	}
+
+	const ECollectionShareType::Type CastShareType = CollectionManagerScriptingSubsystemUtil::ECollectionScriptingShareType_To_ECollectionShareType(ShareType);
+	if (CollectionContainer->CollectionExists(Collection, CastShareType))
+	{
+		FText Error;
+		if (!CollectionContainer->EmptyCollection(Collection, CastShareType, &Error))
+		{
+			CollectionManagerScriptingSubsystemUtil::LogLastCollectionManagerError(Error);
+			return false;
+		}
+	}
+	else
+	{
+		FText Error;
+		if (!CollectionContainer->IsValidCollectionName(Collection.ToString(), ECollectionShareType::CST_All, &Error))
+		{
+			CollectionManagerScriptingSubsystemUtil::LogLastCollectionManagerError(Error);
+			return false;
+		}
+
+		if (!CollectionContainer->CreateCollection(Collection, CastShareType, ECollectionStorageMode::Static, &Error))
+		{
+			CollectionManagerScriptingSubsystemUtil::LogLastCollectionManagerError(Error);
+			return false;
+		}
+	}
+
+	OutNewOrEmptyCollection.Name = Collection;
+	OutNewOrEmptyCollection.ShareType = ShareType;
+	OutNewOrEmptyCollection.Container = CollectionContainer->GetCollectionSource()->GetName();
+
+	return true;
+}
+
 bool UCollectionManagerScriptingSubsystem::GetCollections(const FCollectionScriptingContainerSource Container, TArray<FCollectionScriptingRef>& OutCollections)
 {
 	const TSharedPtr<ICollectionContainer> CollectionContainer = FCollectionManagerModule::GetModule().Get().FindCollectionContainer(Container.Name);
diff --git a/Engine/Source/Editor/UnrealEd/Public/Subsystems/CollectionManagerScriptingSubsystem.h b/Engine/Source/Editor/UnrealEd/Public/Subsystems/CollectionManagerScriptingSubsystem.h
--- a/Engine/Source/Editor/UnrealEd/Public/Subsystems/CollectionManagerScriptingSubsystem.h
+++ b/Engine/Source/Editor/UnrealEd/Public/Subsystems/CollectionManagerScriptingSubsystem.h
@@ -67,6 +67,19 @@ public:
 	UFUNCTION(BlueprintCallable, Category="CollectionManager")
 	bool CreateCollection(const FCollectionScriptingContainerSource Container, const FName Collection, const ECollectionScriptingShareType ShareType, FCollectionScriptingRef& OutNewCollection);
 
+	/**
+	 * Create a new collection with the given name and share type in the provided collection container if it doesn't already exist, or empty the existing collection if it does.
+	 *
+	 * @param Container The container the collection should be created or found in. "None" defaults to the base game's container.
+	 * @param Collection Name of the collection to create or empty.
+	 * @param ShareType Whether the collection should be local, private, or shared?
+	 * @param OutNewOrEmptyCollection The collection that was created or emptied.
+	 *
+	 * @return True if the collection was created or emptied, false otherwise (see the output log for details on error).
+	 */
+	UFUNCTION(BlueprintCallable, Category="CollectionManager")
+	bool CreateOrEmptyCollection(const FCollectionScriptingContainerSource Container, const FName Collection, const ECollectionScriptingShareType ShareType, FCollectionScriptingRef& OutNewOrEmptyCollection);
+
 	/**
 	 * Get all available collections in the specified container.
 	 *
