Updated GetCollections and GetCollectionsByName to return true if OutCollections contains at least 1 collection.

#rnx

diff --git a/Engine/Source/Editor/UnrealEd/Private/Subsystems/CollectionManagerScriptingSubsystem.cpp b/Engine/Source/Editor/UnrealEd/Private/Subsystems/CollectionManagerScriptingSubsystem.cpp
--- a/Engine/Source/Editor/UnrealEd/Private/Subsystems/CollectionManagerScriptingSubsystem.cpp
+++ b/Engine/Source/Editor/UnrealEd/Private/Subsystems/CollectionManagerScriptingSubsystem.cpp
@@ -199,7 +199,7 @@ bool UCollectionManagerScriptingSubsystem::GetCollections(const FCollectionScrip
 			};
 		});
 
-	return true;
+	return OutCollections.Num() > 0;
 }
 
 bool UCollectionManagerScriptingSubsystem::DestroyCollection(const FCollectionScriptingRef& Collection)
@@ -431,7 +431,7 @@ bool UCollectionManagerScriptingSubsystem::GetCollectionsByName(const FCollectio
 		};
 	});
 	
-	return true;
+	return OutCollections.Num() > 0;
 }
 
 bool UCollectionManagerScriptingSubsystem::GetAssetsInCollection(const FCollectionScriptingRef& Collection, TArray<FAssetData>& OutAssets)
diff --git a/Engine/Source/Editor/UnrealEd/Public/Subsystems/CollectionManagerScriptingSubsystem.h b/Engine/Source/Editor/UnrealEd/Public/Subsystems/CollectionManagerScriptingSubsystem.h
--- a/Engine/Source/Editor/UnrealEd/Public/Subsystems/CollectionManagerScriptingSubsystem.h
+++ b/Engine/Source/Editor/UnrealEd/Public/Subsystems/CollectionManagerScriptingSubsystem.h
@@ -86,7 +86,7 @@ public:
 	 * @param Container The container to retrieve collections from. "None" defaults to the base game's container.
 	 * @param OutCollections The collections found in the container.
 	 *
-	 * @return True if the container exists, false otherwise (see the output log for details on error).
+	 * @return True if OutCollections contains at least one collection, false otherwise (see the output log for details on error).
 	 */
 	UFUNCTION(BlueprintCallable, Category="CollectionManager")
 	bool GetCollections(const FCollectionScriptingContainerSource Container, TArray<FCollectionScriptingRef>& OutCollections);
@@ -285,7 +285,7 @@ public:
 	 * @param Collection The collection to look for.
 	 * @param OutCollections The collections found.
 	 *
-	 * @return True if the container exists, false otherwise (see the output log for details on error).
+	 * @return True if OutCollections contains at least one collection, false otherwise (see the output log for details on error).
 	 */
 	UFUNCTION(BlueprintCallable, Category = "CollectionManager")
 	bool GetCollectionsByName(const FCollectionScriptingContainerSource Container, const FName Collection, TArray<FCollectionScriptingRef>& OutCollections);
