Convert viewport scalability warning to a reset button.
- Scalability warning only shows up if current settings do not match defaults?including device profile defaults.
- Make existing Scalability "Defaults" button always accessible.
- Highlight scalability header buttons when all values are in that state.

diff --git a/Engine/Source/Editor/EditorStyle/Private/StarshipStyle.cpp b/Engine/Source/Editor/EditorStyle/Private/StarshipStyle.cpp
--- a/Engine/Source/Editor/EditorStyle/Private/StarshipStyle.cpp
+++ b/Engine/Source/Editor/EditorStyle/Private/StarshipStyle.cpp
@@ -465,6 +465,7 @@ void FStarshipEditorStyle::FStyle::SetupGeneralStyles()
 		Set("Icons.RecentAssets", new IMAGE_BRUSH_SVG("Starship/Common/RecentAssets", Icon16x16));
 		Set("Icons.Scalability", new IMAGE_BRUSH_SVG("Starship/Common/Scalability_16", Icon20x20));
 		Set("Icons.ViewportScalability", new IMAGE_BRUSH_SVG("Starship/Common/ViewportScalability_16", Icon16x16));
+		Set("Icons.ViewportScalabilityReset", new IMAGE_BRUSH_SVG("Starship/Common/ViewportScalabilityReset_16", Icon16x16));
 		Set("Icons.EllipsisVerticalNarrow",
 			new CORE_IMAGE_BRUSH_SVG("Starship/Common/ellipsis-vertical-narrow", FVector2f(6.f, 24.f)));
 
@@ -4749,6 +4750,8 @@ void FStarshipEditorStyle::FStyle::SetupLevelEditorStyle()
 		Set( "EditorViewport.TranslateRotate2DMode", new IMAGE_BRUSH_SVG("Starship/EditorViewport/TranslateRotate2D", Icon16x16 ) );
 
 		Set( "EditorViewport.ToggleRealTime", new IMAGE_BRUSH_SVG("Starship/Common/Realtime", Icon16x16));
+		Set( "EditorViewport.ToggleRealTimeLocked", new IMAGE_BRUSH_SVG("Starship/Common/RealtimeOff_16", Icon16x16));
+		Set( "EditorViewport.RealTimeReset", new IMAGE_BRUSH_SVG("Starship/Common/RealtimeReset_16", Icon16x16));
 
 		Set( "EditorViewport.LocationGridSnap", new IMAGE_BRUSH_SVG("Starship/EditorViewport/grid", Icon16x16));
 		Set( "EditorViewport.RotationGridSnap", new IMAGE_BRUSH_SVG("Starship/EditorViewport/angle", Icon16x16));
diff --git a/Engine/Source/Editor/LevelEditor/Private/ViewportToolbar/LevelEditorSubmenus.cpp b/Engine/Source/Editor/LevelEditor/Private/ViewportToolbar/LevelEditorSubmenus.cpp
--- a/Engine/Source/Editor/LevelEditor/Private/ViewportToolbar/LevelEditorSubmenus.cpp
+++ b/Engine/Source/Editor/LevelEditor/Private/ViewportToolbar/LevelEditorSubmenus.cpp
@@ -2203,6 +2203,7 @@ FToolMenuEntry CreatePerformanceAndScalabilitySubmenu()
 					FToolMenuSection& UnnamedSection = Submenu->FindOrAddSection(NAME_None);
 
 					UnnamedSection.AddEntry(UE::UnrealEd::CreateToggleRealtimeEntry());
+					UnnamedSection.AddEntry(UE::UnrealEd::CreateResetScalabilitySubmenu());
 
 					if (const UUnrealEdViewportToolbarContext* Context =
 							Submenu->FindContext<UUnrealEdViewportToolbarContext>())
@@ -2253,6 +2254,8 @@ FToolMenuEntry CreatePIEPerformanceAndScalabilitySubmenu()
 
 					FToolMenuEntry& MenuEntry = UnnamedSection.AddMenuEntry(FLevelViewportCommands::Get().ToggleAllowConstrainedAspectRatioInPreview);
 					MenuEntry.SetShowInToolbarTopLevel(true);
+					
+					UnnamedSection.AddEntry(UE::UnrealEd::CreateResetScalabilitySubmenu());
 				}
 
 				{
diff --git a/Engine/Source/Editor/UnrealEd/Private/SScalabilitySettings.cpp b/Engine/Source/Editor/UnrealEd/Private/SScalabilitySettings.cpp
--- a/Engine/Source/Editor/UnrealEd/Private/SScalabilitySettings.cpp
+++ b/Engine/Source/Editor/UnrealEd/Private/SScalabilitySettings.cpp
@@ -151,8 +151,10 @@ TSharedRef<SWidget> SScalabilitySettings::MakeButtonWidget(const FText& InName,
 
 TSharedRef<SWidget> SScalabilitySettings::MakeHeaderButtonWidget(const FText& InName, int32 InQualityLevel, const FText& InToolTip, Scalability::EQualityLevelBehavior Behavior)
 {
-	return SNew(SButton)
-		.OnClicked(this, &SScalabilitySettings::OnHeaderClicked, InQualityLevel, Behavior)
+	return SNew(SCheckBox)
+		.Style(FAppStyle::Get(), "DetailsView.SectionButton")
+		.OnCheckStateChanged(this, &SScalabilitySettings::OnHeaderClicked, InQualityLevel, Behavior)
+		.IsChecked(this, &SScalabilitySettings::GetHeaderChecked, InQualityLevel, Behavior)
 		.ToolTipText(InToolTip)
 		.Content()
 		[
@@ -175,11 +177,17 @@ TSharedRef<SWidget> SScalabilitySettings::MakeAutoButtonWidget()
 
 TSharedRef<SWidget> SScalabilitySettings::MakeDefaultButtonWidget()
 {
-	return SNew(SButton)
-		.OnClicked(this, &SScalabilitySettings::OnDefaultClicked)
-		.ToolTipText(LOCTEXT("DefaultButtonTooltip", "Resets the Quality Levels to what the Device Profile sets"))
-		.Visibility_Lambda([]() {
-			return (GEditor && GEditor->IsFeatureLevelPreviewActive()) ? EVisibility::Visible : EVisibility::Collapsed;
+	return SNew(SCheckBox)
+		.Style(FAppStyle::Get(), "DetailsView.SectionButton")
+		.OnCheckStateChanged(this, &SScalabilitySettings::OnDefaultClicked)
+		.IsChecked(this, &SScalabilitySettings::GetDefaultCheckState)
+		.ToolTipText_Lambda([]
+		{
+			if (GEditor && GEditor->IsFeatureLevelPreviewActive())
+			{
+				return LOCTEXT("DefaultButtonWithProfileTooltip", "Resets the Quality Levels to what the Device Profile sets");
+			}
+			return LOCTEXT("DefaultButtonTooltip", "Resets the Quality Levels to the default.");
 		})
 		.Content()
 		[
@@ -188,7 +196,7 @@ TSharedRef<SWidget> SScalabilitySettings::MakeDefaultButtonWidget()
 		];
 }
 
-FReply SScalabilitySettings::OnHeaderClicked(int32 InQualityLevel, Scalability::EQualityLevelBehavior Behavior)
+void SScalabilitySettings::OnHeaderClicked(ECheckBoxState CheckState, int32 InQualityLevel, Scalability::EQualityLevelBehavior Behavior)
 {
 	Scalability::FQualityLevels LevelCounts = Scalability::GetQualityLevelCounts();
 	switch (Behavior)
@@ -205,7 +213,25 @@ FReply SScalabilitySettings::OnHeaderClicked(int32 InQualityLevel, Scalability::
 	Scalability::SaveState(GEditorSettingsIni);
 	GEditor->RedrawAllViewports();
 	GEditor->RedrawLevelEditingViewports();
-	return FReply::Handled();
+}
+
+ECheckBoxState SScalabilitySettings::GetHeaderChecked(int32 InQualityLevel, Scalability::EQualityLevelBehavior Behavior) const
+{
+	Scalability::FQualityLevels CurrentLevels = Scalability::GetQualityLevels();
+	
+	Scalability::FQualityLevels TargetLevels;
+	switch (Behavior)
+	{
+	case Scalability::EQualityLevelBehavior::ERelativeToMax:
+		TargetLevels.SetFromSingleQualityLevelRelativeToMax(InQualityLevel);
+		break;		
+	case Scalability::EQualityLevelBehavior::EAbsolute:
+	default:
+		TargetLevels.SetFromSingleQualityLevel(InQualityLevel);
+		break;
+	}
+	
+	return CurrentLevels == TargetLevels ? ECheckBoxState::Checked : ECheckBoxState::Unchecked;
 }
 
 FReply SScalabilitySettings::OnAutoClicked()
@@ -221,7 +247,7 @@ FReply SScalabilitySettings::OnAutoClicked()
 	return FReply::Handled();
 }
 
-FReply SScalabilitySettings::OnDefaultClicked()
+void SScalabilitySettings::OnDefaultClicked(ECheckBoxState CheckState)
 {
 	Scalability::ResetQualityLevelsToDefault();
 	CachedQualityLevels = Scalability::GetQualityLevels();
@@ -229,7 +255,11 @@ FReply SScalabilitySettings::OnDefaultClicked()
 
 	GEditor->RedrawAllViewports();
 	GEditor->RedrawLevelEditingViewports();
-	return FReply::Handled();
+}
+
+ECheckBoxState SScalabilitySettings::GetDefaultCheckState() const
+{
+	return Scalability::GetQualityLevels() == Scalability::GetDefaultQualityLevels() ? ECheckBoxState::Checked : ECheckBoxState::Unchecked;
 }
 
 SGridPanel::FSlot::FSlotArguments SScalabilitySettings::MakeGridSlot(int32 InCol, int32 InRow, int32 InColSpan /*= 1*/, int32 InRowSpan /*= 1*/)
diff --git a/Engine/Source/Editor/UnrealEd/Private/ViewportToolbar/UnrealEdViewportToolbar.cpp b/Engine/Source/Editor/UnrealEd/Private/ViewportToolbar/UnrealEdViewportToolbar.cpp
--- a/Engine/Source/Editor/UnrealEd/Private/ViewportToolbar/UnrealEdViewportToolbar.cpp
+++ b/Engine/Source/Editor/UnrealEd/Private/ViewportToolbar/UnrealEdViewportToolbar.cpp
@@ -3755,7 +3755,7 @@ FToolMenuEntry CreateScreenPercentageSubmenu()
 
 FToolMenuEntry CreateScalabilitySubmenu()
 {
-	FToolMenuEntry ScalabilitySubmenu = FToolMenuEntry::InitSubMenu(
+	return FToolMenuEntry::InitSubMenu(
 		"Scalability",
 		LOCTEXT("ScalabilitySubMenu", "Viewport Scalability"),
 		TAttribute<FText>::CreateLambda([]()
@@ -3781,24 +3781,39 @@ FToolMenuEntry CreateScalabilitySubmenu()
 		false,
 		FSlateIcon(FAppStyle::GetAppStyleSetName(), "Icons.ViewportScalability")
 	);
+}
 
-	ScalabilitySubmenu.SetShowInToolbarTopLevel(TAttribute<bool>::Create(&UE::UnrealEd::IsScalabilityWarningVisible));
-
-	ScalabilitySubmenu.ToolBarData.ResizeParams.AllowClipping = TAttribute<bool>::CreateLambda(
-		[]() -> bool
-		{
-			return !UE::UnrealEd::IsScalabilityWarningVisible();
-		}
+FToolMenuEntry CreateResetScalabilitySubmenu()
+{
+	// Hide this entry in non-toolbar contexts.
+	// The actual Defaults button is handled by SScalabilitySettings
+	FUIAction InvisibleAction;
+	InvisibleAction.IsActionVisibleDelegate.BindLambda([] { return false; });
+				
+	FToolMenuEntry ResetToDefaultsEntry = FToolMenuEntry::InitMenuEntry(
+		"ResetToDefault",
+		LOCTEXT("ResetToDefaultLabel", "Reset To Defaults"),
+		GetScalabilityWarningTooltip(),
+		FSlateIcon(FAppStyle::GetAppStyleSetName(), "Icons.ViewportScalabilityReset"),
+		FUIAction(InvisibleAction)
 	);
-
-	// We need the warning style when we've raised this entry to the top-level toolbar to draw
-	// attention to it. However, styles are not driven by an attribute, so they cannot be updated after an
-	// FToolMenuEntry has been created. Therefore we always apply the style instead because the style
-	// will not make this entry look different when it appears in a menu, but when it appears in a toolbar
-	// we get the desired warning effect.
-	ScalabilitySubmenu.StyleNameOverride = "ViewportToolbarWarning";
+				
+	FToolUIAction ToolbarAction;
+	ToolbarAction.ExecuteAction.BindLambda([](const FToolMenuContext& Context)
+	{
+		Scalability::ResetQualityLevelsToDefault();
+		GEditor->RedrawAllViewports();
+	});
+				
+	ResetToDefaultsEntry.ToolBarData.ActionOverride = ToolbarAction;
+	ResetToDefaultsEntry.ToolBarData.ResizeParams.AllowClipping = TAttribute<bool>::CreateLambda([]
+	{
+		return !UE::UnrealEd::IsScalabilityWarningVisible();
+	});
+	ResetToDefaultsEntry.SetShowInToolbarTopLevel(TAttribute<bool>::Create(&UE::UnrealEd::IsScalabilityWarningVisible));
+	ResetToDefaultsEntry.StyleNameOverride = "ViewportToolbarWarning";
 	
-	return ScalabilitySubmenu;
+	return ResetToDefaultsEntry;
 }
 
 FText GetCameraSpeedTooltip()
@@ -3849,9 +3864,12 @@ FToolMenuEntry CreatePerformanceAndScalabilitySubmenu()
 
 bool IsScalabilityWarningVisible()
 {
-	// This method returns magic numbers. 3 means epic
-	return GetDefault<UEditorPerformanceSettings>()->bEnableScalabilityWarningIndicator
-		&& Scalability::GetQualityLevels().GetMinQualityLevel() != 3;
+	if (!GetDefault<UEditorPerformanceSettings>()->bEnableScalabilityWarningIndicator)
+	{
+		return false;
+	}
+	
+	return Scalability::GetQualityLevels() != Scalability::GetDefaultQualityLevels();
 }
 
 FText GetScalabilityWarningLabel()
@@ -3874,7 +3892,7 @@ FText GetScalabilityWarningLabel()
 FText GetScalabilityWarningTooltip()
 {
 	return LOCTEXT(
-		"ScalabilityWarning_ToolTip", "Non-default scalability settings could be affecting what is shown in this viewport.\nFor example you may experience lower visual quality, reduced particle counts, and other artifacts that don't match what the scene would look like when running outside of the editor. Click to make changes."
+		"ScalabilityWarning_ToolTip", "Non-default scalability settings could be affecting what is shown in this viewport.\nFor example you may experience lower visual quality, reduced particle counts, and other artifacts that don't match what the scene would look like when running outside of the editor.\n\nClick to reset scalability settings to default."
 	);
 }
 
@@ -3938,113 +3956,103 @@ FToolMenuEntry CreateToggleRealtimeEntry()
 					return;
 				}
 
-				TWeakPtr<SEditorViewport> EditorViewportWeak;
+				TWeakPtr<SEditorViewport> EditorViewportWeak = EditorViewportContext->Viewport;
+				
 				FToolUIAction RealtimeToggleAction;
-				if (EditorViewportContext)
-				{
-					EditorViewportWeak = EditorViewportContext->Viewport;
-
-					RealtimeToggleAction.ExecuteAction = FToolMenuExecuteAction::CreateLambda(
-						[EditorViewportWeak](const FToolMenuContext& Context) -> void
+				RealtimeToggleAction.ExecuteAction = FToolMenuExecuteAction::CreateLambda(
+					[EditorViewportWeak](const FToolMenuContext& Context) -> void
+					{
+						if (TSharedPtr<SEditorViewport> EditorViewport = EditorViewportWeak.Pin())
 						{
-							if (TSharedPtr<SEditorViewport> EditorViewport = EditorViewportWeak.Pin())
-							{
-								EditorViewport->OnToggleRealtime();
-							}
+							EditorViewport->OnToggleRealtime();
 						}
-					);
+					}
+				);
 
-					RealtimeToggleAction.CanExecuteAction = FToolMenuCanExecuteAction::CreateLambda(
-						[EditorViewportWeak](const FToolMenuContext& Context) -> bool
+				RealtimeToggleAction.CanExecuteAction = FToolMenuCanExecuteAction::CreateLambda(
+					[EditorViewportWeak](const FToolMenuContext& Context) -> bool
+					{
+						if (TSharedPtr<SEditorViewport> EditorViewport = EditorViewportWeak.Pin())
 						{
-							if (TSharedPtr<SEditorViewport> EditorViewport = EditorViewportWeak.Pin())
-							{
-								return !EditorViewport->GetViewportClient()->IsRealtimeOverrideSet();
-							}
-							return true;
+							return !EditorViewport->GetViewportClient()->IsRealtimeOverrideSet();
 						}
-					);
+						return true;
+					}
+				);
 
-					RealtimeToggleAction.GetActionCheckState = FToolMenuGetActionCheckState::CreateLambda(
-						[EditorViewportWeak](const FToolMenuContext& Context) -> ECheckBoxState
+				RealtimeToggleAction.GetActionCheckState = FToolMenuGetActionCheckState::CreateLambda(
+					[EditorViewportWeak](const FToolMenuContext& Context) -> ECheckBoxState
+					{
+						if (TSharedPtr<SEditorViewport> EditorViewport = EditorViewportWeak.Pin())
 						{
-							if (TSharedPtr<SEditorViewport> EditorViewport = EditorViewportWeak.Pin())
-							{
-								return EditorViewport->IsRealtime() ? ECheckBoxState::Checked : ECheckBoxState::Unchecked;
-							}
-							return ECheckBoxState::Undetermined;
+							return EditorViewport->IsRealtime() ? ECheckBoxState::Checked : ECheckBoxState::Unchecked;
 						}
-					);
-				}
+						return ECheckBoxState::Undetermined;
+					}
+				);
 
-				TAttribute<FText> Tooltip;
+				TAttribute<FText> Tooltip = TAttribute<FText>::CreateLambda([EditorViewportWeak]
 				{
-					const FText RealtimeTooltip = LOCTEXT("ToggleRealtimeTooltip", "Toggle realtime rendering of the viewport");
-
-					// If we can find a context with a viewport, use that to adjust the tooltip
-					// based on the viewport's realtime status.
-					if (EditorViewportContext)
+					if (const TSharedPtr<SEditorViewport> EditorViewport = EditorViewportWeak.Pin())
 					{
-						Tooltip = TAttribute<FText>::CreateLambda(
-							[EditorViewportWeak, RealtimeTooltip]() -> FText
-							{
-								if (const TSharedPtr<SEditorViewport> EditorViewport = EditorViewportWeak.Pin())
-								{
-									if (!EditorViewport->IsRealtime())
-									{
-										return LOCTEXT(
-											"ToggleRealtimeTooltip_WarnRealtimeOff", "Warning: This viewport is not updating in realtime. Click to turn on realtime mode."
-										);
-									}
-
-									if (EditorViewport->GetViewportClient()->IsRealtimeOverrideSet())
-									{
-										const FText Format = LOCTEXT(
-											"ToggleRealtimeTooltip_RealtimeOverrides",
-											"Realtime rendering cannot be toggled because an override has been set: {0}"
-										);
-										const FText Message =
-											EditorViewport->GetViewportClient()->GetRealtimeOverrideMessage();
-
-										return FText::Format(Format, Message);
-									}
-								}
+						if (EditorViewport->GetViewportClient()->IsRealtimeOverrideSet())
+						{
+							const FText Format = LOCTEXT(
+								"ToggleRealtimeTooltip_RealtimeOverrides",
+								"Realtime rendering cannot be toggled because an override has been set: {0}"
+							);
+							const FText Message =
+								EditorViewport->GetViewportClient()->GetRealtimeOverrideMessage();
 
-								return RealtimeTooltip;
-							}
-						);
+							return FText::Format(Format, Message);
+						}
+						
+						if (!EditorViewport->IsRealtime())
+						{
+							return LOCTEXT(
+								"ToggleRealtimeTooltip_WarnRealtimeOff", "Warning: This viewport is not updating in realtime. Click to turn on realtime mode."
+							);
+						}
 					}
-					else
+
+					return LOCTEXT("ToggleRealtimeTooltip", "Toggle realtime rendering of the viewport");
+				});
+				
+				TAttribute<FSlateIcon> Icon = TAttribute<FSlateIcon>::CreateLambda([EditorViewportWeak]
+				{
+					if (const TSharedPtr<SEditorViewport> EditorViewport = EditorViewportWeak.Pin())
 					{
-						Tooltip = RealtimeTooltip;
+						if (TSharedPtr<FEditorViewportClient> Client = EditorViewport->GetViewportClient())
+						{
+							if (Client->IsRealtimeOverrideSet())
+							{
+								return FSlateIcon(FAppStyle::GetAppStyleSetName(), "EditorViewport.ToggleRealTimeLocked");
+							}
+						}
 					}
-				}
+
+					return FSlateIcon(FAppStyle::GetAppStyleSetName(), "EditorViewport.RealTimeReset");
+				});
 
 				FToolMenuEntry ToggleRealtime = FToolMenuEntry::InitMenuEntry(
 					"ToggleRealtime",
 					LOCTEXT("ToggleRealtimeLabel", "Realtime Viewport"),
 					Tooltip,
-					FSlateIcon(FAppStyle::GetAppStyleSetName(), "EditorViewport.ToggleRealTime"),
+					Icon,
 					RealtimeToggleAction,
 					EUserInterfaceActionType::ToggleButton
 				);
 
-				// If we can find a context with a viewport, bind the top-level status of the
-				// realtime button to the viewport's realtime state where we show the realtime
-				// toggle in the top-level if the viewport is NOT realtime.
-				if (EditorViewportContext)
-				{
-					ToggleRealtime.SetShowInToolbarTopLevel(TAttribute<bool>::CreateLambda(
-						[WeakViewport = EditorViewportContext->Viewport]() -> bool
+				ToggleRealtime.SetShowInToolbarTopLevel(TAttribute<bool>::CreateLambda(
+					[WeakViewport = EditorViewportContext->Viewport]() -> bool
+					{
+						if (const TSharedPtr<SEditorViewport> EditorViewport = WeakViewport.Pin())
 						{
-							if (const TSharedPtr<SEditorViewport> EditorViewport = WeakViewport.Pin())
-							{
-								return ShouldShowViewportRealtimeWarning(*EditorViewport->GetViewportClient());
-							}
-							return false;
+							return ShouldShowViewportRealtimeWarning(*EditorViewport->GetViewportClient());
 						}
-					));
-				}
+						return false;
+					}
+				));
 
 				ToggleRealtime.ToolBarData.ResizeParams.AllowClipping = TAttribute<bool>::CreateLambda(
 					[EditorViewportWeak]() -> bool
@@ -4065,7 +4073,7 @@ FToolMenuEntry CreateToggleRealtimeEntry()
 
 				// We need the warning style when we've raised this entry to the top-level toolbar to draw
 				// attention to it. However, styles are not driven by an attribute, so they cannot be updated after an
-				// FToolMenuEntry has been created. Therefore we always apply the style instead because the style
+				// FToolMenuEntry has been created. Therefor we always apply the style instead because the style
 				// will not make this entry look different when it appears in a menu, but when it appears in a toolbar
 				// we get the desired warning effect.
 				ToggleRealtime.StyleNameOverride = "ViewportToolbarWarning";
diff --git a/Engine/Source/Editor/UnrealEd/Public/SScalabilitySettings.h b/Engine/Source/Editor/UnrealEd/Public/SScalabilitySettings.h
--- a/Engine/Source/Editor/UnrealEd/Public/SScalabilitySettings.h
+++ b/Engine/Source/Editor/UnrealEd/Public/SScalabilitySettings.h
@@ -68,13 +68,15 @@ private:
 	TSharedRef<SWidget> MakeDefaultButtonWidget();
 
 	// Callback for when a quality level header button is pressed
-	FReply OnHeaderClicked(int32 InQualityLevel, Scalability::EQualityLevelBehavior Behavior);
+	void OnHeaderClicked(ECheckBoxState CheckState, int32 InQualityLevel, Scalability::EQualityLevelBehavior Behavior);
+	ECheckBoxState GetHeaderChecked(int32 InQualityLevel, Scalability::EQualityLevelBehavior Behavior) const;
 
 	// Callback for auto benchmark button
 	FReply OnAutoClicked();
 	
 	// Callback for default button
-	FReply OnDefaultClicked();
+	void OnDefaultClicked(ECheckBoxState CheckState);
+	ECheckBoxState GetDefaultCheckState() const;
 
 	// Create a gridslot for the group quality level with all the required formatting
 	SGridPanel::FSlot::FSlotArguments MakeGridSlot(int32 InCol, int32 InRow, int32 InColSpan = 1, int32 InRowSpan = 1);
diff --git a/Engine/Source/Editor/UnrealEd/Public/ViewportToolbar/UnrealEdViewportToolbar.h b/Engine/Source/Editor/UnrealEd/Public/ViewportToolbar/UnrealEdViewportToolbar.h
--- a/Engine/Source/Editor/UnrealEd/Public/ViewportToolbar/UnrealEdViewportToolbar.h
+++ b/Engine/Source/Editor/UnrealEd/Public/ViewportToolbar/UnrealEdViewportToolbar.h
@@ -265,6 +265,7 @@ UNREALED_API void ConstructScreenPercentageMenu(UToolMenu* InMenu);
 UNREALED_API FToolMenuEntry CreateScreenPercentageSubmenu();
 
 UNREALED_API FToolMenuEntry CreateScalabilitySubmenu();
+UNREALED_API FToolMenuEntry CreateResetScalabilitySubmenu();
 
 UNREALED_API FText GetCameraSpeedTooltip();
 
diff --git a/Engine/Source/Runtime/Engine/Private/Scalability.cpp b/Engine/Source/Runtime/Engine/Private/Scalability.cpp
--- a/Engine/Source/Runtime/Engine/Private/Scalability.cpp
+++ b/Engine/Source/Runtime/Engine/Private/Scalability.cpp
@@ -464,7 +464,7 @@ void ApplyCachedQualityLevels()
 	}
 }
 
-void ResetQualityLevelsToDefault()
+FQualityLevels GetDefaultQualityLevels()
 {
 	// We directly use UDeviceProfileManager::DeviceProfileManagerSingleton here so we don't unexpectedly initialize the device profile system during cooking.
 	if (ScalabilityShaderPlatform != EShaderPlatform::SP_NumPlatforms && UDeviceProfileManager::DeviceProfileManagerSingleton)
@@ -474,10 +474,19 @@ void ResetQualityLevelsToDefault()
 		{
 			if (Scalability::FQualityLevels* DefaultQualityLevels = DefaultPreviewShaderPlatformToCachedQualityLevels.Find(MakeTuple(ScalabilityShaderPlatform, PreviewDeviceProfile->GetFName())))
 			{
-				SetQualityLevels(*DefaultQualityLevels, false);
+				return *DefaultQualityLevels;
 			}
 		}
 	}
+	
+	FQualityLevels QualityLevels;
+	QualityLevels.SetFromSingleQualityLevelRelativeToMax(1);
+	return QualityLevels;
+}
+
+void ResetQualityLevelsToDefault()
+{
+	SetQualityLevels(GetDefaultQualityLevels(), false);
 }
 #endif
 
diff --git a/Engine/Source/Runtime/Engine/Public/Scalability.h b/Engine/Source/Runtime/Engine/Public/Scalability.h
--- a/Engine/Source/Runtime/Engine/Public/Scalability.h
+++ b/Engine/Source/Runtime/Engine/Public/Scalability.h
@@ -176,7 +176,9 @@ namespace Scalability
 
 #if WITH_EDITOR
 	ENGINE_API void ApplyCachedQualityLevels();
-	/** Resets the Quality Levels to what the Device Profile sets */
+	/** Gets the default quality level from the Device Profile or system defaults */
+	ENGINE_API FQualityLevels GetDefaultQualityLevels();
+	/** Resets the Quality Levels to the Device Profile or system defaults */
 	ENGINE_API void ResetQualityLevelsToDefault();
 #endif
 
