Remove AddActorFromSelection flow from Viewport's TryPlacingAssetObject

diff --git a/Engine/Source/Editor/UnrealEd/Private/LevelEditorViewport.cpp b/Engine/Source/Editor/UnrealEd/Private/LevelEditorViewport.cpp
--- a/Engine/Source/Editor/UnrealEd/Private/LevelEditorViewport.cpp
+++ b/Engine/Source/Editor/UnrealEd/Private/LevelEditorViewport.cpp
@@ -307,18 +307,7 @@ TArray<FTypedElementHandle> FLevelEditorViewportClient::TryPlacingAssetObject(UL
 			ActorFactory = GEditor->FindActorFactoryForActorClass( ObjectClass );
 		}
 
-		// TODO: This path looks to be an ancient and unused one- AddActorFromSelection queries a seldom
-		// used separate object selection set, and replacing ObjToUse with whatever happens to be stored
-		// there (and only if ObjToUse happened to be an actor) does not make sense. This should be removed.
-		if ( ActorFactory != NULL )
-		{
-			AActor* Actor = FActorFactoryAssetProxy::AddActorFromSelection(ObjectClass, NULL, 
-				PlacementOptions.bSelectOutput, PlacementOptions.ObjectFlags, ActorFactory);
-			AddActorToPlaced(Actor);
-			bPlaced = (Actor != nullptr);
-		}
-
-		if (!bPlaced && ActorFactory != NULL )
+		if (ActorFactory != NULL )
 		{
 			UE::AssetPlacementUtil::FExtraPlaceAssetOptions OverridenPlacementOptions = PlacementOptions;
 			OverridenPlacementOptions.FactoryToUse = TScriptInterface<IAssetFactoryInterface>(ActorFactory);
