World Partition - Builders: Fix issues with DL loading logic in the builders.
* Consider the default loaded in editor state
* Renamed bLoadInitiallyActiveDataLayers to bForceExcludeNonInitiallyActivatedRuntimeDataLayer and changed it's default value from true to false.
* Simplified conditions to make the logic easier to understand

diff --git a/Engine/Plugins/PCG/Source/PCGEditor/Private/WorldPartitionBuilder/PCGWorldPartitionBuilder.cpp b/Engine/Plugins/PCG/Source/PCGEditor/Private/WorldPartitionBuilder/PCGWorldPartitionBuilder.cpp
--- a/Engine/Plugins/PCG/Source/PCGEditor/Private/WorldPartitionBuilder/PCGWorldPartitionBuilder.cpp
+++ b/Engine/Plugins/PCG/Source/PCGEditor/Private/WorldPartitionBuilder/PCGWorldPartitionBuilder.cpp
@@ -548,7 +548,7 @@ bool UPCGWorldPartitionBuilder::PreRun(UWorld* World, FPackageSourceControlHelpe
 	// Setup DataLayers
 	// Set Base Builder Class fields
 	bLoadNonDynamicDataLayers = Args.bLoadEditorOnlyDataLayers;
-	bLoadInitiallyActiveDataLayers = Args.bLoadActivatedRuntimeDataLayers;
+	bExcludeNonInitiallyActivatedRuntimeDataLayer = Args.bLoadActivatedRuntimeDataLayers;
 
 	IncludedDataLayers.Append(Args.IncludedDataLayers);
 	ExcludedDataLayers.Append(Args.ExcludedDataLayers);
diff --git a/Engine/Source/Editor/UnrealEd/Private/WorldPartition/WorldPartitionBuilder.cpp b/Engine/Source/Editor/UnrealEd/Private/WorldPartition/WorldPartitionBuilder.cpp
--- a/Engine/Source/Editor/UnrealEd/Private/WorldPartition/WorldPartitionBuilder.cpp
+++ b/Engine/Source/Editor/UnrealEd/Private/WorldPartition/WorldPartitionBuilder.cpp
@@ -303,6 +303,22 @@ bool UWorldPartitionBuilder::Run(UWorld* World, FPackageSourceControlHelper& Pac
 	return PostRun(World, PackageHelper, bResult);
 }
 
+bool UWorldPartitionBuilder::IsForceIncludedDataLayer(const UDataLayerInstance* InDataLayerInstance) const
+{
+	PRAGMA_DISABLE_DEPRECATION_WARNINGS
+	const FName DataLayerShortName(InDataLayerInstance->GetDataLayerShortName());
+	return IncludedDataLayers.Contains(InDataLayerInstance->GetAsset()) || DataLayerShortNames.Contains(DataLayerShortName);
+	PRAGMA_ENABLE_DEPRECATION_WARNINGS
+}
+
+bool UWorldPartitionBuilder::IsForceExcludedDataLayer(const UDataLayerInstance* InDataLayerInstance) const
+{
+	PRAGMA_DISABLE_DEPRECATION_WARNINGS
+	const FName DataLayerShortName(InDataLayerInstance->GetDataLayerShortName());
+	return ExcludedDataLayers.Contains(InDataLayerInstance->GetAsset()) || ExcludedDataLayerShortNames.Contains(DataLayerShortName);
+	PRAGMA_ENABLE_DEPRECATION_WARNINGS
+}
+
 void UWorldPartitionBuilder::LoadDataLayers(UWorld* InWorld)
 {
 	if (UWorldPartition* WorldPartition = InWorld->GetWorldPartition())
@@ -314,26 +330,30 @@ void UWorldPartitionBuilder::LoadDataLayers(UWorld* InWorld)
 		bool bUpdateEditorCells = false;
 		DataLayerManager->ForEachDataLayerInstance([&bUpdateEditorCells, this](UDataLayerInstance* DataLayer)
 		{
-			const FName DataLayerShortName(DataLayer->GetDataLayerShortName());
-
-			PRAGMA_DISABLE_DEPRECATION_WARNINGS
-			// Load all Non Excluded Data Layers + Non DynamicallyLoaded Data Layers + Initially Active Data Layers + Data Layers provided by builder
-			const bool bLoadedInEditor = !ExcludedDataLayers.Contains(DataLayer->GetAsset()) && !ExcludedDataLayerShortNames.Contains(DataLayerShortName) && ((bLoadNonDynamicDataLayers && !DataLayer->IsRuntime()) ||
-											(bLoadInitiallyActiveDataLayers && DataLayer->GetInitialRuntimeState() == EDataLayerRuntimeState::Activated) ||
-											IncludedDataLayers.Contains(DataLayer->GetAsset()) || DataLayerShortNames.Contains(DataLayerShortName));
-			PRAGMA_ENABLE_DEPRECATION_WARNINGS
-			
-			if (DataLayer->IsLoadedInEditor() != bLoadedInEditor)
+			// Forced inclusion
+			const bool bForceIncludeEditorDataLayer = bLoadNonDynamicDataLayers && !DataLayer->IsRuntime();
+			const bool bIsForceIncluded = IsForceIncludedDataLayer(DataLayer) || bForceIncludeEditorDataLayer;
+
+			// Forced exclusion
+			const bool bForceExcludeNonInitiallyActivatedRuntimeDataLayer = (bExcludeNonInitiallyActivatedRuntimeDataLayer && DataLayer->IsRuntime()) ? DataLayer->GetInitialRuntimeState() != EDataLayerRuntimeState::Activated : false;
+			const bool bIsForceExcluded = IsForceExcludedDataLayer(DataLayer) || bForceExcludeNonInitiallyActivatedRuntimeDataLayer;
+
+			// Default state
+			const bool bIsDefaultLoaded = DataLayer->IsLoadedInEditor();
+
+			// Final state
+			const bool bShouldBeLoaded = !bIsForceExcluded && (bIsForceIncluded || bIsDefaultLoaded);
+			if (DataLayer->IsLoadedInEditor() != bShouldBeLoaded)
 			{
 				bUpdateEditorCells = true;
-				DataLayer->SetIsLoadedInEditor(bLoadedInEditor, /*bFromUserChange*/false);
-				if (RequiresCommandletRendering() && bLoadedInEditor)
+				DataLayer->SetIsLoadedInEditor(bShouldBeLoaded, /*bFromUserChange*/false);
+				if (RequiresCommandletRendering() && bShouldBeLoaded)
 				{
 					DataLayer->SetIsInitiallyVisible(true);
 				}
 			}
 			
-			UE_LOG(LogWorldPartitionBuilder, Display, TEXT("DataLayer '%s' Loaded: %d"), *UDataLayerInstance::GetDataLayerText(DataLayer).ToString(), bLoadedInEditor ? 1 : 0);
+			UE_LOG(LogWorldPartitionBuilder, Display, TEXT("DataLayer '%s' Loaded: %d"), *UDataLayerInstance::GetDataLayerText(DataLayer).ToString(), bShouldBeLoaded ? 1 : 0);
 			
 			return true;
 		});
diff --git a/Engine/Source/Editor/UnrealEd/Public/WorldPartition/WorldPartitionBuilder.h b/Engine/Source/Editor/UnrealEd/Public/WorldPartition/WorldPartitionBuilder.h
--- a/Engine/Source/Editor/UnrealEd/Public/WorldPartition/WorldPartitionBuilder.h
+++ b/Engine/Source/Editor/UnrealEd/Public/WorldPartition/WorldPartitionBuilder.h
@@ -133,9 +133,9 @@ protected:
 		return Args;
 	}
 
-	UE_DEPRECATED(5.3, "Please use OnFilesModified")
+	UE_DEPRECATED(5.3, "Use OnFilesModified() instead")
 	UNREALED_API bool AutoSubmitFiles(const TArray<FString>& InModifiedFiles, const FString& InChangelistDescription) const;
-	UE_DEPRECATED(5.3, "Please use OnPackagesModified")
+	UE_DEPRECATED(5.3, "Use OnPackagesModified() instead")
 	UNREALED_API bool AutoSubmitPackages(const TArray<UPackage*>& InModifiedPackages, const FString& InChangelistDescription) const;
 
 	UNREALED_API virtual UWorld::InitializationValues GetWorldInitializationValues() const;
@@ -143,19 +143,25 @@ protected:
 	UE_DEPRECATED(5.3, "You must override the version that takes a World parameter")
 	virtual bool PreWorldInitialization(FPackageSourceControlHelper& PackageHelper) final { return true; }
 
+	bool IsForceIncludedDataLayer(const UDataLayerInstance* InDataLayerInstance) const;
+	bool IsForceExcludedDataLayer(const UDataLayerInstance* InDataLayerInstance) const;
+
 	int32 IterativeCellSize = 102400;
 	int32 IterativeCellOverlapSize = 0;
 	FBox  IterativeWorldBounds;
 
-	UE_DEPRECATED(5.6, "Please use IncludedDataLayers")
+	UE_DEPRECATED(5.6, "Data layers to include must now be specified using IncludedDataLayers")
 	TSet<FName> DataLayerShortNames;
-	UE_DEPRECATED(5.6, "Please use ExcludedDataLayers")
+	UE_DEPRECATED(5.6, "Data layers to exclude must now be specified using ExcludedDataLayers")
 	TSet<FName> ExcludedDataLayerShortNames;
 
 	TSet<TObjectPtr<UDataLayerAsset>> IncludedDataLayers;
 	TSet<TObjectPtr<UDataLayerAsset>> ExcludedDataLayers;
 
 	bool bLoadNonDynamicDataLayers = true;
+	bool bExcludeNonInitiallyActivatedRuntimeDataLayer = false;
+
+	UE_DEPRECATED(5.6, "This has been renamed to bExcludeNonInitiallyActivatedRuntimeDataLayer")
 	bool bLoadInitiallyActiveDataLayers = true;
 
 	FModifiedFilesHandler ModifiedFilesHandler;
diff --git a/Engine/Source/Runtime/Engine/Private/WorldPartition/RuntimeSpatialHash/RuntimeSpatialHashHLOD.cpp b/Engine/Source/Runtime/Engine/Private/WorldPartition/RuntimeSpatialHash/RuntimeSpatialHashHLOD.cpp
--- a/Engine/Source/Runtime/Engine/Private/WorldPartition/RuntimeSpatialHash/RuntimeSpatialHashHLOD.cpp
+++ b/Engine/Source/Runtime/Engine/Private/WorldPartition/RuntimeSpatialHash/RuntimeSpatialHashHLOD.cpp
@@ -371,14 +371,17 @@ static void UpdateHLODGridsActors(UWorld* World, const UActorDescContainerInstan
 	TMap<FName, ASpatialHashRuntimeGridInfo*> ExistingGridActors;
 	for (UActorDescContainerInstance::TConstIterator<ASpatialHashRuntimeGridInfo> Iterator(ContainerInstance); Iterator; ++Iterator)
 	{
-		if (ASpatialHashRuntimeGridInfo* GridActor = CastChecked<ASpatialHashRuntimeGridInfo>(Iterator->GetActor()))
+		if (Iterator->IsLoaded())
 		{
-			const FSpatialHashRuntimeGrid* HLODGrid = HLODGrids.Find(GridActor->GridSettings.GridName);
-			if (HLODGrid && GridActor->GridSettings.Priority && GridActor->GridSettings.HLODLayer)
-			{
-				check(GridActor->GetContentBundleGuid() == ContainerInstance->GetContentBundleGuid());
-				ExistingGridActors.Emplace(GridActor->GridSettings.GridName, GridActor);
-			}
+		    if (ASpatialHashRuntimeGridInfo* GridActor = CastChecked<ASpatialHashRuntimeGridInfo>(Iterator->GetActor()))
+		    {
+			    const FSpatialHashRuntimeGrid* HLODGrid = HLODGrids.Find(GridActor->GridSettings.GridName);
+			    if (HLODGrid && GridActor->GridSettings.Priority && GridActor->GridSettings.HLODLayer)
+			    {
+				    check(GridActor->GetContentBundleGuid() == ContainerInstance->GetContentBundleGuid());
+				    ExistingGridActors.Emplace(GridActor->GridSettings.GridName, GridActor);
+			    }
+		    }
 		}
 	}
 
