Fix potential crash in the ResavePackages commandlet.

- Prevent async compilation tasks from running while saving the package as we might run into threading issues between ::ResetLoaders and FBulkData.
- Added code comments to explain why the wait was added and detailing potential optimizations if it comes up as a performance issue. In testing the impact was relatively minor.

diff --git a/Engine/Source/Editor/UnrealEd/Private/Commandlets/ContentCommandlets.cpp b/Engine/Source/Editor/UnrealEd/Private/Commandlets/ContentCommandlets.cpp
--- a/Engine/Source/Editor/UnrealEd/Private/Commandlets/ContentCommandlets.cpp
+++ b/Engine/Source/Editor/UnrealEd/Private/Commandlets/ContentCommandlets.cpp
@@ -865,8 +865,11 @@ void UResavePackagesCommandlet::LoadAndSaveOnePackage(const FString& Filename)
 			if (bResaveOnDemand && bSavePackage)
 			{
 				// Finish all async loading for the package, in case the systems marking it for load do so only from the async loading
-				FAssetCompilingManager::Get().FinishAllCompilation();
-				FlushAsyncLoading();
+				{
+					TRACE_CPUPROFILER_EVENT_SCOPE(UResavePackagesCommandlet::LoadAndSaveOnePackage::ResaveOnDemandFlush);
+					FAssetCompilingManager::Get().FinishAllCompilation();
+					FlushAsyncLoading();
+				}
 				// Skip saving the package if no systems requested it
 				FScopeLock ScopeLock(&ResaveOnDemandPackagesLock);
 				bSavePackage = ResaveOnDemandPackages.Contains(Package->GetFName());
@@ -877,6 +880,17 @@ void UResavePackagesCommandlet::LoadAndSaveOnePackage(const FString& Filename)
 				TArray<UObject*> ObjectsInOuter;
 				GetObjectsWithOuter(Package, ObjectsInOuter);
 
+				{
+					// Added to guard against FBulkData being used in an async task while ::ResetLoaders is called later on this thread.
+					// Can be removed when bulkdata attachment issues are sorted.
+					// If this shows up as a perf issue we could consider only waiting if the package's FLinkerLoad has attached
+					// FBulkData.
+					// NOTE: This call will wait for any async job related to one of the UObjects in the package, the ::FinishAllCompilation
+					// call above (if bResaveOnDemand is true) will wait for all sync compilation jobs to complete.
+					TRACE_CPUPROFILER_EVENT_SCOPE(UResavePackagesCommandlet::LoadAndSaveOnePackage::FinishCompilationForObjects);
+					FAssetCompilingManager::Get().FinishCompilationForObjects(ObjectsInOuter);
+				}
+
 				if (ObjectsInOuter.IsEmpty())
 				{
 					bSavePackage = false;
