Fix the 'resavepackages' commandlet to work correctly when '-resaveondemand=bulkdataduplicates' is applied.

- Add the 'ResavePackages' commandlet to the config value 'ini:Editor:[BulkDataRegistry]CommandletsUsingBulkDataRegistry' so that the bulkdata registry is enabled for it.
- The ResavePackages commandlet will now give an error if the user is running with  '-ResaveOnDemand=BulkDataDuplicates' but the BulkDataRegistry is disabled along with info on how to enable it as without the BulkDataRegistry the 'BulkDataDuplicates' operation will do nothing, so the commandlet will never re-save any package. All the commandlet will do is load all of the packages available to it and unload them again.
- If the BulkDataRegistry is enabled for the ResavePackages commandlet then we need to call 'FTickableEditorObject::TickObjects' each time we do a garbage collection pass so that the BulkDataRegistry can clear out temp payloads that may have been loaded.

diff --git a/Engine/Config/BaseEditor.ini b/Engine/Config/BaseEditor.ini
--- a/Engine/Config/BaseEditor.ini
+++ b/Engine/Config/BaseEditor.ini
@@ -669,6 +669,9 @@ bCanModifyPluginsFromBrowser=true
 [EditorBulkData]
 LogCookedStatus=false
 
+[BulkDataRegistry]
++CommandletsUsingBulkDataRegistry=ResavePackagesCommandlet
+
 [DeviceDiscovery]
 AppleRescanRate=10
 
diff --git a/Engine/Source/Editor/UnrealEd/Private/Commandlets/ContentCommandlets.cpp b/Engine/Source/Editor/UnrealEd/Private/Commandlets/ContentCommandlets.cpp
--- a/Engine/Source/Editor/UnrealEd/Private/Commandlets/ContentCommandlets.cpp
+++ b/Engine/Source/Editor/UnrealEd/Private/Commandlets/ContentCommandlets.cpp
@@ -46,6 +46,7 @@
 #include "Particles/ParticleEmitter.h"
 #include "Particles/ParticleSystem.h"
 #include "PlatformInfo.h"
+#include "Serialization/BulkDataRegistry.h"
 #include "SourceControlHelpers.h"
 #include "SourceControlOperations.h"
 #include "StaticMeshCompiler.h"
@@ -383,6 +384,12 @@ int32 UResavePackagesCommandlet::InitializeResaveParameters( const TArray<FStrin
 			*SystemNames);
 		bResaveOnDemand = true;
 		UE::SavePackageUtilities::OnAddResaveOnDemandPackage.BindUObject(this, &UResavePackagesCommandlet::OnAddResaveOnDemandPackage);
+
+		if (!IBulkDataRegistry::IsEnabled() && ResaveOnDemandSystems.Contains(FName("bulkdataduplicates")))
+		{
+			UE_LOG(LogContentCommandlet, Error, TEXT("The BulkDataRegistry must be enabled to use -ResaveOnDemand=BulkDataDuplicates! Try using the -AllowCommandletBulkDataRegistry cmdline switch."));
+			return 1;
+		}
 	}
 
 	// This option allows the dependency graph and soft object path redirect map to be populated. This is useful if you want soft object references to redirectors to be followed to the destination asset at save time.
@@ -1022,6 +1029,11 @@ void UResavePackagesCommandlet::LoadAndSaveOnePackage(const FString& Filename)
 				FAssetCompilingManager::Get().FinishAllCompilation();
 			}
 
+			if (IBulkDataRegistry::IsEnabled())
+			{
+				FTickableEditorObject::TickObjects(0.0f);
+			}
+
 			if (GarbageCollectionFrequency > 1)
 			{
 				UE_LOG(LogContentCommandlet, Display, TEXT("GC"));
