[CookOnTheFlyServer] AI and Navigation wold subsystems are not required for ValidateSourcePackage

diff --git a/Engine/Source/Editor/UnrealEd/Private/CookOnTheFlyServer.cpp b/Engine/Source/Editor/UnrealEd/Private/CookOnTheFlyServer.cpp
--- a/Engine/Source/Editor/UnrealEd/Private/CookOnTheFlyServer.cpp
+++ b/Engine/Source/Editor/UnrealEd/Private/CookOnTheFlyServer.cpp
@@ -3343,6 +3343,8 @@ EDataValidationResult UCookOnTheFlyServer::ValidateSourcePackage(UE::Cook::FPack
 		IVS.EnableTraceCollision(true);
 		IVS.SetTransactional(false);
 		IVS.CreateWorldPartition(true);
+		IVS.CreateAISystem(false);
+		IVS.CreateNavigation(false);
 
 		World->InitWorld(IVS);
 		bRunCleanupWorld = true;
