Add definition for ForceReplaceReferences that take a list of replace requests

diff --git a/Engine/Source/Editor/UnrealEd/Private/ObjectTools.cpp b/Engine/Source/Editor/UnrealEd/Private/ObjectTools.cpp
--- a/Engine/Source/Editor/UnrealEd/Private/ObjectTools.cpp
+++ b/Engine/Source/Editor/UnrealEd/Private/ObjectTools.cpp
@@ -1491,6 +1491,12 @@ namespace ObjectTools
 		ForceReplaceReferences(ObjectToReplaceWith, ObjectsToReplace, ObjectsToReplaceWithin, ReplaceInfo, false);
 	}
 
+	void ForceReplaceReferences(TArrayView<FReplaceRequest> Requests, TSet<UObject*>& ObjectsToReplaceWithin)
+	{
+		FForceReplaceInfo ReplaceInfo;
+		ForceReplaceReferences(Requests, ObjectsToReplaceWithin, ReplaceInfo, false);
+	}
+
 	FConsolidationResults ConsolidateObjects(TArrayView<FReplaceRequest> Requests, TSet<UObject*>& ObjectsToConsolidateWithin, TSet<UObject*>& ObjectsToNotConsolidateWithin, bool bShouldDeleteAfterConsolidate, bool bWarnAboutRootSet)
 	{
 		for (FReplaceRequest Request : Requests)
