Cooker: SkipOnlyEditorOnly: When running in the legacy -skiponlyeditoronly=false mode, add the new-style dependencies that UObjects add during SavePackage via calls to AddRuntimeDependency, since some classes are relying on that new API to add their required dependencies.
#rnx

diff --git a/Engine/Source/Editor/UnrealEd/Private/Cooker/CookSavePackage.cpp b/Engine/Source/Editor/UnrealEd/Private/Cooker/CookSavePackage.cpp
--- a/Engine/Source/Editor/UnrealEd/Private/Cooker/CookSavePackage.cpp
+++ b/Engine/Source/Editor/UnrealEd/Private/Cooker/CookSavePackage.cpp
@@ -739,7 +739,7 @@ void FSaveCookedPackageContext::CalculatePlatformRuntimeDependencies()
 
 	// Add imports and softobjectpaths from the save to the SaveDependencies for the current Platform and for all Platforms
 	FName PackageFName = Package->GetFName();
-	if (COTFS.bSkipOnlyEditorOnly && HasSavePackageResult())
+	if (HasSavePackageResult())
 	{
 		TConstArrayView<const ITargetPlatform*> ReachablePlatforms(&TargetPlatform, 1);
 		for (const TArray<FName>* DependencyNames : { &SavePackageResult.ImportPackages, &SavePackageResult.SoftPackageReferences })
