Make SEditorViewportToolBarMenu handle a null weak  ParentToolbar pointer correctly.
- Make legacy Level Viewport Toolbar extension handle potential null weak pointers correctly.

diff --git a/Engine/Source/Editor/LevelEditor/Private/SLevelViewportToolBar.cpp b/Engine/Source/Editor/LevelEditor/Private/SLevelViewportToolBar.cpp
--- a/Engine/Source/Editor/LevelEditor/Private/SLevelViewportToolBar.cpp
+++ b/Engine/Source/Editor/LevelEditor/Private/SLevelViewportToolBar.cpp
@@ -93,9 +93,13 @@ public:
 			Menu->AddDynamicSection("LevelSection", FNewToolMenuDelegate::CreateLambda([](UToolMenu* InMenu) {
 				if (UEditorViewportViewMenuContext* Context = InMenu->FindContext<UEditorViewportViewMenuContext>())
 				{
-					TSharedPtr<SLevelViewportToolBar> LevelViewportToolBar = StaticCastSharedPtr<SLevelViewportToolBar>(
-						Context->EditorViewportViewMenu.Pin()->GetParentToolBar().Pin());
-					LevelViewportToolBar->FillViewMenu(InMenu);
+					if (TSharedPtr<const SEditorViewportToolbarMenu> Menu = Context->EditorViewportViewMenu.Pin())
+					{
+						if (TSharedPtr<SLevelViewportToolBar> LevelViewportToolBar = StaticCastSharedPtr<SLevelViewportToolBar>(Menu->GetParentToolBar().Pin()))
+						{
+							LevelViewportToolBar->FillViewMenu(InMenu);
+						}
+					}
 				}
 			}));
 		}
diff --git a/Engine/Source/Editor/UnrealEd/Private/SEditorViewportToolBarMenu.cpp b/Engine/Source/Editor/UnrealEd/Private/SEditorViewportToolBarMenu.cpp
--- a/Engine/Source/Editor/UnrealEd/Private/SEditorViewportToolBarMenu.cpp
+++ b/Engine/Source/Editor/UnrealEd/Private/SEditorViewportToolBarMenu.cpp
@@ -35,7 +35,6 @@ void SEditorViewportToolbarMenu::Construct( const FArguments& Declaration )
 	LabelIconBrush = Declaration._LabelIcon;
 
 	ParentToolBar = Declaration._ParentToolBar;
-	checkf(ParentToolBar.IsValid(), TEXT("The parent toolbar must be specified") );
 
 	TSharedPtr<SWidget> ButtonContent;
 
@@ -118,10 +117,15 @@ FText SEditorViewportToolbarMenu::GetFilteredToolTipText(TAttribute<FText> ToolT
 {
 	// If we're part of a toolbar that has a currently open menu, then we suppress our tool-tip
 	// as it will just get in the way
-	TWeakPtr<SMenuAnchor> OpenedMenu = ParentToolBar.Pin()->GetOpenMenu();
-	if (OpenedMenu.IsValid() && OpenedMenu.Pin()->IsOpen())
+	if (TSharedPtr<SViewportToolBar> Toolbar = ParentToolBar.Pin())
 	{
-		return FText::GetEmpty();
+		if (TSharedPtr<SMenuAnchor> OpenedMenu = Toolbar->GetOpenMenu().Pin())
+		{
+			if (OpenedMenu->IsOpen())
+			{
+				return FText::GetEmpty();
+			}
+		}
 	}
 
 	return ToolTipText.Get();
@@ -133,13 +137,19 @@ FReply SEditorViewportToolbarMenu::OnMenuClicked()
 	if (MenuAnchor->ShouldOpenDueToClick())
 	{
 		MenuAnchor->SetIsOpen( true );
-		ParentToolBar.Pin()->SetOpenMenu( MenuAnchor );
+		if (TSharedPtr<SViewportToolBar> Toolbar = ParentToolBar.Pin())
+		{
+			Toolbar->SetOpenMenu(MenuAnchor);
+		}
 	}
 	else
 	{
 		MenuAnchor->SetIsOpen( false );
-		TSharedPtr<SMenuAnchor> NullAnchor;
-		ParentToolBar.Pin()->SetOpenMenu( NullAnchor );
+		if (TSharedPtr<SViewportToolBar> Toolbar = ParentToolBar.Pin())
+		{
+			TSharedPtr<SMenuAnchor> NullAnchor;
+			Toolbar->SetOpenMenu(NullAnchor);
+		}
 	}
 
 	return FReply::Handled();
@@ -149,15 +159,19 @@ void SEditorViewportToolbarMenu::OnMouseEnter( const FGeometry& MyGeometry, cons
 {
 	SCompoundWidget::OnMouseEnter(MyGeometry, MouseEvent);
 
-	// See if there is another menu on the same tool bar already open
-	TWeakPtr<SMenuAnchor> OpenedMenu = ParentToolBar.Pin()->GetOpenMenu();
-	if( OpenedMenu.IsValid() && OpenedMenu.Pin()->IsOpen() && OpenedMenu.Pin() != MenuAnchor )
+	if (TSharedPtr<SViewportToolBar> Toolbar = ParentToolBar.Pin())
 	{
-		// There is another menu open so we open this menu and close the other
-		ParentToolBar.Pin()->SetOpenMenu( MenuAnchor ); 
-		MenuAnchor->SetIsOpen( true );
+		// See if there is another menu on the same tool bar already open
+		if (TSharedPtr<SMenuAnchor> OpenedMenu = Toolbar->GetOpenMenu().Pin())
+		{
+			if (OpenedMenu != MenuAnchor)
+			{
+				// There is another menu open so we open this menu and close the other
+				Toolbar->SetOpenMenu( MenuAnchor ); 
+				MenuAnchor->SetIsOpen( true );	
+			}
+		}
 	}
-
 }
 
 EVisibility SEditorViewportToolbarMenu::GetLabelIconVisibility() const
