Updated FCollectionScriptingContainerSource properties to support VisibleAnywhere and FCollectionScriptingRef properties to support EditAnywhere.

diff --git a/Engine/Source/Editor/UnrealEd/Public/Subsystems/CollectionManagerScriptingSubsystem.h b/Engine/Source/Editor/UnrealEd/Public/Subsystems/CollectionManagerScriptingSubsystem.h
--- a/Engine/Source/Editor/UnrealEd/Public/Subsystems/CollectionManagerScriptingSubsystem.h
+++ b/Engine/Source/Editor/UnrealEd/Public/Subsystems/CollectionManagerScriptingSubsystem.h
@@ -14,11 +14,11 @@ struct FCollectionScriptingContainerSource
 	GENERATED_BODY()
 
 	/** The name of the container. Defaults to the base game's container. */
-	UPROPERTY(BlueprintReadOnly, Category="CollectionManager")
+	UPROPERTY(VisibleAnywhere, BlueprintReadOnly, Category="CollectionManager")
 	FName Name = NAME_Game;
 
 	/** Friendly title of the container. */
-	UPROPERTY(BlueprintReadOnly, Category="CollectionManager")
+	UPROPERTY(VisibleAnywhere, BlueprintReadOnly, Category="CollectionManager")
 	FText Title = NSLOCTEXT("CollectionScriptingContainerSource", "ProjectCollectionSource_Name", "Collections");
 };
 
@@ -28,15 +28,15 @@ struct FCollectionScriptingRef
 	GENERATED_BODY()
 
 	/** The name (not title) of the container that owns this collection. Defaults to the base game's container. */
-	UPROPERTY(BlueprintReadWrite, Category="CollectionManager")
+	UPROPERTY(EditAnywhere, BlueprintReadWrite, Category="CollectionManager")
 	FName Container = NAME_Game;
 
 	/** Friendly name of the collection. */
-	UPROPERTY(BlueprintReadWrite, Category="CollectionManager")
+	UPROPERTY(EditAnywhere, BlueprintReadWrite, Category="CollectionManager")
 	FName Name;
 
 	/** Share type of this collection, which controls its visibility. */
-	UPROPERTY(BlueprintReadWrite, Category="CollectionManager")
+	UPROPERTY(EditAnywhere, BlueprintReadWrite, Category="CollectionManager")
 	ECollectionScriptingShareType ShareType = ECollectionScriptingShareType::Local;
 };
 
