Asset Editor Subsystem: Walk the class hierarchy for per interface overrides to check if a super class has an interface with an override

diff --git a/Engine/Source/Editor/UnrealEd/Private/Subsystems/BrowseToAssetOverrideSubsystem.cpp b/Engine/Source/Editor/UnrealEd/Private/Subsystems/BrowseToAssetOverrideSubsystem.cpp
--- a/Engine/Source/Editor/UnrealEd/Private/Subsystems/BrowseToAssetOverrideSubsystem.cpp
+++ b/Engine/Source/Editor/UnrealEd/Private/Subsystems/BrowseToAssetOverrideSubsystem.cpp
@@ -46,18 +46,21 @@ FName UBrowseToAssetOverrideSubsystem::GetBrowseToAssetOverride(const UObject* O
 	// Query all the class interfaces to see if there's a valid per-interface override for this instance
 	if (PerInterfaceOverrides.Num() > 0)
 	{
-		UClass* ObjectClass = Object->GetClass();
-		for (const FImplementedInterface& Interface : ObjectClass->Interfaces)
+		for (const UClass* ObjectClass = Object->GetClass(); ObjectClass; ObjectClass = ObjectClass->GetSuperClass())
 		{
-			if (const FBrowseToAssetOverrideDelegate* CallbackPtr = PerInterfaceOverrides.Find(Interface.Class->GetClassPathName()))
+			for (const FImplementedInterface& Interface : ObjectClass->Interfaces)
 			{
-				FName CallbackBrowseToAssetOverride = CallbackPtr->IsBound() ? CallbackPtr->Execute(Object) : FName();
-				if (!CallbackBrowseToAssetOverride.IsNone())
+				if (const FBrowseToAssetOverrideDelegate* CallbackPtr = PerInterfaceOverrides.Find(Interface.Class->GetClassPathName()))
 				{
-					return CallbackBrowseToAssetOverride;
+					FName CallbackBrowseToAssetOverride = CallbackPtr->IsBound() ? CallbackPtr->Execute(Object) : FName();
+					if (!CallbackBrowseToAssetOverride.IsNone())
+					{
+						return CallbackBrowseToAssetOverride;
+					}
 				}
 			}
-		}		
+		}
+				
 	}
 
 
