InteractiveToolsFramework: Change the offset used in ortho raycasts to be less huge, which causes numerical issues when trying to hit gizmos.

diff --git a/Engine/Source/Editor/UnrealEd/Private/Tools/EdModeInteractiveToolsContext.cpp b/Engine/Source/Editor/UnrealEd/Private/Tools/EdModeInteractiveToolsContext.cpp
--- a/Engine/Source/Editor/UnrealEd/Private/Tools/EdModeInteractiveToolsContext.cpp
+++ b/Engine/Source/Editor/UnrealEd/Private/Tools/EdModeInteractiveToolsContext.cpp
@@ -709,7 +709,7 @@ FRay UEditorInteractiveToolsContext::GetRayFromMousePos(FEditorViewportClient* V
 	// ray origin "outside" the scene (which is a disaster for numerical precision !! ... )
 	if (ViewportClient->IsOrtho())
 	{
-		RayOrigin -= 0.1 * HALF_WORLD_MAX * RayDirection;
+		RayOrigin -= UE_FLOAT_HUGE_DISTANCE * RayDirection;
 	}
 
 	return FRay(RayOrigin, RayDirection, true);
