Unshelved from pending changelist '41851935':

Add an option to convert back the short debug names to full names for the VertexFactories. When ShaderSymbols.info was produced with r.DumpShaderDebugShortNames=1, we need to convert back the vertex factory names.

#tests locally CookShaders cmdlet

diff --git a/Engine/Source/Editor/UnrealEd/Private/Commandlets/CookShadersCommandlet.cpp b/Engine/Source/Editor/UnrealEd/Private/Commandlets/CookShadersCommandlet.cpp
--- a/Engine/Source/Editor/UnrealEd/Private/Commandlets/CookShadersCommandlet.cpp
+++ b/Engine/Source/Editor/UnrealEd/Private/Commandlets/CookShadersCommandlet.cpp
@@ -64,7 +64,7 @@ namespace CookShadersCommandlet {
 	};
 	static_assert(UE_ARRAY_COUNT(MaterialQualityLevelNames) == EMaterialQualityLevel::Num + 1, "Missing entry from material quality level names.");
 
-	bool LoadAndParse(const FString& Path, const FString& Filter, TArray<FInfoRecord>& OutInfo)
+	bool LoadAndParse(const FString& Path, const FString& Filter, TArray<FInfoRecord>& OutInfo, bool bUseShortNames)
 	{
 		IFileManager& FileManager = IFileManager::Get();
 		TUniquePtr<FArchive> Reader = TUniquePtr<FArchive>(FileManager.CreateFileReader(*Path));
@@ -154,8 +154,25 @@ namespace CookShadersCommandlet {
 
 									// if it has 5 or more parts, Num-3 is the vertex factory
 									if (Substrings.Num() >= 5)
-									{
+									{										
 										Record.VertexFactory = Substrings[Substrings.Num() - 3];
+ 										if(bUseShortNames)
+										{ 
+											Record.VertexFactory.ReplaceInline(TEXT("Land"), TEXT("Landscape"));
+											Record.VertexFactory.ReplaceInline(TEXT("Inst"), TEXT("Instanced"));
+											Record.VertexFactory.ReplaceInline(TEXT("VF"), TEXT("VertexFactory"));
+											Record.VertexFactory.ReplaceInline(TEXT("APEX"), TEXT("GPUSkinAPEXCloth"));
+											Record.VertexFactory.ReplaceInline(TEXT("_1"), TEXT("true"));
+											Record.VertexFactory.ReplaceInline(TEXT("_0"), TEXT("false"));
+											if (Record.VertexFactory.Contains(TEXT("GPUSkin")))
+											{
+												Record.VertexFactory.InsertAt(0, TEXT("T"));
+											}
+											else
+											{
+												Record.VertexFactory.InsertAt(0, TEXT("F"));
+											}
+										}
 									}
 
 									// if it has 6 parts, Num-4 is the pipeline
@@ -223,6 +240,8 @@ int32 UCookShadersCommandlet::Main(const FString& Params)
 		UE_LOG(LogCookShadersCommandlet, Log, TEXT(" Optional: -material=<string>             (Cook this material if you don't have a .info file, can be Global for global shaders)."));
 		UE_LOG(LogCookShadersCommandlet, Log, TEXT(" Optional: -noglobals                     (Don't do global shaders, even if they match the filter.)"));
 		UE_LOG(LogCookShadersCommandlet, Log, TEXT(" Optional: -nomaterialinstances           (Don't do material instances)"));
+		UE_LOG(LogCookShadersCommandlet, Log, TEXT(" Optional: -useshortnames				  (ShaderSymbols.info was produced with r.DumpShaderDebugShortNames=1. We need to convert back the vertex factory names"));
+
 		return 0;
 	}
 
@@ -237,6 +256,7 @@ int32 UCookShadersCommandlet::Main(const FString& Params)
 	FParse::Value(*Params, TEXT("ShaderSymbolsExport="), ExportPath, true);
 	const bool bNoGlobals = Switches.Contains(TEXT("noglobals"));
 	const bool bNoMaterialInstances = Switches.Contains(TEXT("nomaterialinstances"));
+	const bool bUseShortNames = Switches.Contains(TEXT("useshortnames"));
 
 	TArray<FInfoRecord> Info;
 
@@ -255,7 +275,7 @@ int32 UCookShadersCommandlet::Main(const FString& Params)
 	// Load info file if requested
 	if (!InfoFilePath.IsEmpty())
 	{
-		if (!LoadAndParse(InfoFilePath, Filter, Info))
+		if (!LoadAndParse(InfoFilePath, Filter, Info, bUseShortNames))
 		{
 			UE_LOG(LogCookShadersCommandlet, Log, TEXT("Unabled to read / parse info file '%s'"), *InfoFilePath);
 			return 0;
