Unshelved from pending changelist '41852713':

[Transforms] Fix issue with pan dragging in ortho viewports.

With LUF enabled, the Y axis was often inverted.  Right ortho was using Left ortho X calculations, which were backwards.

diff --git a/Engine/Source/Editor/UnrealEd/Private/EditorViewportClient.cpp b/Engine/Source/Editor/UnrealEd/Private/EditorViewportClient.cpp
--- a/Engine/Source/Editor/UnrealEd/Private/EditorViewportClient.cpp
+++ b/Engine/Source/Editor/UnrealEd/Private/EditorViewportClient.cpp
@@ -3899,24 +3899,25 @@ void FEditorViewportClient::OnOrthoZoom( const struct FInputEventState& InputSta
 		}
 		int32 DeltaFromCenterX = -(CenterX - (ViewportWidth>>1));
 		int32 DeltaFromCenterY =  (CenterY - (ViewportHeight>>1));
+		const bool bIsUsingLUFCoordinates = AxisDisplayInfo::GetAxisDisplayCoordinateSystem() == EAxisList::LeftUpForward;
 		switch( GetViewportType() )
 		{
-		case LVT_OrthoXY:
-			OldOffsetFromCenter.Set(DeltaFromCenterX, -DeltaFromCenterY, 0.0f);
+		case LVT_OrthoTop:
+			OldOffsetFromCenter.Set(DeltaFromCenterX, bIsUsingLUFCoordinates ? DeltaFromCenterY : -DeltaFromCenterY, 0.0f);
 			break;
-		case LVT_OrthoXZ:
+		case LVT_OrthoLeft:
 			OldOffsetFromCenter.Set(DeltaFromCenterX, 0.0f, DeltaFromCenterY);
 			break;
-		case LVT_OrthoYZ:
+		case LVT_OrthoBack:
 			OldOffsetFromCenter.Set(0.0f, DeltaFromCenterX, DeltaFromCenterY);
 			break;
-		case LVT_OrthoNegativeXY:
-			OldOffsetFromCenter.Set(-DeltaFromCenterX, -DeltaFromCenterY, 0.0f);
+		case LVT_OrthoBottom:
+			OldOffsetFromCenter.Set(-DeltaFromCenterX, bIsUsingLUFCoordinates ? DeltaFromCenterY : -DeltaFromCenterY, 0.0f);
 			break;
-		case LVT_OrthoNegativeXZ:
+		case LVT_OrthoRight:
 			OldOffsetFromCenter.Set(-DeltaFromCenterX, 0.0f, DeltaFromCenterY);
 			break;
-		case LVT_OrthoNegativeYZ:
+		case LVT_OrthoFront:
 			OldOffsetFromCenter.Set(0.0f, -DeltaFromCenterX, DeltaFromCenterY);
 			break;
 		case LVT_OrthoFreelook:
@@ -4801,28 +4802,33 @@ FVector FEditorViewportClient::TranslateDelta( FKey InKey, float InDelta, bool I
 				{
 					switch( GetViewportType() )
 					{
-					case LVT_OrthoXY:
+					case LVT_OrthoTop:
 						vec = FVector (-Y * UnitsPerPixel, -X * UnitsPerPixel, 0.0f);
 						break;
-					case LVT_OrthoXZ:
+					case LVT_OrthoLeft:
 						vec = FVector(X * UnitsPerPixel, 0.f, Y * UnitsPerPixel);
 						break;
-					case LVT_OrthoYZ:
+					case LVT_OrthoBack:
 						vec = FVector(0.f, X * UnitsPerPixel, Y * UnitsPerPixel);
 						break;
- 					case LVT_OrthoNegativeXY:
+ 					case LVT_OrthoBottom:
  						vec = FVector(-Y * UnitsPerPixel, X * UnitsPerPixel, 0.0f);
  						break;
-					case LVT_OrthoNegativeXZ:
-						vec = FVector(X * UnitsPerPixel, 0.f, Y * UnitsPerPixel);
+					case LVT_OrthoRight:
+						vec = FVector(-X * UnitsPerPixel, 0.f, Y * UnitsPerPixel);
 						break;
-					case LVT_OrthoNegativeYZ:
+					case LVT_OrthoFront:
 						vec = FVector(0.f, -X * UnitsPerPixel, Y * UnitsPerPixel);
 						break;
 					case LVT_OrthoFreelook:
 					case LVT_Perspective:
 						break;
 					}
+
+					if (IsOrtho() && AxisDisplayInfo::GetAxisDisplayCoordinateSystem() == EAxisList::LeftUpForward)
+					{
+						vec.Y = -vec.Y;
+					}
 				}
 			}
 		}
