Fixed actors not being consistently discovered when running with bSearchAllAssets=false

#rnx

diff --git a/Engine/Source/Editor/UnrealEd/Private/Commandlets/GatherTextFromAssetsCommandlet.cpp b/Engine/Source/Editor/UnrealEd/Private/Commandlets/GatherTextFromAssetsCommandlet.cpp
--- a/Engine/Source/Editor/UnrealEd/Private/Commandlets/GatherTextFromAssetsCommandlet.cpp
+++ b/Engine/Source/Editor/UnrealEd/Private/Commandlets/GatherTextFromAssetsCommandlet.cpp
@@ -855,6 +855,7 @@ bool UGatherTextFromAssetsCommandlet::DiscoverExternalActors(TArray<FAssetData>&
 	// Append the actors that are directly known by each world (by looking for their external actors under the expected path)
 	if (PartitionedWorldPackageNames.Num() > 0)
 	{
+		// Note: This doesn't add AActor to ClassPaths as that doesn't work correctly doing a partial asset scan (see bSearchAllAssets)
 		FARFilter Filter;
 		for (const FName PartitionedWorldPackageName : PartitionedWorldPackageNames)
 		{
@@ -862,33 +863,38 @@ bool UGatherTextFromAssetsCommandlet::DiscoverExternalActors(TArray<FAssetData>&
 			Filter.PackagePaths.Add(*ULevel::GetExternalActorsPath(*PackageNameStr));
 		}
 		Filter.bRecursivePaths = true;
-		Filter.ClassPaths.Add(AActor::StaticClass()->GetClassPathName());
-		Filter.bRecursiveClasses = true;
 		Filter.bIncludeOnlyOnDiskAssets = true;
 		Filter.WithoutPackageFlags = PKG_Cooked;
 		
-		AssetRegistry.GetAssets(Filter, InOutAssetDataArray);
+		TArray<FAssetData> PotentialExternalActors;
+		AssetRegistry.GetAssets(Filter, PotentialExternalActors);
+		for (FAssetData& PotentialExternalActor : PotentialExternalActors)
+		{
+			if (FWorldPartitionActorDescUtils::IsValidActorDescriptorFromAssetData(PotentialExternalActor))
+			{
+				InOutAssetDataArray.Add(MoveTemp(PotentialExternalActor));
+			}
+		}
 	}
 
 	// Append additional actors for each world
 	// These are actors added via ExternalDataLayers or ContentBundles, and will be referenced by the GameFeatureData assets that add them (via a GameFeatureAction)
 	if (GameFeatureDataPackageNames.Num() > 0)
 	{
-		TArray<FAssetData> GameFeatureDataActorDependencies;
+		TArray<FAssetData> GameFeatureDataDependencies;
 		{
+			// Note: This doesn't add AActor to ClassPaths as that doesn't work correctly doing a partial asset scan (see bSearchAllAssets)
 			FARFilter Filter;
 			for (const FName GameFeatureDataPackageName : GameFeatureDataPackageNames)
 			{
 				AssetRegistry.GetDependencies(GameFeatureDataPackageName, Filter.PackageNames);
 			}
-			Filter.ClassPaths.Add(AActor::StaticClass()->GetClassPathName());
-			Filter.bRecursiveClasses = true;
 			Filter.bIncludeOnlyOnDiskAssets = true;
 			Filter.WithoutPackageFlags = PKG_Cooked;
 
 			if (Filter.PackageNames.Num() > 0)
 			{
-				AssetRegistry.GetAssets(Filter, GameFeatureDataActorDependencies);
+				AssetRegistry.GetAssets(Filter, GameFeatureDataDependencies);
 			}
 		}
 
@@ -898,19 +904,19 @@ bool UGatherTextFromAssetsCommandlet::DiscoverExternalActors(TArray<FAssetData>&
 		// If neither of the above is true then we include all actors related to the GFDs in this gather, as we assume this is a plugin hosting external actors for worlds owned by another localization target
 		const bool bHasWorldFilter = WorldPackageFilter.IsSet() || PartitionedWorldPackageNames.Num() > 0;
 		const TSet<FName>& GameFeatureDataActorsWorldPackageFilter = WorldPackageFilter.Get(PartitionedWorldPackageNames);
-		for (FAssetData& GameFeatureDataActorDependency : GameFeatureDataActorDependencies)
+		for (FAssetData& GameFeatureDataDependency : GameFeatureDataDependencies)
 		{
-			const FName OptionalOuterPathName = GameFeatureDataActorDependency.GetOptionalOuterPathName();
-			if (!OptionalOuterPathName.IsNone())
+			const FName OptionalOuterPathName = GameFeatureDataDependency.GetOptionalOuterPathName();
+			if (!OptionalOuterPathName.IsNone() && FWorldPartitionActorDescUtils::IsValidActorDescriptorFromAssetData(GameFeatureDataDependency))
 			{
 				const FName OptionalOuterPackageName = FName(FPackageName::ObjectPathToPackageName(FNameBuilder(OptionalOuterPathName).ToView()));
 				if (!bHasWorldFilter || GameFeatureDataActorsWorldPackageFilter.Contains(OptionalOuterPackageName))
 				{
-					InOutAssetDataArray.Add(MoveTemp(GameFeatureDataActorDependency));
+					InOutAssetDataArray.Add(MoveTemp(GameFeatureDataDependency));
 				}
 				else
 				{
-					UE_LOG(LogGatherTextFromAssetsCommandlet, VeryVerbose, TEXT("Skipping external actor package (%s) as its associated world package (%s) is not relevant to this gather."), *GameFeatureDataActorDependency.GetSoftObjectPath().ToString(), *OptionalOuterPackageName.ToString());
+					UE_LOG(LogGatherTextFromAssetsCommandlet, VeryVerbose, TEXT("Skipping external actor package (%s) as its associated world package (%s) is not relevant to this gather."), *GameFeatureDataDependency.GetSoftObjectPath().ToString(), *OptionalOuterPackageName.ToString());
 				}
 			}
 		}
