Allow controlling cook.CVarControl via cook server commandline arguments (-CookCVarControl), and forcing an override cook device profile (-DeviceProfile).
RenderingValidation's Quest 3 automated testing script is used as a test case.

diff --git a/Engine/Source/Editor/UnrealEd/Private/CookOnTheFlyServer.cpp b/Engine/Source/Editor/UnrealEd/Private/CookOnTheFlyServer.cpp
--- a/Engine/Source/Editor/UnrealEd/Private/CookOnTheFlyServer.cpp
+++ b/Engine/Source/Editor/UnrealEd/Private/CookOnTheFlyServer.cpp
@@ -6099,13 +6099,15 @@ static EMPCookGeneratorSplit ParseMPCookGeneratorSplitFromString(const FString&
 
 namespace UE::Cook::CVarControl
 {
-	void UpdateCVars(FBeginCookContext& BeginContext)
+	void UpdateCVars(FBeginCookContext& BeginContext, UE::Cook::FBeginCookConfigSettings& Settings)
 	{
 		static FName CookTimeCVarTag("CookTimeCVars");
 
-		if (GCookTimeCVarControl == 2 || GCookTimeCVarControl == 3)
+		const int32 CookTimeCVarControl = (Settings.OverrideCookCVarControl >= 0 && Settings.OverrideCookCVarControl <= 3) ? Settings.OverrideCookCVarControl : GCookTimeCVarControl;
+
+		if (CookTimeCVarControl == 2 || CookTimeCVarControl == 3 || !Settings.OverrideDeviceProfileName.IsNone())
 		{
-			checkf(BeginContext.TargetPlatforms.Num() == 1, TEXT("When using cook.CVarControl in mode 2 or 3, only a single TargetPlatform may be cooked at once."));
+			checkf(BeginContext.TargetPlatforms.Num() == 1, TEXT("When using cook.CVarControl in mode 2 or 3, or specifying a device profile override, only a single TargetPlatform may be cooked at once."));
 		}
 
 
@@ -6119,9 +6121,9 @@ namespace UE::Cook::CVarControl
 		for (const ITargetPlatform* Platform : BeginContext.TargetPlatforms)
 		{
 			FName PlatformName = *Platform->IniPlatformName();
-			FName DPName = *Platform->CookingDeviceProfileName();
+			FName DPName = Settings.OverrideDeviceProfileName.IsNone() ? *Platform->CookingDeviceProfileName() : Settings.OverrideDeviceProfileName;
 
-			if (GCookTimeCVarControl != 0)
+			if (CookTimeCVarControl != 0)
 			{
 				// register that when we cook for this platform, we will want to use the given DP when looking up CVar values
 				// if it matches the platform, we don't set it, because that's the default, so don't actually do anything special
@@ -6130,17 +6132,17 @@ namespace UE::Cook::CVarControl
 					ConsoleVariablePlatformMapping::RegisterPlatformToDeviceProfileMapping(PlatformName, DPName);
 				}
 			}
-			if (GCookTimeCVarControl == 2)
+			if (CookTimeCVarControl == 2)
 			{
 				IConsoleManager::Get().StompPlatformCVars(PlatformName, DPName.ToString(), CookTimeCVarTag, ECVF_SetByCode, ECVF_Preview, ECVF_Cheat);
 			}
-			else if (GCookTimeCVarControl == 3)
+			else if (CookTimeCVarControl == 3)
 			{
 				IConsoleManager::Get().StompPlatformCVars(PlatformName, DPName.ToString(), CookTimeCVarTag, ECVF_SetByCode, ECVF_Default, ECVF_Default);
 			}
 
 			// for now, only do this if we are performing any cvar control modes
-			if (GCookTimeCVarControl != 0)
+			if (CookTimeCVarControl != 0)
 			{
 				if (Platform != TPM->GetRunningTargetPlatform())
 				{
@@ -10365,6 +10367,9 @@ void FBeginCookConfigSettings::LoadLocal(FBeginCookContext& BeginContext)
 	FParse::Value(FCommandLine::Get(), TEXT("-CookShowInstigator="), CookShowInstigator);
 	LoadNeverCookLocal(BeginContext);
 
+	FParse::Value(FCommandLine::Get(), TEXT("-CookCVarControl="), OverrideCookCVarControl);
+	FParse::Value(FCommandLine::Get(), TEXT("-DeviceProfile="), OverrideDeviceProfileName);
+
 	for (const ITargetPlatform* TargetPlatform : BeginContext.TargetPlatforms)
 	{
 		FConfigFile PlatformEngineIni;
@@ -10388,7 +10393,7 @@ void UCookOnTheFlyServer::SetBeginCookConfigSettings(FBeginCookContext& BeginCon
 	PackageDatas->SetBeginCookConfigSettings(Settings.CookShowInstigator);
 	UE::Cook::FGenerationHelper::SetBeginCookConfigSettings();
 	SetNeverCookPackageConfigSettings(BeginContext, Settings);
-	UE::Cook::CVarControl::UpdateCVars(BeginContext);
+	UE::Cook::CVarControl::UpdateCVars(BeginContext, Settings);
 }
 
 namespace UE::Cook
diff --git a/Engine/Source/Editor/UnrealEd/Private/Cooker/CookTypes.h b/Engine/Source/Editor/UnrealEd/Private/Cooker/CookTypes.h
--- a/Engine/Source/Editor/UnrealEd/Private/Cooker/CookTypes.h
+++ b/Engine/Source/Editor/UnrealEd/Private/Cooker/CookTypes.h
@@ -506,6 +506,8 @@ namespace UE::Cook
 		bool bHybridIterativeEnabled = true;
 		bool bHybridIterativeAllowAllClasses = false;
 		TArray<FName> NeverCookPackageList;
+		int32 OverrideCookCVarControl = -1;
+		FName OverrideDeviceProfileName;
 		TFastPointerMap<const ITargetPlatform*, TSet<FName>> PlatformSpecificNeverCookPackages;
 
 		friend FCbWriter& ::operator<<(FCbWriter& Writer, const UE::Cook::FBeginCookConfigSettings& Value);
