Update -DeviceProfile and -CookCVarControl CookOnTheFlyServer options to support multiprocess cooking.

diff --git a/Engine/Source/Editor/UnrealEd/Classes/CookOnTheSide/CookOnTheFlyServer.h b/Engine/Source/Editor/UnrealEd/Classes/CookOnTheSide/CookOnTheFlyServer.h
--- a/Engine/Source/Editor/UnrealEd/Classes/CookOnTheSide/CookOnTheFlyServer.h
+++ b/Engine/Source/Editor/UnrealEd/Classes/CookOnTheSide/CookOnTheFlyServer.h
@@ -1657,6 +1657,12 @@ private:
 
 	TRefCountPtr<UE::Cook::IMPCollector> ConfigCollector;
 
+	/** Override for the platform-dependent base device profile used for cooking, parsed from the commandline (-DeviceProfile=) */
+	FName OverrideDeviceProfileName;
+
+	/** Override for Cook.CVarControl, parsed from the commandline (-CookCVarControl=) */
+	int32 OverrideCookCVarControl = -1;
+
 	/**
 	 * Heap of Pollables to tick, ordered by NextTimeSeconds.
 	 * Only accessible by PumpPollables when bPollablesInTick==true, can be accessed outside of PollablesLock.
diff --git a/Engine/Source/Editor/UnrealEd/Private/CookOnTheFlyServer.cpp b/Engine/Source/Editor/UnrealEd/Private/CookOnTheFlyServer.cpp
--- a/Engine/Source/Editor/UnrealEd/Private/CookOnTheFlyServer.cpp
+++ b/Engine/Source/Editor/UnrealEd/Private/CookOnTheFlyServer.cpp
@@ -6099,18 +6099,16 @@ static EMPCookGeneratorSplit ParseMPCookGeneratorSplitFromString(const FString&
 
 namespace UE::Cook::CVarControl
 {
-	void UpdateCVars(FBeginCookContext& BeginContext, UE::Cook::FBeginCookConfigSettings& Settings)
+	void UpdateCVars(FBeginCookContext& BeginContext, FName OverrideDeviceProfileName, int32 OverrideCookCVarControl)
 	{
 		static FName CookTimeCVarTag("CookTimeCVars");
 
-		const int32 CookTimeCVarControl = (Settings.OverrideCookCVarControl >= 0 && Settings.OverrideCookCVarControl <= 3) ? Settings.OverrideCookCVarControl : GCookTimeCVarControl;
-
-		if (CookTimeCVarControl == 2 || CookTimeCVarControl == 3 || !Settings.OverrideDeviceProfileName.IsNone())
+		const int32 CookTimeCVarControl = (OverrideCookCVarControl >= 0 && OverrideCookCVarControl <= 3) ? OverrideCookCVarControl : GCookTimeCVarControl;
+		if (CookTimeCVarControl == 2 || CookTimeCVarControl == 3 || !OverrideDeviceProfileName.IsNone())
 		{
-			checkf(BeginContext.TargetPlatforms.Num() == 1, TEXT("When using cook.CVarControl in mode 2 or 3, or specifying a device profile override, only a single TargetPlatform may be cooked at once."));
+			checkf(BeginContext.TargetPlatforms.Num() == 1, TEXT("When using Cook.CVarControl in mode 2 or 3, or specifying a device profile override, only a single TargetPlatform may be cooked at once."));
 		}
 
-
 		// in case we had cooked before on this run, reset everything
 		IConsoleManager::Get().UnsetAllConsoleVariablesWithTag(CookTimeCVarTag);
 
@@ -6121,7 +6119,7 @@ namespace UE::Cook::CVarControl
 		for (const ITargetPlatform* Platform : BeginContext.TargetPlatforms)
 		{
 			FName PlatformName = *Platform->IniPlatformName();
-			FName DPName = Settings.OverrideDeviceProfileName.IsNone() ? *Platform->CookingDeviceProfileName() : Settings.OverrideDeviceProfileName;
+			FName DPName = OverrideDeviceProfileName.IsNone() ? *Platform->CookingDeviceProfileName() : OverrideDeviceProfileName;
 
 			if (CookTimeCVarControl != 0)
 			{
@@ -6346,6 +6344,9 @@ void UCookOnTheFlyServer::SetInitializeConfigSettings(UE::Cook::FInitializeConfi
 
 	bDeterminismDebug = FParse::Param(FCommandLine::Get(), TEXT("cookdeterminism")) ||
 		FParse::Param(FCommandLine::Get(), TEXT("diffonlybase"));
+
+	FParse::Value(FCommandLine::Get(), TEXT("-CookCVarControl="), OverrideCookCVarControl);
+	FParse::Value(FCommandLine::Get(), TEXT("-DeviceProfile="), OverrideDeviceProfileName);
 }
 
 void UCookOnTheFlyServer::ParseCookFilters()
@@ -10367,9 +10368,6 @@ void FBeginCookConfigSettings::LoadLocal(FBeginCookContext& BeginContext)
 	FParse::Value(FCommandLine::Get(), TEXT("-CookShowInstigator="), CookShowInstigator);
 	LoadNeverCookLocal(BeginContext);
 
-	FParse::Value(FCommandLine::Get(), TEXT("-CookCVarControl="), OverrideCookCVarControl);
-	FParse::Value(FCommandLine::Get(), TEXT("-DeviceProfile="), OverrideDeviceProfileName);
-
 	for (const ITargetPlatform* TargetPlatform : BeginContext.TargetPlatforms)
 	{
 		FConfigFile PlatformEngineIni;
@@ -10393,7 +10391,7 @@ void UCookOnTheFlyServer::SetBeginCookConfigSettings(FBeginCookContext& BeginCon
 	PackageDatas->SetBeginCookConfigSettings(Settings.CookShowInstigator);
 	UE::Cook::FGenerationHelper::SetBeginCookConfigSettings();
 	SetNeverCookPackageConfigSettings(BeginContext, Settings);
-	UE::Cook::CVarControl::UpdateCVars(BeginContext, Settings);
+	UE::Cook::CVarControl::UpdateCVars(BeginContext, OverrideDeviceProfileName, OverrideCookCVarControl);
 }
 
 namespace UE::Cook
diff --git a/Engine/Source/Editor/UnrealEd/Private/Cooker/CookTypes.h b/Engine/Source/Editor/UnrealEd/Private/Cooker/CookTypes.h
--- a/Engine/Source/Editor/UnrealEd/Private/Cooker/CookTypes.h
+++ b/Engine/Source/Editor/UnrealEd/Private/Cooker/CookTypes.h
@@ -506,8 +506,6 @@ namespace UE::Cook
 		bool bHybridIterativeEnabled = true;
 		bool bHybridIterativeAllowAllClasses = false;
 		TArray<FName> NeverCookPackageList;
-		int32 OverrideCookCVarControl = -1;
-		FName OverrideDeviceProfileName;
 		TFastPointerMap<const ITargetPlatform*, TSet<FName>> PlatformSpecificNeverCookPackages;
 
 		friend FCbWriter& ::operator<<(FCbWriter& Writer, const UE::Cook::FBeginCookConfigSettings& Value);
