Changed UE_LOG warnings/errors instances during the localization GATHER process to UE_LOGFMT to send proper structured logging to Horde for the new Horde localization workflow

diff --git a/Engine/Source/Developer/Localization/Private/LocTextHelper.cpp b/Engine/Source/Developer/Localization/Private/LocTextHelper.cpp
--- a/Engine/Source/Developer/Localization/Private/LocTextHelper.cpp
+++ b/Engine/Source/Developer/Localization/Private/LocTextHelper.cpp
@@ -1,18 +1,19 @@
 // Copyright Epic Games, Inc. All Rights Reserved.
 
 #include "LocTextHelper.h"
-#include "PlatformInfo.h"
+
+#include "Dom/JsonObject.h"
+#include "Internationalization/BreakIterator.h"
+#include "Internationalization/IBreakIterator.h"
+#include "Logging/StructuredLog.h"
+#include "Misc/DataDrivenPlatformInfoRegistry.h"
 #include "Misc/EnumClassFlags.h"
-#include "Misc/Paths.h"
 #include "Misc/FileHelper.h"
-#include "Misc/DataDrivenPlatformInfoRegistry.h"
+#include "Misc/Paths.h"
+#include "PlatformInfo.h"
 #include "Serialization/Csv/CsvParser.h"
-#include "Internationalization/IBreakIterator.h"
-#include "Internationalization/BreakIterator.h"
-
-#include "Dom/JsonObject.h"
-#include "Serialization/JsonInternationalizationManifestSerializer.h"
 #include "Serialization/JsonInternationalizationArchiveSerializer.h"
+#include "Serialization/JsonInternationalizationManifestSerializer.h"
 #include "Serialization/JsonInternationalizationMetadataSerializer.h"
 
 #include UE_INLINE_GENERATED_CPP_BY_NAME(LocTextHelper)
@@ -20,6 +21,10 @@
 #define LOCTEXT_NAMESPACE "LocTextHelper"
 
 DEFINE_LOG_CATEGORY_STATIC(LogLocTextHelper, Log, All);
+namespace LocTextHelper
+{
+	static constexpr int32 LocalizationLogIdentifier = 304;
+}
 
 bool FLocTextPlatformSplitUtils::ShouldSplitPlatformData(const ELocTextPlatformSplitMode& InSplitMode)
 {
@@ -636,16 +641,17 @@ void FLocTextHelper::TrimManifest()
 						const FString DependencyEntryFullSrcLoc = (!DependencyFileName.IsEmpty()) ? DependencyFileName : ConflictingContext->SourceLocation;
 
 						// There is a dependency manifest entry that has the same namespace and keys as our main manifest entry but the source text differs.
-						UE_LOG(LogLocTextHelper, Warning, TEXT("Text conflict for namespace \"%s\" and key \"%s\"%s. The conflicting sources are \"%s\"%s and \"%s\"%s (from manifest dependency \"%s\")."),
-							*SanitizeLogOutput(ManifestEntry->Namespace.GetString()),
-							*SanitizeLogOutput(Context.Key.GetString()),
-							(Context.KeyMetadataObj ? *FString::Printf(TEXT(" (meta-data \"%s\")"), *SanitizeLogOutput(FJsonInternationalizationMetaDataSerializer::MetadataToString(Context.KeyMetadataObj))) : TEXT("")),
-							*SanitizeLogOutput(ManifestEntry->Source.Text),
-							(ManifestEntry->Source.MetadataObj ? *FString::Printf(TEXT(" (meta-data \"%s\")"), *SanitizeLogOutput(FJsonInternationalizationMetaDataSerializer::MetadataToString(ManifestEntry->Source.MetadataObj))) : TEXT("")),
-							*SanitizeLogOutput(DependencyEntry->Source.Text),
-							(DependencyEntry->Source.MetadataObj ? *FString::Printf(TEXT(" (meta-data \"%s\")"), *SanitizeLogOutput(FJsonInternationalizationMetaDataSerializer::MetadataToString(DependencyEntry->Source.MetadataObj))) : TEXT("")),
-							*DependencyFileName
-							);
+						UE_LOGFMT(LogLocTextHelper, Warning, "Text conflict for namespace '{locNamespace}' and key '{locKey}'{locKeyMetaData}. Main manifest entry is '{text}'{textMetaData}. Conflict from manifest dependency is {conflictLocation}:'{conflictText}'{conflictTextMetaData}.",
+							("locNamespace", *SanitizeLogOutput(ManifestEntry->Namespace.GetString())),
+							("locKey", *SanitizeLogOutput(Context.Key.GetString())),
+							("locKeyMetaData", (Context.KeyMetadataObj ? *FString::Printf(TEXT(" (meta-data \"%s\")"), *SanitizeLogOutput(FJsonInternationalizationMetaDataSerializer::MetadataToString(Context.KeyMetadataObj))) : TEXT(""))),
+							("text", *SanitizeLogOutput(ManifestEntry->Source.Text)),
+							("textMetaData", (ManifestEntry->Source.MetadataObj ? *FString::Printf(TEXT(" (meta-data \"%s\")"), *SanitizeLogOutput(FJsonInternationalizationMetaDataSerializer::MetadataToString(ManifestEntry->Source.MetadataObj))) : TEXT(""))),
+							("conflictLocation", *DependencyFileName),
+							("conflictText", *SanitizeLogOutput(DependencyEntry->Source.Text)),
+							("conflictTextMetaData", (DependencyEntry->Source.MetadataObj ? *FString::Printf(TEXT(" (meta-data \"%s\")"), *SanitizeLogOutput(FJsonInternationalizationMetaDataSerializer::MetadataToString(DependencyEntry->Source.MetadataObj))) : TEXT(""))),
+							("id", LocTextHelper::LocalizationLogIdentifier)
+						);
 						
 						ConflictTracker.AddConflict(ManifestEntry->Namespace, Context.Key, Context.KeyMetadataObj, ManifestEntry->Source, *Context.SourceLocation);
 						ConflictTracker.AddConflict(ManifestEntry->Namespace, Context.Key, Context.KeyMetadataObj, DependencyEntry->Source, DependencyEntryFullSrcLoc);
@@ -657,14 +663,15 @@ void FLocTextHelper::TrimManifest()
 					const bool bAddSuccessful = TrimmedManifest->AddSource(ManifestEntry->Namespace, ManifestEntry->Source, Context);
 					if (!bAddSuccessful)
 					{
-						UE_LOG(LogLocTextHelper, Error, TEXT("Failed to add text for namespace \"%s\" and key \"%s\"%s with source \"%s\"%s from manifest dependency \"%s\"."),
-							*ManifestEntry->Namespace.GetString(),
-							*Context.Key.GetString(),
-							(Context.KeyMetadataObj ? *FString::Printf(TEXT(" (meta-data \"%s\")"), *FJsonInternationalizationMetaDataSerializer::MetadataToString(Context.KeyMetadataObj)) : TEXT("")),
-							*ManifestEntry->Source.Text,
-							(ManifestEntry->Source.MetadataObj ? *FString::Printf(TEXT(" (meta-data \"%s\")"), *FJsonInternationalizationMetaDataSerializer::MetadataToString(ManifestEntry->Source.MetadataObj)) : TEXT("")),
-							*DependencyFileName
-							);
+						UE_LOGFMT(LogLocTextHelper, Error, "Failed to add text for namespace '{locNamespace}' and key '{locKey}'{locKeyMetaData} with source '{text}'{textMetaData} from manifest dependency '{file}'.",
+							("file", *DependencyFileName),
+							("locNamespace", *ManifestEntry->Namespace.GetString()),
+							("locKey", *Context.Key.GetString()),
+							("locKeyMetaData", (Context.KeyMetadataObj ? *FString::Printf(TEXT(" (meta-data \"%s\")"), *FJsonInternationalizationMetaDataSerializer::MetadataToString(Context.KeyMetadataObj)) : TEXT(""))),
+							("text", *ManifestEntry->Source.Text),
+							("textMetaData", (ManifestEntry->Source.MetadataObj ? *FString::Printf(TEXT(" (meta-data \"%s\")"), *FJsonInternationalizationMetaDataSerializer::MetadataToString(ManifestEntry->Source.MetadataObj)) : TEXT(""))),
+							("id", LocTextHelper::LocalizationLogIdentifier)
+						);
 					}
 				}
 			}
@@ -971,18 +978,19 @@ bool FLocTextHelper::AddSourceText(const FLocKey& InNamespace, const FLocItem& I
 			const FManifestContext* ConflictingContext = ExistingEntry->FindContext(InContext.Key, InContext.KeyMetadataObj);
 			const FString& ExistingEntrySourceLocation = (!ExistingEntryFileName.IsEmpty()) ? ExistingEntryFileName : ConflictingContext->SourceLocation;
 
-			UE_LOG(LogLocTextHelper, Warning, TEXT("%s: Text conflict%s for namespace \"%s\" and key \"%s\"%s. The conflicting sources are \"%s\"%s and \"%s\"%s."), 
-				*InContext.SourceLocation,
-				(InDescription ? *FString::Printf(TEXT(" from %s"), **InDescription) : TEXT("")),
-				*SanitizeLogOutput(InNamespace.GetString()),
-				*SanitizeLogOutput(InContext.Key.GetString()),
-				(InContext.KeyMetadataObj ? *FString::Printf(TEXT(" (meta-data \"%s\")"), *SanitizeLogOutput(FJsonInternationalizationMetaDataSerializer::MetadataToString(InContext.KeyMetadataObj))) : TEXT("")),
-				*SanitizeLogOutput(InSource.Text),
-				(InSource.MetadataObj ? *FString::Printf(TEXT(" (meta-data \"%s\")"), *SanitizeLogOutput(FJsonInternationalizationMetaDataSerializer::MetadataToString(InSource.MetadataObj))) : TEXT("")),
-				*SanitizeLogOutput(ExistingEntry->Source.Text),
-				(ExistingEntry->Source.MetadataObj ? *FString::Printf(TEXT(" (meta-data \"%s\")"), *SanitizeLogOutput(FJsonInternationalizationMetaDataSerializer::MetadataToString(ExistingEntry->Source.MetadataObj))) : TEXT(""))
-				);
-			UE_LOG(LogLocTextHelper, Warning, TEXT("%s: See conflicting location."), *ExistingEntrySourceLocation);
+			UE_LOGFMT(LogLocTextHelper, Warning, "Text conflict{from} for namespace '{locNamespace}' and key '{locKey}'{locKeyMetaData}. First entry is {location}:'{text}'{textMetaData}. Conflicting entry is {conflictLocation}:'{conflictText}'{conflictTextMetaData}.", 
+				("from", (InDescription ? *FString::Printf(TEXT(" from %s"), **InDescription) : TEXT(""))),
+				("location", *InContext.SourceLocation),
+				("locNamespace", *SanitizeLogOutput(InNamespace.GetString())),
+				("locKey", *SanitizeLogOutput(InContext.Key.GetString())),
+				("locKeyMetaData", (InContext.KeyMetadataObj ? *FString::Printf(TEXT(" (meta-data \"%s\")"), *SanitizeLogOutput(FJsonInternationalizationMetaDataSerializer::MetadataToString(InContext.KeyMetadataObj))) : TEXT(""))),
+				("text", *SanitizeLogOutput(InSource.Text)),
+				("textMetaData", (InSource.MetadataObj ? *FString::Printf(TEXT(" (meta-data \"%s\")"), *SanitizeLogOutput(FJsonInternationalizationMetaDataSerializer::MetadataToString(InSource.MetadataObj))) : TEXT(""))),
+				("conflictLocation", *ExistingEntrySourceLocation),
+				("conflictText", *SanitizeLogOutput(ExistingEntry->Source.Text)),
+				("conflictTextMetaData", (ExistingEntry->Source.MetadataObj ? *FString::Printf(TEXT(" (meta-data \"%s\")"), *SanitizeLogOutput(FJsonInternationalizationMetaDataSerializer::MetadataToString(ExistingEntry->Source.MetadataObj))) : TEXT(""))),
+				("id", LocTextHelper::LocalizationLogIdentifier)
+			);
 
 			ConflictTracker.AddConflict(InNamespace, InContext.Key, InContext.KeyMetadataObj, InSource, InContext.SourceLocation);
 			ConflictTracker.AddConflict(InNamespace, InContext.Key, InContext.KeyMetadataObj, ExistingEntry->Source, ExistingEntrySourceLocation);
@@ -993,15 +1001,16 @@ bool FLocTextHelper::AddSourceText(const FLocKey& InNamespace, const FLocItem& I
 		bAddSuccessful = Manifest->AddSource(InNamespace, InSource, InContext);
 		if (!bAddSuccessful)
 		{
-			UE_LOG(LogLocTextHelper, Error, TEXT("%s: Failed to add text%s for namespace \"%s\" and key \"%s\"%s with source \"%s\"%s."),
-				*InContext.SourceLocation,
-				(InDescription ? *FString::Printf(TEXT(" from %s"), **InDescription) : TEXT("")),
-				*InNamespace.GetString(),
-				*InContext.Key.GetString(),
-				(InContext.KeyMetadataObj ? *FString::Printf(TEXT(" (meta-data \"%s\")"), *FJsonInternationalizationMetaDataSerializer::MetadataToString(InContext.KeyMetadataObj)) : TEXT("")),
-				*InSource.Text,
-				(InSource.MetadataObj ? *FString::Printf(TEXT(" (meta-data \"%s\")"), *FJsonInternationalizationMetaDataSerializer::MetadataToString(InSource.MetadataObj)) : TEXT(""))
-				);
+			UE_LOGFMT(LogLocTextHelper, Error, "{location}: Failed to add text{from} for namespace '{locNamespace}' and key '{locKey}'{locKeyMetaData} with source '{text}'{textMetaData}.",
+				("location", *InContext.SourceLocation),
+				("from", (InDescription ? *FString::Printf(TEXT(" from %s"), **InDescription) : TEXT(""))),
+				("locNamespace", *InNamespace.GetString()),
+				("locKey", *InContext.Key.GetString()),
+				("locKeyMetaData", (InContext.KeyMetadataObj ? *FString::Printf(TEXT(" (meta-data \"%s\")"), *FJsonInternationalizationMetaDataSerializer::MetadataToString(InContext.KeyMetadataObj)) : TEXT(""))),
+				("text", *InSource.Text),
+				("textMetaData", (InSource.MetadataObj ? *FString::Printf(TEXT(" (meta-data \"%s\")"), *FJsonInternationalizationMetaDataSerializer::MetadataToString(InSource.MetadataObj)) : TEXT(""))),
+				("id", LocTextHelper::LocalizationLogIdentifier)
+			);
 		}
 	}
 
@@ -1291,12 +1300,19 @@ FLocTextWordCounts FLocTextHelper::GetWordCountReport(const FDateTime& InTimesta
 			FText BaseReportError;
 			if (!WordCounts.FromCSV(BaseReportCSV, &BaseReportError))
 			{
-				UE_LOG(LogLocTextHelper, Warning, TEXT("Failed to parse base word count report '%s': %s"), InBaseReportFilePath, *BaseReportError.ToString());
+				UE_LOGFMT(LogLocTextHelper, Warning, "Failed to parse base word count report '{report}': {reportError}",
+					("report", InBaseReportFilePath),
+					("reportError", *BaseReportError.ToString()),
+					("id", LocTextHelper::LocalizationLogIdentifier)
+				);
 			}
 		}
 		else
 		{
-			UE_LOG(LogLocTextHelper, Warning, TEXT("Failed to load base word count report '%s'."), InBaseReportFilePath);
+			UE_LOGFMT(LogLocTextHelper, Warning, "Failed to load base word count report '{report}'.", 
+				("report", InBaseReportFilePath),
+				("id", LocTextHelper::LocalizationLogIdentifier)
+			);
 		}
 	}
 
@@ -1725,14 +1741,15 @@ bool FLocTextHelper::SaveManifestImpl(const TSharedRef<const FInternationalizati
 
 					if (!ManifestToUpdate->AddSource(ManifestEntry->Namespace, ManifestEntry->Source, Context))
 					{
-						UE_LOG(LogLocTextHelper, Error, TEXT("Failed to add text for namespace \"%s\" and key \"%s\"%s with source \"%s\"%s%s."),
-							*ManifestEntry->Namespace.GetString(),
-							*Context.Key.GetString(),
-							(Context.KeyMetadataObj ? *FString::Printf(TEXT(" (meta-data \"%s\")"), *FJsonInternationalizationMetaDataSerializer::MetadataToString(Context.KeyMetadataObj)) : TEXT("")),
-							*ManifestEntry->Source.Text,
-							(ManifestEntry->Source.MetadataObj ? *FString::Printf(TEXT(" (meta-data \"%s\")"), *FJsonInternationalizationMetaDataSerializer::MetadataToString(ManifestEntry->Source.MetadataObj)) : TEXT("")),
-							(Context.PlatformName.IsNone() ? TEXT("") : *FString::Printf(TEXT(" for platform \"%s\""), *Context.PlatformName.ToString()))
-							);
+						UE_LOGFMT(LogLocTextHelper, Error, "Failed to add text for namespace '{locNamespace}' and key '{locKey}'{locKeyMetaData} with source '{text}'{textMetaData}{forPlatform}.",
+							("locNamespace", *ManifestEntry->Namespace.GetString()),
+							("locKey", *Context.Key.GetString()),
+							("locKeyMetaData", (Context.KeyMetadataObj ? *FString::Printf(TEXT(" (meta-data \"%s\")"), *FJsonInternationalizationMetaDataSerializer::MetadataToString(Context.KeyMetadataObj)) : TEXT(""))),
+							("text", *ManifestEntry->Source.Text),
+							("textMetaData", (ManifestEntry->Source.MetadataObj ? *FString::Printf(TEXT(" (meta-data \"%s\")"), *FJsonInternationalizationMetaDataSerializer::MetadataToString(ManifestEntry->Source.MetadataObj)) : TEXT(""))),
+							("forPlatform", (Context.PlatformName.IsNone() ? TEXT("") : *FString::Printf(TEXT(" for platform \"%s\""), *Context.PlatformName.ToString()))),
+							("id", LocTextHelper::LocalizationLogIdentifier)
+						);
 					}
 				}
 			}
diff --git a/Engine/Source/Developer/Localization/Private/TextLocalizationResourceGenerator.cpp b/Engine/Source/Developer/Localization/Private/TextLocalizationResourceGenerator.cpp
--- a/Engine/Source/Developer/Localization/Private/TextLocalizationResourceGenerator.cpp
+++ b/Engine/Source/Developer/Localization/Private/TextLocalizationResourceGenerator.cpp
@@ -1,13 +1,19 @@
 // Copyright Epic Games, Inc. All Rights Reserved.
 
 #include "TextLocalizationResourceGenerator.h"
-#include "Internationalization/TextLocalizationResource.h"
-#include "Misc/Paths.h"
+
 #include "Internationalization/Culture.h"
-#include "Misc/ConfigCacheIni.h"
+#include "Internationalization/TextLocalizationResource.h"
 #include "LocTextHelper.h"
+#include "Logging/StructuredLog.h"
+#include "Misc/ConfigCacheIni.h"
+#include "Misc/Paths.h"
 
 DEFINE_LOG_CATEGORY_STATIC(LogTextLocalizationResourceGenerator, Log, All);
+namespace TextLocalizationResourceGenerator
+{
+	static constexpr int32 LocalizationLogIdentifier = 304;
+}
 
 namespace UE::TextLocalizationResourceGenerator::Private
 {
@@ -154,12 +160,20 @@ bool FTextLocalizationResourceGenerator::GenerateLocRes(const FLocTextHelper& In
 					TArray<FString> ValidationErrors;
 					if (!FmtPattern.ValidatePattern(Culture, ValidationErrors))
 					{
-						FString Message = FString::Printf(TEXT("Format pattern '%s' (%s,%s) generated the following validation errors for '%s':"), *TranslationText.Text, *InManifestEntry->Namespace.GetString(), *Context.Key.GetString(), *InCultureToGenerate);
+						FString ValidationErrorsText = TEXT("");
 						for (const FString& ValidationError : ValidationErrors)
 						{
-							Message += FString::Printf(TEXT("\n  - %s"), *ValidationError);
+							ValidationErrorsText += FString::Printf(TEXT("\n  - %s"), *ValidationError);
 						}
-						UE_LOG(LogTextLocalizationResourceGenerator, Warning, TEXT("%s"), *FLocTextHelper::SanitizeLogOutput(Message));
+						UE_LOGFMT(LogTextLocalizationResourceGenerator, Warning, "{location}: Format pattern '{text}' ({locNamespace},{locKey}) generated the following validation errors for '{cultureCode}': {error}",
+							("location", *InLocTextHelper.GetTargetPath() + FString("/") + InCultureToGenerate + FString("/") + *InLocTextHelper.GetTargetName() + FString(".archive")),
+							("cultureCode", *InCultureToGenerate),
+							("locNamespace", *InManifestEntry->Namespace.GetString()),
+							("locKey", *Context.Key.GetString()),
+							("text", *FLocTextHelper::SanitizeLogOutput(TranslationText.Text)),
+							("error", *FLocTextHelper::SanitizeLogOutput(ValidationErrorsText)),
+							("id", TextLocalizationResourceGenerator::LocalizationLogIdentifier)
+						);
 					}
 				}
 
@@ -167,37 +181,51 @@ bool FTextLocalizationResourceGenerator::GenerateLocRes(const FLocTextHelper& In
 				if (EnumHasAnyFlags(InGenerateFlags, EGenerateLocResFlags::ValidateSafeWhitespace) && TranslationText.Text.Len() > 0)
 				{
 					auto IsUnsafeWhitespace = [](const TCHAR InChar)
-					{
-						// Unsafe whitespace is any whitespace character, except new-lines
-						return FText::IsWhitespace(InChar) && !(InChar == TEXT('\r') || InChar == TEXT('\n'));
-					};
+						{
+							// Unsafe whitespace is any whitespace character, except new-lines
+							return FText::IsWhitespace(InChar) && !(InChar == TEXT('\r') || InChar == TEXT('\n'));
+						};
 
 					if (IsUnsafeWhitespace(TranslationText.Text[0]) || IsUnsafeWhitespace(TranslationText.Text[TranslationText.Text.Len() - 1]))
 					{
-						UE_LOG(LogTextLocalizationResourceGenerator, Warning, TEXT("%s"), *FLocTextHelper::SanitizeLogOutput(FString::Printf(TEXT("Translation '%s' (%s,%s) has leading or trailing whitespace for '%s'"), *TranslationText.Text, *InManifestEntry->Namespace.GetString(), *Context.Key.GetString(), *InCultureToGenerate)));
+						UE_LOGFMT(LogTextLocalizationResourceGenerator, Warning, "{location}: Translation '{text}' ({locNamespace},{locKey}) has leading or trailing whitespace for '{cultureCode}'.",
+							("location", *InLocTextHelper.GetTargetPath() + FString("/") + InCultureToGenerate + FString("/") + *InLocTextHelper.GetTargetName() + FString(".archive")),
+							("cultureCode", *InCultureToGenerate),
+							("locNamespace", *InManifestEntry->Namespace.GetString()),
+							("locKey", *Context.Key.GetString()),
+							("text", *FLocTextHelper::SanitizeLogOutput(TranslationText.Text)),
+							("id", TextLocalizationResourceGenerator::LocalizationLogIdentifier)
+						);
 					}
 				}
-				
-				if (!bIsNativeCulture && EnumHasAnyFlags(InGenerateFlags, EGenerateLocResFlags::ValidateRichTextTags) && 
+
+				if (!bIsNativeCulture && EnumHasAnyFlags(InGenerateFlags, EGenerateLocResFlags::ValidateRichTextTags) &&
 					!UE::TextLocalizationResourceGenerator::Private::ValidateRichTextTags(InManifestEntry->Source.Text, TranslationText.Text))
 				{
-					UE_LOG(LogTextLocalizationResourceGenerator, Warning, TEXT("%s"), *FLocTextHelper::SanitizeLogOutput(FString::Printf(TEXT("Broken Rich Text Tag detected in a translation. An unbalanced tag (a complete/incomplet opening rich text tag (i.e. <TagName>) with an incomplet/complete closing tag(</>)) was detected in the translation but not in the source text. Find the problematic tag in the translation and fix it in the PO file to remove this warning. \tTarget: '%s'\tCulture: '%s'\tSource: '%s'\tTranslation: '%s' (%s,%s)"), *InLocTextHelper.GetTargetPath(), *InCultureToGenerate, *InManifestEntry->Source.Text, *TranslationText.Text, *InManifestEntry->Namespace.GetString(), *Context.Key.GetString())));
+					UE_LOGFMT(LogTextLocalizationResourceGenerator, Warning, "Broken Rich Text Tag detected in a translation. An unbalanced tag (a complete/incomplet opening rich text tag (i.e. <TagName>) with an incomplet/complete closing tag(</>)) was detected a the translation but not in its source text. Find the problematic tag in the translation and fix the translation to remove this warning. Translation File:'{translationFile}' Namespace And Key:'{locNamespace},{locKey}' Translation Text To Fix:'{text}'.",
+						("cultureCode", *InCultureToGenerate),
+						("locNamespace", *InManifestEntry->Namespace.GetString()),
+						("locKey", *Context.Key.GetString()),
+						("translationFile", *InLocTextHelper.GetTargetPath() + FString("/") + InCultureToGenerate + FString("/") + *InLocTextHelper.GetTargetName() + FString(".archive")),
+						("text", *FLocTextHelper::SanitizeLogOutput(TranslationText.Text)),
+						("id", TextLocalizationResourceGenerator::LocalizationLogIdentifier)
+					);
 				}
+			}
 
-				// Find the LocRes to update
-				FTextLocalizationResource* LocResToUpdate = &OutPlatformAgnosticLocRes;
-				if (!Context.PlatformName.IsNone())
+			// Find the LocRes to update
+			FTextLocalizationResource* LocResToUpdate = &OutPlatformAgnosticLocRes;
+			if (!Context.PlatformName.IsNone())
+			{
+				if (TSharedRef<FTextLocalizationResource>* PerPlatformLocRes = OutPerPlatformLocRes.Find(Context.PlatformName))
 				{
-					if (TSharedRef<FTextLocalizationResource>* PerPlatformLocRes = OutPerPlatformLocRes.Find(Context.PlatformName))
-					{
-						LocResToUpdate = &PerPlatformLocRes->Get();
-					}
+					LocResToUpdate = &PerPlatformLocRes->Get();
 				}
-				check(LocResToUpdate);
-
-				// Add this entry to the LocRes
-				LocResToUpdate->AddEntry(InManifestEntry->Namespace.GetString(), Context.Key.GetString(), InManifestEntry->Source.Text, TranslationText.Text, InPriority, InLocResID);
 			}
+			check(LocResToUpdate);
+
+			// Add this entry to the LocRes
+			LocResToUpdate->AddEntry(InManifestEntry->Namespace.GetString(), Context.Key.GetString(), InManifestEntry->Source.Text, TranslationText.Text, InPriority, InLocResID);
 		}
 
 		return true; // continue enumeration
@@ -216,7 +244,9 @@ bool FTextLocalizationResourceGenerator::GenerateLocResAndUpdateLiveEntriesFromC
 	FString NativeCulture;
 	if (!GConfig->GetString(*SectionName, TEXT("NativeCulture"), NativeCulture, InConfigFilePath))
 	{
-		UE_LOG(LogTextLocalizationResourceGenerator, Error, TEXT("No native culture specified."));
+		UE_LOGFMT(LogTextLocalizationResourceGenerator, Error, "No native culture specified.",
+			("id", TextLocalizationResourceGenerator::LocalizationLogIdentifier)
+		);
 		return false;
 	}
 
@@ -224,7 +254,9 @@ bool FTextLocalizationResourceGenerator::GenerateLocResAndUpdateLiveEntriesFromC
 	FString SourcePath;
 	if (!GConfig->GetString(*SectionName, TEXT("SourcePath"), SourcePath, InConfigFilePath))
 	{
-		UE_LOG(LogTextLocalizationResourceGenerator, Error, TEXT("No source path specified."));
+		UE_LOGFMT(LogTextLocalizationResourceGenerator, Error, "No source path specified.",
+			("id", TextLocalizationResourceGenerator::LocalizationLogIdentifier)
+		);
 		return false;
 	}
 
@@ -232,7 +264,9 @@ bool FTextLocalizationResourceGenerator::GenerateLocResAndUpdateLiveEntriesFromC
 	FString DestinationPath;
 	if (!GConfig->GetString(*SectionName, TEXT("DestinationPath"), DestinationPath, InConfigFilePath))
 	{
-		UE_LOG(LogTextLocalizationResourceGenerator, Error, TEXT("No destination path specified."));
+		UE_LOGFMT(LogTextLocalizationResourceGenerator, Error, "No destination path specified.",
+			("id", TextLocalizationResourceGenerator::LocalizationLogIdentifier)
+		);
 		return false;
 	}
 
@@ -240,7 +274,9 @@ bool FTextLocalizationResourceGenerator::GenerateLocResAndUpdateLiveEntriesFromC
 	FString ManifestName;
 	if (!GConfig->GetString(*SectionName, TEXT("ManifestName"), ManifestName, InConfigFilePath))
 	{
-		UE_LOG(LogTextLocalizationResourceGenerator, Error, TEXT("No manifest name specified."));
+		UE_LOGFMT(LogTextLocalizationResourceGenerator, Error, "No manifest name specified.",
+			("id", TextLocalizationResourceGenerator::LocalizationLogIdentifier)
+		);
 		return false;
 	}
 
@@ -248,7 +284,9 @@ bool FTextLocalizationResourceGenerator::GenerateLocResAndUpdateLiveEntriesFromC
 	FString ArchiveName;
 	if (!GConfig->GetString(*SectionName, TEXT("ArchiveName"), ArchiveName, InConfigFilePath))
 	{
-		UE_LOG(LogTextLocalizationResourceGenerator, Error, TEXT("No archive name specified."));
+		UE_LOGFMT(LogTextLocalizationResourceGenerator, Error, "No archive name specified.",
+			("id", TextLocalizationResourceGenerator::LocalizationLogIdentifier)
+		);
 		return false;
 	}
 
@@ -256,7 +294,9 @@ bool FTextLocalizationResourceGenerator::GenerateLocResAndUpdateLiveEntriesFromC
 	FString ResourceName;
 	if (!GConfig->GetString(*SectionName, TEXT("ResourceName"), ResourceName, InConfigFilePath))
 	{
-		UE_LOG(LogTextLocalizationResourceGenerator, Error, TEXT("No resource name specified."));
+		UE_LOGFMT(LogTextLocalizationResourceGenerator, Error, "No resource name specified.",
+			("id", TextLocalizationResourceGenerator::LocalizationLogIdentifier)
+		);
 		return false;
 	}
 
@@ -291,7 +331,9 @@ bool FTextLocalizationResourceGenerator::GenerateLocResAndUpdateLiveEntriesFromC
 
 	if (CulturesToGenerate.Num() == 0)
 	{
-		UE_LOG(LogTextLocalizationResourceGenerator, Error, TEXT("No cultures to generate were specified."));
+		UE_LOGFMT(LogTextLocalizationResourceGenerator, Error, "No cultures to generate were specified.",
+			("id", TextLocalizationResourceGenerator::LocalizationLogIdentifier)
+		);
 		return false;
 	}
 
@@ -301,7 +343,10 @@ bool FTextLocalizationResourceGenerator::GenerateLocResAndUpdateLiveEntriesFromC
 		FText LoadError;
 		if (!LocTextHelper.LoadAll(ELocTextHelperLoadFlags::LoadOrCreate, &LoadError))
 		{
-			UE_LOG(LogTextLocalizationResourceGenerator, Error, TEXT("%s"), *LoadError.ToString());
+			UE_LOGFMT(LogTextLocalizationResourceGenerator, Error, "Load error: {error}",
+				("error", *LoadError.ToString()),
+				("id", TextLocalizationResourceGenerator::LocalizationLogIdentifier)
+			);
 			return false;
 		}
 	}
@@ -317,7 +362,10 @@ bool FTextLocalizationResourceGenerator::GenerateLocResAndUpdateLiveEntriesFromC
 
 		if (!GenerateLocRes(LocTextHelper, CultureName, InGenerateFlags, FTextKey(ResourceFilePath), TextLocalizationResource, Unused_PerPlatformLocRes, CultureIndex))
 		{
-			UE_LOG(LogTextLocalizationResourceGenerator, Error, TEXT("Failed to generate localization resource for culture '%s'."), *CultureName);
+			UE_LOGFMT(LogTextLocalizationResourceGenerator, Error, "Failed to generate localization resource for culture '{cultureName}'.",
+				("cultureName", *CultureName),
+				("id", TextLocalizationResourceGenerator::LocalizationLogIdentifier)
+			);
 			return false;
 		}
 	}
diff --git a/Engine/Source/Editor/UnrealEd/Private/Commandlets/GatherTextCommandlet.cpp b/Engine/Source/Editor/UnrealEd/Private/Commandlets/GatherTextCommandlet.cpp
--- a/Engine/Source/Editor/UnrealEd/Private/Commandlets/GatherTextCommandlet.cpp
+++ b/Engine/Source/Editor/UnrealEd/Private/Commandlets/GatherTextCommandlet.cpp
@@ -1,20 +1,26 @@
 // Copyright Epic Games, Inc. All Rights Reserved.
 
 #include "Commandlets/GatherTextCommandlet.h"
+
 #include "Commandlets/GatherTextFromSourceCommandlet.h"
+#include "GeneralProjectSettings.h"
+#include "HAL/FileManager.h"
+#include "Logging/StructuredLog.h"
 #include "Misc/App.h"
-#include "Misc/Paths.h"
 #include "Misc/ConfigCacheIni.h"
+#include "Misc/FileHelper.h"
+#include "Misc/Paths.h"
 #include "Misc/ScopedSlowTask.h"
-#include "HAL/FileManager.h"
+#include "SourceControlHelpers.h"
 #include "UObject/Class.h"
-#include "UObject/Package.h"
 #include "UObject/GCObjectScopeGuard.h"
-#include "SourceControlHelpers.h"
-#include "GeneralProjectSettings.h"
-#include "Misc/FileHelper.h"
+#include "UObject/Package.h"
 
 DEFINE_LOG_CATEGORY_STATIC(LogGatherTextCommandlet, Log, All);
+namespace GatherTextCommandlet
+{
+	static constexpr int32 LocalizationLogIdentifier = 304;
+}
 
 /**
  *	UGatherTextCommandlet
@@ -76,12 +82,18 @@ int32 UGatherTextCommandlet::Execute(const FString& Params, const TSharedPtr<con
 			}
 			else
 			{
-				UE_LOG(LogGatherTextCommandlet, Warning, TEXT("There are no config file paths in specified config ,list '%s'. Please check to see the is correctly populated."), **ConfigListFileParamPtr);
+				UE_LOGFMT(LogGatherTextCommandlet, Warning, "There are no config file paths in specified config ,list '{configList}'. Please check to see the is correctly populated.",
+					("configList", **ConfigListFileParamPtr),
+					("id", GatherTextCommandlet::LocalizationLogIdentifier)
+				);
 			}
 		}
 		else
 		{
-			UE_LOG(LogGatherTextCommandlet, Warning, TEXT("Specified config list file '%s' does not exist. No additional config files from -ConfigList can be added."), **ConfigListFileParamPtr);
+			UE_LOGFMT(LogGatherTextCommandlet, Warning, "Specified config list file '{configList}' does not exist. No additional config files from -ConfigList can be added.",
+				("configList", **ConfigListFileParamPtr),
+				("id", GatherTextCommandlet::LocalizationLogIdentifier)
+			);
 		}
 	}
 
@@ -100,7 +112,10 @@ int32 UGatherTextCommandlet::Execute(const FString& Params, const TSharedPtr<con
 
 	if (GatherTextConfigPaths.Num() == 0)
 	{
-		UE_LOG(LogGatherTextCommandlet, Error, TEXT("-config or -ConfigList not specified. If -ConfigList was specified, please check that the file path is valid.\n%s"), *UsageText);
+		UE_LOGFMT(LogGatherTextCommandlet, Error, "-config or -ConfigList not specified. If -ConfigList was specified, please check that the file path is valid.\n{usageText}",
+			("usageText", *UsageText),
+			("id", GatherTextCommandlet::LocalizationLogIdentifier)
+		);
 		return -1;
 	}
 
@@ -132,7 +147,10 @@ int32 UGatherTextCommandlet::Execute(const FString& Params, const TSharedPtr<con
 		FText SCCErrorStr;
 		if (!CommandletSourceControlInfo->IsReady(SCCErrorStr))
 		{
-			UE_LOG(LogGatherTextCommandlet, Error, TEXT("Revision Control error: %s"), *SCCErrorStr.ToString());
+			UE_LOGFMT(LogGatherTextCommandlet, Error, "Revision Control error: {error}",
+				("error", *SCCErrorStr.ToString()),
+				("id", GatherTextCommandlet::LocalizationLogIdentifier)
+			);
 			return -1;
 		}
 	}
@@ -146,7 +164,9 @@ int32 UGatherTextCommandlet::Execute(const FString& Params, const TSharedPtr<con
 
 		if (SlowTask.ShouldCancel() || (EmbeddedContext && EmbeddedContext->ShouldAbort()))
 		{
-			UE_LOG(LogGatherTextCommandlet, Error, TEXT("GatherText aborted!"));
+			UE_LOGFMT(LogGatherTextCommandlet, Error, "GatherText aborted!",
+				("id", GatherTextCommandlet::LocalizationLogIdentifier)
+			);
 			return -1;
 		}
 
@@ -173,10 +193,16 @@ int32 UGatherTextCommandlet::Execute(const FString& Params, const TSharedPtr<con
 		}
 		else
 		{
-			UE_LOG(LogGatherTextCommandlet, Error, TEXT("%s"), *SCCErrorStr.ToString());
+			UE_LOGFMT(LogGatherTextCommandlet, Error, "Revision control error: {error}",
+				("error", *SCCErrorStr.ToString()),
+				("id", GatherTextCommandlet::LocalizationLogIdentifier)
+			);
 			if (!CommandletSourceControlInfo->CleanUp(SCCErrorStr))
 			{
-				UE_LOG(LogGatherTextCommandlet, Error, TEXT("%s"), *SCCErrorStr.ToString());
+				UE_LOGFMT(LogGatherTextCommandlet, Error, "Revision control error post-cleanup: {error}",
+					("error", *SCCErrorStr.ToString()),
+					("id", GatherTextCommandlet::LocalizationLogIdentifier)
+				);
 			}
 			return -1;
 		}
@@ -194,7 +220,10 @@ int32 UGatherTextCommandlet::ProcessGatherConfig(const FString& GatherTextConfig
 
 	if (!GConfig->FindConfigFile(*GatherTextConfigPath))
 	{
-		UE_LOG(LogGatherTextCommandlet, Error, TEXT("Loading Config File \"%s\" failed."), *GatherTextConfigPath);
+		UE_LOGFMT(LogGatherTextCommandlet, Error, "Loading Config File '{configFile}' failed.",
+			("configFile", *GatherTextConfigPath),
+			("id", GatherTextCommandlet::LocalizationLogIdentifier)
+		);
 		return -1; 
 	}
 
@@ -284,7 +313,9 @@ int32 UGatherTextCommandlet::ProcessGatherConfig(const FString& GatherTextConfig
 
 		if (SlowTask.ShouldCancel() || (EmbeddedContext && EmbeddedContext->ShouldAbort()))
 		{
-			UE_LOG(LogGatherTextCommandlet, Error, TEXT("GatherText aborted!"));
+			UE_LOGFMT(LogGatherTextCommandlet, Error, "GatherText aborted!",
+				("id", GatherTextCommandlet::LocalizationLogIdentifier)
+			);
 			return -1;
 		}
 
@@ -298,7 +329,11 @@ int32 UGatherTextCommandlet::ProcessGatherConfig(const FString& GatherTextConfig
 		UClass* CommandletClass = FindFirstObject<UClass>(*CommandletClassName, EFindFirstObjectOptions::None, ELogVerbosity::Warning, TEXT("UGatherTextCommandlet::ProcessGatherConfig"));
 		if (!CommandletClass)
 		{
-			UE_LOG(LogGatherTextCommandlet, Error, TEXT("The commandlet name %s in section %s is invalid."), *CommandletClassName, *StepName);
+			UE_LOGFMT(LogGatherTextCommandlet, Error, "The commandlet name {commandletName} in section {section} is invalid.",
+				("commandletName", *CommandletClassName),
+				("section", *StepName),
+				("id", GatherTextCommandlet::LocalizationLogIdentifier)
+			);
 			continue;
 		}
 
@@ -325,13 +360,20 @@ int32 UGatherTextCommandlet::ProcessGatherConfig(const FString& GatherTextConfig
 
 		if( 0 != Commandlet->Main( GeneratedCmdLine ) )
 		{
-			UE_LOG(LogGatherTextCommandlet, Error, TEXT("%s-%s reported an error."), *StepName, *CommandletClassName);
+			UE_LOGFMT(LogGatherTextCommandlet, Error, "{stepName}-{commandletName} reported an error.",
+				("stepName", *StepName),
+				("commandletName", *CommandletClassName),
+				("id", GatherTextCommandlet::LocalizationLogIdentifier)
+			);
 			if( CommandletSourceControlInfo.IsValid() )
 			{
 				FText SCCErrorStr;
 				if( !CommandletSourceControlInfo->CleanUp( SCCErrorStr ) )
 				{
-					UE_LOG(LogGatherTextCommandlet, Error, TEXT("%s"), *SCCErrorStr.ToString());
+					UE_LOGFMT(LogGatherTextCommandlet, Error, "Revision Control error: {error}",
+						("error", *SCCErrorStr.ToString()),
+						("id", GatherTextCommandlet::LocalizationLogIdentifier)
+					);
 				}
 			}
 			return -1;
@@ -388,7 +430,9 @@ int32 UGatherTextCommandlet::ProcessGatherConfig(const FString& GatherTextConfig
 		}
 		else
 		{
-			UE_LOG(LogGatherTextCommandlet, Warning, TEXT("No destination path specified in the 'CommonSettings' section. Cannot check for stale per-platform data!"));
+			UE_LOGFMT(LogGatherTextCommandlet, Warning, "No destination path specified in the 'CommonSettings' section. Cannot check for stale per-platform data!",
+				("id", GatherTextCommandlet::LocalizationLogIdentifier)
+			);
 		}
 	}
 
diff --git a/Engine/Source/Editor/UnrealEd/Private/Commandlets/GatherTextFromAssetsCommandlet.cpp b/Engine/Source/Editor/UnrealEd/Private/Commandlets/GatherTextFromAssetsCommandlet.cpp
--- a/Engine/Source/Editor/UnrealEd/Private/Commandlets/GatherTextFromAssetsCommandlet.cpp
+++ b/Engine/Source/Editor/UnrealEd/Private/Commandlets/GatherTextFromAssetsCommandlet.cpp
@@ -307,7 +307,10 @@ void UGatherTextFromAssetsCommandlet::ProcessGatherableTextDataArray(const TArra
 			{
 				if (TextSourceSiteContext.KeyName.IsEmpty())
 				{
-					UE_LOG(LogGatherTextFromAssetsCommandlet, Warning, TEXT("Detected missing key on asset \"%s\"."), *TextSourceSiteContext.SiteDescription);
+					UE_LOGFMT(LogGatherTextFromAssetsCommandlet, Warning, "Detected missing key on asset '{location}'.",
+						("location", *TextSourceSiteContext.SiteDescription),
+						("id", GatherTextFromAssetsCommandlet::LocalizationLogIdentifier)
+					);
 					continue;
 				}
 
@@ -578,7 +581,10 @@ bool UGatherTextFromAssetsCommandlet::BuildCollectionFilter(FARFilter& InOutFilt
 		if (!CollectionManager.TryParseCollectionPath(Collection, &CollectionContainer, &CollectionName, &ShareType) ||
 			!CollectionContainer->GetObjectsInCollection(CollectionName, ShareType, InOutFilter.SoftObjectPaths, ECollectionRecursionFlags::SelfAndChildren))
 		{
-			UE_LOG(LogGatherTextFromAssetsCommandlet, Error, TEXT("Failed get objects in specified collection: %s"), *Collection);
+			UE_LOGFMT(LogGatherTextFromAssetsCommandlet, Error, "Failed get objects in specified collection: {collection}",
+				("collection", *Collection),
+				("id", GatherTextFromAssetsCommandlet::LocalizationLogIdentifier)
+			);
 			bHasFailedToGetACollection = true;
 		}
 	}
@@ -1026,7 +1032,9 @@ bool UGatherTextFromAssetsCommandlet::ParseCommandLineHelper(const FString& InCo
 	// Get destination path
 	if (!GetPathFromConfig(*SectionName, TEXT("DestinationPath"), DestinationPath, GatherTextConfigPath))
 	{
-		UE_LOG(LogGatherTextFromAssetsCommandlet, Error, TEXT("No destination path specified."));
+		UE_LOGFMT(LogGatherTextFromAssetsCommandlet, Error, "No destination path specified.",
+			("id", GatherTextFromAssetsCommandlet::LocalizationLogIdentifier)
+		);
 		return false;
 	}
 
@@ -1038,7 +1046,11 @@ bool UGatherTextFromAssetsCommandlet::ParseCommandLineHelper(const FString& InCo
 			FText OutError;
 			if (!GatherManifestHelper->AddDependency(ManifestDependency, &OutError))
 			{
-				UE_LOG(LogGatherTextFromAssetsCommandlet, Error, TEXT("The GatherTextFromAssets commandlet couldn't load the specified manifest dependency: '%s'. %s"), *ManifestDependency, *OutError.ToString());
+				UE_LOGFMT(LogGatherTextFromAssetsCommandlet, Error, "The GatherTextFromAssets commandlet couldn't load the specified manifest dependency: '{manifestDependency}'. {error}",
+					("manifestDependency", *ManifestDependency),
+					("error", *OutError.ToString()),
+					("id", GatherTextFromAssetsCommandlet::LocalizationLogIdentifier)
+				);
 				bHasFailedToAddManifestDependency = true;
 			}
 		}
@@ -1058,7 +1070,10 @@ bool UGatherTextFromAssetsCommandlet::ParseCommandLineHelper(const FString& InCo
 
 			if (ModuleLoadResult != EModuleLoadResult::Success)
 			{
-				UE_LOG(LogGatherTextFromAssetsCommandlet, Warning, TEXT("Failed to preload dependent module %s. Please check if the modules have been renamed or moved to another folder."), *ModuleName);
+				UE_LOGFMT(LogGatherTextFromAssetsCommandlet, Warning, "Failed to preload dependent module {module}. Please check if the modules have been renamed or moved to another folder.",
+					("module", *ModuleName),
+					("id", GatherTextFromAssetsCommandlet::LocalizationLogIdentifier)
+				);
 				bHasFailedToPreloadAnyModules = true;
 				continue;
 			}
@@ -1191,7 +1206,7 @@ void UGatherTextFromAssetsCommandlet::ProcessAndRemoveCachedPackages(TMap<FName,
 				// if this package's localization ID is a duplicate
 				if (const FName* ExistingLongPackageName = AssignedPackageLocalizationIds.Find(PackagePendingGather.PackageLocalizationId))
 				{
-					UE_LOGFMT(LogGatherTextFromAssetsCommandlet, Warning, "Package '{file}' and '{conflict}' have the same localization ID ({locID}). Please reset one of these (Asset Localization -> Reset Localization ID) to avoid conflicts.", ("id", GatherTextFromAssetsCommandlet::LocalizationLogIdentifier), ("file", *PackageNameStr), ("conflict", *ExistingLongPackageName->ToString()), ("locID", *PackagePendingGather.PackageLocalizationId));
+					UE_LOGFMT(LogGatherTextFromAssetsCommandlet, Warning, "Package '{file}' and '{conflictFile}' have the same localization ID ({locKey}). Please reset one of these (Asset Localization -> Reset Localization ID) to avoid conflicts.", ("file", *PackageNameStr), ("conflictFile", *ExistingLongPackageName->ToString()), ("locKey", *PackagePendingGather.PackageLocalizationId), ("id", GatherTextFromAssetsCommandlet::LocalizationLogIdentifier));
 					if (bFixPackageLocalizationIdConflict) 
 					{
 						bThisPackageHasLocIdConflictToFix = true;
@@ -1277,7 +1292,11 @@ void UGatherTextFromAssetsCommandlet::MergeInExternalActorsWithStaleOrMissingCac
 		}
 		else
 		{
-			UE_LOG(LogGatherTextFromAssetsCommandlet, Warning, TEXT("Failed to queue world package '%s' for %d external actor(s)."), *StaleExternalActorsPair.Key.ToString(), StaleExternalActorsPair.Value.Num());
+			UE_LOGFMT(LogGatherTextFromAssetsCommandlet, Warning, "Failed to queue world package '{package}' for {nbExternalActors} external actor(s).",
+				("package", *StaleExternalActorsPair.Key.ToString()),
+				("nbExternalActors", StaleExternalActorsPair.Value.Num()),
+				("id", GatherTextFromAssetsCommandlet::LocalizationLogIdentifier)
+			);
 		}
 	}
 }
@@ -1320,7 +1339,9 @@ bool UGatherTextFromAssetsCommandlet::LoadAndProcessUncachedPackages(TArray<FNam
 
 		if (SlowTask.ShouldCancel() || (EmbeddedContext && EmbeddedContext->ShouldAbort()))
 		{
-			UE_LOG(LogGatherTextFromAssetsCommandlet, Error, TEXT("GatherText aborted!"));
+			UE_LOGFMT(LogGatherTextFromAssetsCommandlet, Error, "GatherText aborted!",
+				("id", GatherTextFromAssetsCommandlet::LocalizationLogIdentifier)
+			);
 			return false;
 		}
 
@@ -1339,7 +1360,10 @@ bool UGatherTextFromAssetsCommandlet::LoadAndProcessUncachedPackages(TArray<FNam
 
 		if (!Package)
 		{
-			UE_LOG(LogGatherTextFromAssetsCommandlet, Warning, TEXT("Failed to load package: '%s'."), *PackageNameStr);
+			UE_LOGFMT(LogGatherTextFromAssetsCommandlet, Warning, "Failed to load package: '{package}'.",
+				("package", *PackageNameStr),
+				("id", GatherTextFromAssetsCommandlet::LocalizationLogIdentifier)
+			);
 			++NumPackagesFailedLoading;
 			continue;
 		}
@@ -1364,7 +1388,10 @@ bool UGatherTextFromAssetsCommandlet::LoadAndProcessUncachedPackages(TArray<FNam
 				UE_LOG(LogGatherTextFromAssetsCommandlet, Display, TEXT("Resaving package: '%s'..."), *PackageNameStr);
 				if (!ResavePackage(PackageNameStr, Package, PackagePendingGather.PackageFilename))
 				{
-					UE_LOG(LogGatherTextFromAssetsCommandlet, Warning, TEXT("Failed to resave package: '%s'."), *PackageNameStr);
+					UE_LOGFMT(LogGatherTextFromAssetsCommandlet, Warning, "Failed to resave package: '{package}'.",
+						("package", *PackageNameStr),
+						("id", GatherTextFromAssetsCommandlet::LocalizationLogIdentifier)
+					);
 				}
 			}
 		}
@@ -1423,7 +1450,10 @@ bool UGatherTextFromAssetsCommandlet::LoadAndProcessUncachedPackages(TArray<FNam
 				UE_LOG(LogGatherTextFromAssetsCommandlet, Display, TEXT("Resaving package: '%s'..."), *PackageNameStr);
 				if (!ResavePackage(PackageNameStr, Package, PackagePendingGather.PackageFilename))
 				{
-					UE_LOG(LogGatherTextFromAssetsCommandlet, Warning, TEXT("Failed to resave package: '%s'."), *PackageNameStr);
+					UE_LOGFMT(LogGatherTextFromAssetsCommandlet, Warning, "Failed to resave package: '{package}'.",
+						("package", *PackageNameStr),
+						("id", GatherTextFromAssetsCommandlet::LocalizationLogIdentifier)
+					);
 				}
 			}
 
@@ -1489,7 +1519,10 @@ bool UGatherTextFromAssetsCommandlet::LoadAndProcessUncachedPackages(TArray<FNam
 			UE_LOG(LogGatherTextFromAssetsCommandlet, Display, TEXT("Resaving package: '%s'..."), *PackageNameStr);
 			if (!ResavePackage(PackageNameStr, Package, PackagePendingGather.PackageFilename))
 			{
-				UE_LOG(LogGatherTextFromAssetsCommandlet, Warning, TEXT("Failed to resave package: '%s'."), *PackageNameStr);
+				UE_LOGFMT(LogGatherTextFromAssetsCommandlet, Warning, "Failed to resave package: '{package}'.",
+					("package", *PackageNameStr),
+					("id", GatherTextFromAssetsCommandlet::LocalizationLogIdentifier)
+				);
 			}
 		}
 
@@ -1535,7 +1568,10 @@ void UGatherTextFromAssetsCommandlet::ReportStaleGatherCache(TArray<FName>& InPa
 
 	if (!bStaleGatherCacheReportSaved)
 	{
-		UE_LOG(LogGatherTextFromAssetsCommandlet, Error, TEXT("Failed to save report: '%s'."), *StaleGatherCacheReportFilename);
+		UE_LOGFMT(LogGatherTextFromAssetsCommandlet, Error, "Failed to save report: '{report}'.",
+			("report", *StaleGatherCacheReportFilename),
+			("id", GatherTextFromAssetsCommandlet::LocalizationLogIdentifier)
+		);
 	}
 }
 
@@ -1751,7 +1787,9 @@ bool UGatherTextFromAssetsCommandlet::GetConfigurationScript(const TMap<FString,
 	}
 	else
 	{
-		UE_LOG(LogGatherTextFromAssetsCommandlet, Error, TEXT("No config specified."));
+		UE_LOGFMT(LogGatherTextFromAssetsCommandlet, Error, "No config specified.",
+			("id", GatherTextFromAssetsCommandlet::LocalizationLogIdentifier)
+		);
 		return false;
 	}
 
@@ -1763,7 +1801,9 @@ bool UGatherTextFromAssetsCommandlet::GetConfigurationScript(const TMap<FString,
 	}
 	else
 	{
-		UE_LOG(LogGatherTextFromAssetsCommandlet, Error, TEXT("No config section specified."));
+		UE_LOGFMT(LogGatherTextFromAssetsCommandlet, Error, "No config section specified.",
+			("id", GatherTextFromAssetsCommandlet::LocalizationLogIdentifier)
+		);
 		return false;
 	}
 
@@ -1787,13 +1827,19 @@ bool UGatherTextFromAssetsCommandlet::ConfigureFromScript(const FString& GatherT
 		if (IncludePaths.Num())
 		{
 			IncludePathFilters.Append(IncludePaths);
-			UE_LOG(LogGatherTextFromAssetsCommandlet, Warning, TEXT("IncludePaths detected in section %s. IncludePaths is deprecated, please use IncludePathFilters."), *SectionName);
+			UE_LOGFMT(LogGatherTextFromAssetsCommandlet, Warning, "IncludePaths detected in section {section}. IncludePaths is deprecated, please use IncludePathFilters.",
+				("section", *SectionName),
+				("id", GatherTextFromAssetsCommandlet::LocalizationLogIdentifier)
+			);
 		}
 	}
 
 	if (IncludePathFilters.Num() == 0)
 	{
-		UE_LOG(LogGatherTextFromAssetsCommandlet, Error, TEXT("No include path filters in section %s."), *SectionName);
+		UE_LOGFMT(LogGatherTextFromAssetsCommandlet, Error, "No include path filters in section {section}.",
+			("section", *SectionName),
+			("id", GatherTextFromAssetsCommandlet::LocalizationLogIdentifier)
+		);
 		bHasFatalError = true;
 	}
 
@@ -1815,7 +1861,10 @@ bool UGatherTextFromAssetsCommandlet::ConfigureFromScript(const FString& GatherT
 						CollectionContainer->CollectionExists(CollectionName, ShareType);
 					if (!bDoesCollectionExist)
 					{
-						UE_LOG(LogGatherTextFromAssetsCommandlet, Error, TEXT("Failed to find collection \"%s\", collection does not exist."), *Collection);
+						UE_LOGFMT(LogGatherTextFromAssetsCommandlet, Error, "Failed to find collection '{collection}', collection does not exist.",
+							("collection", *Collection),
+							("id", GatherTextFromAssetsCommandlet::LocalizationLogIdentifier)
+						);
 						bHasFatalError = true;
 					}
 				}
@@ -1836,7 +1885,10 @@ bool UGatherTextFromAssetsCommandlet::ConfigureFromScript(const FString& GatherT
 		if (ExcludePaths.Num())
 		{
 			ExcludePathFilters.Append(ExcludePaths);
-			UE_LOG(LogGatherTextFromAssetsCommandlet, Warning, TEXT("ExcludePaths detected in section %s. ExcludePaths is deprecated, please use ExcludePathFilters."), *SectionName);
+			UE_LOGFMT(LogGatherTextFromAssetsCommandlet, Warning, "ExcludePaths detected in section {section}. ExcludePaths is deprecated, please use ExcludePathFilters.",
+				("section", *SectionName),
+				("id", GatherTextFromAssetsCommandlet::LocalizationLogIdentifier)
+			);
 		}
 	}
 
@@ -1850,13 +1902,19 @@ bool UGatherTextFromAssetsCommandlet::ConfigureFromScript(const FString& GatherT
 		if (PackageExtensions.Num())
 		{
 			PackageFileNameFilters.Append(PackageExtensions);
-			UE_LOG(LogGatherTextFromAssetsCommandlet, Warning, TEXT("PackageExtensions detected in section %s. PackageExtensions is deprecated, please use PackageFileNameFilters."), *SectionName);
+			UE_LOGFMT(LogGatherTextFromAssetsCommandlet, Warning, "PackageExtensions detected in section {section}. PackageExtensions is deprecated, please use PackageFileNameFilters.",
+				("section", *SectionName),
+				("id", GatherTextFromAssetsCommandlet::LocalizationLogIdentifier)
+			);
 		}
 	}
 
 	if (PackageFileNameFilters.Num() == 0)
 	{
-		UE_LOG(LogGatherTextFromAssetsCommandlet, Error, TEXT("No package file name filters in section %s."), *SectionName);
+		UE_LOGFMT(LogGatherTextFromAssetsCommandlet, Error, "No package file name filters in section {section}.",
+			("section", *SectionName),
+			("id", GatherTextFromAssetsCommandlet::LocalizationLogIdentifier)
+		);
 		bHasFatalError = true;
 	}
 
@@ -1908,7 +1966,9 @@ bool UGatherTextFromAssetsCommandlet::ConfigureFromScript(const FString& GatherT
 		int32 MaxUsedMemoryMB = 0;
 		if (GConfig->GetInt(TEXT("GatherTextFromAssets"), TEXT("MaxMemoryAllowance"), MaxUsedMemoryMB, GEditorIni))
 		{
-			UE_LOG(LogGatherTextFromAssetsCommandlet, Warning, TEXT("The MaxMemoryAllowance config option is deprecated, please use MaxUsedMemory."), *SectionName);
+			UE_LOGFMT(LogGatherTextFromAssetsCommandlet, Warning, "The MaxMemoryAllowance config option is deprecated, please use MaxUsedMemory.",
+				("id", GatherTextFromAssetsCommandlet::LocalizationLogIdentifier)
+			);
 		}
 		else
 		{
diff --git a/Engine/Source/Editor/UnrealEd/Private/Commandlets/GatherTextFromMetadataCommandlet.cpp b/Engine/Source/Editor/UnrealEd/Private/Commandlets/GatherTextFromMetadataCommandlet.cpp
--- a/Engine/Source/Editor/UnrealEd/Private/Commandlets/GatherTextFromMetadataCommandlet.cpp
+++ b/Engine/Source/Editor/UnrealEd/Private/Commandlets/GatherTextFromMetadataCommandlet.cpp
@@ -10,24 +10,29 @@
 #include "Logging/LogCategory.h"
 #include "Logging/LogMacros.h"
 #include "Logging/LogVerbosity.h"
+#include "Logging/StructuredLog.h"
 #include "Misc/AssertionMacros.h"
 #include "Misc/PackageName.h"
 #include "Misc/Paths.h"
 #include "Modules/ModuleManager.h"
+#include "ProfilingDebugging/ScopedTimers.h"
 #include "SourceCodeNavigation.h"
 #include "Templates/Casts.h"
 #include "Templates/SharedPointer.h"
 #include "Trace/Detail/Channel.h"
 #include "UObject/Class.h"
 #include "UObject/Package.h"
+#include "UObject/UnrealType.h"
 #include "UObject/UObjectHash.h"
 #include "UObject/UObjectIterator.h"
-#include "UObject/UnrealType.h"
-#include "ProfilingDebugging/ScopedTimers.h"
 
 class UObject;
 
 DEFINE_LOG_CATEGORY_STATIC(LogGatherTextFromMetaDataCommandlet, Log, All);
+namespace GatherTextFromMetaDataCommandlet
+{
+	static constexpr int32 LocalizationLogIdentifier = 304;
+}
 
 //////////////////////////////////////////////////////////////////////////
 //GatherTextFromMetaDataCommandlet
@@ -63,7 +68,9 @@ int32 UGatherTextFromMetaDataCommandlet::Main( const FString& Params )
 	}
 	else
 	{
-		UE_LOG(LogGatherTextFromMetaDataCommandlet, Error, TEXT("No config specified."));
+		UE_LOGFMT(LogGatherTextFromMetaDataCommandlet, Error, "No config specified.",
+			("id", GatherTextFromMetaDataCommandlet::LocalizationLogIdentifier)
+		);
 		return -1;
 	}
 
@@ -77,7 +84,9 @@ int32 UGatherTextFromMetaDataCommandlet::Main( const FString& Params )
 	}
 	else
 	{
-		UE_LOG(LogGatherTextFromMetaDataCommandlet, Error, TEXT("No config section specified."));
+		UE_LOGFMT(LogGatherTextFromMetaDataCommandlet, Error, "No config section specified.",
+			("id", GatherTextFromMetaDataCommandlet::LocalizationLogIdentifier)
+		);
 		return -1;
 	}
 
@@ -101,13 +110,19 @@ int32 UGatherTextFromMetaDataCommandlet::Main( const FString& Params )
 		if (IncludePaths.Num())
 		{
 			IncludePathFilters.Append(IncludePaths);
-			UE_LOG(LogGatherTextFromMetaDataCommandlet, Warning, TEXT("IncludePaths detected in section %s. IncludePaths is deprecated, please use IncludePathFilters."), *SectionName);
+			UE_LOGFMT(LogGatherTextFromMetaDataCommandlet, Warning, "IncludePaths detected in section {section}. IncludePaths is deprecated, please use IncludePathFilters.",
+				("section", *SectionName),
+				("id", GatherTextFromMetaDataCommandlet::LocalizationLogIdentifier)
+			);
 		}
 	}
 
 	if (IncludePathFilters.Num() == 0)
 	{
-		UE_LOG(LogGatherTextFromMetaDataCommandlet, Error, TEXT("No include path filters in section %s."), *SectionName);
+		UE_LOGFMT(LogGatherTextFromMetaDataCommandlet, Error, "No include path filters in section {section}.",
+			("section", *SectionName),
+			("id", GatherTextFromMetaDataCommandlet::LocalizationLogIdentifier)
+		);
 		return -1;
 	}
 
@@ -122,7 +137,10 @@ int32 UGatherTextFromMetaDataCommandlet::Main( const FString& Params )
 		if (ExcludePaths.Num())
 		{
 			ExcludePathFilters.Append(ExcludePaths);
-			UE_LOG(LogGatherTextFromMetaDataCommandlet, Warning, TEXT("ExcludePaths detected in section %s. ExcludePaths is deprecated, please use ExcludePathFilters."), *SectionName);
+			UE_LOGFMT(LogGatherTextFromMetaDataCommandlet, Warning, "ExcludePaths detected in section {section}. ExcludePaths is deprecated, please use ExcludePathFilters.",
+				("section", *SectionName),
+				("id", GatherTextFromMetaDataCommandlet::LocalizationLogIdentifier)
+			);
 		}
 	}
 
@@ -190,7 +208,12 @@ int32 UGatherTextFromMetaDataCommandlet::Main( const FString& Params )
 					const UClass* UFieldClass = FindFirstObject<UClass>(*FieldTypeStr, EFindFirstObjectOptions::None, ELogVerbosity::Warning, TEXT("Looking for field types to include or exclude in GatherTextFromMetadata commandlet"));
 					if (!FieldClass && !UFieldClass)
 					{
-						UE_LOG(LogGatherTextFromMetaDataCommandlet, Warning, TEXT("Field Type %s was not found (from %s in section %s). Did you forget a ModulesToPreload entry?"), *FieldTypeStr, InConfigKey, *SectionName);
+						UE_LOGFMT(LogGatherTextFromMetaDataCommandlet, Warning, "Field Type {fieldType} was not found (from {config} in section {section}). Did you forget a ModulesToPreload entry?",
+							("fieldType", *FieldTypeStr),
+							("config", InConfigKey),
+							("section", *SectionName),
+							("id", GatherTextFromMetaDataCommandlet::LocalizationLogIdentifier)
+						);
 						continue;
 					}
 
@@ -247,7 +270,12 @@ int32 UGatherTextFromMetaDataCommandlet::Main( const FString& Params )
 					const UStruct* FieldOwnerType = FindFirstObject<UStruct>(*FieldOwnerTypeStr, EFindFirstObjectOptions::EnsureIfAmbiguous);
 					if (!FieldOwnerType)
 					{
-						UE_LOG(LogGatherTextFromMetaDataCommandlet, Warning, TEXT("Field Owner Type %s was not found (from %s in section %s). Did you forget a ModulesToPreload entry?"), *FieldOwnerTypeStr, InConfigKey, *SectionName);
+						UE_LOGFMT(LogGatherTextFromMetaDataCommandlet, Warning, "Field Owner Type {fieldOwner} was not found (from {config} in section {section}). Did you forget a ModulesToPreload entry?",
+							("fieldOwner", *FieldOwnerTypeStr),
+							("config", InConfigKey),
+							("section", *SectionName),
+							("id", GatherTextFromMetaDataCommandlet::LocalizationLogIdentifier)
+						);
 						continue;
 					}
 
@@ -286,7 +314,11 @@ int32 UGatherTextFromMetaDataCommandlet::Main( const FString& Params )
 		FText OutError;
 		if (!GatherManifestHelper->AddDependency(ManifestDependency, &OutError))
 		{
-			UE_LOG(LogGatherTextFromMetaDataCommandlet, Error, TEXT("The GatherTextFromMetaData commandlet couldn't load the specified manifest dependency: '%s'. %s"), *ManifestDependency, *OutError.ToString());
+			UE_LOGFMT(LogGatherTextFromMetaDataCommandlet, Error, "The GatherTextFromMetaData commandlet couldn't load the specified manifest dependency: '{manifestDependency}'. {error}",
+				("manifestDependency", *ManifestDependency),
+				("error", *OutError.ToString()),
+				("id", GatherTextFromMetaDataCommandlet::LocalizationLogIdentifier)
+			);
 			return -1;
 		}
 	}
diff --git a/Engine/Source/Editor/UnrealEd/Private/Commandlets/GatherTextFromSourceCommandlet.cpp b/Engine/Source/Editor/UnrealEd/Private/Commandlets/GatherTextFromSourceCommandlet.cpp
--- a/Engine/Source/Editor/UnrealEd/Private/Commandlets/GatherTextFromSourceCommandlet.cpp
+++ b/Engine/Source/Editor/UnrealEd/Private/Commandlets/GatherTextFromSourceCommandlet.cpp
@@ -113,7 +113,9 @@ int32 UGatherTextFromSourceCommandlet::Main( const FString& Params )
 	}
 	else
 	{
-		UE_LOG(LogGatherTextFromSourceCommandlet, Error, TEXT("No config specified."));
+		UE_LOGFMT(LogGatherTextFromSourceCommandlet, Error, "No config specified.",
+			("id", GatherTextFromSourceCommandlet::LocalizationLogIdentifier)
+		);
 		return -1;
 	}
 
@@ -127,7 +129,9 @@ int32 UGatherTextFromSourceCommandlet::Main( const FString& Params )
 	}
 	else
 	{
-		UE_LOG(LogGatherTextFromSourceCommandlet, Error, TEXT("No config section specified."));
+		UE_LOGFMT(LogGatherTextFromSourceCommandlet, Error, "No config section specified.",
+			("id", GatherTextFromSourceCommandlet::LocalizationLogIdentifier)
+		);
 		return -1;
 	}
 
@@ -150,13 +154,19 @@ int32 UGatherTextFromSourceCommandlet::Main( const FString& Params )
 		if (IncludePaths.Num())
 		{
 			SearchDirectoryPaths.Append(IncludePaths);
-			UE_LOG(LogGatherTextFromSourceCommandlet, Warning, TEXT("IncludePaths detected in section %s. IncludePaths is deprecated, please use SearchDirectoryPaths."), *SectionName);
+			UE_LOGFMT(LogGatherTextFromSourceCommandlet, Warning, "IncludePaths detected in section {section}. IncludePaths is deprecated, please use SearchDirectoryPaths.",
+				("section", *SectionName),
+				("id", GatherTextFromSourceCommandlet::LocalizationLogIdentifier)
+			);
 		}
 	}
 
 	if (SearchDirectoryPaths.Num() == 0)
 	{
-		UE_LOG(LogGatherTextFromSourceCommandlet, Warning, TEXT("No search directory paths in section %s."), *SectionName);
+		UE_LOGFMT(LogGatherTextFromSourceCommandlet, Warning, "No search directory paths in section {section}.",
+			("section", *SectionName),
+			("id", GatherTextFromSourceCommandlet::LocalizationLogIdentifier)
+		);
 		return 0;
 	}
 
@@ -171,7 +181,10 @@ int32 UGatherTextFromSourceCommandlet::Main( const FString& Params )
 		if (ExcludePaths.Num())
 		{
 			ExcludePathFilters.Append(ExcludePaths);
-			UE_LOG(LogGatherTextFromSourceCommandlet, Warning, TEXT("ExcludePaths detected in section %s. ExcludePaths is deprecated, please use ExcludePathFilters."), *SectionName);
+			UE_LOGFMT(LogGatherTextFromSourceCommandlet, Warning, "ExcludePaths detected in section {section}. ExcludePaths is deprecated, please use ExcludePathFilters.",
+				("section", *SectionName),
+				("id", GatherTextFromSourceCommandlet::LocalizationLogIdentifier)
+			);
 		}
 	}
 
@@ -186,13 +199,19 @@ int32 UGatherTextFromSourceCommandlet::Main( const FString& Params )
 		if (SourceFileSearchFilters.Num())
 		{
 			FileNameFilters.Append(SourceFileSearchFilters);
-			UE_LOG(LogGatherTextFromSourceCommandlet, Warning, TEXT("SourceFileSearchFilters detected in section %s. SourceFileSearchFilters is deprecated, please use FileNameFilters."), *SectionName);
+			UE_LOGFMT(LogGatherTextFromSourceCommandlet, Warning, "SourceFileSearchFilters detected in section {section}. SourceFileSearchFilters is deprecated, please use FileNameFilters.",
+				("section", *SectionName),
+				("id", GatherTextFromSourceCommandlet::LocalizationLogIdentifier)
+			);
 		}
 	}
 
 	if (FileNameFilters.Num() == 0)
 	{
-		UE_LOG(LogGatherTextFromSourceCommandlet, Warning, TEXT("No source filters in section %s"), *SectionName);
+		UE_LOGFMT(LogGatherTextFromSourceCommandlet, Warning, "No source filters in section {section}",
+			("section", *SectionName),
+			("id", GatherTextFromSourceCommandlet::LocalizationLogIdentifier)
+		);
 		return 0;
 	}
 
@@ -228,6 +247,7 @@ int32 UGatherTextFromSourceCommandlet::Main( const FString& Params )
 		}
 
 		UE_LOG(LogGatherTextFromSourceCommandlet, Display, TEXT("The GatherTextFromSource commandlet couldn't find any source files matching (%s) in the specified directories:\n%s"), *SourceFileSearchFiltersString, *SpecifiedDirectoriesString);
+		
 		return 0;
 	}
 
@@ -240,7 +260,11 @@ int32 UGatherTextFromSourceCommandlet::Main( const FString& Params )
 		FText OutError;
 		if (!GatherManifestHelper->AddDependency(ManifestDependency, &OutError))
 		{
-			UE_LOG(LogGatherTextFromSourceCommandlet, Error, TEXT("The GatherTextFromSource commandlet couldn't load the specified manifest dependency: '%s'. %s"), *ManifestDependency, *OutError.ToString());
+			UE_LOGFMT(LogGatherTextFromSourceCommandlet, Error, "The GatherTextFromSource commandlet couldn't load the specified manifest dependency: '{manifestDependency}'. {error}",
+				("manifestDependency", *ManifestDependency),
+				("error", *OutError.ToString()),
+				("id", GatherTextFromSourceCommandlet::LocalizationLogIdentifier)
+			);
 			return -1;
 		}
 	}
@@ -579,17 +603,27 @@ void UGatherTextFromSourceCommandlet::RunPass(EGatherSourcePasses Pass, bool Sho
 			FString SourceFileText;
 			if (!FFileHelper::LoadFileToString(SourceFileText, *SourceFile))
 			{
-				UE_LOG(LogGatherTextFromSourceCommandlet, Error, TEXT("%s: GatherTextSource failed to open file"), *ParseCtxt.Filename);
+				UE_LOGFMT(LogGatherTextFromSourceCommandlet, Error, "{failingFile}: GatherTextFromSource failed to open file",
+					("failingFile", *ParseCtxt.Filename),
+					("id", GatherTextFromSourceCommandlet::LocalizationLogIdentifier)
+				);
 			}
 			else
 			{
 				if (!ParseSourceText(SourceFileText, Parsables, ParseCtxt, PrepassResults))
 				{
-					UE_LOG(LogGatherTextFromSourceCommandlet, Warning, TEXT("%s: GatherTextSource error(s) parsing source file"), *ParseCtxt.Filename);
+					UE_LOGFMT(LogGatherTextFromSourceCommandlet, Warning, "{failingFile}: GatherTextSource error(s) parsing source file",
+						("failingFile", *ParseCtxt.Filename),
+						("id", GatherTextFromSourceCommandlet::LocalizationLogIdentifier)
+					);
 				}
 				else if (ParseCtxt.WithinNamespaceDefineLineNumber != INDEX_NONE)
 				{
-					UE_LOG(LogGatherTextFromSourceCommandlet, Warning, TEXT("%s(%d): Missing '#undef LOCTEXT_NAMESPACE' for '#define LOCTEXT_NAMESPACE'"), *ParseCtxt.Filename, ParseCtxt.WithinNamespaceDefineLineNumber);
+					UE_LOGFMT(LogGatherTextFromSourceCommandlet, Warning, "{file}({line}): Missing '#undef LOCTEXT_NAMESPACE' for '#define LOCTEXT_NAMESPACE'",
+						("file", *ParseCtxt.Filename),
+						("line", ParseCtxt.WithinNamespaceDefineLineNumber),
+						("id", GatherTextFromSourceCommandlet::LocalizationLogIdentifier)
+					);
 				}
 			}
 		}
@@ -609,7 +643,11 @@ void UGatherTextFromSourceCommandlet::RunPass(EGatherSourcePasses Pass, bool Sho
 			{
 				if (ParsedStringTablePair.Value.SourceLocation.Line == INDEX_NONE)
 				{
-					UE_LOG(LogGatherTextFromSourceCommandlet, Warning, TEXT("String table with ID '%s' had %d entries parsed for it, but the table was never registered. Skipping for gather."), *ParsedStringTablePair.Key.ToString(), ParsedStringTablePair.Value.TableEntries.Num());
+					UE_LOGFMT(LogGatherTextFromSourceCommandlet, Warning, "String table with ID '{stringTableID}' had {nbEntries} entries parsed for it, but the table was never registered. Skipping for gather.",
+						("stringTableID", *ParsedStringTablePair.Key.ToString()),
+						("nbEntries", ParsedStringTablePair.Value.TableEntries.Num()),
+						("id", GatherTextFromSourceCommandlet::LocalizationLogIdentifier)
+					);
 				}
 				else
 				{
@@ -884,14 +922,20 @@ FString UGatherTextFromSourceCommandlet::RemoveStringFromTextMacro(const FString
 		int32 OpenQuoteIdx = TextMacro.Find(TEXT("\""), ESearchCase::CaseSensitive);
 		if (0 > OpenQuoteIdx || TextMacro.Len() - 1 == OpenQuoteIdx)
 		{
-			UE_LOG(LogGatherTextFromSourceCommandlet, Warning, TEXT("Missing quotes in %s"), *FLocTextHelper::SanitizeLogOutput(IdentForLogging));
+			UE_LOGFMT(LogGatherTextFromSourceCommandlet, Warning, "Missing quotes in {text}",
+				("text", *FLocTextHelper::SanitizeLogOutput(IdentForLogging)),
+				("id", GatherTextFromSourceCommandlet::LocalizationLogIdentifier)
+			);
 		}
 		else
 		{
 			int32 CloseQuoteIdx = TextMacro.Find(TEXT("\""), ESearchCase::CaseSensitive, ESearchDir::FromStart, OpenQuoteIdx+1);
 			if (0 > CloseQuoteIdx)
 			{
-				UE_LOG(LogGatherTextFromSourceCommandlet, Warning, TEXT("Missing quotes in %s"), *FLocTextHelper::SanitizeLogOutput(IdentForLogging));
+				UE_LOGFMT(LogGatherTextFromSourceCommandlet, Warning, "Missing quotes in {text}",
+					("text", *FLocTextHelper::SanitizeLogOutput(IdentForLogging)),
+					("id", GatherTextFromSourceCommandlet::LocalizationLogIdentifier)
+				);
 			}
 			else
 			{
@@ -1363,7 +1407,11 @@ bool UGatherTextFromSourceCommandlet::ParseSourceText(const FString& Text, const
 	// Handle a string C++11 raw string literal that was never closed as this is likely a false positive that needs to be fixed in the parser
 	if (ParseCtxt.WithinStringLiteral && !ParseCtxt.RawStringLiteralClosingDelim.IsEmpty())
 	{
-		UE_LOG(LogGatherTextFromSourceCommandlet, Warning, TEXT("%s(%s): A C++11 raw string literal was not correctly terminated."), *ParseCtxt.Filename, ParseCtxt.WithinStartingLine);
+		UE_LOGFMT(LogGatherTextFromSourceCommandlet, Warning, "{file}({line}): A C++11 raw string literal was not correctly terminated.",
+			("file", *ParseCtxt.Filename),
+			("line", ParseCtxt.WithinStartingLine),
+			("id", GatherTextFromSourceCommandlet::LocalizationLogIdentifier)
+		);
 	}
 
 	return true;
@@ -1451,7 +1499,14 @@ void UGatherTextFromSourceCommandlet::PrunePrepassResults(TArray<FParsedNestedMa
 				{
 					++NestedMacroStats.DuplicateExcluded;
 
-					UE_LOG(LogGatherTextFromSourceCommandlet, Error, TEXT("Excluding duplicate %s macro in header files: %s:%d %s:%d"), *Results[i].MacroName, *Results[i].Filename, Results[i].LineStart, *Results[j].Filename, Results[j].LineStart);
+					UE_LOGFMT(LogGatherTextFromSourceCommandlet, Error, "Excluding duplicate {macroName} macros in header files: {file}({line}) and {conflictFile}({conflictLine}).",
+						("macroName", *Results[i].MacroName),
+						("file", *Results[i].Filename),
+						("line", Results[i].LineStart),
+						("conflictFile", *Results[j].Filename),
+						("conflictLine", Results[j].LineStart),
+						("id", GatherTextFromSourceCommandlet::LocalizationLogIdentifier)
+					);
 
 					Results[i].bExclude = true;
 					Results[j].bExclude = true;
@@ -1669,7 +1724,11 @@ void UGatherTextFromSourceCommandlet::FSourceFileParseContext::SetDefine(const F
 		// #define LOC_DEFINE_REGION
 		if (ExcludedRegion)
 		{
-			UE_LOG(LogGatherTextFromSourceCommandlet, Warning, TEXT("%s(%d): Found a '#define LOC_DEFINE_REGION' within another '#define LOC_DEFINE_REGION'"), *Filename, LineNumber);
+			UE_LOGFMT(LogGatherTextFromSourceCommandlet, Warning, "{file}({line}): Found a '#define LOC_DEFINE_REGION' within another '#define LOC_DEFINE_REGION'",
+				("file", *Filename),
+				("line", LineNumber),
+				("id", GatherTextFromSourceCommandlet::LocalizationLogIdentifier)
+			);
 		}
 		else
 		{
@@ -1684,7 +1743,11 @@ void UGatherTextFromSourceCommandlet::FSourceFileParseContext::SetDefine(const F
 			// #define LOCTEXT_NAMESPACE <namespace>
 			if (WithinNamespaceDefineLineNumber != INDEX_NONE)
 			{
-				UE_LOG(LogGatherTextFromSourceCommandlet, Warning, TEXT("%s(%d): Found a '#define LOCTEXT_NAMESPACE' within another '#define LOCTEXT_NAMESPACE'"), *Filename, LineNumber);
+				UE_LOGFMT(LogGatherTextFromSourceCommandlet, Warning, "{file}({line}): Found a '#define LOCTEXT_NAMESPACE' within another '#define LOCTEXT_NAMESPACE'",
+					("file", *Filename),
+					("line", LineNumber),
+					("id", GatherTextFromSourceCommandlet::LocalizationLogIdentifier)
+				);
 			}
 			else
 			{
@@ -1715,7 +1778,11 @@ void UGatherTextFromSourceCommandlet::FSourceFileParseContext::RemoveDefine(cons
 		// #undef LOC_DEFINE_REGION
 		if (!ExcludedRegion)
 		{
-			UE_LOGFMT(LogGatherTextFromSourceCommandlet, Warning, "{file}({line}): Found an '#undef LOC_DEFINE_REGION' without a corresponding '#define LOC_DEFINE_REGION'", ("id", GatherTextFromSourceCommandlet::LocalizationLogIdentifier), ("file", *Filename), ("line", LineNumber));
+			UE_LOGFMT(LogGatherTextFromSourceCommandlet, Warning, "{file}({line}): Found an '#undef LOC_DEFINE_REGION' without a corresponding '#define LOC_DEFINE_REGION'",
+				("file", *Filename), 
+				("line", LineNumber), 
+				("id", GatherTextFromSourceCommandlet::LocalizationLogIdentifier)
+			);
 		}
 		else
 		{
@@ -1730,7 +1797,11 @@ void UGatherTextFromSourceCommandlet::FSourceFileParseContext::RemoveDefine(cons
 			// #undef LOCTEXT_NAMESPACE
 			if (WithinNamespaceDefineLineNumber == INDEX_NONE)
 			{
-				UE_LOGFMT(LogGatherTextFromSourceCommandlet, Warning, "{file}({line}): Found an '#undef LOCTEXT_NAMESPACE' without a corresponding '#define LOCTEXT_NAMESPACE'", ("id", GatherTextFromSourceCommandlet::LocalizationLogIdentifier), ("file", *Filename), ("line", LineNumber));
+				UE_LOGFMT(LogGatherTextFromSourceCommandlet, Warning, "{file}({line}): Found an '#undef LOCTEXT_NAMESPACE' without a corresponding '#define LOCTEXT_NAMESPACE'",
+					("file", *Filename), 
+					("line", LineNumber), 
+					("id", GatherTextFromSourceCommandlet::LocalizationLogIdentifier)
+				);
 			}
 			else
 			{
@@ -1748,7 +1819,13 @@ bool UGatherTextFromSourceCommandlet::FSourceFileParseContext::AddStringTableImp
 	FParsedStringTable& ParsedStringTable = ParsedStringTables.FindOrAdd(InTableId);
 	if (ParsedStringTable.SourceLocation.Line != INDEX_NONE)
 	{
-		UE_LOG(LogGatherTextFromSourceCommandlet, Warning, TEXT("%s: String table with ID \"%s\" was already parsed at \"%s\". Ignoring additional definition."), *FSourceLocation(Filename, LineNumber).ToString(), *InTableId.ToString(), *ParsedStringTable.SourceLocation.ToString());
+		UE_LOGFMT(LogGatherTextFromSourceCommandlet, Warning, "{file}({line}): String table with ID '{stringTableID}' was already parsed at '{conflictLocation}'. Ignoring additional definition.",
+			("file", *Filename),
+			("line", LineNumber),
+			("stringTableID", *InTableId.ToString()),
+			("conflictLocation", *ParsedStringTable.SourceLocation.ToString()),
+			("id", GatherTextFromSourceCommandlet::LocalizationLogIdentifier)
+		);
 		return false;
 	}
 
@@ -1774,7 +1851,14 @@ bool UGatherTextFromSourceCommandlet::FSourceFileParseContext::AddStringTableEnt
 		}
 		else
 		{
-			UE_LOG(LogGatherTextFromSourceCommandlet, Warning, TEXT("%s: String table entry with ID \"%s\" and key \"%s\" was already parsed at \"%s\". Ignoring additional definition."), *FSourceLocation(Filename, LineNumber).ToString(), *InTableId.ToString(), *InKey, *ExistingEntry->SourceLocation.ToString());
+			UE_LOGFMT(LogGatherTextFromSourceCommandlet, Warning, "{file}({line}): String table entry with ID '{stringTableID}' and key '{locKey}' was already parsed at '{conflictLocation}'. Ignoring additional definition.",
+				("file", *Filename),
+				("line", LineNumber),
+				("stringTableID", *InTableId.ToString()),
+				("locKey", *InKey),
+				("conflictLocation", *ExistingEntry->SourceLocation.ToString()),
+				("id", GatherTextFromSourceCommandlet::LocalizationLogIdentifier)
+			);
 			return false;
 		}
 	}
@@ -1807,7 +1891,14 @@ bool UGatherTextFromSourceCommandlet::FSourceFileParseContext::AddStringTableEnt
 		}
 		else
 		{
-			UE_LOG(LogGatherTextFromSourceCommandlet, Warning, TEXT("%s: String table entry meta-data with ID \"%s\" and key \"%s\" was already parsed at \"%s\". Ignoring additional definition."), *FSourceLocation(Filename, LineNumber).ToString(), *InTableId.ToString(), *InKey, *ExistingMetaData->SourceLocation.ToString());
+			UE_LOGFMT(LogGatherTextFromSourceCommandlet, Warning, "{file}({line}): String table entry meta-data with ID '{stringTableID}' and key '{locKey}' was already parsed at '{conflictLocation}'. Ignoring additional definition.",
+				("file", *Filename),
+				("line", LineNumber),
+				("stringTableID", *InTableId.ToString()),
+				("locKey", *InKey),
+				("conflictLocation", *ExistingMetaData->SourceLocation.ToString()),
+				("id", GatherTextFromSourceCommandlet::LocalizationLogIdentifier)
+			);
 			return false;
 		}
 	}
@@ -1853,7 +1944,13 @@ void UGatherTextFromSourceCommandlet::FSourceFileParseContext::AddStringTableFro
 		}
 		else
 		{
-			UE_LOG(LogGatherTextFromSourceCommandlet, Warning, TEXT("%s: String table with ID \"%s\" failed to import strings from \"%s\"."), *FSourceLocation(Filename, LineNumber).ToString(), *InTableId.ToString(), *FullImportPath);
+			UE_LOGFMT(LogGatherTextFromSourceCommandlet, Warning, "{file}({line}): String table with ID '{stringTableID}' failed to import strings from '{importPath}'.",
+				("file", *Filename),
+				("line", LineNumber),
+				("stringTableID", *InTableId.ToString()),
+				("importPath", *FullImportPath),
+				("id", GatherTextFromSourceCommandlet::LocalizationLogIdentifier)
+			);
 		}
 	}
 }
@@ -2001,7 +2098,13 @@ bool UGatherTextFromSourceCommandlet::FMacroDescriptor::ParseArgumentString(cons
 				// Here, we found the end of the line, but we don't know if it's the end of the argument.
 				if (!ArgsGatherer.Gather(ArgStart, UE_PTRDIFF_TO_INT32(Cursor - ArgStart)))
 				{
-					UE_LOG(LogGatherTextFromSourceCommandlet, Warning, TEXT("%s(%d): Parsing Arguments failed in %s macro. %s"), *Context.Filename, Context.LineNumber, *GetToken(), *FLocTextHelper::SanitizeLogOutput(Context.LineText.TrimStartAndEnd()));
+					UE_LOGFMT(LogGatherTextFromSourceCommandlet, Warning, "{file}({line}): Parsing Arguments failed in {macroName} macro. {context}",
+						("file", *Context.Filename),
+						("line", Context.LineNumber),
+						("macroName", *GetToken()),
+						("context", *FLocTextHelper::SanitizeLogOutput(Context.LineText.TrimStartAndEnd())),
+						("id", GatherTextFromSourceCommandlet::LocalizationLogIdentifier)
+					);
 					return false;
 				}
 			}
@@ -2023,7 +2126,13 @@ bool UGatherTextFromSourceCommandlet::FMacroDescriptor::ParseArgumentString(cons
 				{
 					if (!ArgsGatherer.Gather(ArgStart, UE_PTRDIFF_TO_INT32(Cursor - ArgStart) - 1))
 					{
-						UE_LOG(LogGatherTextFromSourceCommandlet, Warning, TEXT("%s(%d): Parsing Arguments failed in %s macro. %s"), *Context.Filename, Context.LineNumber, *GetToken(), *FLocTextHelper::SanitizeLogOutput(Context.LineText.TrimStartAndEnd()));
+						UE_LOGFMT(LogGatherTextFromSourceCommandlet, Warning, "{file}({line}): Parsing Arguments failed in {macroName} macro. {context}",
+							("file", *Context.Filename),
+							("line", Context.LineNumber),
+							("macroName", *GetToken()),
+							("context", *FLocTextHelper::SanitizeLogOutput(Context.LineText.TrimStartAndEnd())),
+							("id", GatherTextFromSourceCommandlet::LocalizationLogIdentifier)
+						);
 						return false;
 					}
 				}
@@ -2068,7 +2177,13 @@ bool UGatherTextFromSourceCommandlet::FMacroDescriptor::ParseArgumentString(cons
 
 			if (0 > BracketStack)
 			{
-				UE_LOG(LogGatherTextFromSourceCommandlet, Warning, TEXT("%s(%d): Unexpected bracket ')' in %s macro while parsing. %s"), *Context.Filename, Context.LineNumber, *GetToken(), *FLocTextHelper::SanitizeLogOutput(Context.LineText.TrimStartAndEnd()));
+				UE_LOGFMT(LogGatherTextFromSourceCommandlet, Warning, "{file}({line}): Unexpected bracket ')' in {macroName} macro while parsing. {context}",
+					("file", *Context.Filename),
+					("line", Context.LineNumber),
+					("macroName", *GetToken()),
+					("context", *FLocTextHelper::SanitizeLogOutput(Context.LineText.TrimStartAndEnd())),
+					("id", GatherTextFromSourceCommandlet::LocalizationLogIdentifier)
+				);
 				return false;
 			}
 		}		
@@ -2081,12 +2196,24 @@ bool UGatherTextFromSourceCommandlet::FMacroDescriptor::ParseArgumentString(cons
 					// We have a single bracket open and we found a ',' this is the end of the argument. If Bracket stack is > 1 it means that we are in a function call in one of the parameters.
 					if (!ArgsGatherer.Gather(ArgStart, UE_PTRDIFF_TO_INT32(Cursor - ArgStart)))
 					{
-						UE_LOG(LogGatherTextFromSourceCommandlet, Warning, TEXT("%s(%d): Parsing Arguments failed in %s macro. %s"), *Context.Filename, Context.LineNumber, *GetToken(), *FLocTextHelper::SanitizeLogOutput(Context.LineText.TrimStartAndEnd()));
+						UE_LOGFMT(LogGatherTextFromSourceCommandlet, Warning, "{file}({line}): Parsing Arguments failed in {macroName} macro. {context}",
+							("file", *Context.Filename),
+							("line", Context.LineNumber),
+							("macroName", *GetToken()),
+							("context", *FLocTextHelper::SanitizeLogOutput(Context.LineText.TrimStartAndEnd())),
+							("id", GatherTextFromSourceCommandlet::LocalizationLogIdentifier)
+						);
 						return false;
 					}
 					if (!ArgsGatherer.EndArgument())
 					{
-						UE_LOG(LogGatherTextFromSourceCommandlet, Warning, TEXT("%s(%d): Parsing Arguments failed in %s macro. %s"), *Context.Filename, Context.LineNumber, *GetToken(), *FLocTextHelper::SanitizeLogOutput(Context.LineText.TrimStartAndEnd()));
+						UE_LOGFMT(LogGatherTextFromSourceCommandlet, Warning, "{file}({line}): Parsing Arguments failed in {macroName} macro. {context}",
+							("file", *Context.Filename),
+							("line", Context.LineNumber),
+							("macroName", *GetToken()),
+							("context", *FLocTextHelper::SanitizeLogOutput(Context.LineText.TrimStartAndEnd())),
+							("id", GatherTextFromSourceCommandlet::LocalizationLogIdentifier)
+						);
 						return false;
 					}
 
@@ -2099,7 +2226,13 @@ bool UGatherTextFromSourceCommandlet::FMacroDescriptor::ParseArgumentString(cons
 					//	    , "Param3")
 					if (!ArgsGatherer.EndArgument())
 					{
-						UE_LOG(LogGatherTextFromSourceCommandlet, Warning, TEXT("%s(%d): Parsing Arguments failed in %s macro. %s"), *Context.Filename, Context.LineNumber, *GetToken(), *FLocTextHelper::SanitizeLogOutput(Context.LineText.TrimStartAndEnd()));
+						UE_LOGFMT(LogGatherTextFromSourceCommandlet, Warning, "{file}({line}): Parsing Arguments failed in {macroName} macro. {context}",
+							("file", *Context.Filename),
+							("line", Context.LineNumber),
+							("macroName", *GetToken()),
+							("context", *FLocTextHelper::SanitizeLogOutput(Context.LineText.TrimStartAndEnd())),
+							("id", GatherTextFromSourceCommandlet::LocalizationLogIdentifier)
+						);
 						return false;
 					}
 					ArgStart = Cursor + 1;
@@ -2116,7 +2249,13 @@ bool UGatherTextFromSourceCommandlet::FMacroDescriptor::ParseArgumentString(cons
 				{
 					if (!ArgsGatherer.Gather(ArgStart, UE_PTRDIFF_TO_INT32(Cursor - ArgStart) - 1))
 					{
-						UE_LOG(LogGatherTextFromSourceCommandlet, Warning, TEXT("%s(%d): Parsing Arguments failed in %s macro. %s"), *Context.Filename, Context.LineNumber, *GetToken(), *FLocTextHelper::SanitizeLogOutput(Context.LineText.TrimStartAndEnd()));
+						UE_LOGFMT(LogGatherTextFromSourceCommandlet, Warning, "{file}({line}): Parsing Arguments failed in {macroName} macro. {context}",
+							("file", *Context.Filename),
+							("line", Context.LineNumber),
+							("macroName", *GetToken()),
+							("context", *FLocTextHelper::SanitizeLogOutput(Context.LineText.TrimStartAndEnd())),
+							("id", GatherTextFromSourceCommandlet::LocalizationLogIdentifier)
+						);
 						return false;
 					}
 				}
@@ -2132,13 +2271,25 @@ bool UGatherTextFromSourceCommandlet::FMacroDescriptor::ParseArgumentString(cons
 			{
 				if (!ArgsGatherer.Gather(ArgStart, UE_PTRDIFF_TO_INT32(Cursor - ArgStart)))
 				{
-					UE_LOG(LogGatherTextFromSourceCommandlet, Warning, TEXT("%s(%d): Parsing Arguments failed in %s macro. %s"), *Context.Filename, Context.LineNumber, *GetToken(), *FLocTextHelper::SanitizeLogOutput(Context.LineText.TrimStartAndEnd()));
+					UE_LOGFMT(LogGatherTextFromSourceCommandlet, Warning, "{file}({line}): Parsing Arguments failed in {macroName} macro. {context}",
+						("file", *Context.Filename),
+						("line", Context.LineNumber),
+						("macroName", *GetToken()),
+						("context", *FLocTextHelper::SanitizeLogOutput(Context.LineText.TrimStartAndEnd())),
+						("id", GatherTextFromSourceCommandlet::LocalizationLogIdentifier)
+					);
 					return false;
 				}
 			}
 			if (!ArgsGatherer.EndArgument())
 			{
-				UE_LOG(LogGatherTextFromSourceCommandlet, Warning, TEXT("%s(%d): Parsing Arguments failed in %s macro. %s"), *Context.Filename, Context.LineNumber, *GetToken(), *FLocTextHelper::SanitizeLogOutput(Context.LineText.TrimStartAndEnd()));
+				UE_LOGFMT(LogGatherTextFromSourceCommandlet, Warning, "{file}({line}): Parsing Arguments failed in {macroName} macro. {context}",
+					("file", *Context.Filename),
+					("line", Context.LineNumber),
+					("macroName", *GetToken()),
+					("context", *FLocTextHelper::SanitizeLogOutput(Context.LineText.TrimStartAndEnd())),
+					("id", GatherTextFromSourceCommandlet::LocalizationLogIdentifier)
+				);
 				return false;
 			}
 			break;
@@ -2204,7 +2355,13 @@ bool UGatherTextFromSourceCommandlet::FMacroDescriptor::ParseArgsFromNextLines(F
 	{
 		if(i >= Context.TextLines.Num())
 		{
-			UE_LOG(LogGatherTextFromSourceCommandlet, Warning, TEXT("%s(%d): We reached end of file while parsing specified %s macro for arguments. %s"), *Context.Filename, Context.LineNumber, *GetToken(), *FLocTextHelper::SanitizeLogOutput(Context.LineText.TrimStartAndEnd()));
+			UE_LOGFMT(LogGatherTextFromSourceCommandlet, Warning, "{file}({line}): We reached end of file while parsing specified {macroName} macro for arguments. {context}",
+				("file", *Context.Filename),
+				("line", Context.LineNumber),
+				("macroName", *GetToken()),
+				("context", *FLocTextHelper::SanitizeLogOutput(Context.LineText.TrimStartAndEnd())),
+				("id", GatherTextFromSourceCommandlet::LocalizationLogIdentifier)
+			);
 			return false;
 		}
 
@@ -2224,7 +2381,13 @@ bool UGatherTextFromSourceCommandlet::FMacroDescriptor::ParseArgsFromNextLines(F
 			// We did not find the opening bracket on the first line and on the second non empty line, we give up
 			if (OpenBracketIdx == INDEX_NONE)
 			{
-				UE_LOG(LogGatherTextFromSourceCommandlet, Warning, TEXT("%s(%d): Opening bracket '(' not found while parsing specified %s macro for arguments. %s"), *Context.Filename, Context.LineNumber, *GetToken(), *FLocTextHelper::SanitizeLogOutput(Context.LineText.TrimStartAndEnd()));
+				UE_LOGFMT(LogGatherTextFromSourceCommandlet, Warning, "{file}({line}): Opening bracket '(' not found while parsing specified {macroName} macro for arguments. {context}",
+					("file", *Context.Filename),
+					("line", Context.LineNumber),
+					("macroName", *GetToken()),
+					("context", *FLocTextHelper::SanitizeLogOutput(Context.LineText.TrimStartAndEnd())),
+					("id", GatherTextFromSourceCommandlet::LocalizationLogIdentifier)
+				);
 				return false;
 			}
 			BracketStack = 1;
@@ -2270,7 +2433,11 @@ void UGatherTextFromSourceCommandlet::FUICommandMacroDescriptor::TryParseArgs(co
 	if (Identifier.IsEmpty())
 	{
 		//The command doesn't have an identifier so we can't gather it
-		UE_LOG(LogGatherTextFromSourceCommandlet, Warning, TEXT("%s: %s macro has an empty identifier and cannot be gathered."), *SourceLocation , *GetToken());
+		UE_LOGFMT(LogGatherTextFromSourceCommandlet, Warning, "{location}: {macroName} macro has an empty identifier and cannot be gathered.",
+			("location", *SourceLocation),
+			("macroName", *GetToken()),
+			("id", GatherTextFromSourceCommandlet::LocalizationLogIdentifier)
+		);
 		return;
 	}
 
@@ -2328,7 +2495,15 @@ void UGatherTextFromSourceCommandlet::FUICommandMacroDescriptor::TryParse(const
 			// Validate that we got the rightnumber of Arguments
 			if (Arguments.Num() < GetMinNumberOfArgument())
 			{
-				UE_LOG(LogGatherTextFromSourceCommandlet, Warning, TEXT("%s(%d): Expected at least %d arguments for %s macro, but got %d. %s"), *Context.Filename, Context.LineNumber, GetMinNumberOfArgument(), *GetToken(), Arguments.Num(), *FLocTextHelper::SanitizeLogOutput(Context.LineText.TrimStartAndEnd()));
+				UE_LOGFMT(LogGatherTextFromSourceCommandlet, Warning, "{file}({line}): Expected at least {expectedNbArguments} arguments for {macroName} macro, but got {nbArguments}. {context}",
+					("file", *Context.Filename),
+					("line", Context.LineNumber),
+					("expectedNbArguments", GetMinNumberOfArgument()),
+					("macroName", *GetToken()),
+					("nbArguments", Arguments.Num()),
+					("context", *FLocTextHelper::SanitizeLogOutput(Context.LineText.TrimStartAndEnd())),
+					("id", GatherTextFromSourceCommandlet::LocalizationLogIdentifier)
+				);
 				return;
 			}
 			
@@ -2351,7 +2526,15 @@ void UGatherTextFromSourceCommandlet::FUICommandExtMacroDescriptor::TryParse(con
 			// Need at least 5 arguments
 			if (Arguments.Num() < GetMinNumberOfArgument())
 			{
-				UE_LOG(LogGatherTextFromSourceCommandlet, Warning, TEXT("%s(%d): Expected at least %d arguments for %s macro, but got %d. %s"), *Context.Filename, Context.LineNumber, GetMinNumberOfArgument(), *GetToken(), Arguments.Num(), *FLocTextHelper::SanitizeLogOutput(Context.LineText.TrimStartAndEnd()));
+				UE_LOGFMT(LogGatherTextFromSourceCommandlet, Warning, "{file}({line}): Expected at least {expectedNbArguments} arguments for {macroName} macro, but got {nbArguments}. {context}",
+					("file", *Context.Filename),
+					("line", Context.LineNumber),
+					("expectedNbArguments", GetMinNumberOfArgument()),
+					("macroName", *GetToken()),
+					("nbArguments", Arguments.Num()),
+					("context", *FLocTextHelper::SanitizeLogOutput(Context.LineText.TrimStartAndEnd())),
+					("id", GatherTextFromSourceCommandlet::LocalizationLogIdentifier)
+				);
 				return;
 			}
 			TryParseArgs(Text, Context, Arguments, 2);
@@ -2522,7 +2705,12 @@ void UGatherTextFromSourceCommandlet::FNestedMacroDescriptor::TryParse(const FSt
 
 		if (Arg.Contains(TEXT("##")))
 		{
-			UE_LOG(LogGatherTextFromSourceCommandlet, Warning, TEXT("%s(%d): Concatenation in %s macro with \"##\" not supported"), *Context.Filename, Context.LineNumber, *MacroName);
+			UE_LOGFMT(LogGatherTextFromSourceCommandlet, Warning, "{file}({line}): Concatenation in {macroName} macro with '##' not supported",
+				("file", *Context.Filename),
+				("line", Context.LineNumber),
+				("macroName", *MacroName),
+				("id", GatherTextFromSourceCommandlet::LocalizationLogIdentifier)
+			);
 			++NestedMacroStats.Concatenation;
 			return;
 		}
@@ -2533,7 +2721,12 @@ void UGatherTextFromSourceCommandlet::FNestedMacroDescriptor::TryParse(const FSt
 		
 			if (Pos != PosLast)
 			{
-				UE_LOG(LogGatherTextFromSourceCommandlet, Warning, TEXT("%s(%d): Variadic in %s macro with \"...\" must be last param"), *Context.Filename, Context.LineNumber, *MacroName);
+				UE_LOGFMT(LogGatherTextFromSourceCommandlet, Warning, "{file}({line}): Variadic in {macroName} macro with '...' must be last param.",
+					("file", *Context.Filename),
+					("line", Context.LineNumber),
+					("macroName", *MacroName),
+					("id", GatherTextFromSourceCommandlet::LocalizationLogIdentifier)
+				);
 				return;
 			}
 			bVariadic = true;
@@ -2545,13 +2738,29 @@ void UGatherTextFromSourceCommandlet::FNestedMacroDescriptor::TryParse(const FSt
 	{
 		if ((ArgArray.Num() - 1) > ArgArrayValues.Num())
 		{
-			UE_LOG(LogGatherTextFromSourceCommandlet, Error, TEXT("%s(%d): Expected minimum of %d arguments for %s variadic macro, but got %d. %s"), *Context.Filename, Context.LineNumber, (ArgArray.Num() - 1), *GetToken(), ArgArrayValues.Num(), *FLocTextHelper::SanitizeLogOutput(Context.LineText.TrimStartAndEnd()));
+			UE_LOGFMT(LogGatherTextFromSourceCommandlet, Error, "{file}({line}): Expected minimum of {minimumNbArguments} arguments for {macroName} variadic macro, but got {nbArguments}. {context}",
+				("file", *Context.Filename),
+				("line", Context.LineNumber),
+				("expectedNbArguments", (ArgArray.Num() - 1)),
+				("macroName", *GetToken()),
+				("nbArguments", ArgArrayValues.Num()),
+				("context", *FLocTextHelper::SanitizeLogOutput(Context.LineText.TrimStartAndEnd())),
+				("id", GatherTextFromSourceCommandlet::LocalizationLogIdentifier)
+			);
 			return;
 		}
 	}
 	else if (ArgArray.Num() != ArgArrayValues.Num())
 	{
-		UE_LOG(LogGatherTextFromSourceCommandlet, Error, TEXT("%s(%d): Expected %d arguments for %s macro, but got %d. %s"), *Context.Filename, Context.LineNumber, ArgArray.Num(), *GetToken(), ArgArrayValues.Num(), *FLocTextHelper::SanitizeLogOutput(Context.LineText.TrimStartAndEnd()));
+		UE_LOGFMT(LogGatherTextFromSourceCommandlet, Error, "{file}({line}): Expected {expectedNbArguments} arguments for {macroName} macro, but got {nbArguments}. {context}",
+			("file", *Context.Filename),
+			("line", Context.LineNumber),
+			("expectedNbArguments", ArgArray.Num()),
+			("macroName", *GetToken()),
+			("nbArguments", ArgArrayValues.Num()),
+			("context", *FLocTextHelper::SanitizeLogOutput(Context.LineText.TrimStartAndEnd())),
+			("id", GatherTextFromSourceCommandlet::LocalizationLogIdentifier)
+		);
 		return;
 	}
 
@@ -2619,7 +2828,12 @@ void UGatherTextFromSourceCommandlet::FNestedMacroDescriptor::TryParse(const FSt
 		int32 PosClose = FindMatching(MacroInner, TCHAR('('), TCHAR(')'), 1);
 		if (PosClose < 0)
 		{
-			UE_LOG(LogGatherTextFromSourceCommandlet, Error, TEXT("%s(%d): Missing matching closing bracket in %s macro"), *Context.Filename, Context.LineNumber, *MacroName);
+			UE_LOGFMT(LogGatherTextFromSourceCommandlet, Error, "{file}({line}): Missing matching closing bracket in {macroName} macro",
+				("file", *Context.Filename),
+				("line", Context.LineNumber),
+				("macroName", *MacroName),
+				("id", GatherTextFromSourceCommandlet::LocalizationLogIdentifier)
+			);
 			delete InnerDescriptor;
 			return;
 		}
@@ -2773,7 +2987,15 @@ void UGatherTextFromSourceCommandlet::FStringMacroDescriptor::TryParse(const FSt
 
 			if (NumArgs != Arguments.Num())
 			{
-				UE_LOG(LogGatherTextFromSourceCommandlet, Warning, TEXT("%s(%d): Expected %d arguments for %s macro, but got %d. %s"), *Context.Filename, Context.LineNumber, Arguments.Num(), *GetToken(), NumArgs, *FLocTextHelper::SanitizeLogOutput(Context.LineText.TrimStartAndEnd()));
+				UE_LOGFMT(LogGatherTextFromSourceCommandlet, Warning, "{file}({line}): Expected {expectedNbArguments} arguments for {macroName} macro, but got {nbArguments}. {context}",
+					("file", *Context.Filename),
+					("line", Context.LineNumber),
+					("expectedNbArguments", Arguments.Num()),
+					("macroName", *GetToken()),
+					("nbArguments", NumArgs),
+					("context", *FLocTextHelper::SanitizeLogOutput(Context.LineText.TrimStartAndEnd())),
+					("id", GatherTextFromSourceCommandlet::LocalizationLogIdentifier)
+				);
 			}
 			else
 			{
@@ -2825,14 +3047,22 @@ void UGatherTextFromSourceCommandlet::FStringMacroDescriptor::TryParse(const FSt
 				if (Identifier.IsEmpty())
 				{
 					// The command doesn't have an identifier so we can't gather it
-					UE_LOG(LogGatherTextFromSourceCommandlet, Warning, TEXT("%s: %s macro has an empty identifier and cannot be gathered."), *SourceLocation, *GetToken() );
+					UE_LOGFMT(LogGatherTextFromSourceCommandlet, Warning, "{location}: {macroName} macro has an empty identifier and cannot be gathered.",
+						("location", *SourceLocation),
+						("macroName", *GetToken()),
+						("id", GatherTextFromSourceCommandlet::LocalizationLogIdentifier)
+					);
 					return;
 				}
 
 				if (SourceText.IsEmpty())
 				{
 					// The command doesn't have a source text so we can't gather it
-					UE_LOG(LogGatherTextFromSourceCommandlet, Warning, TEXT("%s: %s macro has an empty source text and cannot be gathered."), *SourceLocation, *GetToken() );
+					UE_LOGFMT(LogGatherTextFromSourceCommandlet, Warning, "{location}: {macroName} macro has an empty source text and cannot be gathered.",
+						("location", *SourceLocation),
+						("macroName", *GetToken()),
+						("id", GatherTextFromSourceCommandlet::LocalizationLogIdentifier)
+					);
 					return;
 				}
 
@@ -2840,7 +3070,11 @@ void UGatherTextFromSourceCommandlet::FStringMacroDescriptor::TryParse(const FSt
 				{
 					if (!Namespace.IsSet())
 					{
-						UE_LOG(LogGatherTextFromSourceCommandlet, Warning, TEXT("%s: %s macro doesn't define a namespace and no external namespace was set. An empty namspace will be used."), *SourceLocation, *GetToken() );
+						UE_LOGFMT(LogGatherTextFromSourceCommandlet, Warning, "{location}: {macroName} macro doesn't define a namespace and no external namespace was set. An empty namspace will be used.",
+							("location", *SourceLocation),
+							("macroName", *GetToken()),
+							("id", GatherTextFromSourceCommandlet::LocalizationLogIdentifier)
+						);
 						Namespace = FString();
 					}
 
@@ -2874,7 +3108,15 @@ void UGatherTextFromSourceCommandlet::FStringTableMacroDescriptor::TryParse(cons
 		{
 			if (Arguments.Num() != GetMinNumberOfArgument())
 			{
-				UE_LOG(LogGatherTextFromSourceCommandlet, Warning, TEXT("%s(%d): Expected %d arguments for %s macro, but got %d. %s"), *Context.Filename, Context.LineNumber, GetMinNumberOfArgument(), *GetToken(), Arguments.Num(), *FLocTextHelper::SanitizeLogOutput(Context.LineText.TrimStartAndEnd()));
+				UE_LOGFMT(LogGatherTextFromSourceCommandlet, Warning, "{file}({line}): Expected {expectedNbArguments} arguments for {macroName} macro, but got {nbArguments}. {context}",
+					("file", *Context.Filename),
+					("line", Context.LineNumber),
+					("expectedNbArguments", GetMinNumberOfArgument()),
+					("macroName", *GetToken()),
+					("nbArguments", Arguments.Num()),
+					("context", *FLocTextHelper::SanitizeLogOutput(Context.LineText.TrimStartAndEnd())),
+					("id", GatherTextFromSourceCommandlet::LocalizationLogIdentifier)
+				);
 			}
 			else
 			{
@@ -2893,7 +3135,12 @@ void UGatherTextFromSourceCommandlet::FStringTableMacroDescriptor::TryParse(cons
 
 					if (TableIdName.IsNone())
 					{
-						UE_LOG(LogGatherTextFromSourceCommandlet, Warning, TEXT("%s(%d): %s macro has an empty identifier and cannot be gathered."), *Context.Filename, Context.LineNumber, *GetToken());
+						UE_LOGFMT(LogGatherTextFromSourceCommandlet, Warning, "{file}({line}): {macroName} macro has an empty identifier and cannot be gathered.",
+							("file", *Context.Filename),
+							("line", Context.LineNumber),
+							("macroName", *GetToken()),
+							("id", GatherTextFromSourceCommandlet::LocalizationLogIdentifier)
+						);
 					}
 					else
 					{
@@ -2917,7 +3164,15 @@ void UGatherTextFromSourceCommandlet::FStringTableFromFileMacroDescriptor::TryPa
 		{
 			if (Arguments.Num() != GetMinNumberOfArgument())
 			{
-				UE_LOG(LogGatherTextFromSourceCommandlet, Warning, TEXT("%s(%d): Expected %d arguments for %s macro, but got %d. %s"), *Context.Filename, Context.LineNumber, GetMinNumberOfArgument(), *GetToken(), Arguments.Num(), *FLocTextHelper::SanitizeLogOutput(Context.LineText.TrimStartAndEnd()));
+				UE_LOGFMT(LogGatherTextFromSourceCommandlet, Warning, "{file}({line}): Expected {expectedNbArguments} arguments for {macroName} macro, but got {nbArguments}. {context}",
+					("file", *Context.Filename),
+					("line", Context.LineNumber),
+					("expectedNbArguments", GetMinNumberOfArgument()),
+					("macroName", *GetToken()),
+					("nbArguments", Arguments.Num()),
+					("context", *FLocTextHelper::SanitizeLogOutput(Context.LineText.TrimStartAndEnd())),
+					("id", GatherTextFromSourceCommandlet::LocalizationLogIdentifier)
+				);
 			}
 			else
 			{
@@ -2939,7 +3194,12 @@ void UGatherTextFromSourceCommandlet::FStringTableFromFileMacroDescriptor::TryPa
 
 					if (TableIdName.IsNone())
 					{
-						UE_LOG(LogGatherTextFromSourceCommandlet, Warning, TEXT("%s(%d): %s macro has an empty identifier and cannot be gathered."), *Context.Filename, Context.LineNumber, *GetToken());
+						UE_LOGFMT(LogGatherTextFromSourceCommandlet, Warning, "{file}({line}): {macroName} macro has an empty identifier and cannot be gathered.",
+							("file", *Context.Filename),
+							("line", Context.LineNumber),
+							("macroName", *GetToken()),
+							("id", GatherTextFromSourceCommandlet::LocalizationLogIdentifier)
+						);
 					}
 					else
 					{
@@ -2963,7 +3223,15 @@ void UGatherTextFromSourceCommandlet::FStringTableEntryMacroDescriptor::TryParse
 		{
 			if (Arguments.Num() != GetMinNumberOfArgument())
 			{
-				UE_LOG(LogGatherTextFromSourceCommandlet, Warning, TEXT("%s(%d): Expected %d arguments for %s macro, but got %d. %s"), *Context.Filename, Context.LineNumber, GetMinNumberOfArgument(), *GetToken(), Arguments.Num(), *FLocTextHelper::SanitizeLogOutput(Context.LineText.TrimStartAndEnd()));
+				UE_LOGFMT(LogGatherTextFromSourceCommandlet, Warning, "{file}({line}): Expected {expectedNbArguments} arguments for {macroName} macro, but got {nbArguments}. {context}",
+					("file", *Context.Filename),
+					("line", Context.LineNumber),
+					("expectedNbArguments", GetMinNumberOfArgument()),
+					("macroName", *GetToken()),
+					("nbArguments", Arguments.Num()),
+					("context", *FLocTextHelper::SanitizeLogOutput(Context.LineText.TrimStartAndEnd())),
+					("id", GatherTextFromSourceCommandlet::LocalizationLogIdentifier)
+				);
 			}
 			else
 			{
@@ -2985,7 +3253,12 @@ void UGatherTextFromSourceCommandlet::FStringTableEntryMacroDescriptor::TryParse
 
 					if (TableIdName.IsNone() || Key.IsEmpty())
 					{
-						UE_LOG(LogGatherTextFromSourceCommandlet, Warning, TEXT("%s(%d): %s macro has an empty identifier and cannot be gathered."), *Context.Filename, Context.LineNumber, *GetToken());
+						UE_LOGFMT(LogGatherTextFromSourceCommandlet, Warning, "{file}({line}): {macroName} macro has an empty identifier and cannot be gathered.",
+							("file", *Context.Filename),
+							("line", Context.LineNumber),
+							("macroName", *GetToken()),
+							("id", GatherTextFromSourceCommandlet::LocalizationLogIdentifier)
+						);
 					}
 					else if (!SourceString.IsEmpty())
 					{
@@ -3009,7 +3282,15 @@ void UGatherTextFromSourceCommandlet::FStringTableEntryMetaDataMacroDescriptor::
 		{
 			if (Arguments.Num() != GetMinNumberOfArgument())
 			{
-				UE_LOG(LogGatherTextFromSourceCommandlet, Warning, TEXT("%s(%d): Expected %d arguments for %s macro, but got %d. %s"), *Context.Filename, Context.LineNumber, GetMinNumberOfArgument(), *GetToken(), Arguments.Num(), *FLocTextHelper::SanitizeLogOutput(Context.LineText.TrimStartAndEnd()));
+				UE_LOGFMT(LogGatherTextFromSourceCommandlet, Warning, "{file}({line}): Expected {expectedNbArguments} arguments for {macroName} macro, but got {nbArguments}. {context}",
+					("file", *Context.Filename),
+					("line", Context.LineNumber),
+					("expectedNbArguments", GetMinNumberOfArgument()),
+					("macroName", *GetToken()),
+					("nbArguments", Arguments.Num()),
+					("context", *FLocTextHelper::SanitizeLogOutput(Context.LineText.TrimStartAndEnd())),
+					("id", GatherTextFromSourceCommandlet::LocalizationLogIdentifier)
+				);
 			}
 			else
 			{
@@ -3035,7 +3316,12 @@ void UGatherTextFromSourceCommandlet::FStringTableEntryMetaDataMacroDescriptor::
 
 					if (TableIdName.IsNone() || Key.IsEmpty() || MetaDataIdName.IsNone())
 					{
-						UE_LOG(LogGatherTextFromSourceCommandlet, Warning, TEXT("%s(%d): %s macro has an empty identifier and cannot be gathered."), *Context.Filename, Context.LineNumber, *GetToken());
+						UE_LOGFMT(LogGatherTextFromSourceCommandlet, Warning, "{file}({line}): {macroName} macro has an empty identifier and cannot be gathered.",
+							("file", *Context.Filename),
+							("line", Context.LineNumber),
+							("macroName", *GetToken()),
+							("id", GatherTextFromSourceCommandlet::LocalizationLogIdentifier)
+						);
 					}
 					else if (!MetaData.IsEmpty())
 					{
@@ -3079,7 +3365,14 @@ void UGatherTextFromSourceCommandlet::FStructuredLogMacroDescriptor::TryParse(co
 
 			if (Arguments.Num() < GetMinNumberOfArgument())
 			{
-				UE_LOG(LogGatherTextFromSourceCommandlet, Warning, TEXT("%s: Expected at least %d arguments for %s macro, but got %d. %s"), *SourceLocation, GetMinNumberOfArgument(), *GetToken(), Arguments.Num(), *FLocTextHelper::SanitizeLogOutput(Context.LineText.TrimStartAndEnd()));
+				UE_LOGFMT(LogGatherTextFromSourceCommandlet, Warning, "{location}: Expected at least {expectedNbArguments} arguments for {macroName} macro, but got {nbArguments}. {context}",
+					("location", *SourceLocation),
+					("expectedNbArguments", GetMinNumberOfArgument()),
+					("macroName", *GetToken()),
+					("nbArguments", Arguments.Num()),
+					("context", *FLocTextHelper::SanitizeLogOutput(Context.LineText.TrimStartAndEnd())),
+					("id", GatherTextFromSourceCommandlet::LocalizationLogIdentifier)
+				);
 				return;
 			}
 
@@ -3106,7 +3399,11 @@ void UGatherTextFromSourceCommandlet::FStructuredLogMacroDescriptor::TryParse(co
 
 			if (Namespace.IsEmpty())
 			{
-				UE_LOG(LogGatherTextFromSourceCommandlet, Warning, TEXT("%s: %s macro doesn't define a namespace and no external namespace was set. An empty namspace will be used."), *SourceLocation, *GetToken());
+				UE_LOGFMT(LogGatherTextFromSourceCommandlet, Warning, "{location}: {macroName} macro doesn't define a namespace and no external namespace was set. An empty namespace will be used.",
+					("location", *SourceLocation),
+					("macroName", *GetToken()),
+					("id", GatherTextFromSourceCommandlet::LocalizationLogIdentifier)
+				);
 			}
 
 			FString Key = ParseArg(MoveTemp(Arguments[NextArg++]), TEXT("Key"));
diff --git a/Engine/Source/Editor/UnrealEd/Private/Commandlets/GenerateGatherArchiveCommandlet.cpp b/Engine/Source/Editor/UnrealEd/Private/Commandlets/GenerateGatherArchiveCommandlet.cpp
--- a/Engine/Source/Editor/UnrealEd/Private/Commandlets/GenerateGatherArchiveCommandlet.cpp
+++ b/Engine/Source/Editor/UnrealEd/Private/Commandlets/GenerateGatherArchiveCommandlet.cpp
@@ -14,11 +14,16 @@
 #include "LocTextHelper.h"
 #include "Logging/LogCategory.h"
 #include "Logging/LogMacros.h"
+#include "Logging/StructuredLog.h"
+#include "ProfilingDebugging/ScopedTimers.h"
 #include "Templates/SharedPointer.h"
 #include "Trace/Detail/Channel.h"
-#include "ProfilingDebugging/ScopedTimers.h"
 
 DEFINE_LOG_CATEGORY_STATIC(LogGenerateArchiveCommandlet, Log, All);
+namespace GenerateGatherArchiveCommandlet
+{
+	static constexpr int32 LocalizationLogIdentifier = 304;
+}
 
 void ConditionTranslationMetadata(TSharedRef<FLocMetadataValue> MetadataValue)
 {
@@ -158,7 +163,9 @@ int32 UGenerateGatherArchiveCommandlet::Main( const FString& Params )
 	}
 	else
 	{
-		UE_LOG(LogGenerateArchiveCommandlet, Error, TEXT("No config specified."));
+		UE_LOGFMT(LogGenerateArchiveCommandlet, Error, "No config specified.",
+			("id", GenerateGatherArchiveCommandlet::LocalizationLogIdentifier)
+		);
 		return -1;
 	}
 
@@ -172,7 +179,9 @@ int32 UGenerateGatherArchiveCommandlet::Main( const FString& Params )
 	}
 	else
 	{
-		UE_LOG(LogGenerateArchiveCommandlet, Error, TEXT("No config section specified."));
+		UE_LOGFMT(LogGenerateArchiveCommandlet, Error, "No config section specified.",
+			("id", GenerateGatherArchiveCommandlet::LocalizationLogIdentifier)
+		);
 		return -1;
 	}
 
@@ -180,7 +189,9 @@ int32 UGenerateGatherArchiveCommandlet::Main( const FString& Params )
 	FString ManifestName;
 	if( !GetStringFromConfig( *SectionName, TEXT("ManifestName"), ManifestName, GatherTextConfigPath ) )
 	{
-		UE_LOG( LogGenerateArchiveCommandlet, Error, TEXT("No manifest name specified.") );
+		UE_LOGFMT(LogGenerateArchiveCommandlet, Error, "No manifest name specified.",
+			("id", GenerateGatherArchiveCommandlet::LocalizationLogIdentifier)
+		);
 		return -1;
 	}
 
@@ -188,7 +199,9 @@ int32 UGenerateGatherArchiveCommandlet::Main( const FString& Params )
 	FString ArchiveName;
 	if (!(GetStringFromConfig(*SectionName, TEXT("ArchiveName"), ArchiveName, GatherTextConfigPath)))
 	{
-		UE_LOG(LogGenerateArchiveCommandlet, Error, TEXT("No archive name specified."));
+		UE_LOGFMT(LogGenerateArchiveCommandlet, Error, "No archive name specified.",
+			("id", GenerateGatherArchiveCommandlet::LocalizationLogIdentifier)
+		);
 		return -1;
 	}
 
@@ -196,7 +209,9 @@ int32 UGenerateGatherArchiveCommandlet::Main( const FString& Params )
 	FString NativeCulture;
 	if (!GetStringFromConfig(*SectionName, TEXT("NativeCulture"), NativeCulture, GatherTextConfigPath))
 	{
-		UE_LOG(LogGenerateArchiveCommandlet, Error, TEXT("No native culture specified."));
+		UE_LOGFMT(LogGenerateArchiveCommandlet, Error, "No native culture specified.",
+			("id", GenerateGatherArchiveCommandlet::LocalizationLogIdentifier)
+		);
 		return -1;
 	}
 
@@ -206,7 +221,9 @@ int32 UGenerateGatherArchiveCommandlet::Main( const FString& Params )
 	
 	if( CulturesToGenerate.Num() == 0 )
 	{
-		UE_LOG(LogGenerateArchiveCommandlet, Error, TEXT("No cultures specified for generation."));
+		UE_LOGFMT(LogGenerateArchiveCommandlet, Error, "No cultures specified for generation.",
+			("id", GenerateGatherArchiveCommandlet::LocalizationLogIdentifier)
+		);
 		return -1;
 	}
 
@@ -214,7 +231,9 @@ int32 UGenerateGatherArchiveCommandlet::Main( const FString& Params )
 	FString DestinationPath;
 	if( !GetPathFromConfig( *SectionName, TEXT("DestinationPath"), DestinationPath, GatherTextConfigPath ) )
 	{
-		UE_LOG( LogGenerateArchiveCommandlet, Error, TEXT("No destination path specified.") );
+		UE_LOGFMT(LogGenerateArchiveCommandlet, Error, "No destination path specified.",
+			("id", GenerateGatherArchiveCommandlet::LocalizationLogIdentifier)
+		);
 		return -1;
 	}
 
@@ -225,7 +244,10 @@ int32 UGenerateGatherArchiveCommandlet::Main( const FString& Params )
 		FText LoadError;
 		if (!LocTextHelper.LoadAll(ELocTextHelperLoadFlags::LoadOrCreate, &LoadError))
 		{
-			UE_LOG(LogGenerateArchiveCommandlet, Error, TEXT("%s"), *LoadError.ToString());
+			UE_LOGFMT(LogGenerateArchiveCommandlet, Error, "Load error: {error}",
+				("error", *LoadError.ToString()),
+				("id", GenerateGatherArchiveCommandlet::LocalizationLogIdentifier)
+			);
 			return -1;
 		}
 	}
@@ -295,7 +317,10 @@ int32 UGenerateGatherArchiveCommandlet::Main( const FString& Params )
 			FText SaveError;
 			if (!LocTextHelper.SaveArchive(CultureName, &SaveError))
 			{
-				UE_LOG(LogGenerateArchiveCommandlet, Error, TEXT("%s"), *SaveError.ToString());
+				UE_LOGFMT(LogGenerateArchiveCommandlet, Error, "Save error: {error}",
+					("error", *SaveError.ToString()),
+					("id", GenerateGatherArchiveCommandlet::LocalizationLogIdentifier)
+				);
 				return -1;
 			}
 		}
diff --git a/Engine/Source/Editor/UnrealEd/Private/Commandlets/GenerateGatherManifestCommandlet.cpp b/Engine/Source/Editor/UnrealEd/Private/Commandlets/GenerateGatherManifestCommandlet.cpp
--- a/Engine/Source/Editor/UnrealEd/Private/Commandlets/GenerateGatherManifestCommandlet.cpp
+++ b/Engine/Source/Editor/UnrealEd/Private/Commandlets/GenerateGatherManifestCommandlet.cpp
@@ -10,13 +10,18 @@
 #include "LocTextHelper.h"
 #include "Logging/LogCategory.h"
 #include "Logging/LogMacros.h"
+#include "Logging/StructuredLog.h"
 #include "Misc/FileHelper.h"
 #include "Misc/Paths.h"
+#include "ProfilingDebugging/ScopedTimers.h"
 #include "Templates/SharedPointer.h"
 #include "Trace/Detail/Channel.h"
-#include "ProfilingDebugging/ScopedTimers.h"
 
 DEFINE_LOG_CATEGORY_STATIC(LogGenerateManifestCommandlet, Log, All);
+namespace GenerateGatherManifestCommandlet
+{
+	static constexpr int32 LocalizationLogIdentifier = 304;
+}
 
 namespace GenerateManifestHelper
 {
@@ -96,7 +101,9 @@ int32 UGenerateGatherManifestCommandlet::Main( const FString& Params )
 	}
 	else
 	{
-		UE_LOG( LogGenerateManifestCommandlet, Error, TEXT("No config specified.") );
+		UE_LOGFMT(LogGenerateManifestCommandlet, Error, "No config specified.",
+			("id", GenerateGatherManifestCommandlet::LocalizationLogIdentifier)
+		);
 		return -1;
 	}
 
@@ -110,7 +117,9 @@ int32 UGenerateGatherManifestCommandlet::Main( const FString& Params )
 	}
 	else
 	{
-		UE_LOG( LogGenerateManifestCommandlet, Error, TEXT("No config section specified.") );
+		UE_LOGFMT(LogGenerateManifestCommandlet, Error, "No config section specified.",
+			("id", GenerateGatherManifestCommandlet::LocalizationLogIdentifier)
+		);
 		return -1;
 	}
 
@@ -118,7 +127,9 @@ int32 UGenerateGatherManifestCommandlet::Main( const FString& Params )
 	FString DestinationPath;
 	if( !GetPathFromConfig( *SectionName, TEXT("DestinationPath"), DestinationPath, GatherTextConfigPath ) )
 	{
-		UE_LOG( LogGenerateManifestCommandlet, Error, TEXT("No destination path specified.") );
+		UE_LOGFMT(LogGenerateManifestCommandlet, Error, "No destination path specified.",
+			("id", GenerateGatherManifestCommandlet::LocalizationLogIdentifier)
+		);
 		return -1;
 	}
 
@@ -126,13 +137,19 @@ int32 UGenerateGatherManifestCommandlet::Main( const FString& Params )
 	FString ManifestName;
 	if( !GetStringFromConfig( *SectionName, TEXT("ManifestName"), ManifestName, GatherTextConfigPath ) )
 	{
-		UE_LOG( LogGenerateManifestCommandlet, Error, TEXT("No manifest name specified.") );
+		UE_LOGFMT(LogGenerateManifestCommandlet, Error, "No manifest name specified.",
+			("id", GenerateGatherManifestCommandlet::LocalizationLogIdentifier)
+		);
 		return -1;
 	}
 
 	if (!GenerateManifestHelper::IsManifestFileExtensionValid(ManifestName))
 	{
-		UE_LOG(LogGenerateManifestCommandlet, Error, TEXT("Found manifest file %s is malformed. All manifest files should have a %s extension."), *ManifestName, *GenerateManifestHelper::GetManifestFileExtension());
+		UE_LOGFMT(LogGenerateManifestCommandlet, Error, "Found manifest file {file} is malformed. All manifest files should have a {extension} extension.",
+			("file", *ManifestName),
+			("extension", *GenerateManifestHelper::GetManifestFileExtension()),
+			("id", GenerateGatherManifestCommandlet::LocalizationLogIdentifier)
+		);
 		return -1;
 	}
 
@@ -151,7 +168,11 @@ int32 UGenerateGatherManifestCommandlet::Main( const FString& Params )
 	{
 		if (!GenerateManifestHelper::IsManifestFileExtensionValid(ManifestDependency))
 		{
-			UE_LOG(LogGenerateManifestCommandlet, Error, TEXT("Found manifest dependency %s is malformed. All manifest files should have a %s extension."), *ManifestDependency, *GenerateManifestHelper::GetManifestFileExtension());
+			UE_LOGFMT(LogGenerateManifestCommandlet, Error, "Found manifest dependency {file} is malformed. All manifest files should have a {extension} extension.",
+				("file", *ManifestDependency),
+				("extension", *GenerateManifestHelper::GetManifestFileExtension()),
+				("id", GenerateGatherManifestCommandlet::LocalizationLogIdentifier)
+			);
 			return -1;
 		}
 	}
@@ -163,14 +184,21 @@ int32 UGenerateGatherManifestCommandlet::Main( const FString& Params )
 			FString PreviewManifestDependency= GeneratePreviewManifestHelper::GetPreviewManifestFilename(ManifestDependency);
 			if (!FPaths::FileExists(PreviewManifestDependency))
 			{
-				UE_LOG(LogGenerateManifestCommandlet, Warning, TEXT("Preview manifest dependency %s does not exist. Make sure to generate all preview manifest dependencies of this localization target before trying again."), *PreviewManifestDependency);
+				UE_LOGFMT(LogGenerateManifestCommandlet, Warning, "Preview manifest dependency {dependencyManifest} does not exist. Make sure to generate all preview manifest dependencies of this localization target before trying again.",
+					("dependencyManifest", *PreviewManifestDependency),
+					("id", GenerateGatherManifestCommandlet::LocalizationLogIdentifier)
+				);
 				continue;
 			}
 			// make sure the preview manifest file is newer than the regular manifest file
 			// we could be dealing with a scenario where the preview manifest file is from a previous run and isn't up to date or failed to be deleted
 			else if (IFileManager::Get().GetTimeStamp(*PreviewManifestDependency) < IFileManager::Get().GetTimeStamp(*ManifestDependency))
 			{
-				UE_LOG(LogGenerateManifestCommandlet, Warning, TEXT("Preview manifest dependency %s is older than the original manifest %s. Preview manifest is out of date and should be regenerated as a dependency."), *PreviewManifestDependency, *ManifestDependency);
+				UE_LOGFMT(LogGenerateManifestCommandlet, Warning, "Preview manifest dependency {previewDependencyManifest} is older than the original manifest {originalManifest}. Preview manifest is out of date and should be regenerated as a dependency.",
+					("previewDependencyManifest", *PreviewManifestDependency),
+					("originalManifest", *ManifestDependency),
+					("id", GenerateGatherManifestCommandlet::LocalizationLogIdentifier)
+				);
 				continue;
 			}
 			ManifestDependency = MoveTemp(PreviewManifestDependency);
@@ -182,7 +210,11 @@ int32 UGenerateGatherManifestCommandlet::Main( const FString& Params )
 		FText OutError;
 		if (!GatherManifestHelper->AddDependency(ManifestDependency, &OutError))
 		{
-			UE_LOG(LogGenerateManifestCommandlet, Error, TEXT("Failed to add manifest dependency %s. Failure reason : %s"), *ManifestDependency, *OutError.ToString());
+			UE_LOGFMT(LogGenerateManifestCommandlet, Error, "Failed to add manifest dependency {dependencyManifest}. Error: {error}",
+				("dependencyManifest", *ManifestDependency),
+				("error", *OutError.ToString()),
+				("id", GenerateGatherManifestCommandlet::LocalizationLogIdentifier)
+			);
 			return -1;
 		}
 	}
@@ -194,7 +226,10 @@ int32 UGenerateGatherManifestCommandlet::Main( const FString& Params )
 	FText ManifestSaveError;
 	if (!GatherManifestHelper->SaveManifest(ManifestPath, &ManifestSaveError))
 	{
-		UE_LOG(LogGenerateManifestCommandlet, Error, TEXT("%s"), *ManifestSaveError.ToString());
+		UE_LOGFMT(LogGenerateManifestCommandlet, Error, "Save error: {error}",
+			("error", *ManifestSaveError.ToString()),
+			("id", GenerateGatherManifestCommandlet::LocalizationLogIdentifier)
+		);
 		return -1;
 	}
 
diff --git a/Engine/Source/Editor/UnrealEd/Private/Commandlets/InternationalizationExportCommandlet.cpp b/Engine/Source/Editor/UnrealEd/Private/Commandlets/InternationalizationExportCommandlet.cpp
--- a/Engine/Source/Editor/UnrealEd/Private/Commandlets/InternationalizationExportCommandlet.cpp
+++ b/Engine/Source/Editor/UnrealEd/Private/Commandlets/InternationalizationExportCommandlet.cpp
@@ -10,14 +10,19 @@
 #include "LocTextHelper.h"
 #include "Logging/LogCategory.h"
 #include "Logging/LogMacros.h"
+#include "Logging/StructuredLog.h"
 #include "PortableObjectPipeline.h"
+#include "ProfilingDebugging/ScopedTimers.h"
 #include "Templates/SharedPointer.h"
 #include "Trace/Detail/Channel.h"
 #include "UObject/Class.h"
 #include "UObject/UObjectGlobals.h"
-#include "ProfilingDebugging/ScopedTimers.h"
 
 DEFINE_LOG_CATEGORY_STATIC(LogInternationalizationExportCommandlet, Log, All);
+namespace InternationalizationExportCommandlet
+{
+	static constexpr int32 LocalizationLogIdentifier = 304;
+}
 
 /**
 *	UInternationalizationExportCommandlet
@@ -37,7 +42,9 @@ int32 UInternationalizationExportCommandlet::Main(const FString& Params)
 	}
 	else
 	{
-		UE_LOG(LogInternationalizationExportCommandlet, Error, TEXT("No config specified."));
+		UE_LOGFMT(LogInternationalizationExportCommandlet, Error, "No config specified.",
+			("id", InternationalizationExportCommandlet::LocalizationLogIdentifier)
+		);
 		return -1;
 	}
 
@@ -49,7 +56,9 @@ int32 UInternationalizationExportCommandlet::Main(const FString& Params)
 	}
 	else
 	{
-		UE_LOG(LogInternationalizationExportCommandlet, Error, TEXT("No config section specified."));
+		UE_LOGFMT(LogInternationalizationExportCommandlet, Error, "No config section specified.",
+			("id", InternationalizationExportCommandlet::LocalizationLogIdentifier)
+		);
 		return -1;
 	}
 
@@ -57,7 +66,9 @@ int32 UInternationalizationExportCommandlet::Main(const FString& Params)
 	FString NativeCultureName;
 	if (!GetStringFromConfig(*SectionName, TEXT("NativeCulture"), NativeCultureName, ConfigPath))
 	{
-		UE_LOG(LogInternationalizationExportCommandlet, Error, TEXT("No native culture specified."));
+		UE_LOGFMT(LogInternationalizationExportCommandlet, Error, "No native culture specified.",
+			("id", InternationalizationExportCommandlet::LocalizationLogIdentifier)
+		);
 		return false;
 	}
 
@@ -65,7 +76,9 @@ int32 UInternationalizationExportCommandlet::Main(const FString& Params)
 	FString ManifestName;
 	if (!GetStringFromConfig(*SectionName, TEXT("ManifestName"), ManifestName, ConfigPath))
 	{
-		UE_LOG(LogInternationalizationExportCommandlet, Error, TEXT("No manifest name specified."));
+		UE_LOGFMT(LogInternationalizationExportCommandlet, Error, "No manifest name specified.",
+			("id", InternationalizationExportCommandlet::LocalizationLogIdentifier)
+		);
 		return false;
 	}
 
@@ -73,7 +86,9 @@ int32 UInternationalizationExportCommandlet::Main(const FString& Params)
 	FString ArchiveName;
 	if (!GetStringFromConfig(*SectionName, TEXT("ArchiveName"), ArchiveName, ConfigPath))
 	{
-		UE_LOG(LogInternationalizationExportCommandlet, Error, TEXT("No archive name specified."));
+		UE_LOGFMT(LogInternationalizationExportCommandlet, Error, "No archive name specified.",
+			("id", InternationalizationExportCommandlet::LocalizationLogIdentifier)
+		);
 		return false;
 	}
 
@@ -81,7 +96,9 @@ int32 UInternationalizationExportCommandlet::Main(const FString& Params)
 	FString SourcePath;
 	if (!GetPathFromConfig(*SectionName, TEXT("SourcePath"), SourcePath, ConfigPath))
 	{
-		UE_LOG(LogInternationalizationExportCommandlet, Error, TEXT("No source path specified."));
+		UE_LOGFMT(LogInternationalizationExportCommandlet, Error, "No source path specified.",
+			("id", InternationalizationExportCommandlet::LocalizationLogIdentifier)
+		);
 		return -1;
 	}
 
@@ -89,7 +106,9 @@ int32 UInternationalizationExportCommandlet::Main(const FString& Params)
 	FString DestinationPath;
 	if (!GetPathFromConfig(*SectionName, TEXT("DestinationPath"), DestinationPath, ConfigPath))
 	{
-		UE_LOG(LogInternationalizationExportCommandlet, Error, TEXT("No destination path specified."));
+		UE_LOGFMT(LogInternationalizationExportCommandlet, Error, "No destination path specified.",
+			("id", InternationalizationExportCommandlet::LocalizationLogIdentifier)
+		);
 		return -1;
 	}
 
@@ -97,7 +116,9 @@ int32 UInternationalizationExportCommandlet::Main(const FString& Params)
 	FString Filename;
 	if (!GetStringFromConfig(*SectionName, TEXT("PortableObjectName"), Filename, ConfigPath))
 	{
-		UE_LOG(LogInternationalizationExportCommandlet, Error, TEXT("No portable object name specified."));
+		UE_LOGFMT(LogInternationalizationExportCommandlet, Error, "No portable object name specified.",
+			("id", InternationalizationExportCommandlet::LocalizationLogIdentifier)
+		);
 		return -1;
 	}
 
@@ -105,7 +126,9 @@ int32 UInternationalizationExportCommandlet::Main(const FString& Params)
 	TArray<FString> CulturesToGenerate;
 	if (GetStringArrayFromConfig(*SectionName, TEXT("CulturesToGenerate"), CulturesToGenerate, ConfigPath) == 0)
 	{
-		UE_LOG(LogInternationalizationExportCommandlet, Error, TEXT("No cultures specified for generation."));
+		UE_LOGFMT(LogInternationalizationExportCommandlet, Error, "No cultures specified for generation.",
+			("id", InternationalizationExportCommandlet::LocalizationLogIdentifier)
+		);
 		return -1;
 	}
 
@@ -154,7 +177,9 @@ int32 UInternationalizationExportCommandlet::Main(const FString& Params)
 
 	if (!bDoImport && !bDoExport)
 	{
-		UE_LOG(LogInternationalizationExportCommandlet, Error, TEXT("Import/Export operation not detected.  Use bExportLoc or bImportLoc in config section."));
+		UE_LOGFMT(LogInternationalizationExportCommandlet, Error, "Import/Export operation not detected. Use bExportLoc or bImportLoc in config section.",
+			("id", InternationalizationExportCommandlet::LocalizationLogIdentifier)
+		);
 		return -1;
 	}
 
@@ -168,7 +193,10 @@ int32 UInternationalizationExportCommandlet::Main(const FString& Params)
 			FText LoadError;
 			if (!LocTextHelper.LoadAll(ELocTextHelperLoadFlags::LoadOrCreate, &LoadError))
 			{
-				UE_LOG(LogInternationalizationExportCommandlet, Error, TEXT("%s"), *LoadError.ToString());
+				UE_LOGFMT(LogInternationalizationExportCommandlet, Error, "Load error: {error}",
+					("error", *LoadError.ToString()),
+					("id", InternationalizationExportCommandlet::LocalizationLogIdentifier)
+				);
 				return false;
 			}
 		}
@@ -176,7 +204,9 @@ int32 UInternationalizationExportCommandlet::Main(const FString& Params)
 		// Import all PO files
 		if (!PortableObjectPipeline::ImportAll(LocTextHelper, SourcePath, Filename, TextCollapseMode, POFormat, bUseCultureDirectory))
 		{
-			UE_LOG(LogInternationalizationExportCommandlet, Error, TEXT("Failed to import localization files."));
+			UE_LOGFMT(LogInternationalizationExportCommandlet, Error, "Failed to import localization files.",
+				("id", InternationalizationExportCommandlet::LocalizationLogIdentifier)
+			);
 			return -1;
 		}
 	}
@@ -194,7 +224,10 @@ int32 UInternationalizationExportCommandlet::Main(const FString& Params)
 			FText LoadError;
 			if (!LocTextHelper.LoadAll(ELocTextHelperLoadFlags::LoadOrCreate, &LoadError))
 			{
-				UE_LOG(LogInternationalizationExportCommandlet, Error, TEXT("%s"), *LoadError.ToString());
+				UE_LOGFMT(LogInternationalizationExportCommandlet, Error, "Load error: {error}",
+					("error", *LoadError.ToString()),
+					("id", InternationalizationExportCommandlet::LocalizationLogIdentifier)
+				);
 				return false;
 			}
 		}
@@ -202,7 +235,9 @@ int32 UInternationalizationExportCommandlet::Main(const FString& Params)
 		// Export all PO files
 		if (!PortableObjectPipeline::ExportAll(LocTextHelper, DestinationPath, Filename, TextCollapseMode, POFormat, bShouldPersistComments, bUseCultureDirectory))
 		{
-			UE_LOG(LogInternationalizationExportCommandlet, Error, TEXT("Failed to export localization files."));
+			UE_LOGFMT(LogInternationalizationExportCommandlet, Error, "Failed to export localization files.",
+				("id", InternationalizationExportCommandlet::LocalizationLogIdentifier)
+			);
 			return -1;
 		}
 	}
