Cooker: Fix soft lock in IsCachedCookedPlatformData calls.
When an object in a package we are trying to save is moved out of the package, stop waiting for it.
When an object in a package we are trying to save is marked for garbage, stop waiting for it.
When an object is reallocated (destroying the object and recreating it, without a garbage collect), detect that via WeakPtr and reset the object's need to have BeginCacheForCookedPlatformData called on it.
#rnx

diff --git a/Engine/Source/Editor/UnrealEd/Private/Cooker/CookPackageData.cpp b/Engine/Source/Editor/UnrealEd/Private/Cooker/CookPackageData.cpp
--- a/Engine/Source/Editor/UnrealEd/Private/Cooker/CookPackageData.cpp
+++ b/Engine/Source/Editor/UnrealEd/Private/Cooker/CookPackageData.cpp
@@ -1954,22 +1954,42 @@ bool FPendingCookedPlatformData::PollIsComplete()
 		Release();
 		return true;
 	}
-	else
+
+	// If something (another object's BeginCacheForCookedPlatformData, maybe) has marked the object as
+	// garbage, or renamed it out of the package, then we no longer need to wait on it.
+	// We might have removed the packagedata from the save state and no longer have a cached UPackage* on it,
+	// so compare current package vs original package by name instead of pointer.
+	FName CurrentPackageName = LocalObject->GetPackage()->GetFName();
+	if (CurrentPackageName != PackageData.GetPackageName())
+	{
+		UE_LOG(LogCook, Warning,
+			TEXT("We were waiting for IsCachedCookedPlatformData to return true for %s in package %s, but that object has been moved out of the package. We will stop waiting on it."),
+			*LocalObject->GetFullName(), *PackageData.GetPackageName().ToString());
+		Release();
+		return true;
+	}
+	if (LocalObject->HasAnyFlags(EObjectFlags::RF_MirroredGarbage))
 	{
+		UE_LOG(LogCook, Warning,
+			TEXT("We were waiting for IsCachedCookedPlatformData to return true for %s, but that object is now marked for garbage. We will stop waiting on it."),
+			*LocalObject->GetFullName());
+		Release();
+		return true;
+	}
+
 #if DEBUG_COOKONTHEFLY
-		UE_LOG(LogCook, Display, TEXT("Object %s isn't cached yet"), *LocalObject->GetFullName());
+	UE_LOG(LogCook, Display, TEXT("%s isn't cached yet"), *LocalObject->GetFullName());
 #endif
-		/*if ( LocalObject->IsA(UMaterial::StaticClass()) )
+	/*if ( LocalObject->IsA(UMaterial::StaticClass()) )
+	{
+		if (GShaderCompilingManager->HasShaderJobs() == false)
 		{
-			if (GShaderCompilingManager->HasShaderJobs() == false)
-			{
-				UE_LOG(LogCook, Warning,
-				TEXT("Shader compiler is in a bad state!  Shader %s is finished compile but shader compiling manager did not notify shader.  "),
-					*LocalObject->GetPathName());
-			}
-		}*/
-		return false;
-	}
+			UE_LOG(LogCook, Warning,
+			TEXT("Shader compiler is in a bad state!  Shader %s is finished compile but shader compiling manager did not notify shader.  "),
+				*LocalObject->GetPathName());
+		}
+	}*/
+	return false;
 }
 
 void FPendingCookedPlatformData::Release()
@@ -2009,10 +2029,9 @@ void FPendingCookedPlatformData::ClearCachedCookedPlatformData(UObject* Object,
 {
 	FPackageDatas& PackageDatas = PackageData.GetPackageDatas();
 	UCookOnTheFlyServer& COTFS = PackageDatas.GetCookOnTheFlyServer();
-	using FCCPDMapType = TMap<UObject*, FCachedCookedPlatformDataState>;
-	FCCPDMapType& CCPDs = PackageDatas.GetCachedCookedPlatformDataObjects();
+	FMapOfCachedCookedPlatformDataState& CCPDs = PackageDatas.GetCachedCookedPlatformDataObjects();
 
-	uint32 ObjectKeyHash = FCCPDMapType::KeyFuncsType::GetKeyHash(Object);
+	uint32 ObjectKeyHash = FMapOfCachedCookedPlatformDataState::KeyFuncsType::GetKeyHash(Object);
 	FCachedCookedPlatformDataState* CCPDState = CCPDs.FindByHash(ObjectKeyHash, Object);
 	if (!CCPDState)
 	{
@@ -3227,4 +3246,82 @@ bool FCachedCookedPlatformDataState::IsReferenced() const
 	return !PackageDatas.IsEmpty();
 }
 
+void FCachedCookedPlatformDataState::Construct(UObject* Object)
+{
+	WeakPtr = Object;
+	bInitialized = true;
+}
+
+FCachedCookedPlatformDataState& FMapOfCachedCookedPlatformDataState::Add(
+	UObject* Object, const FCachedCookedPlatformDataState& Value)
+{
+	FCachedCookedPlatformDataState* Existing = &Super::FindOrAdd(Object);
+	*Existing = Value;
+	if (!Existing->bInitialized)
+	{
+		Existing->Construct(Object);
+	}
+	return *Existing;
+}
+
+FCachedCookedPlatformDataState& FMapOfCachedCookedPlatformDataState::FindOrAdd(UObject* Object)
+{
+	FCachedCookedPlatformDataState* Existing = &Super::FindOrAdd(Object);
+	if (!Existing->bInitialized)
+	{
+		Existing->Construct(Object);
+	}
+	if (!Existing->WeakPtr.Get())
+	{
+		Remove(Object);
+		Existing = &Super::FindOrAdd(Object);
+	}
+	return *Existing;
+}
+
+FCachedCookedPlatformDataState* FMapOfCachedCookedPlatformDataState::Find(UObject* Object)
+{
+	FCachedCookedPlatformDataState* Existing = Super::Find(Object);
+	if (!Existing)
+	{
+		return nullptr;
+	}
+	if (!Existing->WeakPtr.Get())
+	{
+		Remove(Object);
+		return nullptr;
+	}
+	return Existing;
+}
+
+FCachedCookedPlatformDataState& FMapOfCachedCookedPlatformDataState::FindOrAddByHash(uint32 KeyHash, UObject* Object)
+{
+	FCachedCookedPlatformDataState* Existing = &Super::FindOrAddByHash(KeyHash, Object);
+	if (!Existing->bInitialized)
+	{
+		Existing->Construct(Object);
+	}
+	if (!Existing->WeakPtr.Get())
+	{
+		Remove(Object);
+		Existing = &Super::FindOrAddByHash(KeyHash, Object);
+	}
+	return *Existing;
+}
+
+FCachedCookedPlatformDataState* FMapOfCachedCookedPlatformDataState::FindByHash(uint32 KeyHash, UObject* Object)
+{
+	FCachedCookedPlatformDataState* Existing = Super::FindByHash(KeyHash, Object);
+	if (!Existing)
+	{
+		return nullptr;
+	}
+	if (!Existing->WeakPtr.Get())
+	{
+		Remove(Object);
+		return nullptr;
+	}
+	return Existing;
+}
+
 } // namespace UE::Cook
diff --git a/Engine/Source/Editor/UnrealEd/Private/Cooker/CookPackageData.h b/Engine/Source/Editor/UnrealEd/Private/Cooker/CookPackageData.h
--- a/Engine/Source/Editor/UnrealEd/Private/Cooker/CookPackageData.h
+++ b/Engine/Source/Editor/UnrealEd/Private/Cooker/CookPackageData.h
@@ -76,10 +76,18 @@ const TCHAR* LexToString(ECachedCookedPlatformDataEvent);
  */
 struct FCachedCookedPlatformDataState
 {
+	void Construct(UObject* Object);
 	void AddRefFrom(FPackageData* PackageData);
 	void ReleaseFrom(FPackageData* PackageData);
 	bool IsReferenced() const;
 
+	/**
+	* The weakpointer that was constructed from the UObject* key when we added the data to the cache.
+	* If this pointer no longer equals the key, then the UObject* at the memory position given by the key
+	* has been destroyed (and possibly a new object allocated into its same memory) and we should clear the
+	* the cache state.
+	*/
+	FWeakObjectPtr WeakPtr;
 	/**
 	 * The packages that have called any of the BeginCacheForCookedPlatformData family of function
 	 * on this object. Usually only a single 1, sometimes 2, see comment in AddPackageData.
@@ -87,6 +95,20 @@ struct FCachedCookedPlatformDataState
 	TArray<FPackageData*, TInlineAllocator<2>> PackageDatas;
 	/** The per-platform state of which BeginCacheForCookedPlatformData events have been passed. */
 	TMap<const ITargetPlatform*, ECachedCookedPlatformDataEvent> PlatformStates;
+	bool bInitialized = false;
+};
+
+class FMapOfCachedCookedPlatformDataState : public TMap<UObject*, FCachedCookedPlatformDataState>
+{
+public:
+	using Super = TMap<UObject*, FCachedCookedPlatformDataState>;
+	using Super::TMap;
+
+	FCachedCookedPlatformDataState& Add(UObject* Object, const FCachedCookedPlatformDataState& Value);
+	FCachedCookedPlatformDataState& FindOrAdd(UObject* Object);
+	FCachedCookedPlatformDataState* Find(UObject* Object);
+	FCachedCookedPlatformDataState& FindOrAddByHash(uint32 TypeHash, UObject* Object);
+	FCachedCookedPlatformDataState* FindByHash(uint32 TypeHash, UObject* Object);
 };
 
 /**
@@ -1408,7 +1430,7 @@ public:
 	template <typename CallbackType>
 	void LockAndEnumeratePackageDatas(CallbackType&& Callback);
 
-	TMap<UObject*, FCachedCookedPlatformDataState>& GetCachedCookedPlatformDataObjects();
+	FMapOfCachedCookedPlatformDataState& GetCachedCookedPlatformDataObjects();
 	void CachedCookedPlatformDataObjectsPostGarbageCollect(const TSet<UObject*>& SaveQueueObjectsThatStillExist);
 
 private:
@@ -1435,7 +1457,7 @@ private:
 	/* Guarded by ExistenceLock. Duplicates information on FPackageData, but can be read/write from any thread. */
 	TMap<FName, FThreadsafePackageData> ThreadsafePackageDatas;
 	TRingBuffer<FPendingCookedPlatformDataContainer> PendingCookedPlatformDataLists;
-	TMap<UObject*, FCachedCookedPlatformDataState> CachedCookedPlatformDataObjects;
+	FMapOfCachedCookedPlatformDataState CachedCookedPlatformDataObjects;
 	uint32 NextLeafToRootRank = 0;
 	int32 PendingCookedPlatformDataNum = 0;
 	FRequestQueue RequestQueue;
@@ -2063,7 +2085,7 @@ inline void FPackageDatas::LockAndEnumeratePackageDatas(CallbackType&& Callback)
 	EnumeratePackageDatasWithinLock(Forward<CallbackType>(Callback));
 }
 
-inline TMap<UObject*, FCachedCookedPlatformDataState>& FPackageDatas::GetCachedCookedPlatformDataObjects()
+inline FMapOfCachedCookedPlatformDataState& FPackageDatas::GetCachedCookedPlatformDataObjects()
 {
 	return CachedCookedPlatformDataObjects;
 }
