Fixed reference tree Show Script Objects menu item not displaying checked state.

diff --git a/Engine/Source/Editor/UnrealEd/Private/Dialogs/DlgReferenceTree.cpp b/Engine/Source/Editor/UnrealEd/Private/Dialogs/DlgReferenceTree.cpp
--- a/Engine/Source/Editor/UnrealEd/Private/Dialogs/DlgReferenceTree.cpp
+++ b/Engine/Source/Editor/UnrealEd/Private/Dialogs/DlgReferenceTree.cpp
@@ -476,7 +476,7 @@ void SReferenceTree::FillViewEntries( FMenuBuilder& MenuBuilder )
 void SReferenceTree::FillOptionsEntries( FMenuBuilder& MenuBuilder )
 {
 	MenuBuilder.AddMenuEntry( NSLOCTEXT("ReferenceTreeView", "ShowScriptObjects", "Show Script Objects"), NSLOCTEXT("ReferenceTreeView", "ShowScriptObjects_Tooltip", "Toggles displaying script objects in the tree."), FSlateIcon(), 
-		FUIAction(FExecuteAction::CreateRaw(this, &SReferenceTree::OnShowScriptReferences), FCanExecuteAction(), FIsActionChecked::CreateRaw(this, &SReferenceTree::OnShowScriptReferences_Checked) ) );
+		FUIAction(FExecuteAction::CreateRaw(this, &SReferenceTree::OnShowScriptReferences), FCanExecuteAction(), FIsActionChecked::CreateRaw(this, &SReferenceTree::OnShowScriptReferences_Checked)), NAME_None, EUserInterfaceActionType::Check );
 }
 
 void SReferenceTree::SetAllExpansionStates(bool bInExpansionState)
