Added support for showing Verse paths to the reference tree.

diff --git a/Engine/Source/Editor/UnrealEd/Private/Dialogs/DlgReferenceTree.cpp b/Engine/Source/Editor/UnrealEd/Private/Dialogs/DlgReferenceTree.cpp
--- a/Engine/Source/Editor/UnrealEd/Private/Dialogs/DlgReferenceTree.cpp
+++ b/Engine/Source/Editor/UnrealEd/Private/Dialogs/DlgReferenceTree.cpp
@@ -1,6 +1,7 @@
 // Copyright Epic Games, Inc. All Rights Reserved.
 
 #include "Dialogs/DlgReferenceTree.h"
+#include "Containers/VersePath.h"
 #include "UObject/UObjectIterator.h"
 #include "Widgets/SWindow.h"
 #include "Framework/Application/SlateApplication.h"
@@ -8,6 +9,7 @@
 #include "Framework/Commands/UIAction.h"
 #include "Framework/Commands/UICommandList.h"
 #include "Framework/MultiBox/MultiBoxBuilder.h"
+#include "IAssetTools.h"
 #include "Widgets/Input/SButton.h"
 #include "Styling/AppStyle.h"
 #include "Editor/UnrealEdEngine.h"
@@ -141,6 +143,20 @@ struct FReferenceTreeDataContainer
 
 	FReferenceTreeDataContainer(UObject* InObject) : Object(InObject)
 	{}
+
+	FText GetToolTipText() const
+	{
+		if (IAssetTools::Get().ShowingContentVersePath())
+		{
+			UE::Core::FVersePath VersePath = Object->GetVersePath();
+			if (VersePath.IsValid())
+			{
+				return FText::FromString(MoveTemp(VersePath).ToString());
+			}
+		}
+
+		return FText::FromString(Object->GetPathName());
+	}
 };
 
 /** Widget that represents a row in the sound wave compression's tree control.  Generates widgets for each column on demand. */
@@ -563,5 +579,5 @@ TSharedRef< ITableRow > SReferenceTree::OnGenerateRowForReferenceTree( FReferenc
 	return
 		SNew( SReferenceTreeRow, OwnerTable )
 		.Item( Item )
-		.ToolTipText( FText::FromString(Item->Object->GetFullName()) );
+		.ToolTipText( Item.ToSharedRef(), &FReferenceTreeDataContainer::GetToolTipText );
 }
