[ContentBrowser] New setting to expand Tooltip by default

diff --git a/Engine/Source/Editor/ContentBrowser/Private/AssetViewWidgets.cpp b/Engine/Source/Editor/ContentBrowser/Private/AssetViewWidgets.cpp
--- a/Engine/Source/Editor/ContentBrowser/Private/AssetViewWidgets.cpp
+++ b/Engine/Source/Editor/ContentBrowser/Private/AssetViewWidgets.cpp
@@ -65,6 +65,7 @@
 #include "SlotBase.h"
 #include "SourceControlHelpers.h"
 #include "AssetRegistry/AssetRegistryModule.h"
+#include "Settings/ContentBrowserSettings.h"
 #include "Styling/ISlateStyle.h"
 #include "Styling/SlateBrush.h"
 #include "Styling/SlateTypes.h"
@@ -700,6 +701,11 @@ void SAssetViewItem::CacheDirtyExternalPackageInfo()
 	}
 }
 
+bool SAssetViewItem::ShouldExpandTooltip() const
+{
+	return GetDefault<UContentBrowserSettings>()->GetAlwaysExpandTooltips();
+}
+
 const FSlateBrush* SAssetViewItem::GetDirtyImage() const
 {
 	return IsDirty() ? AssetDirtyBrush : NULL;
@@ -1872,6 +1878,7 @@ void SAssetListItem::Construct( const FArguments& InArgs )
 			ThumbnailConfig.GenericThumbnailSize = MakeAttributeSP(this, &SAssetListItem::GetGenericThumbnailSize);
 			ThumbnailConfig.AssetSystemInfoProvider = MakeShared<FAssetSystemContentBrowserInfoProvider>(AssetItem);
 			ThumbnailConfig.bAllowAssetStatusThumbnailOverlay = true;
+			ThumbnailConfig.AlwaysExpandTooltip = MakeAttributeSP(this, &SAssetListItem::ShouldExpandTooltip);
 			ThumbnailConfig.bShowAssetChip = true;
 		}
 
@@ -2351,6 +2358,7 @@ void SAssetTileItem::Construct( const FArguments& InArgs )
 			ThumbnailConfig.bShowAssetChip = true;
 			ThumbnailConfig.AssetChipBorderImageOverride = TAttribute<const FSlateBrush*>::CreateSP(this, &SAssetTileItem::GetAssetAreaOverlayBackgroundImage);
 			ThumbnailConfig.bCanDisplayEditModePrimitiveTools = CanDisplayPrimitiveTools();
+			ThumbnailConfig.AlwaysExpandTooltip = MakeAttributeSP(this, &SAssetTileItem::ShouldExpandTooltip);
 			ThumbnailConfig.IsEditModeVisible = MakeAttributeSP(this, &SAssetTileItem::GetThumbnailEditModeUIVisibility);
 		}
 
diff --git a/Engine/Source/Editor/ContentBrowser/Private/Menus/AssetViewMenus.cpp b/Engine/Source/Editor/ContentBrowser/Private/Menus/AssetViewMenus.cpp
--- a/Engine/Source/Editor/ContentBrowser/Private/Menus/AssetViewMenus.cpp
+++ b/Engine/Source/Editor/ContentBrowser/Private/Menus/AssetViewMenus.cpp
@@ -308,6 +308,28 @@ FDelayedAutoRegisterHelper SAssetView::AssetViewOptionsMenuRegistration(
 						);
 					}
 
+					// AssetDetails Section
+					if (UE::Editor::ContentBrowser::IsNewStyleEnabled())
+					{
+						const FToolMenuInsert AssetDetailsSectionInsert("Search", EToolMenuInsertType::After);
+						FToolMenuSection& AssetDetailsSection = InMenu->FindOrAddSection(
+							"AssetDetails",
+							LOCTEXT("AssetDetailsName", "Asset Details"),
+							AssetDetailsSectionInsert);
+
+						AssetDetailsSection.AddMenuEntry(
+							"ThumbnailTooltip",
+							LOCTEXT("ThumbnailTooltipExpanded", "Always Expand Tooltips"),
+							LOCTEXT("ThumbnailTooltipExpandedTooltip", "Toggle Asset and Folder tooltip expansion default state"),
+							FSlateIcon(),
+							FUIAction(
+									FExecuteAction::CreateSP(AssetView, &SAssetView::ToggleTooltipExpandedState),
+									FCanExecuteAction(),
+									FIsActionChecked::CreateSP(AssetView, &SAssetView::IsTooltipExpandedByDefault)
+								),
+							EUserInterfaceActionType::ToggleButton);
+					}
+
 					// Show/View Section
 					{
 						FToolMenuSection& ShowSection = InMenu->FindOrAddSection(
diff --git a/Engine/Source/Editor/ContentBrowser/Private/SAssetView.cpp b/Engine/Source/Editor/ContentBrowser/Private/SAssetView.cpp
--- a/Engine/Source/Editor/ContentBrowser/Private/SAssetView.cpp
+++ b/Engine/Source/Editor/ContentBrowser/Private/SAssetView.cpp
@@ -5666,6 +5666,19 @@ EVisibility SAssetView::GetColumnViewVisibility() const
 	return GetCurrentViewType() == EAssetViewType::Column ? EVisibility::Visible : EVisibility::Collapsed;
 }
 
+void SAssetView::ToggleTooltipExpandedState()
+{
+	bool bNewState = !GetDefault<UContentBrowserSettings>()->GetAlwaysExpandTooltips();
+
+	GetMutableDefault<UContentBrowserSettings>()->SetAlwaysExpandTooltips(bNewState);
+	GetMutableDefault<UContentBrowserSettings>()->PostEditChange();
+}
+
+bool SAssetView::IsTooltipExpandedByDefault()
+{
+	return GetDefault<UContentBrowserSettings>()->GetAlwaysExpandTooltips();
+}
+
 void SAssetView::ToggleThumbnailEditMode()
 {
 	bThumbnailEditMode = !bThumbnailEditMode;
diff --git a/Engine/Source/Editor/ContentBrowser/Public/AssetViewWidgets.h b/Engine/Source/Editor/ContentBrowser/Public/AssetViewWidgets.h
--- a/Engine/Source/Editor/ContentBrowser/Public/AssetViewWidgets.h
+++ b/Engine/Source/Editor/ContentBrowser/Public/AssetViewWidgets.h
@@ -265,6 +265,9 @@ protected:
 	/** Caches the ExternalPackagesText */
 	void CacheDirtyExternalPackageInfo();
 
+	/** If true it will expand the tooltip by default */
+	bool ShouldExpandTooltip() const;
+
 protected:
 	/** Data for a cached display tag for this item (used in the tooltip, and also as the display string in column views) */
 	struct FTagDisplayItem
diff --git a/Engine/Source/Editor/ContentBrowser/Public/SAssetView.h b/Engine/Source/Editor/ContentBrowser/Public/SAssetView.h
--- a/Engine/Source/Editor/ContentBrowser/Public/SAssetView.h
+++ b/Engine/Source/Editor/ContentBrowser/Public/SAssetView.h
@@ -834,6 +834,12 @@ private:
 	/** Gets the visibility of the column view */
 	EVisibility GetColumnViewVisibility() const;
 
+	/** Toggles tooltip expanded default state */
+	void ToggleTooltipExpandedState();
+
+	/** True if the default state of tooltip is expanded */
+	bool IsTooltipExpandedByDefault();
+
 	/** Toggles thumbnail editing mode */
 	void ToggleThumbnailEditMode();
 
diff --git a/Engine/Source/Editor/Documentation/Private/SDocumentationToolTip.cpp b/Engine/Source/Editor/Documentation/Private/SDocumentationToolTip.cpp
--- a/Engine/Source/Editor/Documentation/Private/SDocumentationToolTip.cpp
+++ b/Engine/Source/Editor/Documentation/Private/SDocumentationToolTip.cpp
@@ -59,6 +59,7 @@ void SDocumentationToolTip::Construct( const FArguments& InArgs )
 	Shortcut = InArgs._Shortcut;
 	OverrideFullTooltipContent = InArgs._OverrideExtendedToolTipContent;
 	OverridePromptContent = InArgs._OverridePromptContent;
+	AlwaysExpandTooltip = InArgs._AlwaysExpandTooltip;
 
 	ExcerptName = InArgs._ExcerptName;
 	IsShowingFullTip = false;
@@ -610,7 +611,7 @@ void SDocumentationToolTip::Tick( const FGeometry& AllottedGeometry, const doubl
 		}
 	}
 
-	if ( !IsShowingFullTip && ModifierKeys.IsAltDown() && ModifierKeys.IsControlDown() )
+	if ( !IsShowingFullTip && ((ModifierKeys.IsAltDown() && ModifierKeys.IsControlDown()) || ShouldAlwaysExpand()) )
 	{
 		if (!OverrideFullTooltipContent.IsValid())
 		{
@@ -644,7 +645,7 @@ void SDocumentationToolTip::Tick( const FGeometry& AllottedGeometry, const doubl
 		}
 		IsShowingFullTip = true;
 	}
-	else if ( ( IsShowingFullTip || NeedsUpdate )  && ( !ModifierKeys.IsAltDown() || !ModifierKeys.IsControlDown() ) )
+	else if ( ( IsShowingFullTip || NeedsUpdate )  && (( !ModifierKeys.IsAltDown() || !ModifierKeys.IsControlDown() ) && !ShouldAlwaysExpand() ))
 	{
 		if ( NeedsUpdate )
 		{
@@ -718,6 +719,11 @@ EVisibility SDocumentationToolTip::GetShortcutVisibility() const
 	return EVisibility::Collapsed;
 }
 
+bool SDocumentationToolTip::ShouldAlwaysExpand() const
+{
+	return AlwaysExpandTooltip.IsSet() && AlwaysExpandTooltip.Get();
+}
+
 FOptionalSize SDocumentationToolTip::GetPromptHeight() const
 {
 	const float DesiredHeight = PromptContent->GetDesiredSize().Y;
diff --git a/Engine/Source/Editor/Documentation/Public/SDocumentationToolTip.h b/Engine/Source/Editor/Documentation/Public/SDocumentationToolTip.h
--- a/Engine/Source/Editor/Documentation/Public/SDocumentationToolTip.h
+++ b/Engine/Source/Editor/Documentation/Public/SDocumentationToolTip.h
@@ -42,6 +42,7 @@ public:
 		, _DocumentationMargin(0)
 		, _OverrideExtendedToolTipContent(nullptr)
 		, _OverridePromptContent(nullptr)
+		, _AlwaysExpandTooltip(false)
 		, _Content()
 		{}
 
@@ -83,6 +84,9 @@ public:
 		/** Custom widget to prompt users to hit Ctrl + Alt, if needed to fit with the design of custom content overrides */
 		SLATE_ARGUMENT(TSharedPtr<SWidget>, OverridePromptContent)
 
+		/** If true it will keep the tooltip expanded without the need of the Ctrl + Alt to expand it */
+		SLATE_ATTRIBUTE(bool, AlwaysExpandTooltip)
+
 		/** Arbitrary content to be displayed in the tool tip; overrides any text that may be set. */
 		SLATE_DEFAULT_SLOT( FArguments, Content )
 
@@ -132,6 +136,7 @@ private:
 	UE_API EVisibility GetControlVisibility() const;
 	UE_API EVisibility GetShortcutVisibility() const;
 
+	UE_API bool ShouldAlwaysExpand() const;
 private:
 
 	/** Text block widget */
@@ -173,6 +178,7 @@ private:
 
 	bool bFullTipContentIsReady;
 	bool bIsPromptVisible;
+	TAttribute< bool > AlwaysExpandTooltip;
 };
 
 #undef UE_API
diff --git a/Engine/Source/Editor/UnrealEd/Classes/Settings/ContentBrowserSettings.h b/Engine/Source/Editor/UnrealEd/Classes/Settings/ContentBrowserSettings.h
--- a/Engine/Source/Editor/UnrealEd/Classes/Settings/ContentBrowserSettings.h
+++ b/Engine/Source/Editor/UnrealEd/Classes/Settings/ContentBrowserSettings.h
@@ -192,6 +192,18 @@ public:
 		return IncludeCollectionNames;
 	}
 
+	/** Sets whether the tooltip should be expanded by default */
+	void SetAlwaysExpandTooltips(bool bIsExpanded)
+	{
+		AlwaysExpandTooltips = bIsExpanded;
+	}
+
+	/** Gets whether the tooltip should be expanded by default */
+	bool GetAlwaysExpandTooltips() const
+	{
+		return AlwaysExpandTooltips;
+	}
+
 	/**
 	 * Returns an event delegate that is executed when a setting has changed.
 	 *
@@ -238,6 +250,10 @@ private:
 	UPROPERTY(config)
 	bool IncludeCollectionNames;
 
+	/** Whether to display the full tooltip on Asset and Folder or the small version with the Ctrl + Alt */
+	UPROPERTY(config)
+	bool AlwaysExpandTooltips;
+
 	// Holds an event delegate that is executed when a setting has changed.
 	static UNREALED_API FSettingChangedEvent SettingChangedEvent;
 };
diff --git a/Engine/Source/Editor/UnrealEd/Private/AssetThumbnail.cpp b/Engine/Source/Editor/UnrealEd/Private/AssetThumbnail.cpp
--- a/Engine/Source/Editor/UnrealEd/Private/AssetThumbnail.cpp
+++ b/Engine/Source/Editor/UnrealEd/Private/AssetThumbnail.cpp
@@ -96,6 +96,7 @@ public:
 		, _AllowAssetStatusThumbnailOverlay(false)
 		, _ShowAssetChip(false)
 		, _AssetChipBorderImageOverride()
+		, _AlwaysExpandTooltip(false)
 		, _ColorStripOrientation(EThumbnailColorStripOrientation::HorizontalBottomEdge)
 		{}
 
@@ -123,6 +124,7 @@ public:
 		SLATE_ATTRIBUTE(const FSlateBrush*, AssetChipBorderImageOverride)
 		SLATE_ARGUMENT(bool, CanDisplayEditModePrimitiveTools)
 		SLATE_ATTRIBUTE(EVisibility, IsEditModeVisible)
+		SLATE_ATTRIBUTE(bool, AlwaysExpandTooltip)
 		SLATE_ARGUMENT(EThumbnailColorStripOrientation, ColorStripOrientation)
 
 	SLATE_END_ARGS()
@@ -148,6 +150,7 @@ public:
 		bShowAssetChip = InArgs._ShowAssetChip;
 		AssetChipBorderImageOverride = InArgs._AssetChipBorderImageOverride;
 		EditModeVisibility = InArgs._IsEditModeVisible;
+		AlwaysExpandTooltip = InArgs._AlwaysExpandTooltip;
 		AssetThumbnail->OnAssetDataChanged().AddSP(this, &SAssetThumbnail::OnAssetDataChanged);
 
 		const FAssetData& AssetData = AssetThumbnail->GetAssetData();
@@ -182,7 +185,9 @@ public:
 			else
 			{
 				// Set our tooltip - this will refresh each time it's opened to make sure it's up-to-date
-				SetToolTip(SNew(SAssetThumbnailToolTip).AssetThumbnail(SharedThis(this)));
+				SetToolTip(SNew(SAssetThumbnailToolTip)
+					.AssetThumbnail(SharedThis(this))
+					.AlwaysExpandTooltip(AlwaysExpandTooltip));
 			}
 		}
 
@@ -861,6 +866,7 @@ public:
 		return SNew(SDocumentationToolTip)
 				.OverrideExtendedToolTipContent(ExtendedToolTip)
 				.OverridePromptContent(PromptBox)
+				.AlwaysExpandTooltip(AlwaysExpandTooltip)
 				[
 					OverallTooltipVBox
 				];
@@ -928,6 +934,7 @@ private:
 				SNew(STextBlock)
 				.ColorAndOpacity(FStyleColors::White)
 				.Text(Value)
+				.WrapTextAt(700.0f)
 			]
 		];
 	}
@@ -1485,6 +1492,7 @@ private:
 	bool bShowAssetChip;
 	TAttribute<const FSlateBrush*> AssetChipBorderImageOverride;
 	TAttribute<EVisibility> EditModeVisibility;
+	TAttribute<bool> AlwaysExpandTooltip;
 
 	bool bAllowHintText;
 	bool bAllowRealTimeOnHovered;
@@ -1545,7 +1553,11 @@ bool SAssetThumbnailToolTip::IsInteractive() const
 	if (UE::Editor::ContentBrowser::IsNewStyleEnabled())
 	{
 		TSharedPtr<SAssetThumbnail> AssetViewItemPin = AssetThumbnail.Pin();
-		if (!IsEmpty() && AssetViewItemPin.IsValid() && AssetToolTip.IsValid())
+
+		// Use the SDocumentationTooltip IsInteractive only if we have something to show and if it should not be expanded by default.
+		// This way we can avoid our tooltip to be kept in place when hovering.
+		const bool bShouldAlwaysBeExpanded = AlwaysExpandTooltip.IsSet() && AlwaysExpandTooltip.Get();
+		if (!IsEmpty() && AssetViewItemPin.IsValid() && AssetToolTip.IsValid() && !bShouldAlwaysBeExpanded)
 		{
 			return AssetToolTip->IsInteractive();
 		}
@@ -1681,6 +1693,7 @@ TSharedRef<SWidget> FAssetThumbnail::MakeThumbnailWidget( const FAssetThumbnailC
 			.IsEditModeVisible(InConfig.IsEditModeVisible)
 			.SmallAssetChipBorder(InConfig.bSmallAssetChipBorder)
 			.AssetChipBorderImageOverride(InConfig.AssetChipBorderImageOverride)
+			.AlwaysExpandTooltip(InConfig.AlwaysExpandTooltip)
 			.ColorStripOrientation(InConfig.ColorStripOrientation);
 	}
 	else
diff --git a/Engine/Source/Editor/UnrealEd/Private/AssetThumbnailToolTip.h b/Engine/Source/Editor/UnrealEd/Private/AssetThumbnailToolTip.h
--- a/Engine/Source/Editor/UnrealEd/Private/AssetThumbnailToolTip.h
+++ b/Engine/Source/Editor/UnrealEd/Private/AssetThumbnailToolTip.h
@@ -12,9 +12,11 @@ class SAssetThumbnailToolTip : public SToolTip
 public:
 	SLATE_BEGIN_ARGS(SAssetThumbnailToolTip)
 		: _AssetThumbnail()
+		, _AlwaysExpandTooltip(false)
 	{ }
 
-	SLATE_ARGUMENT(TSharedPtr<SAssetThumbnail>, AssetThumbnail)
+		SLATE_ARGUMENT(TSharedPtr<SAssetThumbnail>, AssetThumbnail)
+		SLATE_ATTRIBUTE(bool, AlwaysExpandTooltip)
 
 	SLATE_END_ARGS()
 
@@ -30,5 +32,6 @@ public:
 	virtual bool IsInteractive() const override;
 private:
 	TWeakPtr<SAssetThumbnail> AssetThumbnail;
+	TAttribute<bool> AlwaysExpandTooltip;
 	TSharedPtr<SDocumentationToolTip> AssetToolTip;
 };
diff --git a/Engine/Source/Editor/UnrealEd/Public/AssetThumbnail.h b/Engine/Source/Editor/UnrealEd/Public/AssetThumbnail.h
--- a/Engine/Source/Editor/UnrealEd/Public/AssetThumbnail.h
+++ b/Engine/Source/Editor/UnrealEd/Public/AssetThumbnail.h
@@ -74,6 +74,7 @@ struct FAssetThumbnailConfig
 		, AssetChipBorderImageOverride()
 		, bCanDisplayEditModePrimitiveTools(false)
 		, IsEditModeVisible()
+		, AlwaysExpandTooltip(false)
 	{
 	}
 
@@ -100,6 +101,7 @@ struct FAssetThumbnailConfig
 	TAttribute<const FSlateBrush*> AssetChipBorderImageOverride;
 	bool bCanDisplayEditModePrimitiveTools;
 	TAttribute<EVisibility> IsEditModeVisible;
+	TAttribute<bool> AlwaysExpandTooltip;
 };
 
 
