[VSM] Allow screenpercentage to scale in editor while VSM visualization modes are active. Previously, activating these visualization modes would lock the screenpercentage to 100%, which would skew the data.

diff --git a/Engine/Source/Editor/UnrealEd/Private/EditorViewportClient.cpp b/Engine/Source/Editor/UnrealEd/Private/EditorViewportClient.cpp
--- a/Engine/Source/Editor/UnrealEd/Private/EditorViewportClient.cpp
+++ b/Engine/Source/Editor/UnrealEd/Private/EditorViewportClient.cpp
@@ -3038,8 +3038,7 @@ bool FEditorViewportClient::SupportsPreviewResolutionFraction() const
 	// Don't do preview screen percentage in certain cases.
 	if (EngineShowFlags.VisualizeBuffer 
 		|| EngineShowFlags.MeshEdges
-		|| EngineShowFlags.VisualizeNanite 
-		|| EngineShowFlags.VisualizeVirtualShadowMap 
+		|| EngineShowFlags.VisualizeNanite
 		|| IsVisualizeCalibrationMaterialEnabled())
 	{
 		return false;
