AssetGuide: Make the leading slash optional.

Addresses issues where the path provided doesn't have a leading slash, e.g.:
UnrealEditor-UnrealEd.dylib`UAssetGuideline::TryDetectIniBranchFromFilename(Filename=string=u"Config/DefaultEngine.ini", OutIniBranch=string=Empty)

diff --git a/Engine/Source/Editor/UnrealEd/Private/AssetGuideline.cpp b/Engine/Source/Editor/UnrealEd/Private/AssetGuideline.cpp
--- a/Engine/Source/Editor/UnrealEd/Private/AssetGuideline.cpp
+++ b/Engine/Source/Editor/UnrealEd/Private/AssetGuideline.cpp
@@ -271,7 +271,7 @@ void UAssetGuideline::BeginDestroy()
 
 bool UAssetGuideline::TryDetectIniBranchFromFilename(const FString& Filename, FString& OutIniBranch)
 {
-	const FRegexPattern DefaultPathPattern(TEXT("^/Config/Default(\\w+)\\.ini$"), ERegexPatternFlags::CaseInsensitive);
+	const FRegexPattern DefaultPathPattern(TEXT("^/?Config/Default(\\w+)\\.ini$"), ERegexPatternFlags::CaseInsensitive);
 	FRegexMatcher Matcher(DefaultPathPattern, Filename);
 	if (Matcher.FindNext())
 	{
