Fix confusing variable names on these console commands

#rnx

diff --git a/Engine/Source/Editor/UnrealEd/Private/DiffUtils.cpp b/Engine/Source/Editor/UnrealEd/Private/DiffUtils.cpp
--- a/Engine/Source/Editor/UnrealEd/Private/DiffUtils.cpp
+++ b/Engine/Source/Editor/UnrealEd/Private/DiffUtils.cpp
@@ -76,14 +76,14 @@ namespace UEDiffUtils_Private
 	static void RunDiffCommand(const TArray<FString>& Args);
 	static void RunMergeCommand(const TArray<FString>& Args);
 
-	FAutoConsoleCommand DiffConsoleCommand(
+	FAutoConsoleCommand MergeConsoleCommand(
 		TEXT("merge"),
 		*FString::Format(TEXT("Either merge three assets or a single conflicted asset.\n{0}"), {MergeSyntaxHelp}),
 		FConsoleCommandWithArgsDelegate::CreateStatic(&RunMergeCommand),
 		ECVF_Default
 	);
 
-	FAutoConsoleCommand MergeConsoleCommand(
+	FAutoConsoleCommand DiffConsoleCommand(
 		TEXT("diff"),
 		*FString::Format(TEXT("diff two assets against one another.\n{0}"), {DiffSyntaxHelp}),
 		FConsoleCommandWithArgsDelegate::CreateStatic(&RunDiffCommand),
