Added support for showing Verse paths to the recent asset menu.

diff --git a/Engine/Source/Editor/UnrealEd/Private/Subsystems/AssetEditorSubsystem.cpp b/Engine/Source/Editor/UnrealEd/Private/Subsystems/AssetEditorSubsystem.cpp
--- a/Engine/Source/Editor/UnrealEd/Private/Subsystems/AssetEditorSubsystem.cpp
+++ b/Engine/Source/Editor/UnrealEd/Private/Subsystems/AssetEditorSubsystem.cpp
@@ -2,12 +2,14 @@
 
 #include "Subsystems/AssetEditorSubsystem.h"
 #include "AssetEditorMessages.h"
+#include "Containers/VersePath.h"
 #include "MessageEndpoint.h"
 #include "EngineAnalytics.h"
 #include "AnalyticsEventAttribute.h"
 #include "UObject/Package.h"
 #include "CoreGlobals.h"
 #include "AssetToolsModule.h"
+#include "IAssetTools.h"
 #include "IMessageContext.h"
 #include "LevelEditor.h"
 #include "Toolkits/AssetEditorToolkit.h"
@@ -199,7 +201,7 @@ void UAssetEditorSubsystem::CullRecentAssetEditorsMap()
 
 void UAssetEditorSubsystem::OnAssetRemoved(const FAssetData& AssetData)
 {
-	FString PathName = FPaths::GetBaseFilename(AssetData.GetObjectPathString(), false);
+	FString PathName = AssetData.PackageName.ToString();
 
 	// We need this early exit because FindMRUItemIdx has a check() for non valid long package names
 	if (!FPackageName::IsValidLongPackageName(PathName))
@@ -222,8 +224,8 @@ void UAssetEditorSubsystem::OnAssetRemoved(const FAssetData& AssetData)
 
 void UAssetEditorSubsystem::OnAssetRenamed(const FAssetData& AssetData, const FString& AssetOldName)
 {
-	FString OldPathName = FPaths::GetBaseFilename(AssetOldName, false);
-	FString NewPathName = FPaths::GetBaseFilename(AssetData.GetObjectPathString(), false);
+	FString OldPathName = FSoftObjectPath(AssetOldName).GetLongPackageName();
+	FString NewPathName = AssetData.PackageName.ToString();
 
 	// We need this early exit because FindMRUItemIdx has a check() for non valid long package names
 	if (!FPackageName::IsValidLongPackageName(OldPathName) || !FPackageName::IsValidLongPackageName(NewPathName))
@@ -1241,6 +1243,8 @@ void UAssetEditorSubsystem::CreateRecentAssetsMenu(UToolMenu* InMenu, const FNam
 {
 	FToolMenuSection& Section = InMenu->FindOrAddSection("Recents");
 	
+	const bool bShowingContentVersePath = IAssetTools::Get().ShowingContentVersePath();
+
 	// Keep adding assets until we reach the end of the MRU list, or we reach the max allowed assets
 	for ( int32 CurRecentIndex = 0; CurRecentIndex < RecentAssetsList->GetNumItems() && CurRecentIndex < MaxRecentAssetsToShowInMenu; ++CurRecentIndex )
 	{
@@ -1251,8 +1255,38 @@ void UAssetEditorSubsystem::CreateRecentAssetsMenu(UToolMenu* InMenu, const FNam
 			continue;
 		}
 
-		const FText ToolTip = FText::Format( LOCTEXT( "RecentAssetsToolTip", "Open {0}" ), FText::FromString( CurRecent ) );
-		const FText Label = FText::FromString( FPaths::GetBaseFilename(CurRecent) );
+		FSoftObjectPath AssetPath(CurRecent);
+
+		if (AssetPath.GetAssetFName().IsNone())
+		{
+			// Mimic LoadObject in FindOrLoadAssetForOpening falling back to the asset with the same name as the package if the asset name is not provided.
+			AssetPath = FSoftObjectPath::ConstructFromAssetPath(FTopLevelAssetPath(AssetPath.GetLongPackageFName(), FPackageName::GetShortFName(AssetPath.GetLongPackageFName())));
+		}
+
+		FText CurPath;
+		if (bShowingContentVersePath)
+		{
+			if (const IAssetRegistry* AssetRegistry = IAssetRegistry::Get())
+			{
+				FAssetData AssetData;
+				if (AssetRegistry->TryGetAssetByObjectPath(AssetPath, AssetData) == UE::AssetRegistry::EExists::Exists)
+				{
+					UE::Core::FVersePath VersePath = AssetData.GetVersePath();
+					if (VersePath.IsValid())
+					{
+						CurPath = FText::FromString(MoveTemp(VersePath).ToString());
+					}
+				}
+			}
+		}
+
+		if (CurPath.IsEmpty())
+		{
+			CurPath = FText::FromString(CurRecent);
+		}
+
+		const FText ToolTip = FText::Format( LOCTEXT( "RecentAssetsToolTip", "Open {0}" ), CurPath );
+		const FText Label = FText::FromString( AssetPath.GetAssetName() );
 
 		Section.AddMenuEntry(
 		NAME_None,
@@ -1260,9 +1294,9 @@ void UAssetEditorSubsystem::CreateRecentAssetsMenu(UToolMenu* InMenu, const FNam
 		ToolTip,
 		FSlateIcon(),
 		FUIAction(
-			FExecuteAction::CreateLambda([this, CurRecent]()
+			FExecuteAction::CreateLambda([this, AssetPath]()
 				{
-					if (UObject* Object = FindOrLoadAssetForOpening(CurRecent))
+					if (UObject* Object = FindOrLoadAssetForOpening(AssetPath))
 					{
 						FText ErrorMessage;
 
