[AssetThumbnail] Fix crash happening on startup in the older CB style

diff --git a/Engine/Source/Editor/UnrealEd/Private/AssetThumbnail.cpp b/Engine/Source/Editor/UnrealEd/Private/AssetThumbnail.cpp
--- a/Engine/Source/Editor/UnrealEd/Private/AssetThumbnail.cpp
+++ b/Engine/Source/Editor/UnrealEd/Private/AssetThumbnail.cpp
@@ -530,9 +530,9 @@ public:
 				PRAGMA_DISABLE_DEPRECATION_WARNINGS
 				TSharedPtr<SWidget> AssetSpecificThumbnailOverlay = AssetTypeActions->GetThumbnailOverlay(AssetData);
 				PRAGMA_ENABLE_DEPRECATION_WARNINGS
-				AssetSpecificThumbnailOverlay->SetVisibility(AllowAssetSpecificThumbnailOverlay.Get() ? EVisibility::Visible : EVisibility::Collapsed);
-				if( AssetSpecificThumbnailOverlay.IsValid() )
+				if (AssetSpecificThumbnailOverlay.IsValid())
 				{
+					AssetSpecificThumbnailOverlay->SetVisibility(AllowAssetSpecificThumbnailOverlay.Get() ? EVisibility::Visible : EVisibility::Collapsed);
 					OverlayWidget->AddSlot()
 					[
 						AssetSpecificThumbnailOverlay.ToSharedRef()
