TextureAssetActions : ResizeTextureSourceDataToNearestPowerOfTwo : if compression was set to UI change to Default

diff --git a/Engine/Source/Editor/UnrealEd/Private/TextureAssetActions.cpp b/Engine/Source/Editor/UnrealEd/Private/TextureAssetActions.cpp
--- a/Engine/Source/Editor/UnrealEd/Private/TextureAssetActions.cpp
+++ b/Engine/Source/Editor/UnrealEd/Private/TextureAssetActions.cpp
@@ -689,11 +689,17 @@ static void DoResizeTextureSourceToPowerOfTwo(UTexture * Texture)
 		Texture->MipGenSettings = TMGS_FromTextureGroup;
 	}
 	Texture->NeverStream = false;
+	
+	// it it was set to UI (uncompressed) change to default :
+	if ( Texture->CompressionSettings == TC_EditorIcon )
+	{
+		Texture->CompressionSettings = TC_Default;
+	}
 
 	// this counts as a reimport :
 	UE::TextureUtilitiesCommon::ApplyDefaultsForNewlyImportedTextures(Texture,true);
 
-	// DownsizeTextureSourceData did the PreEditChange
+	// ResizeTextureSourceDataToNearestPowerOfTwo did the PreEditChange
 	Texture->PostEditChange();
 }
 
