FGenerationHelper::FinishGeneratorPlatformSave : Filter out editor packages from the runtime dependencies. Only keep the game packages.
This fixes an issue with incremental cook cooking editor packages. The first incremental cook cooks a map and adds the editor packages to its runtime dependencies. The second cook will download the dependencies of the map, and try to cook every runtime dependencies. Since the editor packages were not cooked the first time, they are cooked the second time.

#rnx
#incrementalcook

diff --git a/Engine/Source/Editor/UnrealEd/Private/Cooker/CookGenerationHelper.cpp b/Engine/Source/Editor/UnrealEd/Private/Cooker/CookGenerationHelper.cpp
--- a/Engine/Source/Editor/UnrealEd/Private/Cooker/CookGenerationHelper.cpp
+++ b/Engine/Source/Editor/UnrealEd/Private/Cooker/CookGenerationHelper.cpp
@@ -1610,7 +1610,8 @@ void FGenerationHelper::FinishGeneratorPlatformSave(FPackageData& PackageData, b
 	// Set dependencies equal to the global AssetRegistry dependencies plus a dependency on
 	// each generated package.
 	COTFS.AssetRegistry->GetDependencies(PackageData.GetPackageName(), OutPackageDependencies,
-		UE::AssetRegistry::EDependencyCategory::Package);
+		UE::AssetRegistry::EDependencyCategory::Package, UE::AssetRegistry::EDependencyQuery::Game);
+
 	OutPackageDependencies.Reserve(PackagesToGenerate.Num());
 	for (FCookGenerationInfo& GeneratedInfo : GetPackagesToGenerate())
 	{
