Fix UGatherTextFromAssetsCommandlet to support non-partitioned worlds using external actors

diff --git a/Engine/Source/Editor/UnrealEd/Classes/Commandlets/GatherTextFromAssetsCommandlet.h b/Engine/Source/Editor/UnrealEd/Classes/Commandlets/GatherTextFromAssetsCommandlet.h
--- a/Engine/Source/Editor/UnrealEd/Classes/Commandlets/GatherTextFromAssetsCommandlet.h
+++ b/Engine/Source/Editor/UnrealEd/Classes/Commandlets/GatherTextFromAssetsCommandlet.h
@@ -65,7 +65,7 @@ private:
 	void FilterAssetsBasedOnIncludeExcludePaths(TArray<FAssetData>& InOutAssetDataArray) const;
 	
 	bool DiscoverExternalActors(TArray<FAssetData>& InOutAssetDataArray);
-	void RemoveExistingExternalActors(TArray<FAssetData>& InOutAssetDataArray, const TSet<FName>* WorldPackageFilter, TSet<FName>& OutPartitionedWorldPackageNames, TSet<FName>& OutGameFeatureDataPackageNames) const;
+	void RemoveExistingExternalActors(TArray<FAssetData>& InOutAssetDataArray, const TSet<FName>* WorldPackageFilter, TSet<FName>& OutExternalActorsWorldPackageNames, TSet<FName>& OutGameFeatureDataPackageNames) const;
 
 	TSet<FName> GetPackageNamesToGather(const TArray<FAssetData>& InAssetDataArray) const;
 	void PopulatePackagesPendingGather(TSet<FName> PackageNamesToGather);
diff --git a/Engine/Source/Editor/UnrealEd/Private/Commandlets/GatherTextFromAssetsCommandlet.cpp b/Engine/Source/Editor/UnrealEd/Private/Commandlets/GatherTextFromAssetsCommandlet.cpp
--- a/Engine/Source/Editor/UnrealEd/Private/Commandlets/GatherTextFromAssetsCommandlet.cpp
+++ b/Engine/Source/Editor/UnrealEd/Private/Commandlets/GatherTextFromAssetsCommandlet.cpp
@@ -778,8 +778,8 @@ void UGatherTextFromAssetsCommandlet::FilterAssetsBasedOnIncludeExcludePaths(TAr
 	}
 }
 
-/** Remove any external actors that currently exist in InOutAssetDataArray. OutPartitionedWorldPackageNames is populated with the package paths of partitioned worlds.*/
-void UGatherTextFromAssetsCommandlet::RemoveExistingExternalActors(TArray<FAssetData>& InOutAssetDataArray, const TSet<FName>* WorldPackageFilter, TSet<FName>& OutPartitionedWorldPackageNames, TSet<FName>& OutGameFeatureDataPackageNames) const
+/** Remove any external actors that currently exist in InOutAssetDataArray. OutExternalActorsWorldPackageNames is populated with the package paths of worlds using external actors.*/
+void UGatherTextFromAssetsCommandlet::RemoveExistingExternalActors(TArray<FAssetData>& InOutAssetDataArray, const TSet<FName>* WorldPackageFilter, TSet<FName>& OutExternalActorsWorldPackageNames, TSet<FName>& OutGameFeatureDataPackageNames) const
 {
 	auto IsGameFeatureData = [](const UClass* AssetClass)
 	{
@@ -794,7 +794,7 @@ void UGatherTextFromAssetsCommandlet::RemoveExistingExternalActors(TArray<FAsset
 		return false;
 	};
 
-	InOutAssetDataArray.RemoveAll([WorldPackageFilter, &OutPartitionedWorldPackageNames, &OutGameFeatureDataPackageNames, &IsGameFeatureData](const FAssetData& AssetData)
+	InOutAssetDataArray.RemoveAll([WorldPackageFilter, &OutExternalActorsWorldPackageNames, &OutGameFeatureDataPackageNames, &IsGameFeatureData](const FAssetData& AssetData)
 		{
 			const FNameBuilder PackageNameStr(AssetData.PackageName);
 
@@ -804,9 +804,9 @@ void UGatherTextFromAssetsCommandlet::RemoveExistingExternalActors(TArray<FAsset
 				{
 					return true;
 				}
-				else if (ULevel::GetIsLevelPartitionedFromAsset(AssetData))
+				else if (ULevel::GetIsLevelUsingExternalActorsFromAsset(AssetData))
 				{
-					OutPartitionedWorldPackageNames.Add(AssetData.PackageName);
+					OutExternalActorsWorldPackageNames.Add(AssetData.PackageName);
 				}
 			}
 			else if (PackageNameStr.ToView().Contains(FPackagePath::GetExternalActorsFolderName()))
@@ -854,18 +854,18 @@ bool UGatherTextFromAssetsCommandlet::DiscoverExternalActors(TArray<FAssetData>&
 		WorldPackageFilter.Emplace(MoveTemp(FilteredWorldPackages));
 	}
 
-	TSet<FName> PartitionedWorldPackageNames;
+	TSet<FName> ExternalActorsWorldPackageNames;
 	TSet<FName> GameFeatureDataPackageNames;
-	RemoveExistingExternalActors(InOutAssetDataArray, WorldPackageFilter.GetPtrOrNull(), PartitionedWorldPackageNames, GameFeatureDataPackageNames);
+	RemoveExistingExternalActors(InOutAssetDataArray, WorldPackageFilter.GetPtrOrNull(), ExternalActorsWorldPackageNames, GameFeatureDataPackageNames);
 
 	// Append the actors that are directly known by each world (by looking for their external actors under the expected path)
-	if (PartitionedWorldPackageNames.Num() > 0)
+	if (ExternalActorsWorldPackageNames.Num() > 0)
 	{
 		// Note: This doesn't add AActor to ClassPaths as that doesn't work correctly doing a partial asset scan (see bSearchAllAssets)
 		FARFilter Filter;
-		for (const FName PartitionedWorldPackageName : PartitionedWorldPackageNames)
+		for (const FName ExternalActorsWorldPackageName : ExternalActorsWorldPackageNames)
 		{
-			FNameBuilder PackageNameStr(PartitionedWorldPackageName);
+			FNameBuilder PackageNameStr(ExternalActorsWorldPackageName);
 			Filter.PackagePaths.Add(*ULevel::GetExternalActorsPath(*PackageNameStr));
 		}
 		Filter.bRecursivePaths = true;
@@ -906,10 +906,10 @@ bool UGatherTextFromAssetsCommandlet::DiscoverExternalActors(TArray<FAssetData>&
 
 		// External actors may be filtered in two ways;
 		//  1. If WorldCollectionFilters were provided, then we only include actors related to those worlds
-		//  2. If PartitionedWorldPackageNames was populated (meaning there are worlds in this gather), then we only include actors related to those worlds
+		//  2. If ExternalActorsWorldPackageNames was populated (meaning there are worlds in this gather), then we only include actors related to those worlds
 		// If neither of the above is true then we include all actors related to the GFDs in this gather, as we assume this is a plugin hosting external actors for worlds owned by another localization target
-		const bool bHasWorldFilter = WorldPackageFilter.IsSet() || PartitionedWorldPackageNames.Num() > 0;
-		const TSet<FName>& GameFeatureDataActorsWorldPackageFilter = WorldPackageFilter.Get(PartitionedWorldPackageNames);
+		const bool bHasWorldFilter = WorldPackageFilter.IsSet() || ExternalActorsWorldPackageNames.Num() > 0;
+		const TSet<FName>& GameFeatureDataActorsWorldPackageFilter = WorldPackageFilter.Get(ExternalActorsWorldPackageNames);
 		for (FAssetData& GameFeatureDataDependency : GameFeatureDataDependencies)
 		{
 			const FName OptionalOuterPathName = GameFeatureDataDependency.GetOptionalOuterPathName();
