Cooker: Second attempt to fix soft lock in IsCachedCookedPlatformData calls. The previous change failed to fix it because even FWeakObjectPtr is not nulled out when objects are reallocated. This change instead detects the reallocation by listening to NotifyUObjectDeleted.
#rnx

diff --git a/Engine/Source/Editor/UnrealEd/Private/Cooker/CookGarbageCollect.cpp b/Engine/Source/Editor/UnrealEd/Private/Cooker/CookGarbageCollect.cpp
--- a/Engine/Source/Editor/UnrealEd/Private/Cooker/CookGarbageCollect.cpp
+++ b/Engine/Source/Editor/UnrealEd/Private/Cooker/CookGarbageCollect.cpp
@@ -619,6 +619,7 @@ void UCookOnTheFlyServer::PreGarbageCollect()
 {
 	using namespace UE::Cook;
 
+	PackageTracker->SetCollectingGarbage(true);
 	if (!IsInSession())
 	{
 		return;
@@ -816,6 +817,8 @@ void UCookOnTheFlyServer::PostGarbageCollect()
 {
 	using namespace UE::Cook;
 
+	PackageTracker->SetCollectingGarbage(false);
+
 	NumObjectsHistory.AddInstance(GUObjectArray.GetObjectArrayNumMinusAvailable());
 	VirtualMemoryHistory.AddInstance(FPlatformMemory::GetStats().UsedVirtual);
 
diff --git a/Engine/Source/Editor/UnrealEd/Private/Cooker/CookPackageData.cpp b/Engine/Source/Editor/UnrealEd/Private/Cooker/CookPackageData.cpp
--- a/Engine/Source/Editor/UnrealEd/Private/Cooker/CookPackageData.cpp
+++ b/Engine/Source/Editor/UnrealEd/Private/Cooker/CookPackageData.cpp
@@ -3224,6 +3224,11 @@ void FPackageDatas::CachedCookedPlatformDataObjectsPostGarbageCollect(
 	}
 }
 
+void FPackageDatas::CachedCookedPlatformDataObjectsOnDestroyedOutsideOfGC(const UObject* DestroyedObject)
+{
+	CachedCookedPlatformDataObjects.Remove(DestroyedObject);
+}
+
 void FCachedCookedPlatformDataState::AddRefFrom(FPackageData* PackageData)
 {
 	// Most objects will only be referenced by a single package.
diff --git a/Engine/Source/Editor/UnrealEd/Private/Cooker/CookPackageData.h b/Engine/Source/Editor/UnrealEd/Private/Cooker/CookPackageData.h
--- a/Engine/Source/Editor/UnrealEd/Private/Cooker/CookPackageData.h
+++ b/Engine/Source/Editor/UnrealEd/Private/Cooker/CookPackageData.h
@@ -1432,6 +1432,7 @@ public:
 
 	FMapOfCachedCookedPlatformDataState& GetCachedCookedPlatformDataObjects();
 	void CachedCookedPlatformDataObjectsPostGarbageCollect(const TSet<UObject*>& SaveQueueObjectsThatStillExist);
+	void CachedCookedPlatformDataObjectsOnDestroyedOutsideOfGC(const UObject* DestroyedObject);
 
 private:
 	/**
diff --git a/Engine/Source/Editor/UnrealEd/Private/Cooker/PackageTracker.cpp b/Engine/Source/Editor/UnrealEd/Private/Cooker/PackageTracker.cpp
--- a/Engine/Source/Editor/UnrealEd/Private/Cooker/PackageTracker.cpp
+++ b/Engine/Source/Editor/UnrealEd/Private/Cooker/PackageTracker.cpp
@@ -218,15 +218,20 @@ void FPackageTracker::OnCreatePackage(FName PackageName)
 	PackageStream.Add({ PackageName, MoveTemp(Instigator), EPackageStreamEvent::PackageLoad });
 }
 
-void FPackageTracker::NotifyUObjectDeleted(const class UObjectBase* Object, int32 Index)
+void FPackageTracker::NotifyUObjectDeleted(const class UObjectBase* ObjectBase, int32 Index)
 {
-	if (Object->GetClass() == UPackage::StaticClass())
+	if (ObjectBase->GetClass() == UPackage::StaticClass())
 	{
-		UPackage* Package = const_cast<UPackage*>(static_cast<const UPackage*>(Object));
+		UPackage* Package = const_cast<UPackage*>(static_cast<const UPackage*>(ObjectBase));
 
 		FWriteScopeLock ScopeLock(Lock);
 		LoadedPackages.Remove(Package);
 	}
+	if (!bCollectingGarbage)
+	{
+		const UObject* Object = static_cast<const UObject*>(ObjectBase);
+		COTFS.PackageDatas->CachedCookedPlatformDataObjectsOnDestroyedOutsideOfGC(Object);
+	}
 }
 
 void FPackageTracker::OnUObjectArrayShutdown()
diff --git a/Engine/Source/Editor/UnrealEd/Private/Cooker/PackageTracker.h b/Engine/Source/Editor/UnrealEd/Private/Cooker/PackageTracker.h
--- a/Engine/Source/Editor/UnrealEd/Private/Cooker/PackageTracker.h
+++ b/Engine/Source/Editor/UnrealEd/Private/Cooker/PackageTracker.h
@@ -307,6 +307,11 @@ public:
 			+ ActiveInstances->Map.Num()*sizeof(FPackageStreamInstancedPackage);
 	}
 
+	void SetCollectingGarbage(bool bInCollectingGarbage)
+	{
+		bCollectingGarbage = bInCollectingGarbage;
+	}
+
 	/**
 	 * When package A loads B loads C, and we don't want to tell the cooker about B (because it is e.g. an
 	 * insetanced package), we need to calculate the reason that A loaded C by merging the reason A loaded B with
@@ -355,6 +360,11 @@ private:
 	bool bSubscribed = false;
 
 	// End Data ReadOnlyWhileSubscribed
+
+	// Begin Data ReadWrite from scheduler thread only
+	bool bCollectingGarbage = false;
+
+	// End Data ReadWrite from scheduler thread only
 };
 
 } // namespace UE::Cook
