Ignore scalability resolution scale if currently 0 when checking if the current values match default values.

#rnx

diff --git a/Engine/Source/Editor/UnrealEd/Private/SScalabilitySettings.cpp b/Engine/Source/Editor/UnrealEd/Private/SScalabilitySettings.cpp
--- a/Engine/Source/Editor/UnrealEd/Private/SScalabilitySettings.cpp
+++ b/Engine/Source/Editor/UnrealEd/Private/SScalabilitySettings.cpp
@@ -231,6 +231,14 @@ ECheckBoxState SScalabilitySettings::GetHeaderChecked(int32 InQualityLevel, Scal
 		break;
 	}
 	
+	if (FMath::IsNearlyZero(CurrentLevels.ResolutionQuality))
+	{
+		// Resolution quality gets initialized to 0.0 in a fresh project, thereby using project defaults.
+		// This should not affect whether a column header is highlighted.
+		CurrentLevels.ResolutionQuality = 0.f;
+		TargetLevels.ResolutionQuality = 0.f;
+	}
+	
 	return CurrentLevels == TargetLevels ? ECheckBoxState::Checked : ECheckBoxState::Unchecked;
 }
 
@@ -259,7 +267,18 @@ void SScalabilitySettings::OnDefaultClicked(ECheckBoxState CheckState)
 
 ECheckBoxState SScalabilitySettings::GetDefaultCheckState() const
 {
-	return Scalability::GetQualityLevels() == Scalability::GetDefaultQualityLevels() ? ECheckBoxState::Checked : ECheckBoxState::Unchecked;
+	Scalability::FQualityLevels CurrentLevels = Scalability::GetQualityLevels();
+	Scalability::FQualityLevels TargetLevels = Scalability::GetDefaultQualityLevels();
+
+	if (FMath::IsNearlyZero(CurrentLevels.ResolutionQuality))
+	{
+		// Resolution quality gets initialized to 0.0 in a fresh project, thereby using project defaults.
+		// This should still show up as "defaulted"
+		CurrentLevels.ResolutionQuality = 0.f;
+		TargetLevels.ResolutionQuality = 0.f;
+	}
+
+	return CurrentLevels == TargetLevels ? ECheckBoxState::Checked : ECheckBoxState::Unchecked;
 }
 
 SGridPanel::FSlot::FSlotArguments SScalabilitySettings::MakeGridSlot(int32 InCol, int32 InRow, int32 InColSpan /*= 1*/, int32 InRowSpan /*= 1*/)
diff --git a/Engine/Source/Editor/UnrealEd/Private/ViewportToolbar/UnrealEdViewportToolbar.cpp b/Engine/Source/Editor/UnrealEd/Private/ViewportToolbar/UnrealEdViewportToolbar.cpp
--- a/Engine/Source/Editor/UnrealEd/Private/ViewportToolbar/UnrealEdViewportToolbar.cpp
+++ b/Engine/Source/Editor/UnrealEd/Private/ViewportToolbar/UnrealEdViewportToolbar.cpp
@@ -3869,7 +3869,18 @@ bool IsScalabilityWarningVisible()
 		return false;
 	}
 	
-	return Scalability::GetQualityLevels() != Scalability::GetDefaultQualityLevels();
+	Scalability::FQualityLevels Current = Scalability::GetQualityLevels();
+	Scalability::FQualityLevels Default = Scalability::GetDefaultQualityLevels();
+	
+	if (FMath::IsNearlyZero(Current.ResolutionQuality))
+	{
+		// Resolution quality gets initialized to 0.0 in a fresh project, thereby using project defaults.
+		// This should still show up as "defaulted"
+		Current.ResolutionQuality = 0.0f;
+		Default.ResolutionQuality = 0.0f;	
+	}
+	
+	return Current != Default;
 }
 
 FText GetScalabilityWarningLabel()
