[EditorAssetSubsystem CVar] Allow asset saving through scripting behaind a CVar since we are unsure if there are edge cases where it shouldn't be allowed

diff --git a/Engine/Source/Editor/UnrealEd/Private/Subsystems/EditorAssetSubsystem.cpp b/Engine/Source/Editor/UnrealEd/Private/Subsystems/EditorAssetSubsystem.cpp
--- a/Engine/Source/Editor/UnrealEd/Private/Subsystems/EditorAssetSubsystem.cpp
+++ b/Engine/Source/Editor/UnrealEd/Private/Subsystems/EditorAssetSubsystem.cpp
@@ -22,6 +22,18 @@
 
 DEFINE_LOG_CATEGORY_STATIC(LogEditorAssetSubsystem, Log, All);
 
+namespace UE::Editor::Private
+{
+	int32 AllowSavingAssetsDuringPIE = 0;
+}
+
+static FAutoConsoleVariableRef CVarToolMenusViewportToolbars(
+	TEXT("Editor.AllowSavingAssetsDuringPIE"),
+	UE::Editor::Private::AllowSavingAssetsDuringPIE,
+	TEXT("Allow assets to get saved during PIE through the EditorAssetSubsystem. Set to 0 (default) to disable and 1 to enable this option."),
+	ECVF_Default
+);
+
 namespace UE::EditorAssetUtils
 {
 	struct ErrorTag{};
@@ -1277,7 +1289,7 @@ bool UEditorAssetSubsystem::SaveLoadedAsset(UObject* AssetToSave, bool bOnlyIfIs
 {
 	TGuardValue<bool> UnattendedScriptGuard(GIsRunningUnattendedScript, true);
 
-	if (!EditorScriptingHelpers::CheckIfInEditorAndPIE() || !UE::EditorAssetUtils::EnsureAssetsLoaded())
+	if ((UE::Editor::Private::AllowSavingAssetsDuringPIE == 0 && !EditorScriptingHelpers::CheckIfInEditorAndPIE()) || !UE::EditorAssetUtils::EnsureAssetsLoaded())
 	{
 		return false;
 	}
@@ -1299,7 +1311,7 @@ bool UEditorAssetSubsystem::SaveLoadedAssets(const TArray<UObject*>& AssetsToSav
 {
 	TGuardValue<bool> UnattendedScriptGuard(GIsRunningUnattendedScript, true);
 
-	if (!EditorScriptingHelpers::CheckIfInEditorAndPIE() || !UE::EditorAssetUtils::EnsureAssetsLoaded())
+	if ((UE::Editor::Private::AllowSavingAssetsDuringPIE == 0 && !EditorScriptingHelpers::CheckIfInEditorAndPIE()) || !UE::EditorAssetUtils::EnsureAssetsLoaded())
 	{
 		return false;
 	}
@@ -1327,7 +1339,7 @@ bool UEditorAssetSubsystem::SaveAsset(const FString& AssetsToSave, bool bOnlyIfI
 {
 	TGuardValue<bool> UnattendedScriptGuard(GIsRunningUnattendedScript, true);
 
-	if (!EditorScriptingHelpers::CheckIfInEditorAndPIE() || !UE::EditorAssetUtils::EnsureAssetsLoaded())
+	if ((UE::Editor::Private::AllowSavingAssetsDuringPIE == 0 && !EditorScriptingHelpers::CheckIfInEditorAndPIE()) || !UE::EditorAssetUtils::EnsureAssetsLoaded())
 	{
 		return false;
 	}
@@ -1366,7 +1378,7 @@ bool UEditorAssetSubsystem::SaveDirectory(const FString& DirectoryPath, bool bOn
 	using UE::EditorAssetUtils::EPackageEnumerationFilter;
 	TGuardValue<bool> UnattendedScriptGuard(GIsRunningUnattendedScript, true);
 
-	if (!EditorScriptingHelpers::CheckIfInEditorAndPIE() || !UE::EditorAssetUtils::EnsureAssetsLoaded())
+	if ((UE::Editor::Private::AllowSavingAssetsDuringPIE == 0 && !EditorScriptingHelpers::CheckIfInEditorAndPIE()) || !UE::EditorAssetUtils::EnsureAssetsLoaded())
 	{
 		return false;
 	}
