- Change color grad editor stop time widget to use spin box
- This allows for math operations and validation of the input data

diff --git a/Engine/Source/Editor/UnrealEd/Private/SColorGradientEditor.cpp b/Engine/Source/Editor/UnrealEd/Private/SColorGradientEditor.cpp
--- a/Engine/Source/Editor/UnrealEd/Private/SColorGradientEditor.cpp
+++ b/Engine/Source/Editor/UnrealEd/Private/SColorGradientEditor.cpp
@@ -26,7 +26,6 @@
 #include "ScopedTransaction.h"
 #include "Misc/Optional.h"
 
-
 #define LOCTEXT_NAMESPACE "SColorGradientEditor"
 
 namespace FColorGradientEditorGeometry
@@ -147,6 +146,11 @@ FLinearColor FGradientStopMark::GetColor( FCurveOwnerInterface& CurveOwner ) con
 	return CurveOwner.GetLinearColorValue(Time);
 }
 
+float FGradientStopMark::GetTime( FCurveOwnerInterface& CurveOwner ) const
+{
+	return Time;
+}
+
 void FGradientStopMark::SetColor( const FLinearColor& InColor, FCurveOwnerInterface& CurveOwner )
 {
 	TArray<FRichCurveEditInfo> Curves = CurveOwner.GetCurves();
@@ -719,23 +723,18 @@ void SColorGradientEditor::OpenGradientStopContextMenu(const FPointerEvent& Mous
 
 		// Set time
 		{
-			TSharedRef<SEditableTextBox> EditableTextBox = 
-				SNew( SEditableTextBox )
-				.MinDesiredWidth(50.0f)
-				.Text( FText::AsNumber( SelectedStop.Time ) )
-				.OnTextCommitted( this, &SColorGradientEditor::OnSetGradientStopTimeFromPopup )
-				.SelectAllTextWhenFocused( true )
-				.ClearKeyboardFocusOnCommit( false )
-				.SelectAllTextOnCommit( true );
-
+			TSharedRef<SSpinBox<float>> TimeWidget =
+				SNew(SSpinBox<float>)
+				.Value( SelectedStop.GetTime(*CurveOwner) )
+				.OnValueCommitted(this, &SColorGradientEditor::OnTimeValueCommited);
 
 			GradientStopMenu.BeginSection( NAME_None, LOCTEXT("TimeMenuSection", "Time") );
 
-			GradientStopMenu.AddWidget( EditableTextBox, FText::GetEmpty() );
+			GradientStopMenu.AddWidget(TimeWidget, FText::GetEmpty() );
 
 			GradientStopMenu.EndSection();
 
-			WidgetToFocus = EditableTextBox;
+			WidgetToFocus = TimeWidget.ToSharedPtr();
 		}
 
 
@@ -893,15 +892,13 @@ void SColorGradientEditor::OnRemoveSelectedGradientStop()
 	DeleteStop( SelectedStop );
 }
 
-void SColorGradientEditor::OnSetGradientStopTimeFromPopup( const FText& NewText, ETextCommit::Type Type )
+void SColorGradientEditor::OnTimeValueCommited(float NewValue, ETextCommit::Type Type)
 {
-	if( NewText.IsNumeric() )
+	if (NewValue != SelectedStop.GetTime(*CurveOwner) )
 	{
-		float NewTime = FCString::Atof( *NewText.ToString() );
-
-		FScopedTransaction Transaction( LOCTEXT("ChangeGradientStopTime", "Change Gradient Stop Time" ) );
+		FScopedTransaction Transaction(LOCTEXT("ChangeGradientStopTime", "Change Gradient Stop Time"));
 		CurveOwner->ModifyOwner();
-		SelectedStop.SetTime( NewTime, *CurveOwner );
+		SelectedStop.SetTime(NewValue, *CurveOwner);
 		CurveOwner->OnCurveChanged(CurveOwner->GetCurves());
 	}
 }
diff --git a/Engine/Source/Editor/UnrealEd/Public/SColorGradientEditor.h b/Engine/Source/Editor/UnrealEd/Public/SColorGradientEditor.h
--- a/Engine/Source/Editor/UnrealEd/Public/SColorGradientEditor.h
+++ b/Engine/Source/Editor/UnrealEd/Public/SColorGradientEditor.h
@@ -30,8 +30,9 @@ public:
 	UNREALED_API bool IsValidAlphaMark( const TArray<FRichCurveEditInfo>& Curves ) const;
 	UNREALED_API bool IsValidColorMark( const TArray<FRichCurveEditInfo>& Curves ) const;
 	UNREALED_API FLinearColor GetColor( FCurveOwnerInterface& CurveOwner ) const;
+	UNREALED_API float GetTime(FCurveOwnerInterface& CurveOwner) const;
 	UNREALED_API void SetColor( const FLinearColor& InColor, FCurveOwnerInterface& CurveOwner );
-	UNREALED_API void SetTime( float NewTime, FCurveOwnerInterface& CurveOwner );
+	UNREALED_API void SetTime(float NewTime, FCurveOwnerInterface& CurveOwner);
 
 public:
 	float Time;
@@ -140,12 +141,9 @@ private:
 	UNREALED_API void OnRemoveSelectedGradientStop();
 
 	/**
-	 * Called when the gradient stop time is changed from the text entry popup
-	 *
-	 * @param NewText	The new time as text
-	 * @param Type		How the value was committed 
+	 * Called when the gradient stop time is changed by typing into the text box
 	 */
-	UNREALED_API void OnSetGradientStopTimeFromPopup( const FText& NewText, ETextCommit::Type Type );
+	UNREALED_API void OnTimeValueCommited( float NewValue, ETextCommit::Type Type );
 
 	/**
 	 * Draws a single gradient stop
