[EditorAssetSubsystem CVar] Hotfix: Added correct name to the cvar definition

diff --git a/Engine/Source/Editor/UnrealEd/Private/Subsystems/EditorAssetSubsystem.cpp b/Engine/Source/Editor/UnrealEd/Private/Subsystems/EditorAssetSubsystem.cpp
--- a/Engine/Source/Editor/UnrealEd/Private/Subsystems/EditorAssetSubsystem.cpp
+++ b/Engine/Source/Editor/UnrealEd/Private/Subsystems/EditorAssetSubsystem.cpp
@@ -27,7 +27,7 @@ namespace UE::Editor::Private
 	int32 AllowSavingAssetsDuringPIE = 0;
 }
 
-static FAutoConsoleVariableRef CVarToolMenusViewportToolbars(
+static FAutoConsoleVariableRef CVarAllowSavingAssetsDuringPIE(
 	TEXT("Editor.AllowSavingAssetsDuringPIE"),
 	UE::Editor::Private::AllowSavingAssetsDuringPIE,
 	TEXT("Allow assets to get saved during PIE through the EditorAssetSubsystem. Set to 0 (default) to disable and 1 to enable this option."),
