PixelInspector: Fixed aspect-ratio-constrained viewport inspecting.

diff --git a/Engine/Source/Editor/UnrealEd/Private/EditorViewportClient.cpp b/Engine/Source/Editor/UnrealEd/Private/EditorViewportClient.cpp
--- a/Engine/Source/Editor/UnrealEd/Private/EditorViewportClient.cpp
+++ b/Engine/Source/Editor/UnrealEd/Private/EditorViewportClient.cpp
@@ -4310,11 +4310,11 @@ void FEditorViewportClient::SetupViewForRendering(FSceneViewFamily& ViewFamily,
 		// Ready to send a request
 		FSceneInterface *SceneInterface = GetScene();
 
-		FVector2D InspectViewportUV(
-			(InspectViewportPos.X + 0.5f) / float(View.UnscaledViewRect.Width()),
-			(InspectViewportPos.Y + 0.5f) / float(View.UnscaledViewRect.Height()));
+		FVector2f InspectViewportUV = FVector2f(InspectViewportPos.X + 0.5f, InspectViewportPos.Y + 0.5f);
+		InspectViewportUV -= FVector2f(View.UnscaledViewRect.Min); // Offset for aspect ratio constrained viewports
+		InspectViewportUV /= FVector2f(View.UnscaledViewRect.Size());
 
-		PixelInspectorModule.CreatePixelInspectorRequest(InspectViewportUV, View.State->GetViewKey(), SceneInterface, bInGameViewMode, View.State->GetPreExposure());
+		PixelInspectorModule.CreatePixelInspectorRequest(FDeprecateSlateVector2D(InspectViewportUV), View.State->GetViewKey(), SceneInterface, bInGameViewMode, View.State->GetPreExposure());
 	}
 	else if (!View.bUsePixelInspector && CurrentMousePos != FIntPoint(-1, -1))
 	{
