Add FGlobalEditorCommonCommands::OpenFindInAllBlueprints, which allows users to bind a shortcut for summoning the Find in all Blueprints dialogue.

diff --git a/Engine/Source/Editor/UnrealEd/Private/Toolkits/GlobalEditorCommonCommands.cpp b/Engine/Source/Editor/UnrealEd/Private/Toolkits/GlobalEditorCommonCommands.cpp
--- a/Engine/Source/Editor/UnrealEd/Private/Toolkits/GlobalEditorCommonCommands.cpp
+++ b/Engine/Source/Editor/UnrealEd/Private/Toolkits/GlobalEditorCommonCommands.cpp
@@ -1,6 +1,8 @@
 // Copyright Epic Games, Inc. All Rights Reserved.
 
 #include "Toolkits/GlobalEditorCommonCommands.h"
+
+#include "FindInBlueprints.h"
 #include "Modules/ModuleManager.h"
 #include "Widgets/DeclarativeSyntaxSupport.h"
 #include "Widgets/SWindow.h"
@@ -64,6 +66,8 @@ void FGlobalEditorCommonCommands::RegisterCommands()
 #else
 	UI_COMMAND(OpenContentBrowserDrawer, "Open Content Browser Drawer", "Opens the content browser drawer from the status bar and focuses the search field", EUserInterfaceActionType::Button, FInputChord(EModifierKey::Control, EKeys::SpaceBar));
 #endif
+
+	UI_COMMAND(OpenFindInAllBlueprints, "Open Find in all Blueprints", "Opens the Find in all Blueprints dialogue", EUserInterfaceActionType::Button, FInputChord());
 }
 
 void FGlobalEditorCommonCommands::MapActions(TSharedRef<FUICommandList>& ToolkitCommands)
@@ -110,6 +114,10 @@ void FGlobalEditorCommonCommands::MapActions(TSharedRef<FUICommandList>& Toolkit
 		Get().OpenLevel,
 		FExecuteAction::CreateStatic(&FLevelEditorActionCallbacks::OpenLevel),
 		FCanExecuteAction::CreateStatic(&FLevelEditorActionCallbacks::OpenLevel_CanExecute));
+
+	ToolkitCommands->MapAction(
+		Get().OpenFindInAllBlueprints,
+		FExecuteAction::CreateStatic(&FGlobalEditorCommonCommands::OnOpenFindInAllBlueprints));
 }
 
 void FGlobalEditorCommonCommands::OnPressedCtrlTab(TSharedPtr<FUICommandInfo> TriggeringCommand)
@@ -236,6 +244,14 @@ void FGlobalEditorCommonCommands::OnOpenOutputLogDrawer()
 	GEditor->GetEditorSubsystem<UStatusBarSubsystem>()->OpenOutputLogDrawer();
 }
 
+void FGlobalEditorCommonCommands::OnOpenFindInAllBlueprints()
+{
+	if (const TSharedPtr<SFindInBlueprints> Dialogue = FFindInBlueprintSearchManager::Get().GetGlobalFindResults())
+	{
+		Dialogue->FocusForUse(false);
+	}
+}
+
 //////////////////////////////////////////////////////////////////////////
 
 #undef LOCTEXT_NAMESPACE
diff --git a/Engine/Source/Editor/UnrealEd/Public/Toolkits/GlobalEditorCommonCommands.h b/Engine/Source/Editor/UnrealEd/Public/Toolkits/GlobalEditorCommonCommands.h
--- a/Engine/Source/Editor/UnrealEd/Public/Toolkits/GlobalEditorCommonCommands.h
+++ b/Engine/Source/Editor/UnrealEd/Public/Toolkits/GlobalEditorCommonCommands.h
@@ -29,7 +29,8 @@ protected:
 	static void OnSummonedConsoleCommandBox();
 	static void OnOpenContentBrowserDrawer();
 	static void OnOpenOutputLogDrawer();
-
+	static void OnOpenFindInAllBlueprints();
+	
 	static TSharedRef<SDockTab> SpawnAssetPicker(const FSpawnTabArgs& InArgs);
 
 	static TSharedPtr<IMenu> OpenPopupMenu(TSharedRef<SWidget> WindowContents, const FVector2D& PopupDesiredSize);
@@ -46,6 +47,7 @@ public:
 	TSharedPtr<FUICommandInfo> SelectNextConsoleExecutor;
 	TSharedPtr<FUICommandInfo> OpenOutputLogDrawer;
 	TSharedPtr<FUICommandInfo> OpenContentBrowserDrawer;
+	TSharedPtr<FUICommandInfo> OpenFindInAllBlueprints;
 
 	/** Level file commands */
 	TSharedPtr<FUICommandInfo> OpenLevel;
