Tweaks for LUF
- Make the animation editor look down the forward axis when using LUF
- Default the thumbnail scene info to look down the forward axis  when using LUF

diff --git a/Engine/Source/Editor/Persona/Private/SAnimationEditorViewport.cpp b/Engine/Source/Editor/Persona/Private/SAnimationEditorViewport.cpp
--- a/Engine/Source/Editor/Persona/Private/SAnimationEditorViewport.cpp
+++ b/Engine/Source/Editor/Persona/Private/SAnimationEditorViewport.cpp
@@ -185,8 +185,7 @@ TSharedRef<FEditorViewportClient> SAnimationEditorViewport::MakeEditorViewportCl
 
 	LevelViewportClient->ViewportType = LVT_Perspective;
 	LevelViewportClient->bSetListenerPosition = false;
-	LevelViewportClient->SetViewLocation(EditorViewportDefs::DefaultPerspectiveViewLocation);
-	LevelViewportClient->SetViewRotation(EditorViewportDefs::DefaultPerspectiveViewRotation);
+	LevelViewportClient->SetInitialViewTransform(LVT_Perspective, EditorViewportDefs::DefaultPerspectiveViewLocation, EditorViewportDefs::DefaultPerspectiveViewRotation, DEFAULT_ORTHOZOOM);
 
 	return LevelViewportClient.ToSharedRef();
 }
diff --git a/Engine/Source/Editor/UnrealEd/Private/UnrealEdMisc.cpp b/Engine/Source/Editor/UnrealEd/Private/UnrealEdMisc.cpp
--- a/Engine/Source/Editor/UnrealEd/Private/UnrealEdMisc.cpp
+++ b/Engine/Source/Editor/UnrealEd/Private/UnrealEdMisc.cpp
@@ -97,6 +97,8 @@
 #include "Cooker/CookConfigAccessTracker.h"
 #include "Cooker/PackageBuildDependencyTracker.h"
 #include "EditorAxisDisplayInfo.h"
+#include "ThumbnailRendering/SceneThumbnailInfo.h"
+
 
 #define USE_UNIT_TESTS 0
 
@@ -235,6 +237,14 @@ FUnrealEdMisc::FUnrealEdMisc() :
 	{
 		EditorAxisDisplayInfo = MakeUnique<FEditorAxisDisplayInfo>();
 		IModularFeatures::Get().RegisterModularFeature(IAxisDisplayInfo::GetModularFeatureName(), EditorAxisDisplayInfo.Get());
+
+		const bool bUsingLUFCoordinateSysem = AxisDisplayInfo::GetAxisDisplayCoordinateSystem() == EAxisList::LeftUpForward;
+
+		if (bUsingLUFCoordinateSysem)
+		{
+			// Adjust the default SceneThumbnailInfo to look down the forward axis
+			USceneThumbnailInfo::StaticClass()->GetDefaultObject<USceneThumbnailInfo>()->OrbitYaw += 90.f;
+		}
 	}
 }
 
