[SourceControl] Need a new unknown source control message

diff --git a/Engine/Source/Editor/UnrealEd/Private/AssetDefinitionDefault.cpp b/Engine/Source/Editor/UnrealEd/Private/AssetDefinitionDefault.cpp
--- a/Engine/Source/Editor/UnrealEd/Private/AssetDefinitionDefault.cpp
+++ b/Engine/Source/Editor/UnrealEd/Private/AssetDefinitionDefault.cpp
@@ -79,10 +79,11 @@ namespace UE::AssetDefinitionDefault::Status
 
 	FText GetSourceControlStatusDescription(const TSharedPtr<IAssetStatusInfoProvider> InAssetStatusInfoProvider)
 	{
+		const FText SourceControlDescriptionEmpty = LOCTEXT("NoStatus", "Couldn't retrieve source control status");
 		FText SourceControlDescription = FText::GetEmpty();
 		if (!InAssetStatusInfoProvider.IsValid())
 		{
-			return SourceControlDescription;
+			return SourceControlDescriptionEmpty;
 		}
 
 		if (ISourceControlModule::Get().IsEnabled() && ISourceControlModule::Get().GetProvider().IsAvailable())
@@ -90,14 +91,11 @@ namespace UE::AssetDefinitionDefault::Status
 			const FString FileName = InAssetStatusInfoProvider->TryGetFilename();
 			if (FSourceControlStatePtr SourceControlState = ISourceControlModule::Get().GetProvider().GetState(FileName, EStateCacheUsage::Use))
 			{
-				TOptional<FText> StatusText = SourceControlState->GetStatusText();
-				if (StatusText.IsSet())
-				{
-					SourceControlDescription = StatusText.GetValue();
-				}
+				// Calling this instead of GetStatusText, since that will check for warnings and won't give out a tooltip for every state unlike GetDisplayTooltip
+				SourceControlDescription = SourceControlState->GetDisplayTooltip();
 			}
 		}
-		return SourceControlDescription;
+		return SourceControlDescription.IsEmpty() ? SourceControlDescriptionEmpty : SourceControlDescription;
 	}
 }
 
