[ContentBrowser] New style fixes based on feedbacks
Added also a couple more fix based on latest feedbacks

diff --git a/Engine/Source/Editor/ContentBrowser/Private/AssetViewWidgets.cpp b/Engine/Source/Editor/ContentBrowser/Private/AssetViewWidgets.cpp
--- a/Engine/Source/Editor/ContentBrowser/Private/AssetViewWidgets.cpp
+++ b/Engine/Source/Editor/ContentBrowser/Private/AssetViewWidgets.cpp
@@ -69,6 +69,7 @@
 #include "Styling/ISlateStyle.h"
 #include "Styling/SlateBrush.h"
 #include "Styling/SlateTypes.h"
+#include "Styling/StarshipCoreStyle.h"
 #include "Styling/StyleColors.h"
 #include "Styling/StyleDefaults.h"
 #include "Styling/WidgetStyle.h"
@@ -1875,7 +1876,9 @@ void SAssetListItem::Construct( const FArguments& InArgs )
 
 		if (UE::Editor::ContentBrowser::IsNewStyleEnabled())
 		{
+			ThumbnailConfig.IsEditModeVisible = MakeAttributeSP(this, &SAssetListItem::GetThumbnailEditModeUIVisibility);
 			ThumbnailConfig.AllowAssetSpecificThumbnailOverlay = MakeAttributeSP(this, &SAssetListItem::ShouldAllowAssetActionOverlay, ThumbnailConfig.bForceGenericThumbnail);
+			ThumbnailConfig.AllowAssetSpecificThumbnailOverlayIndicator = MakeAttributeSP(this, &SAssetListItem::ShouldAllowAssetActionOverlayIndicator, ThumbnailConfig.bForceGenericThumbnail);
 			ThumbnailConfig.GenericThumbnailSize = MakeAttributeSP(this, &SAssetListItem::GetGenericThumbnailSize);
 			ThumbnailConfig.AssetSystemInfoProvider = MakeShared<FAssetSystemContentBrowserInfoProvider>(AssetItem);
 			ThumbnailConfig.AllowAssetStatusThumbnailOverlay = MakeAttributeSP(this, &SAssetListItem::ShouldAllowAssetStatusOverlay, ThumbnailConfig.bForceGenericThumbnail);
@@ -2147,6 +2150,11 @@ bool SAssetListItem::ShouldAllowAssetStatusOverlay(bool bForceGenericThumbnail)
 }
 
 bool SAssetListItem::ShouldAllowAssetActionOverlay(bool bForceGenericThumbnail) const
+{
+	return !bForceGenericThumbnail;
+}
+
+bool SAssetListItem::ShouldAllowAssetActionOverlayIndicator(bool bForceGenericThumbnail) const
 {
 	return !bForceGenericThumbnail && CurrentThumbnailSize.Get() >= EThumbnailSize::Large;
 }
@@ -2472,6 +2480,7 @@ void SAssetTileItem::Construct( const FArguments& InArgs )
 		constexpr float NameAreaBoxTopPadding = 6.f;
 		constexpr float ClassNameMaxHeight = 14.f;
 		const FMargin FolderImagePadding = FMargin(2.f, 2.f, 2.f, 2.f);
+		const FMargin FolderOverlayNamePaddingFromLeft = FMargin(9.f, 0.f, 0.f, 0.f);
 
 		ChildSlot
 		.Padding(FMargin(BorderPadding, BorderPadding, 0.f, 1.f))
@@ -2515,19 +2524,26 @@ void SAssetTileItem::Construct( const FArguments& InArgs )
 
 							// Overlay name for tiny folders
 							+ SOverlay::Slot()
-							.HAlign(HAlign_Center)
+							.Padding(FolderOverlayNamePaddingFromLeft)
+							.HAlign(HAlign_Left)
 							.VAlign(VAlign_Center)
-							.Padding(10.f, 0.f)
 							[
-								SNew(SBorder)
-								.Padding(4.f, 2.f)
-								.BorderImage(FAppStyle::GetBrush("ContentBrowser.AssetTileItem.TinyFolderTextBorder"))
-								.Visibility(this, &SAssetTileItem::GetTinyFolderNameOverlayVisibility)
+								SNew(SBox)
+								.HAlign(HAlign_Center)
+								.VAlign(VAlign_Center)
+								.WidthOverride(this, &SAssetTileItem::GetTinyFolderNameWidth)
 								[
-									SNew(STextBlock)
-									.ColorAndOpacity(FStyleColors::White)
-									.OverflowPolicy(ETextOverflowPolicy::Ellipsis)
-									.Text(GetNameText())
+									SNew(SBorder)
+									.Padding(4.f, 2.f)
+									.BorderImage(FAppStyle::GetBrush("ContentBrowser.AssetTileItem.TinyFolderTextBorder"))
+									.Visibility(this, &SAssetTileItem::GetTinyFolderNameOverlayVisibility)
+									[
+										SNew(STextBlock)
+										.Font(FStyleFonts::Get().Small)
+										.ColorAndOpacity(FStyleColors::White)
+										.OverflowPolicy(ETextOverflowPolicy::Ellipsis)
+										.Text(GetNameText())
+									]
 								]
 							]
 						]
@@ -2807,6 +2823,21 @@ FOptionalSize SAssetTileItem::GetThumbnailBoxSize() const
 	return UE::Editor::ContentBrowser::IsNewStyleEnabled() ? FOptionalSize(ThumbnailDimension.Get() + ThumbnailSizeOffset) : FOptionalSize(ItemWidth.Get()- ThumbnailPadding);
 }
 
+FOptionalSize SAssetTileItem::GetTinyFolderNameWidth() const
+{
+	// Folder offset on each side
+	constexpr float FolderOffset = 5.f;
+	// Folder slot padding on each side
+	constexpr float FolderSlotPadding = 2.f;
+	// Folder shadow offset on the right and bottom side
+	constexpr float FolderShadowOffset = 2.f;
+	// Folder padding between image and actual tiny overlay folder name overlay
+	constexpr float FolderTinyNameOverlayPadding = 4.f;
+
+	const float FolderFullOffset = (FolderOffset * 2.f) + (FolderSlotPadding * 2.f) + FolderShadowOffset + FolderTinyNameOverlayPadding;
+	return FOptionalSize(ThumbnailDimension.Get() - FolderFullOffset);
+}
+
 EVisibility SAssetTileItem::GetTinyFolderNameOverlayVisibility() const
 {
 	if (!UE::Editor::ContentBrowser::IsNewStyleEnabled())
diff --git a/Engine/Source/Editor/ContentBrowser/Public/AssetViewWidgets.h b/Engine/Source/Editor/ContentBrowser/Public/AssetViewWidgets.h
--- a/Engine/Source/Editor/ContentBrowser/Public/AssetViewWidgets.h
+++ b/Engine/Source/Editor/ContentBrowser/Public/AssetViewWidgets.h
@@ -476,6 +476,9 @@ protected:
 	/** Get if AssetAction should be visible */
 	bool ShouldAllowAssetActionOverlay(bool bForceGenericThumbnail) const;
 
+	/** Get if the AssetAction indicator should be visible */
+	bool ShouldAllowAssetActionOverlayIndicator(bool bForceGenericThumbnail) const;
+
 	/** Get if the AssetChip should be visible */
 	bool ShouldUseAssetChip() const;
 
@@ -770,6 +773,9 @@ protected:
 	/** Returns the size of the thumbnail widget */
 	FOptionalSize GetThumbnailBoxSize() const;
 
+	/** Returns the width of the tiny folder name overlay widget */
+	FOptionalSize GetTinyFolderNameWidth() const;
+
 	/** Gets the visibility of the folder name overlay on the folder icon */
 	EVisibility GetTinyFolderNameOverlayVisibility() const;
 
diff --git a/Engine/Source/Editor/UnrealEd/Private/AssetThumbnail.cpp b/Engine/Source/Editor/UnrealEd/Private/AssetThumbnail.cpp
--- a/Engine/Source/Editor/UnrealEd/Private/AssetThumbnail.cpp
+++ b/Engine/Source/Editor/UnrealEd/Private/AssetThumbnail.cpp
@@ -109,6 +109,7 @@ public:
 		SLATE_ARGUMENT(bool, ForceGenericThumbnail)
 		SLATE_ARGUMENT(bool, AllowHintText)
 		SLATE_ATTRIBUTE(bool, AllowAssetSpecificThumbnailOverlay)
+		SLATE_ATTRIBUTE(bool, AllowAssetSpecificThumbnailOverlayIndicator)
 		SLATE_ARGUMENT(bool, AllowRealTimeOnHovered)
 		SLATE_ARGUMENT(EThumbnailLabel::Type, Label)
 		SLATE_ATTRIBUTE(FText, HighlightedText)
@@ -150,6 +151,7 @@ public:
 		ColorStripOrientation = InArgs._ColorStripOrientation;
 		AssetSystemInfoProvider = InArgs._AssetSystemInfoProvider;
 		AllowAssetSpecificThumbnailOverlay = InArgs._AllowAssetSpecificThumbnailOverlay;
+		AllowAssetSpecificThumbnailOverlayIndicator = InArgs._AllowAssetSpecificThumbnailOverlayIndicator;
 		AllowAssetStatusThumbnailOverlay = InArgs._AllowAssetStatusThumbnailOverlay;
 		AssetChipBorderImageOverride = InArgs._AssetChipBorderImageOverride;
 		ShowAssetChip = InArgs._ShowAssetChip;
@@ -461,6 +463,32 @@ public:
 				if (AssetDefinition && AssetDefinition->GetThumbnailActionOverlay(AssetData, OutThumbnailInfo))
 				{
 					constexpr int32 OverlayZOrder = 1;
+					if (OutThumbnailInfo.ActionImageWidget.IsValid())
+					{
+						constexpr float PaddingFromTopLeftBorder = 4.f;
+						constexpr float TopLeftImageSize = 20.f;
+						constexpr float BorderPaddingToActionImageWidget = 4.f;
+
+						OverlayWidget->AddSlot()
+						.ZOrder(OverlayZOrder)
+						.VAlign(VAlign_Top)
+						.HAlign(HAlign_Left)
+						.Padding(PaddingFromTopLeftBorder, PaddingFromTopLeftBorder, 0.f, 0.f)
+						[
+							SNew(SBox)
+							.WidthOverride(TopLeftImageSize)
+							.HeightOverride(TopLeftImageSize)
+							[
+								SNew(SBorder)
+								.BorderImage(FAppStyle::GetBrush(Style, ".AssetThumbnailBar"))
+								.Padding(BorderPaddingToActionImageWidget)
+								.Visibility(this, &SAssetThumbnail::GetActionIconOverlayVisibility)
+								[
+									OutThumbnailInfo.ActionImageWidget.ToSharedRef()
+								]
+							]
+						];
+					}
 
 					PRAGMA_DISABLE_DEPRECATION_WARNINGS
 					if (OutThumbnailInfo.ActionButtonWidget.IsValid())
@@ -532,7 +560,6 @@ public:
 				PRAGMA_ENABLE_DEPRECATION_WARNINGS
 				if (AssetSpecificThumbnailOverlay.IsValid())
 				{
-					AssetSpecificThumbnailOverlay->SetVisibility(AllowAssetSpecificThumbnailOverlay.Get() ? EVisibility::Visible : EVisibility::Collapsed);
 					OverlayWidget->AddSlot()
 					[
 						AssetSpecificThumbnailOverlay.ToSharedRef()
@@ -972,6 +999,12 @@ private:
 		return EVisibility::Collapsed;
 	}
 
+	EVisibility GetActionIconOverlayVisibility() const
+	{
+		const bool bIsEditModeVisible = EditModeVisibility.IsSet() && EditModeVisibility.Get().IsVisible();
+		return AllowAssetSpecificThumbnailOverlayIndicator.Get() && !IsHovered() && !bIsEditModeVisible ? EVisibility::Visible : EVisibility::Collapsed;
+	}
+
 	EVisibility GetActionButtonVisibility() const
 	{
 		return IsHovered() && AllowAssetSpecificThumbnailOverlay.Get() && (!EditModeVisibility.IsSet() || !EditModeVisibility.Get().IsVisible()) ? EVisibility::Visible : EVisibility::Collapsed;
@@ -1514,6 +1547,7 @@ private:
 	TAttribute<EVisibility> EditModeVisibility;
 	TAttribute<bool> AlwaysExpandTooltip;
 	TAttribute<bool> AllowAssetSpecificThumbnailOverlay;
+	TAttribute<bool> AllowAssetSpecificThumbnailOverlayIndicator;
 	TAttribute<bool> AllowAssetStatusThumbnailOverlay;
 	bool bAllowHintText;
 	bool bAllowRealTimeOnHovered;
@@ -1687,11 +1721,18 @@ void FAssetThumbnail::SetAsset( const FAssetData& InAssetData )
 TSharedRef<SWidget> FAssetThumbnail::MakeThumbnailWidget( const FAssetThumbnailConfig& InConfig )
 {
 	PRAGMA_DISABLE_DEPRECATION_WARNINGS
-	const TAttribute<bool> AssetThumbnailOverlayAttribute = InConfig.AllowAssetSpecificThumbnailOverlay.IsSet() ?
-		                                                        InConfig.AllowAssetSpecificThumbnailOverlay :
-		                                                        InConfig.bAllowAssetSpecificThumbnailOverlay;
+	const TAttribute<bool> AssetThumbnailOverlayAttribute =
+		InConfig.AllowAssetSpecificThumbnailOverlay.IsSet()
+		? InConfig.AllowAssetSpecificThumbnailOverlay
+		: InConfig.bAllowAssetSpecificThumbnailOverlay;
 	PRAGMA_ENABLE_DEPRECATION_WARNINGS
 
+	// If not set use the PlayButton attribute instead
+	const TAttribute<bool> AssetThumbnailOverlayIndicatorAttribute =
+		InConfig.AllowAssetSpecificThumbnailOverlayIndicator.IsSet()
+		? InConfig.AllowAssetSpecificThumbnailOverlayIndicator
+		: AssetThumbnailOverlayAttribute;
+
 	TSharedPtr<SWidget> ThumbnailWidget = nullptr;
 	if (UE::Editor::ContentBrowser::IsNewStyleEnabled())
 	{
@@ -1708,6 +1749,7 @@ TSharedRef<SWidget> FAssetThumbnail::MakeThumbnailWidget( const FAssetThumbnailC
 			.AllowRealTimeOnHovered(InConfig.bAllowRealTimeOnHovered)
 			.ClassThumbnailBrushOverride(InConfig.ClassThumbnailBrushOverride)
 			.AllowAssetSpecificThumbnailOverlay(AssetThumbnailOverlayAttribute)
+			.AllowAssetSpecificThumbnailOverlayIndicator(AssetThumbnailOverlayIndicatorAttribute)
 			.AssetTypeColorOverride(InConfig.AssetTypeColorOverride)
 			.Padding(InConfig.Padding)
 			.GenericThumbnailSize(InConfig.GenericThumbnailSize)
diff --git a/Engine/Source/Editor/UnrealEd/Public/AssetThumbnail.h b/Engine/Source/Editor/UnrealEd/Public/AssetThumbnail.h
--- a/Engine/Source/Editor/UnrealEd/Public/AssetThumbnail.h
+++ b/Engine/Source/Editor/UnrealEd/Public/AssetThumbnail.h
@@ -63,6 +63,7 @@ struct FAssetThumbnailConfig
 		, bAllowAssetSpecificThumbnailOverlay( false )
 		PRAGMA_ENABLE_DEPRECATION_WARNINGS
 		, AllowAssetSpecificThumbnailOverlay( false )
+		, AllowAssetSpecificThumbnailOverlayIndicator()
 		, ClassThumbnailBrushOverride( NAME_None )
 		, ThumbnailLabel( EThumbnailLabel::ClassName )
 		, HighlightedText( FText::GetEmpty() )
@@ -96,6 +97,8 @@ struct FAssetThumbnailConfig
 	bool bAllowAssetSpecificThumbnailOverlay;
 	/** Whether to show the asset play button or not */
 	TAttribute< bool > AllowAssetSpecificThumbnailOverlay;
+	/** Whether to show the play indicator or not, by default it will match the AllowAssetSpecificThumbnailOverlay one */
+	TAttribute< bool > AllowAssetSpecificThumbnailOverlayIndicator;
 	FName ClassThumbnailBrushOverride;
 	EThumbnailLabel::Type ThumbnailLabel;
 	TAttribute< FText > HighlightedText;
