Editor Tools - Fixed crash when calling GetToolManager() and there isn't a ModeToolsContext.

#rnx

diff --git a/Engine/Source/Editor/UnrealEd/Private/Tools/UEdMode.cpp b/Engine/Source/Editor/UnrealEd/Private/Tools/UEdMode.cpp
--- a/Engine/Source/Editor/UnrealEd/Private/Tools/UEdMode.cpp
+++ b/Engine/Source/Editor/UnrealEd/Private/Tools/UEdMode.cpp
@@ -204,7 +204,7 @@ UInteractiveToolManager* UEdMode::GetToolManager(EToolsContextScope ToolScope) c
 	}
 	else if (ToolScope == EToolsContextScope::EdMode)
 	{
-		return ModeToolsContext->ToolManager;
+		return ModeToolsContext ? ModeToolsContext->ToolManager : nullptr;
 	}
 	ensure(false);
 	return nullptr;
