Cooker: Migrate to the new structured log formatting API

diff --git a/Engine/Source/Editor/UnrealEd/Classes/CookOnTheSide/CookOnTheFlyServer.h b/Engine/Source/Editor/UnrealEd/Classes/CookOnTheSide/CookOnTheFlyServer.h
--- a/Engine/Source/Editor/UnrealEd/Classes/CookOnTheSide/CookOnTheFlyServer.h
+++ b/Engine/Source/Editor/UnrealEd/Classes/CookOnTheSide/CookOnTheFlyServer.h
@@ -181,7 +181,7 @@ namespace UE::Cook
 	class FCookWorkerServer;
 	class FDiagnostics;
 	class FLogHandler;
-	class FODSCClientData; 
+	class FODSCClientData;
 	class FPackagePreloader;
 	class FPackageWriterMPCollector;
 	class FRequestCluster;
@@ -193,6 +193,7 @@ namespace UE::Cook
 	class FWorkerRequestsRemote;
 	class ICookOnTheFlyRequestManager;
 	class ICookOnTheFlyNetworkServer;
+	class ILogHandler;
 	class IWorkerRequests;
 	enum class ECachedCookedPlatformDataEvent : uint8;
 	enum class ECookPhase : uint8;
@@ -1638,7 +1639,7 @@ private:
 	TUniquePtr<FLayeredCookArtifactReader> AllContextArtifactReader;
 	TSharedPtr<FLooseFilesCookArtifactReader> SharedLooseFilesCookArtifactReader;
 	TUniquePtr<UE::Cook::FCookGCDiagnosticContext> GCDiagnosticContext;
-	TUniquePtr<UE::Cook::FLogHandler> LogHandler;
+	TUniquePtr<UE::Cook::ILogHandler> LogHandler;
 
 	TArray<UE::Cook::FCookSavePackageContext*> SavePackageContexts;
 	/**
diff --git a/Engine/Source/Editor/UnrealEd/Private/CookOnTheFlyServer.cpp b/Engine/Source/Editor/UnrealEd/Private/CookOnTheFlyServer.cpp
--- a/Engine/Source/Editor/UnrealEd/Private/CookOnTheFlyServer.cpp
+++ b/Engine/Source/Editor/UnrealEd/Private/CookOnTheFlyServer.cpp
@@ -5819,7 +5819,7 @@ void UCookOnTheFlyServer::Initialize( ECookMode::Type DesiredCookMode, ECookInit
 	AssetRegistry = IAssetRegistry::Get();
 	GCDiagnosticContext = MakeUnique<UE::Cook::FCookGCDiagnosticContext>();
 	StallDetector = MakeUnique<UE::Cook::FStallDetector>();
-	LogHandler.Reset(new UE::Cook::FLogHandler(*this));
+	LogHandler.Reset(UE::Cook::CreateLogHandler(*this));
 
 	if (!IsCookWorkerMode())
 	{
diff --git a/Engine/Source/Editor/UnrealEd/Private/Cooker/CookLog.cpp b/Engine/Source/Editor/UnrealEd/Private/Cooker/CookLog.cpp
--- a/Engine/Source/Editor/UnrealEd/Private/Cooker/CookLog.cpp
+++ b/Engine/Source/Editor/UnrealEd/Private/Cooker/CookLog.cpp
@@ -14,6 +14,7 @@
 #include "Misc/AssertionMacros.h"
 #include "Misc/DateTime.h"
 #include "Misc/FeedbackContext.h"
+#include "Misc/OutputDevice.h"
 #include "Misc/OutputDeviceRedirector.h"
 #include "Misc/PackageAccessTracking.h"
 #include "Misc/ScopeLock.h"
@@ -102,6 +103,62 @@ bool LoadFromCompactBinary(FCbFieldView Field, FReplicatedLogData& OutLogData)
 	return bOk;
 }
 
+class FLogHandler : public FOutputDevice, public ILogHandler
+{
+public:
+	explicit FLogHandler(UCookOnTheFlyServer& InCOTFS);
+	virtual ~FLogHandler();
+
+	virtual void ReplayLogsFromIncrementallySkipped(TConstArrayView<FReplicatedLogData> LogMessages) override;
+	virtual void ReplayLogFromCookWorker(FReplicatedLogData&& LogData, int32 CookWorkerProfileId) override;
+	virtual void ConditionalPruneReplay() override;
+	virtual void FlushIncrementalCookLogs() override;
+
+	// FOutputDevice
+	virtual void Serialize(const TCHAR* V, ELogVerbosity::Type Verbosity, const FName& Category) override;
+	virtual void Serialize(const TCHAR* V, ELogVerbosity::Type Verbosity, const FName& Category,
+		const double Time) override;
+	virtual void SerializeRecord(const FLogRecord& Record) override;
+	virtual void Flush() override;
+	virtual bool CanBeUsedOnAnyThread() const override;
+	virtual bool CanBeUsedOnMultipleThreads() const override;
+
+private:
+	struct FQueuedLog
+	{
+		FName ActivePackage;
+		FReplicatedLogData LogData;
+	};
+
+	void Marshal(FReplicatedLogData& OutData, FStringView Message, ELogVerbosity::Type Verbosity,
+		const FName& Category);
+	void Marshal(FReplicatedLogData& OutData, const FLogRecord& LogRecord);
+	void UnMarshalAndLog(const FReplicatedLogData& LogData,
+		TFunctionRef<bool(FName, const FString&)> MessagePassesFilter,
+		TFunctionRef<bool(const FString&, FString&)> TryTransformMessage);
+	bool UnMarshal(FCbFieldView Field, FLogRecord& OutLogRecord,
+		TFunctionRef<bool(FName, const FString&)> MessagePassesFilter,
+		TFunctionRef<bool(const FString&, FString&)> TryTransformMessage);
+
+	void ReportActiveLog(FReplicatedLogData&& LogData, FStringView FormatMessage, ELogVerbosity::Type Verbosity);
+	void RecordLogForIncrementalCook(FReplicatedLogData&& LogData, ELogVerbosity::Type Verbosity);
+	void RecordLogForIncrementalCookGameThreadPortion(FName PackageName, FReplicatedLogData&& LogData);
+
+	void PruneReplay();
+
+private:
+	UCookOnTheFlyServer& COTFS;
+	bool bRegistered = false;
+
+	FCriticalSection QueuedLogsForIncrementalCookLock;
+	TArray<FQueuedLog> QueuedLogsForIncrementalCook;
+
+	FCriticalSection TableLock;
+	TArray<FString> StringTable;
+	TArray<FAnsiString> AnsiStringTable;
+	TArray<FUniqueLogTemplate> TemplateTable;
+};
+
 FLogHandler::FLogHandler(UCookOnTheFlyServer& InCOTFS)
 	: COTFS(InCOTFS)
 {
@@ -199,10 +256,6 @@ void FLogHandler::PruneReplay()
 		// Output devices are expected to not use any pointer on a structured log record after completion of the SerializeRecord call.
 		GLog->FlushThreadedLogs();
 	}
-	for (FLogTemplate* LogTemplate : TemplateTable)
-	{
-		DestroyLogTemplate(LogTemplate);
-	}
 
 	StringTable.Empty();
 	AnsiStringTable.Empty();
@@ -236,6 +289,16 @@ void FLogHandler::Flush()
 	PruneReplay();
 }
 
+bool FLogHandler::CanBeUsedOnAnyThread() const
+{
+	return true;
+}
+
+bool FLogHandler::CanBeUsedOnMultipleThreads() const
+{
+	return true;
+}
+
 void FLogHandler::Marshal(FReplicatedLogData& OutData, FStringView Message,
 	ELogVerbosity::Type Verbosity, const FName& Category)
 {
@@ -428,14 +491,12 @@ bool FLogHandler::UnMarshal(FCbFieldView Field, FLogRecord& OutLogRecord,
 
 	if (bHasTextKey)
 	{
-		FLogTemplate* LogTemplate = CreateLogTemplate(OutLogRecord.GetTextNamespace(), OutLogRecord.GetTextKey(), OutLogRecord.GetFormat());
-		TemplateTable.Add(LogTemplate);
+		FLogTemplate* LogTemplate = TemplateTable.Emplace_GetRef(OutLogRecord.GetTextNamespace(), OutLogRecord.GetTextKey(), OutLogRecord.GetFormat()).Get();
 		OutLogRecord.SetTemplate(LogTemplate);
 	}
 	else
 	{
-		FLogTemplate* LogTemplate = CreateLogTemplate(OutLogRecord.GetFormat());
-		TemplateTable.Add(LogTemplate);
+		FLogTemplate* LogTemplate = TemplateTable.Emplace_GetRef(OutLogRecord.GetFormat()).Get();
 		OutLogRecord.SetTemplate(LogTemplate);
 	}
 
@@ -511,4 +572,9 @@ void FLogHandler::RecordLogForIncrementalCookGameThreadPortion(FName ActivePacka
 	PackageData->AddLogMessage(MoveTemp(LogData));
 }
 
+ILogHandler* CreateLogHandler(UCookOnTheFlyServer& COTFS)
+{
+	return new FLogHandler(COTFS);
+}
+
 } // namespace UE::Cook
\ No newline at end of file
diff --git a/Engine/Source/Editor/UnrealEd/Private/Cooker/CookLogPrivate.h b/Engine/Source/Editor/UnrealEd/Private/Cooker/CookLogPrivate.h
--- a/Engine/Source/Editor/UnrealEd/Private/Cooker/CookLogPrivate.h
+++ b/Engine/Source/Editor/UnrealEd/Private/Cooker/CookLogPrivate.h
@@ -8,15 +8,12 @@
 #include "CookOnTheSide/CookLog.h"
 #include "Logging/LogMacros.h"
 #include "Logging/LogVerbosity.h"
-#include "Misc/OutputDevice.h"
 #include "Misc/TVariant.h"
 #include "Serialization/CompactBinary.h"
-#include "UObject/NameTypes.h"
 #include "Templates/Function.h"
+#include "UObject/NameTypes.h"
 
 class UCookOnTheFlyServer;
-namespace UE { class FLogTemplate; }
-namespace UE { class FLogRecord; }
 
 DECLARE_LOG_CATEGORY_EXTERN(LogCookList, Log, All);
 
@@ -41,80 +38,19 @@ bool LoadFromCompactBinary(FCbFieldView Field, FReplicatedLogData& OutLogData);
  * The Cooker's listener to log messages. It passes the logs onto CookWorkerClient for reporting to CookDirector in
  * MPCook, and it stores the log messages on the active package for storage in incremental cooks.
  */
-class FLogHandler : public FOutputDevice
+class ILogHandler
 {
 public:
-	FLogHandler(UCookOnTheFlyServer& InCOTFS);
-	~FLogHandler();
-
-	void ReplayLogsFromIncrementallySkipped(TConstArrayView<FReplicatedLogData> LogMessages);
-	void ReplayLogFromCookWorker(FReplicatedLogData&& LogData, int32 CookWorkerProfileId);
-	void ConditionalPruneReplay();
-	void FlushIncrementalCookLogs();
-
-	// FOutputDevice
-	virtual void Serialize(const TCHAR* V, ELogVerbosity::Type Verbosity, const FName& Category) override;
-	virtual void Serialize(const TCHAR* V, ELogVerbosity::Type Verbosity, const FName& Category,
-		const double Time) override;
-	virtual void SerializeRecord(const FLogRecord& Record) override;
-	virtual void Flush() override;
-	virtual bool CanBeUsedOnAnyThread() const override;
-	virtual bool CanBeUsedOnMultipleThreads() const override;
-
-private:
-	struct FQueuedLog
-	{
-		FName ActivePackage;
-		FReplicatedLogData LogData;
-	};
-
-	void Marshal(FReplicatedLogData& OutData, FStringView Message, ELogVerbosity::Type Verbosity,
-		const FName& Category);
-	void Marshal(FReplicatedLogData& OutData, const FLogRecord& LogRecord);
-	void UnMarshalAndLog(const FReplicatedLogData& LogData,
-		TFunctionRef<bool(FName, const FString&)> MessagePassesFilter,
-		TFunctionRef<bool(const FString&, FString&)> TryTransformMessage);
-	bool UnMarshal(FCbFieldView Field, FLogRecord& OutLogRecord,
-		TFunctionRef<bool(FName, const FString&)> MessagePassesFilter,
-		TFunctionRef<bool(const FString&, FString&)> TryTransformMessage);
-
-	void ReportActiveLog(FReplicatedLogData&& LogData, FStringView FormatMessage, ELogVerbosity::Type Verbosity);
-	void RecordLogForIncrementalCook(FReplicatedLogData&& LogData, ELogVerbosity::Type Verbosity);
-	void RecordLogForIncrementalCookGameThreadPortion(FName PackageName, FReplicatedLogData&& LogData);
-
-	void PruneReplay();
-
-private:
-	UCookOnTheFlyServer& COTFS;
-	bool bRegistered = false;
+	virtual ~ILogHandler() = default;
 
-	FCriticalSection QueuedLogsForIncrementalCookLock;
-	TArray<FQueuedLog> QueuedLogsForIncrementalCook;
-
-	FCriticalSection TableLock;
-	TArray<FString> StringTable;
-	TArray<FAnsiString> AnsiStringTable;
-	TArray<FLogTemplate*> TemplateTable;
+	virtual void ReplayLogsFromIncrementallySkipped(TConstArrayView<FReplicatedLogData> LogMessages) = 0;
+	virtual void ReplayLogFromCookWorker(FReplicatedLogData&& LogData, int32 CookWorkerProfileId) = 0;
+	virtual void ConditionalPruneReplay() = 0;
+	virtual void FlushIncrementalCookLogs() = 0;
 };
 
+ILogHandler* CreateLogHandler(UCookOnTheFlyServer& COTFS);
 
 constexpr FStringView HeartbeatCategoryText(TEXTVIEW("CookWorkerHeartbeat:"));
 
 } // namespace UE::Cook
-
-// Inline Implementations
-
-namespace UE::Cook
-{
-
-inline bool FLogHandler::CanBeUsedOnAnyThread() const
-{
-	return true;
-}
-
-inline bool FLogHandler::CanBeUsedOnMultipleThreads() const
-{
-	return true;
-}
-
-} // namespace UE::Cook
\ No newline at end of file
diff --git a/Engine/Source/Editor/UnrealEd/Private/Cooker/CookWorkerServer.cpp b/Engine/Source/Editor/UnrealEd/Private/Cooker/CookWorkerServer.cpp
--- a/Engine/Source/Editor/UnrealEd/Private/Cooker/CookWorkerServer.cpp
+++ b/Engine/Source/Editor/UnrealEd/Private/Cooker/CookWorkerServer.cpp
@@ -1895,7 +1895,7 @@ FGuid FGeneratorEventMessage::MessageType(TEXT("B6EE94CA70EC4F40B0D2214EDC11ED03
 
 FGuid FLogMessagesMessageHandler::MessageType(TEXT("DB024D28203D4FBAAAF6AAD7080CF277"));
 
-FLogMessagesMessageHandler::FLogMessagesMessageHandler(FLogHandler& InCOTFSLogHandler)
+FLogMessagesMessageHandler::FLogMessagesMessageHandler(ILogHandler& InCOTFSLogHandler)
 	: COTFSLogHandler(InCOTFSLogHandler)
 {
 }
diff --git a/Engine/Source/Editor/UnrealEd/Private/Cooker/CookWorkerServer.h b/Engine/Source/Editor/UnrealEd/Private/Cooker/CookWorkerServer.h
--- a/Engine/Source/Editor/UnrealEd/Private/Cooker/CookWorkerServer.h
+++ b/Engine/Source/Editor/UnrealEd/Private/Cooker/CookWorkerServer.h
@@ -27,7 +27,7 @@ struct FProcHandle;
 namespace UE { class FLogTemplate; }
 namespace UE::CompactBinaryTCP { struct FMarshalledMessage; }
 namespace UE::Cook { class FCookDirector; }
-namespace UE::Cook { class FLogHandler; }
+namespace UE::Cook { class ILogHandler; }
 namespace UE::Cook { struct FDiscoveredPackageReplication; }
 namespace UE::Cook { enum class ECookDirectorThread : uint8; };
 namespace UE::Cook { struct FGeneratorEventMessage; }
@@ -498,7 +498,7 @@ public:
 class FLogMessagesMessageHandler : public IMPCollector
 {
 public:
-	FLogMessagesMessageHandler(FLogHandler& InCOTFSLogHandler);
+	FLogMessagesMessageHandler(ILogHandler& InCOTFSLogHandler);
 	void ClientReportLogMessage(const FReplicatedLogData& LogData);
 
 	// IMPCollector
@@ -508,7 +508,7 @@ public:
 	virtual void ServerReceiveMessage(FMPCollectorServerMessageContext& Context, FCbObjectView Message) override;
 
 private:
-	FLogHandler& COTFSLogHandler;
+	ILogHandler& COTFSLogHandler;
 
 	FCriticalSection QueueLock;
 	TArray<FReplicatedLogData> QueuedLogs;
