IncrementalCook: Turn unsolicited package message to verbose to reduce large amount of log spam on large projects, until we're ready to start fixing the hidden references.
#rnx
#incrementalcook

diff --git a/Engine/Source/Editor/UnrealEd/Private/CookOnTheFlyServer.cpp b/Engine/Source/Editor/UnrealEd/Private/CookOnTheFlyServer.cpp
--- a/Engine/Source/Editor/UnrealEd/Private/CookOnTheFlyServer.cpp
+++ b/Engine/Source/Editor/UnrealEd/Private/CookOnTheFlyServer.cpp
@@ -4822,7 +4822,7 @@ void UCookOnTheFlyServer::ProcessUnsolicitedPackages(TArray<FName>* OutDiscovere
 			}
 			else
 			{
-				UE_LOG(LogCook, Display, TEXT("%s"), *Message);
+				UE_LOG(LogCook, Verbose, TEXT("%s"), *Message);
 			}
 		}
 
