CookCorrectness: Fix bug in the handling for reallocated objects; we were updating the state when they reallocated, but we were missing the handling of the updated state when we next tried to call IsCachedCookedPlatformDataLoaded on them.
Add test content for objects that are reallocated by another object while we ware waiting on them.
#rnx

#virtualized

diff --git a/Engine/Source/Editor/UnrealEd/Private/CookOnTheFlyServer.cpp b/Engine/Source/Editor/UnrealEd/Private/CookOnTheFlyServer.cpp
--- a/Engine/Source/Editor/UnrealEd/Private/CookOnTheFlyServer.cpp
+++ b/Engine/Source/Editor/UnrealEd/Private/CookOnTheFlyServer.cpp
@@ -12628,6 +12628,25 @@ bool UCookOnTheFlyServer::RouteIsCachedCookedPlatformDataLoaded(UE::Cook::FPacka
 			);
 	}
 
+	if (*ExistingEvent != ECachedCookedPlatformDataEvent::BeginCacheForCookedPlatformDataCalled
+		&& *ExistingEvent != ECachedCookedPlatformDataEvent::IsCachedCookedPlatformDataLoadedCalled)
+	{
+		// We are trying to call IsCachedCookedPlatformData on an object without first calling
+		// BeginCacheForCookedPlatformData, which is contractually invalid, and it might cause
+		// the object to misbehave immediately, or it might just never return true.
+		// This can occur when system-specific code reallocates the object out from under us by
+		// calling NewObject on an existing object.
+		// MaterialInstanceConstants in particual will fail to return true forever if we have not
+		// called BeginCacheForCookedPlatformData again after they were reallocated.
+
+		// To handle this reallocation case, call BeginCacheForCookedPlatformData first.
+		UE_LOG(LogCook, Display,
+			TEXT("%s was reallocated after BeginCacheForCookedPlatformData and before IsCacheCookedPlatformData returned true. Calling BeginCacheForCookedPlatformData on it again."),
+			*Obj->GetFullName());
+		Obj->BeginCacheForCookedPlatformData(TargetPlatform);
+		*ExistingEvent = ECachedCookedPlatformDataEvent::BeginCacheForCookedPlatformDataCalled;
+	}
+
 	bool bResult = Obj->IsCachedCookedPlatformDataLoaded(TargetPlatform);
 	if (bResult)
 	{
diff --git a/Engine/Source/Editor/UnrealEd/Private/Cooker/CookGarbageCollect.cpp b/Engine/Source/Editor/UnrealEd/Private/Cooker/CookGarbageCollect.cpp
--- a/Engine/Source/Editor/UnrealEd/Private/Cooker/CookGarbageCollect.cpp
+++ b/Engine/Source/Editor/UnrealEd/Private/Cooker/CookGarbageCollect.cpp
@@ -619,9 +619,9 @@ void UCookOnTheFlyServer::PreGarbageCollect()
 {
 	using namespace UE::Cook;
 
-	PackageTracker->SetCollectingGarbage(true);
 	if (!IsInSession())
 	{
+		PackageTracker->SetCollectingGarbage(true);
 		return;
 	}
 
@@ -803,6 +803,12 @@ void UCookOnTheFlyServer::PreGarbageCollect()
 	// Add all public objects within every package in memory to the UPackage::SoftGCPackageToObjectList container,
 	// so they will be kept in memory if the package is kept in memory.
 	ConstructSoftGCPackageToObjectList(this->SoftGCPackageToObjectListBuffer);
+
+	// We call arbitrary system-specific code through FPendingCookedPlatformData.PollIsComplete
+	// -> IsCachedCookedPlatformDataLoaded above, and we need to continue responding to object reallocations
+	// whenever we call system-specific code. So do not mark that we are ignoring deletions from GC until we
+	// have finished calling into that system-specific code.
+	PackageTracker->SetCollectingGarbage(true);
 }
 
 void UCookOnTheFlyServer::CookerAddReferencedObjects(FReferenceCollector& Collector)
diff --git a/Engine/Source/Editor/UnrealEd/Private/Cooker/CookPackageData.cpp b/Engine/Source/Editor/UnrealEd/Private/Cooker/CookPackageData.cpp
--- a/Engine/Source/Editor/UnrealEd/Private/Cooker/CookPackageData.cpp
+++ b/Engine/Source/Editor/UnrealEd/Private/Cooker/CookPackageData.cpp
@@ -3226,7 +3226,12 @@ void FPackageDatas::CachedCookedPlatformDataObjectsPostGarbageCollect(
 
 void FPackageDatas::CachedCookedPlatformDataObjectsOnDestroyedOutsideOfGC(const UObject* DestroyedObject)
 {
-	CachedCookedPlatformDataObjects.Remove(DestroyedObject);
+	FCachedCookedPlatformDataState RemovedState;
+	if (CachedCookedPlatformDataObjects.RemoveAndCopyValue(const_cast<UObject*>(DestroyedObject), RemovedState))
+	{
+		UE_LOG(LogCook, Display, TEXT("Removing CachedCookedPlatformData state for reallocated object %s."),
+			*DestroyedObject->GetPathName());
+	}
 }
 
 void FCachedCookedPlatformDataState::AddRefFrom(FPackageData* PackageData)
