Allow to specify a maximum chunk size from the chunk index. Syntax is
+MaxChunkSizeOverrideFromChunkIndex=(20, 268435456) in /Script/UnrealEd.ProjectPackagingSettings

so that pakchunk20 gets split into subchunks of maximum 256MiB

diff --git a/Engine/Source/Editor/UnrealEd/Private/Commandlets/AssetRegistryGenerator.cpp b/Engine/Source/Editor/UnrealEd/Private/Commandlets/AssetRegistryGenerator.cpp
--- a/Engine/Source/Editor/UnrealEd/Private/Commandlets/AssetRegistryGenerator.cpp
+++ b/Engine/Source/Editor/UnrealEd/Private/Commandlets/AssetRegistryGenerator.cpp
@@ -264,6 +264,63 @@ bool FAssetRegistryGenerator::ShouldPlatformGenerateStreamingInstallManifest(con
 	return false;
 }
 
+static TMap<int64, int64> GetMaxChunkSizeOverrideFromChunkIndex(const ITargetPlatform* Platform)
+{
+	TMap<int64, int64> MaxChunkSizeOverrideFromChunkIndex;
+	if (Platform)
+	{
+		FConfigFile PlatformIniFile;
+		FConfigCacheIni::LoadLocalIniFile(PlatformIniFile, TEXT("Game"), true, *Platform->IniPlatformName());
+		TArray<FString> ConfigStringArray;
+		if (PlatformIniFile.GetArray(TEXT("/Script/UnrealEd.ProjectPackagingSettings"), TEXT("MaxChunkSizeOverrideFromChunkIndex"), ConfigStringArray))
+		{
+			for (const FString& ConfigString : ConfigStringArray)
+			{
+				FStringView ConfigStringView(ConfigString);
+				ConfigStringView = ConfigStringView.TrimStartAndEnd();
+				if (ConfigStringView.IsEmpty() || ConfigStringView[0] != TEXT('(') || ConfigStringView[ConfigStringView.Len() - 1] != TEXT(')'))
+				{
+					continue;
+				}
+				ConfigStringView.RemovePrefix(1);
+				ConfigStringView.RemoveSuffix(1);
+
+				TArray<FStringView> ConfigStringArrayNameSize;
+				UE::String::ParseTokens(ConfigStringView, TEXTVIEW(","), [&](FStringView Element)
+					{
+						ConfigStringArrayNameSize.Add(Element);
+					});
+
+				if (ConfigStringArrayNameSize.Num() != 2 || ConfigStringArrayNameSize[0].Len() == 0 || ConfigStringArrayNameSize[1].Len() == 0)
+				{
+					continue;
+				}
+
+				const TCHAR* ConfigNameStringViewEnd = ConfigStringArrayNameSize[0].GetData();
+				int64 PakIndex = FCString::Strtoi64(ConfigStringArrayNameSize[0].GetData(), (TCHAR**)&ConfigNameStringViewEnd, 10);
+				if (ConfigNameStringViewEnd == ConfigStringArrayNameSize[0].GetData())
+				{
+					continue;
+				}
+				
+				const TCHAR* ConfigSizeStringViewEnd = ConfigStringArrayNameSize[1].GetData();
+				int64 ConfigSize = FCString::Strtoi64(ConfigStringArrayNameSize[1].GetData(), (TCHAR**)&ConfigSizeStringViewEnd, 10);
+
+				if (ConfigSizeStringViewEnd == ConfigStringArrayNameSize[1].GetData())
+				{
+					continue;
+				}
+
+				MaxChunkSizeOverrideFromChunkIndex.Add(PakIndex, ConfigSize);
+			}
+		}
+	}
+
+	return MaxChunkSizeOverrideFromChunkIndex;
+
+}
+
+
 int64 FAssetRegistryGenerator::GetMaxChunkSizePerPlatform(const ITargetPlatform* Platform) const
 {
 	if (Platform)
@@ -370,6 +427,7 @@ bool FAssetRegistryGenerator::GenerateStreamingInstallManifest(int64 InOverrideC
 		ChunkManifestDir /= InManifestSubDir;
 	}
 	int64 MaxChunkSize = InOverrideChunkSize > 0 ? InOverrideChunkSize : GetMaxChunkSizePerPlatform(TargetPlatform);
+	TMap<int64, int64> MaxChunkSizeOverrideFromIndex = GetMaxChunkSizeOverrideFromChunkIndex(TargetPlatform);
 
 	if (!IFileManager::Get().MakeDirectory(*ChunkManifestDir, true /* Tree */))
 	{
@@ -607,7 +665,10 @@ bool FAssetRegistryGenerator::GenerateStreamingInstallManifest(int64 InOverrideC
 		TArray<FString> ChunkFilenames;
 		Manifest->GenerateValueArray(ChunkFilenames);
 
-		if (MaxChunkSize > 0)
+		int64* MaxChunkSizeForPakIndex = MaxChunkSizeOverrideFromIndex.Find(PakchunkIndex);
+		int64 PakChunkMaxChunkSize = MaxChunkSizeForPakIndex ? *MaxChunkSizeForPakIndex : MaxChunkSize;
+
+		if (PakChunkMaxChunkSize > 0)
 		{
 			PackageFileSizes.Reset();
 			PackageFileSizes.Reserve(Manifest->Num());
@@ -785,12 +846,12 @@ bool FAssetRegistryGenerator::GenerateStreamingInstallManifest(int64 InOverrideC
 			{
 				FString Filename = ChunkFilenames[FilenameIndex];
 				FString PakListLine = FPaths::ConvertRelativePathToFull(Filename.Replace(TEXT("[Platform]"), *Platform));
-				if (MaxChunkSize > 0)
+				if (PakChunkMaxChunkSize > 0)
 				{
 					const int64* ExistingPackageFileSize = PackageFileSizes.Find(PakListLine);
 					const int64 PackageFileSize = ExistingPackageFileSize ? *ExistingPackageFileSize : 0;
 					CurrentPakSize += PackageFileSize;
-					if (MaxChunkSize < CurrentPakSize)
+					if (PakChunkMaxChunkSize < CurrentPakSize)
 					{
 						// early out if we are over memory limit
 						bFinishedAllFiles = false;
@@ -813,7 +874,7 @@ bool FAssetRegistryGenerator::GenerateStreamingInstallManifest(int64 InOverrideC
 			{
 				const TCHAR* UnitsText = TEXT("MB");
 				int32 Unit = 1000*1000;
-				if (MaxChunkSize < Unit * 10)
+				if (PakChunkMaxChunkSize < Unit * 10)
 				{
 					Unit = 1;
 					UnitsText = TEXT("bytes");
@@ -821,7 +882,7 @@ bool FAssetRegistryGenerator::GenerateStreamingInstallManifest(int64 InOverrideC
 				UE_LOGFMT(LogAssetRegistryGenerator, Error, "Failed to add file {NextFilename} to paklist '{PakListFilename}'. The maximum size for a Pakfile is {MaxChunkFileSize}{UnitsText}, but the file to add is {ActualChunkFileSize}{UnitsText}.", 
 					("NextFilename", NextFilename),
 					("PakListFilename", PakListFilename),
-					("MaxChunkFileSize", MaxChunkSize / Unit),
+					("MaxChunkFileSize", PakChunkMaxChunkSize / Unit),
 					("UnitsText", UnitsText),
 					("ActualChunkFileSize", (NextFileSize + Unit - 1) / Unit)	// Round the limit down and round the value up, so that the display always shows that the value is greater than the limit
 				);
