Fix issue where reseting scalability settings to defaults did not save those scalability settings.

#rnx

diff --git a/Engine/Source/Editor/UnrealEd/Private/SScalabilitySettings.cpp b/Engine/Source/Editor/UnrealEd/Private/SScalabilitySettings.cpp
--- a/Engine/Source/Editor/UnrealEd/Private/SScalabilitySettings.cpp
+++ b/Engine/Source/Editor/UnrealEd/Private/SScalabilitySettings.cpp
@@ -260,6 +260,7 @@ void SScalabilitySettings::OnDefaultClicked(ECheckBoxState CheckState)
 	Scalability::ResetQualityLevelsToDefault();
 	CachedQualityLevels = Scalability::GetQualityLevels();
 	Scalability::SetQualityLevels(CachedQualityLevels);
+	Scalability::SaveState(GEditorSettingsIni);
 
 	GEditor->RedrawAllViewports();
 	GEditor->RedrawLevelEditingViewports();
diff --git a/Engine/Source/Editor/UnrealEd/Private/ViewportToolbar/UnrealEdViewportToolbar.cpp b/Engine/Source/Editor/UnrealEd/Private/ViewportToolbar/UnrealEdViewportToolbar.cpp
--- a/Engine/Source/Editor/UnrealEd/Private/ViewportToolbar/UnrealEdViewportToolbar.cpp
+++ b/Engine/Source/Editor/UnrealEd/Private/ViewportToolbar/UnrealEdViewportToolbar.cpp
@@ -3802,6 +3802,7 @@ FToolMenuEntry CreateResetScalabilitySubmenu()
 	ToolbarAction.ExecuteAction.BindLambda([](const FToolMenuContext& Context)
 	{
 		Scalability::ResetQualityLevelsToDefault();
+		Scalability::SaveState(GEditorSettingsIni);
 		GEditor->RedrawAllViewports();
 	});
 				
