[ContentBrowser] Enabled new style cvar by default

diff --git a/Engine/Source/Editor/UnrealEd/Private/UnrealEdGlobals.cpp b/Engine/Source/Editor/UnrealEd/Private/UnrealEdGlobals.cpp
--- a/Engine/Source/Editor/UnrealEd/Private/UnrealEdGlobals.cpp
+++ b/Engine/Source/Editor/UnrealEd/Private/UnrealEdGlobals.cpp
@@ -62,7 +62,7 @@ namespace UE::Editor::ContentBrowser
 	{
 		static FAutoConsoleVariable CVarEnableContentBrowserNewStyle(
 			TEXT("ContentBrowser.EnableNewStyle"),
-			false,
+			true,
 			TEXT("Whether or not to enable the Content Browser restyle (~5.6).")
 		);
 
