Anim Sequence: Fix thumbnail scene references preventing deletion when assets contain notify states

The ActiveNotifyStates array could contain dangling tracked notify state references. Ensuring we reset the anim instance fully for each thumbnail means references do not persist.

diff --git a/Engine/Source/Editor/UnrealEd/Private/ThumbnailHelpers.cpp b/Engine/Source/Editor/UnrealEd/Private/ThumbnailHelpers.cpp
--- a/Engine/Source/Editor/UnrealEd/Private/ThumbnailHelpers.cpp
+++ b/Engine/Source/Editor/UnrealEd/Private/ThumbnailHelpers.cpp
@@ -700,6 +700,7 @@ bool FAnimationSequenceThumbnailScene::SetAnimation(UAnimSequenceBase* InAnimati
 		CleanupComponentChildren(PreviewActor->GetSkeletalMeshComponent());
 		PreviewActor->GetSkeletalMeshComponent()->SetAnimation(NULL);
 		PreviewActor->GetSkeletalMeshComponent()->SetSkeletalMesh(nullptr);
+		PreviewActor->GetSkeletalMeshComponent()->ClearAnimScriptInstance();
 	}
 
 	return bSetSucessfully;
