FLevelEditorViewportClient : remove NormalizeElementsByWidgetMode declaration that has no definition.

This method declaration was added in 40716941 without a corresponding definition.

diff --git a/Engine/Source/Editor/UnrealEd/Public/LevelEditorViewport.h b/Engine/Source/Editor/UnrealEd/Public/LevelEditorViewport.h
--- a/Engine/Source/Editor/UnrealEd/Public/LevelEditorViewport.h
+++ b/Engine/Source/Editor/UnrealEd/Public/LevelEditorViewport.h
@@ -342,7 +342,6 @@ public:
 	 * Get the elements (from the current selection set) that this viewport can manipulate (eg, via the transform gizmo).
 	 */
 	UNREALED_API FTypedElementListConstRef GetElementsToManipulate(const bool bForceRefresh = false);
-	UNREALED_API void NormalizeElementsByWidgetMode(FTypedElementListRef& OutElementList, UE::Widget::EWidgetMode InWidgetMode) const;
 
 	UNREALED_API virtual void SetIsSimulateInEditorViewport( bool bInIsSimulateInEditorViewport ) override;
 
