IncrementalCook: Fix stackoverflow from infinite recursion in the log handler when a log is sent after we register the log handler but before we have finished initializing the cook session.
#rnx

diff --git a/Engine/Source/Editor/UnrealEd/Private/Cooker/CookLog.cpp b/Engine/Source/Editor/UnrealEd/Private/Cooker/CookLog.cpp
--- a/Engine/Source/Editor/UnrealEd/Private/Cooker/CookLog.cpp
+++ b/Engine/Source/Editor/UnrealEd/Private/Cooker/CookLog.cpp
@@ -556,6 +556,13 @@ void FLogHandler::RecordLogForIncrementalCook(UE::Cook::FReplicatedLogData&& Log
 
 void FLogHandler::RecordLogForIncrementalCookGameThreadPortion(FName ActivePackage, FReplicatedLogData&& LogData)
 {
+	if (!COTFS.IsInSession())
+	{
+		// It's illegal to call GetSessionPlatforms below before the cook session has started.
+		// We don't need to record errors before session started for incremental cook, because they come from startup
+		// packages and will be replayed on every cook anyway without our intervention.
+		return;
+	}
 	FPackageData* PackageData = COTFS.PackageDatas->TryAddPackageDataByPackageName(ActivePackage);
 	if (!PackageData)
 	{
