Cooker: Remove log message for CachedCookedPlatformDataObjectsOnDestroyedOutsideOfGC: it causes log spam and is not currently useful because the name it logs is empty.
#rnx

diff --git a/Engine/Source/Editor/UnrealEd/Private/Cooker/CookPackageData.cpp b/Engine/Source/Editor/UnrealEd/Private/Cooker/CookPackageData.cpp
--- a/Engine/Source/Editor/UnrealEd/Private/Cooker/CookPackageData.cpp
+++ b/Engine/Source/Editor/UnrealEd/Private/Cooker/CookPackageData.cpp
@@ -3226,12 +3226,7 @@ void FPackageDatas::CachedCookedPlatformDataObjectsPostGarbageCollect(
 
 void FPackageDatas::CachedCookedPlatformDataObjectsOnDestroyedOutsideOfGC(const UObject* DestroyedObject)
 {
-	FCachedCookedPlatformDataState RemovedState;
-	if (CachedCookedPlatformDataObjects.RemoveAndCopyValue(const_cast<UObject*>(DestroyedObject), RemovedState))
-	{
-		UE_LOG(LogCook, Display, TEXT("Removing CachedCookedPlatformData state for reallocated object %s."),
-			*DestroyedObject->GetPathName());
-	}
+	CachedCookedPlatformDataObjects.Remove(DestroyedObject);
 }
 
 void FCachedCookedPlatformDataState::AddRefFrom(FPackageData* PackageData)
