Fixed crash on FBX export

#rnx

diff --git a/Engine/Source/Editor/UnrealEd/Private/Fbx/FbxMainExport.cpp b/Engine/Source/Editor/UnrealEd/Private/Fbx/FbxMainExport.cpp
--- a/Engine/Source/Editor/UnrealEd/Private/Fbx/FbxMainExport.cpp
+++ b/Engine/Source/Editor/UnrealEd/Private/Fbx/FbxMainExport.cpp
@@ -973,8 +973,8 @@ void FFbxExporter::ExportStaticMesh(AActor* Actor, UStaticMeshComponent* StaticM
 	}
 	else
 	{
-		const int32 LODIndex = (StaticMeshComponent->ForcedLodModel > 0 ? StaticMeshComponent->ForcedLodModel - 1 : /* auto-select*/ 0);
-		if (LODIndex != INDEX_NONE && LODIndex < StaticMeshComponent->LODData.Num())
+		const int32 LODIndex = (StaticMeshComponent->ForcedLodModel > 0 ? FMath::Min(StaticMesh->GetNumLODs(), StaticMeshComponent->ForcedLodModel) - 1 : /* auto-select*/ 0);
+		if (StaticMeshComponent->LODData.IsValidIndex(LODIndex))
 		{
 			ColorBuffer = StaticMeshComponent->LODData[LODIndex].OverrideVertexColors;
 		}
@@ -2882,19 +2882,19 @@ FbxNode* FFbxExporter::ExportActor(AActor* Actor, bool bExportComponents, INodeN
 
 				if (StaticMeshComp && StaticMeshComp->GetStaticMesh())
 				{
+					const int32 LODIndex = (StaticMeshComp->ForcedLodModel > 0 ? FMath::Min(StaticMeshComp->GetStaticMesh()->GetNumLODs(), StaticMeshComp->ForcedLodModel) - 1 : /* auto-select*/ 0);
 					if (USplineMeshComponent* SplineMeshComp = Cast<USplineMeshComponent>(StaticMeshComp))
 					{
 						//TODO: Validate Spline material baking
 						//(Spline mesh staticmesh does not take into account the Spline transforms, might not be accurate/enough to use Spline->StaticMesh for material baking.)
-						ExportSplineMeshToFbx(SplineMeshComp, *SplineMeshComp->GetName(), ExportNode, FFbxMaterialBakingMeshData(SplineMeshComp->GetStaticMesh(), SplineMeshComp, (SplineMeshComp->ForcedLodModel > 0 ? SplineMeshComp->ForcedLodModel - 1 : /* auto-select*/ 0)));
+						ExportSplineMeshToFbx(SplineMeshComp, *SplineMeshComp->GetName(), ExportNode, FFbxMaterialBakingMeshData(SplineMeshComp->GetStaticMesh(), SplineMeshComp, LODIndex));
 					}
 					else if (UInstancedStaticMeshComponent* InstancedMeshComp = Cast<UInstancedStaticMeshComponent>(StaticMeshComp))
 					{
-						ExportInstancedMeshToFbx(InstancedMeshComp, *InstancedMeshComp->GetName(), ExportNode, FFbxMaterialBakingMeshData(InstancedMeshComp->GetStaticMesh(), InstancedMeshComp, (InstancedMeshComp->ForcedLodModel > 0 ? InstancedMeshComp->ForcedLodModel - 1 : /* auto-select*/ 0)));
+						ExportInstancedMeshToFbx(InstancedMeshComp, *InstancedMeshComp->GetName(), ExportNode, FFbxMaterialBakingMeshData(InstancedMeshComp->GetStaticMesh(), InstancedMeshComp, LODIndex));
 					}
 					else
 					{
-						const int32 LODIndex = (StaticMeshComp->ForcedLodModel > 0 ? StaticMeshComp->ForcedLodModel - 1 : /* auto-select*/ 0);
 						const int32 LightmapUVChannel = -1;
 						const TArray<FStaticMaterial>* MaterialOrderOverride = nullptr;
 						const FColorVertexBuffer* ColorBuffer = nullptr;
@@ -5555,7 +5555,12 @@ void FFbxExporter::ExportSplineMeshToFbx(const USplineMeshComponent* SplineMeshC
 	const UStaticMesh* StaticMesh = SplineMeshComp->GetStaticMesh();
 	check(StaticMesh);
 
-	const int32 LODIndex = (SplineMeshComp->ForcedLodModel > 0 ? SplineMeshComp->ForcedLodModel - 1 : /* auto-select*/ 0);
+	const int32 LODIndex = (SplineMeshComp->ForcedLodModel > 0 ? FMath::Min(StaticMesh->GetNumLODs(), SplineMeshComp->ForcedLodModel) - 1 : /* auto-select*/ 0);
+	if (!SplineMeshComp->LODData.IsValidIndex(LODIndex))
+	{
+		return;
+	}
+
 	const FStaticMeshLODResources& RenderMesh = StaticMesh->GetLODForExport(LODIndex);
 
 	// Verify the integrity of the static mesh.
@@ -5794,7 +5799,7 @@ void FFbxExporter::ExportInstancedMeshToFbx(const UInstancedStaticMeshComponent*
 	const UStaticMesh* StaticMesh = InstancedMeshComp->GetStaticMesh();
 	check(StaticMesh);
 
-	const int32 LODIndex = (InstancedMeshComp->ForcedLodModel > 0 ? InstancedMeshComp->ForcedLodModel - 1 : /* auto-select*/ 0);
+	const int32 LODIndex = (InstancedMeshComp->ForcedLodModel > 0 ? FMath::Min(StaticMesh->GetNumLODs(), InstancedMeshComp->ForcedLodModel) - 1 : /* auto-select*/ 0);
 	const int32 NumInstances = InstancedMeshComp->GetInstanceCount();
 	for (int32 InstanceIndex = 0; InstanceIndex < NumInstances; ++InstanceIndex)
 	{
