MPCook: Add more feedback for the unexpected message "which is not a pending package" when parsing results from CookWorkers. It is firing even in non-retraction and non-generator packages, which is unexpected.
#rnx

diff --git a/Engine/Source/Editor/UnrealEd/Private/Cooker/CookWorkerServer.cpp b/Engine/Source/Editor/UnrealEd/Private/Cooker/CookWorkerServer.cpp
--- a/Engine/Source/Editor/UnrealEd/Private/Cooker/CookWorkerServer.cpp
+++ b/Engine/Source/Editor/UnrealEd/Private/Cooker/CookWorkerServer.cpp
@@ -1053,10 +1053,10 @@ void FCookWorkerServer::RecordResults(FPackageResultsMessage& Message)
 		{
 			UE_LOG(LogCook, Display,
 				TEXT("CookWorkerServer %d received FPackageResultsMessage for package %s which is not a pending package. Ignoring it.")
-				TEXT("\n\tState = %s, WorkerId = %s, CookResults = { %s }"),
+				TEXT("\n\tState = %s, WorkerId = %s, CookResults = { %s }, Result.GetSuppressCookReason = %s"),
 				ProfileId, *Result.GetPackageName().ToString(),
 				LexToString(PackageData->GetState()), *Director.GetDisplayName(PackageData->GetWorkerAssignment()),
-				*WriteCookStatus(*PackageData, COTFS.GetSessionPlatforms()));
+				*WriteCookStatus(*PackageData, COTFS.GetSessionPlatforms()), LexToString(Result.GetSuppressCookReason()));
 			continue;
 		}
 		int32 NumPlatforms = OrderedSessionPlatforms.Num();
