Anim Sequence editor
* Fix changing additive animation modes with root motion on causes inconsistent root position when going from an additive mode to none
* Fix a bug where starting an animation with a non-zero playback time (like rewind debugger can do) would cause an incorrect root position

#tests editor animsandbox

diff --git a/Engine/Source/Editor/UnrealEd/Classes/Animation/DebugSkelMeshComponent.h b/Engine/Source/Editor/UnrealEd/Classes/Animation/DebugSkelMeshComponent.h
--- a/Engine/Source/Editor/UnrealEd/Classes/Animation/DebugSkelMeshComponent.h
+++ b/Engine/Source/Editor/UnrealEd/Classes/Animation/DebugSkelMeshComponent.h
@@ -285,6 +285,9 @@ class UDebugSkelMeshComponent : public USkeletalMeshComponent
 	/** Playback time last time ConsumeRootmotion was called */
 	UPROPERTY(transient)
 	float ConsumeRootMotionPreviousPlaybackTime;
+	
+	UPROPERTY(transient)
+	TEnumAsByte<EAdditiveAnimationType> ConsumeRootMotionPreviousAdditiveAnimType = AAT_MAX;
 
 	UPROPERTY(transient)
 	uint32 bShowClothData : 1;
@@ -634,6 +637,8 @@ private:
 	UNREALED_API void RebuildCachedClothBounds();
 
 	UNREALED_API void SetProcessRootMotionModeInternal(EProcessRootMotionMode Mode);
+
+	void ResetRootTransform();
 protected:
 
 	// Overridden to support single clothing ticks
diff --git a/Engine/Source/Editor/UnrealEd/Private/Animation/DebugSkelMeshComponent.cpp b/Engine/Source/Editor/UnrealEd/Private/Animation/DebugSkelMeshComponent.cpp
--- a/Engine/Source/Editor/UnrealEd/Private/Animation/DebugSkelMeshComponent.cpp
+++ b/Engine/Source/Editor/UnrealEd/Private/Animation/DebugSkelMeshComponent.cpp
@@ -215,9 +215,16 @@ void UDebugSkelMeshComponent::ConsumeRootMotion(const FVector& FloorMin, const F
 	// such as physic simulation or turn table rotation may change the transform. 
 	if (PreviewInstance == nullptr)
 	{
+		ConsumeRootMotionPreviousAdditiveAnimType = AAT_MAX;
 		return;
 	}
 
+	if (PreviewInstance->GetAnimSequence() != nullptr && ConsumeRootMotionPreviousAdditiveAnimType != PreviewInstance->GetAnimSequence()->GetAdditiveAnimType())
+	{
+		ConsumeRootMotionPreviousAdditiveAnimType = PreviewInstance->GetAnimSequence()->GetAdditiveAnimType();
+		ResetRootTransform();
+	}
+
 	// Force ProcessRootMotionMode to Ignore if the current asset/animation blueprint is not using root motion. 
 	if (ProcessRootMotionMode != EProcessRootMotionMode::Ignore && DoesCurrentAssetHaveRootMotion() == false)
 	{
@@ -450,7 +457,11 @@ void UDebugSkelMeshComponent::SetProcessRootMotionMode(EProcessRootMotionMode Mo
 void UDebugSkelMeshComponent::SetProcessRootMotionModeInternal(EProcessRootMotionMode Mode)
 {
 	ProcessRootMotionMode = Mode;
+	ResetRootTransform();
+}
 
+void UDebugSkelMeshComponent::ResetRootTransform()
+{
 	FTransform RootMotionTransform = FTransform::Identity;
 
 	if (DoesCurrentAssetHaveRootMotion())
@@ -459,6 +470,7 @@ void UDebugSkelMeshComponent::SetProcessRootMotionModeInternal(EProcessRootMotio
 		{
 			// Reset transform
 			const float CurrentTime = PreviewInstance->GetCurrentTime();
+			ConsumeRootMotionPreviousPlaybackTime = CurrentTime; // Important to cycle this otherwise root motion will be double applied
 			float SectionStartPosition = 0.0f;
 			if (const UAnimMontage* Montage = Cast<UAnimMontage>(PreviewInstance->CurrentAsset))
 			{
