Move display of Asset Editor's ToolbarWidgets into the primary toolbar.
- Adjust padding of SAssetFamilyShortcutBar to better fit with the new location.

#rnx

diff --git a/Engine/Source/Editor/Persona/Private/SAssetFamilyShortcutBar.cpp b/Engine/Source/Editor/Persona/Private/SAssetFamilyShortcutBar.cpp
--- a/Engine/Source/Editor/Persona/Private/SAssetFamilyShortcutBar.cpp
+++ b/Engine/Source/Editor/Persona/Private/SAssetFamilyShortcutBar.cpp
@@ -559,7 +559,7 @@ void SAssetFamilyShortcutBar::BuildShortcuts()
 		FAssetData AssetData = AssetFamily->FindAssetOfType(Class);
 		HorizontalBox->AddSlot()
 		.AutoWidth()
-		.Padding(0.0f, 4.0f, 16.0f, 4.0f)
+		.Padding(0.0f, 0.0f, 16.0f, 0.0f)
 		[
 			SNew(SAssetShortcut, WeakHostingApp.Pin().ToSharedRef(), AssetFamily.ToSharedRef(), AssetData, ThumbnailPool.ToSharedRef())
 			.Visibility_Lambda([Class]()
diff --git a/Engine/Source/Editor/UnrealEd/Private/Toolkits/AssetEditorToolkit.cpp b/Engine/Source/Editor/UnrealEd/Private/Toolkits/AssetEditorToolkit.cpp
--- a/Engine/Source/Editor/UnrealEd/Private/Toolkits/AssetEditorToolkit.cpp
+++ b/Engine/Source/Editor/UnrealEd/Private/Toolkits/AssetEditorToolkit.cpp
@@ -1438,6 +1438,29 @@ void FAssetEditorToolkit::GenerateToolbar()
 	if (!FoundMenu || !FoundMenu->IsRegistered())
 	{
 		FoundMenu = ToolMenus->RegisterMenu(ToolBarName, ParentToolbarName, EMultiBoxType::SlimHorizontalToolBar);
+		
+		FToolMenuSection& MiscSection = FoundMenu->AddSection("MiscWidgets");
+		MiscSection.Alignment = EToolMenuSectionAlign::Last;
+		MiscSection.AddDynamicEntry("ToolbarWidgets", FNewToolMenuSectionDelegate::CreateLambda([](FToolMenuSection& InSection)
+		{
+			UAssetEditorToolkitMenuContext* Context = InSection.FindContext<UAssetEditorToolkitMenuContext>();
+			if (!Context)
+			{
+				return;
+			}
+			
+			TSharedPtr<FAssetEditorToolkit> Toolkit = Context->Toolkit.Pin();
+			if (!Toolkit)
+			{
+				return;
+			}
+			
+			for (const TSharedRef<SWidget>& Widget : Toolkit->ToolbarWidgets)
+			{
+				FToolMenuEntry& Entry = InSection.AddEntry(FToolMenuEntry::InitWidget(NAME_None, Widget, FText::GetEmpty(), true, false, true));
+				Entry.ToolBarData.ResizeParams.AllowClipping = false;
+			}
+		}));
 	}
 
 	FToolMenuContext MenuContext(GetToolkitCommands(), Extender);
@@ -1459,25 +1482,6 @@ void FAssetEditorToolkit::GenerateToolbar()
 	UToolMenu* ReadOnlyToolbar = GenerateReadOnlyToolbar(MenuContext);
 	TSharedRef< class SWidget > ReadOnlyWidget = ToolMenus->GenerateWidget(ReadOnlyToolbar);
 
-	TSharedRef<SWidget> MiscWidgets = SNullWidget::NullWidget;
-
-	if(ToolbarWidgets.Num() > 0)
-	{
-		TSharedRef<SHorizontalBox> MiscWidgetsHBox = SNew(SHorizontalBox);
-
-		for (int32 WidgetIdx = 0; WidgetIdx < ToolbarWidgets.Num(); ++WidgetIdx)
-		{
-			MiscWidgetsHBox->AddSlot()
-				.AutoWidth()
-				.VAlign(VAlign_Center)
-				[
-					ToolbarWidgets[WidgetIdx]
-				];
-		}
-
-		MiscWidgets = MiscWidgetsHBox;
-	}
-	
 	Toolbar = 
 		SNew(SHorizontalBox)
 		+SHorizontalBox::Slot()
@@ -1498,18 +1502,6 @@ void FAssetEditorToolkit::GenerateToolbar()
 		+SHorizontalBox::Slot()
 		.HAlign(HAlign_Right)
 		.AutoWidth()
-		[
-			SNew(SBorder)
-			.VAlign(VAlign_Center)
-			.BorderImage(FAppStyle::Get().GetBrush("AssetEditorToolbar.Background"))
-			.Padding(FMargin(0.0f))
-			[
-				MiscWidgets
-			]
-		]
-		+SHorizontalBox::Slot()
-		.HAlign(HAlign_Right)
-		.AutoWidth()
 		[
 			SNew(SBorder)
 			.VAlign(VAlign_Center)
