IncrementalCook: Increase the warning timeout for waiting on previous build results from the oplog. 17 seconds can currently occur on large projects because we are saturating the bandwidth for connections to the zenserver.
#rnx

diff --git a/Engine/Source/Editor/UnrealEd/Private/Cooker/CookRequestCluster.cpp b/Engine/Source/Editor/UnrealEd/Private/Cooker/CookRequestCluster.cpp
--- a/Engine/Source/Editor/UnrealEd/Private/Cooker/CookRequestCluster.cpp
+++ b/Engine/Source/Editor/UnrealEd/Private/Cooker/CookRequestCluster.cpp
@@ -1215,7 +1215,7 @@ void FRequestCluster::FGraphSearch::ResolveTransitiveBuildDependencyCycle()
 
 void FRequestCluster::FGraphSearch::UpdateDisplay()
 {
-	constexpr double WarningTimeout = 10.0;
+	constexpr double WarningTimeout = 60.0;
 	if (FPlatformTime::Seconds() > LastActivityTime + WarningTimeout && Cluster.IsIncrementalCook())
 	{
 		FScopeLock ScopeLock(&Lock);
