IncrementalCook: Naming Convention: Iterative renamed to Incremental. This is the file renames following the previous change to file contents.
#rnx
#incrementalcook

diff --git a/Engine/Build/Graph/Tests/IncrementalCompare.xml b/Engine/Build/Graph/Tests/IncrementalCompare.xml
new file mode 100644
--- /dev/null
+++ b/Engine/Build/Graph/Tests/IncrementalCompare.xml
@@ -0,0 +1,137 @@
+<?xml version='1.0' ?>
+<BuildGraph xmlns="http://www.epicgames.com/BuildGraph" xmlns:xsi="http://www.w3.org/2001/XMLSchema-instance" xsi:schemaLocation="http://www.epicgames.com/BuildGraph ../Schema.xsd" >
+
+	<!-- Project Options -->
+	<Option Name="ProjectFile" DefaultValue="_Unknown_" Restrict=".+" Description="The projects to cook"/>
+	<Option Name="ProjectName" DefaultValue="_Unknown_" Restrict=".+" Description="The projects to cook"/>
+	<Option Name="ProjectDir" DefaultValue="_Unknown_" Restrict=".+" Description="The projects to cook"/>
+	<Option Name="EditorTarget" DefaultValue="$(ProjectName)Editor" Description="Compile name of the editor target"/>
+
+	<!-- Note we can not make the default='', because the default is used for some validation. -->
+	<Error Message="No project file specified" If="'$(ProjectFile)' == '_Unknown_'"/>
+	<Error Message="No project name specified" If="'$(ProjectName)' == '_Unknown_'"/>
+	<Error Message="No project dir specified" If="'$(ProjectDir)' == '_Unknown_'"/>
+
+	<!-- Platform Options -->
+	<Option Name="TargetPlatform" DefaultValue="Windows" Restrict=".+" Description="The platform to cook"/>
+
+	<!-- Agent Options -->
+	<Option Name="AgentType" DefaultValue="" Description="The platform to cook"/>
+	<Property Name="AgentType" Value="IncrementalCompareAgent_$(ProjectName)_$(TargetPlatform)" If="'$(AgentType)' == ''"/>
+
+	<Option Name="WithUBA" Restrict="true|false" DefaultValue="true" Description="Whether to use UBA compile or not"/>
+	<Property Name="UBACompileArgs" Value=""/>
+	<Property Name="UBACompileArgs" Value="-UBA -ActionStallReportTime=300 -UBALinkRemote" If="$(WithUBA)"/>
+	
+	<!-- Agent=IncrementalCompareAgent variables -->
+	<Option Name="CookProcessCount" DefaultValue="" Description="Number of processes to use during cooking, default is ini-driven."/>
+	<Option Name="MPCookCoreSubscription" DefaultValue="" Description="Over or Under subscription factor for CPU cores to use during cooking, default is ini-driven."/>
+	<Option Name="WithNullCook" Restrict="true|false" DefaultValue="false" Description="Whether to do a null cook or not"/>
+	<Option Name="NullCookCustomConfig" DefaultValue="" Description="What custom configuration to apply if doing a null cook (see -CustomConfig commandline argument)"/>
+	<Property Name="FullPathProjectDir" Value="$(RootDir)/$(ProjectDir)"/>
+	<Property Name="EditorPlatform" Value="Win64"/>
+	<Property Name="CommonCookArgs" Value="-targetplatform=$(TargetPlatform) -cookcultures=en -unversioned -zenstore -cookincremental -IgnoreBlueprintCookDependencies"/>
+	<Property Name="CommonCookArgs" Value="$(CommonCookArgs) -CookProcessCount=$(CookProcessCount)" If="'$(CookProcessCount)' != ''"/>
+	<Property Name="CommonCookArgs" Value="$(CommonCookArgs) -MPCookCoreSubscription=$(MPCookCoreSubscription)" If="'$(MPCookCoreSubscription)' != ''"/>
+	<Property Name="PlatformCookedDir" Value="$(FullPathProjectDir)/Saved/Cooked/$(TargetPlatform)"/>
+	<Property Name="CompareDirIncremental" Value="$(FullPathProjectDir)/Intermediate/IncrementalCompare/Iter/$(TargetPlatform)"/>
+	<Property Name="CompareDirRecook" Value="$(FullPathProjectDir)/Intermediate/IncrementalCompare/Recook/$(TargetPlatform)"/>
+	<Property Name="NewLine" Value="&#10;"/>
+
+	<Property Name="CompileNodeName" Value="Compile $(EditorPlatform) $(ProjectName)"/>
+	<Property Name="NaturalIncrementalCookNodeName" Value="Natural Incremental Cook $(ProjectName) $(TargetPlatform)"/>
+	<Property Name="StoreIncrementalCookNodeName" Value="Store Incremental Cook $(ProjectName) $(TargetPlatform)"/>
+	<Property Name="PrimingNullIncrementalCookNodeName" Value="Priming Null Incremental Cook $(ProjectName) $(TargetPlatform)"/>
+	<Property Name="NullIncrementalCookNodeName" Value="Null Incremental Cook $(ProjectName) $(TargetPlatform)"/>
+	<Property Name="CleanNullIncrementalCookNodeName" Value="Clean Null Incremental Cook $(ProjectName) $(TargetPlatform)"/>
+	<Property Name="FullCookNodeName" Value="Full Cook $(ProjectName) $(TargetPlatform)"/>
+	<Property Name="StoreFullCookNodeName" Value="Store Full Cook $(ProjectName) $(TargetPlatform)"/>
+	<Property Name="DiffCooksNodeName" Value="Diff Cooks $(ProjectName) $(TargetPlatform)"/>
+
+	<Agent Name="IncrementalCompareAgent" Type="$(AgentType)">
+		<Node Name="$(CompileNodeName)">
+			<!-- Compile. -->
+			<Log Message="$(NewLine)========== IncrementalCompare.xml: Compile =========="/>
+			<Compile Target="ShaderCompileWorker" Platform="$(EditorPlatform)" Project="$(ProjectFile)" Configuration="Development" Arguments="$(UBACompileArgs)"/>
+			<Compile Target="ZenLaunch" Platform="$(EditorPlatform)" Configuration="Development" Arguments="$(UBACompileArgs)"/>
+			<Compile Target="$(EditorTarget)" Platform="$(EditorPlatform)" Project="$(ProjectFile)" Configuration="Development" Arguments="$(UBACompileArgs)"/>
+		</Node>
+		<Node Name="$(NaturalIncrementalCookNodeName)" Requires="$(CompileNodeName)">
+			<!-- Execute the incremental cook (we don't pass -forcerecook). -->
+			<Log Message="$(NewLine)========== IncrementalCompare.xml: Incremental Cook =========="/>
+			<Commandlet Name="Cook" Project="$(ProjectFile)" Arguments="$(CommonCookArgs)"/>
+		</Node>
+		<Node Name="$(StoreIncrementalCookNodeName)" Requires="$(NaturalIncrementalCookNodeName)">
+			<!-- Move all cook artifacts to the compare_incremental dir. -->
+			<Log Message="$(NewLine)========== IncrementalCompare.xml: Move =========="/>
+			<Log Message="Delete Directories=$(CompareDirIncremental)"/>
+			<Delete Directories="$(CompareDirIncremental)"/>
+			<Do If="$(WithNullCook)">
+				<Log Message="Copy From=$(PlatformCookedDir)/*.* To=$(CompareDirIncremental)"/>
+				<Copy From="$(PlatformCookedDir)/*.*" To="$(CompareDirIncremental)" ErrorIfNotFound="true"/>
+			</Do>
+			<Do If="!$(WithNullCook)">
+				<Log Message="Move From=$(PlatformCookedDir)/*.* To=$(CompareDirIncremental)"/>
+				<Move From="$(PlatformCookedDir)/*.*" To="$(CompareDirIncremental)" ErrorIfNotFound="true"/>
+				<Log Message="Delete Directories=$(PlatformCookedDir)"/>
+				<Delete Directories="$(PlatformCookedDir)"/>
+			</Do>
+
+			<!-- Write the files from incremental cook's zenstore out into CompareDirIncremental. -->
+			<Log Message="$(NewLine)========== IncrementalCompare.xml: ZenMirror Incremental Cook =========="/>
+			<Log Message="ZenMirror Project=$(ProjectFile) Platform=$(TargetPlatform) DestinationFileDir=$(CompareDirIncremental)"/>
+			<ZenMirror Project="$(ProjectFile)" Platform="$(TargetPlatform)" DestinationFileDir="$(CompareDirIncremental)"/>
+		</Node>
+		<Property Name="FullCookPrequisiteNodeName" Value="$(StoreIncrementalCookNodeName)"/>
+		<Do If="$(WithNullCook)">
+			<Node Name="$(PrimingNullIncrementalCookNodeName)" Requires="$(StoreIncrementalCookNodeName)">
+				<!-- Execute the null incremental cook. -->
+				<Log Message="$(NewLine)========== IncrementalCompare.xml: Priming Null Incremental Cook =========="/>
+				<Property Name="NullCookArgs" Value="$(CommonCookArgs) -dpcvars=cook.displaymode=2"/>
+				<Property Name="NullCookArgs" Value="$(NullCookArgs) -CustomConfig=$(NullCookCustomConfig)" If="'$(NullCookCustomConfig)' != ''"/>
+				<Commandlet Name="Cook" Project="$(ProjectFile)" Arguments="$(NullCookArgs)"/>
+			</Node>
+			<Node Name="$(NullIncrementalCookNodeName)" Requires="$(PrimingNullIncrementalCookNodeName)">
+				<!-- Execute the null incremental cook. -->
+				<Log Message="$(NewLine)========== IncrementalCompare.xml: Null Incremental Cook =========="/>
+				<Property Name="NullCookArgs" Value="$(CommonCookArgs) -dpcvars=cook.displaymode=2"/>
+				<Property Name="NullCookArgs" Value="$(NullCookArgs) -CustomConfig=$(NullCookCustomConfig)" If="'$(NullCookCustomConfig)' != ''"/>
+				<Commandlet Name="Cook" Project="$(ProjectFile)" Arguments="$(NullCookArgs)"/>
+			</Node>
+			<Node Name="$(CleanNullIncrementalCookNodeName)" Requires="$(NullIncrementalCookNodeName)">
+				<!-- Move all cook artifacts to the compare_incremental dir. -->
+				<Log Message="$(NewLine)========== IncrementalCompare.xml: Clean Null =========="/>
+				<Log Message="Delete Directories=$(PlatformCookedDir)"/>
+				<Delete Directories="$(PlatformCookedDir)"/>
+			</Node>
+			<Property Name="FullCookPrequisiteNodeName" Value="$(CleanNullIncrementalCookNodeName)"/>
+		</Do>
+		<Node Name="$(FullCookNodeName)" Requires="$(FullCookPrequisiteNodeName)">
+			<!-- Execute the forcerecook cook. -->
+			<Log Message="$(NewLine)========== IncrementalCompare.xml: Non-Incremental Cook =========="/>
+			<Commandlet Name="Cook" Project="$(ProjectFile)" Arguments="$(CommonCookArgs) -forcerecook"/>
+		</Node>
+		<Node Name="$(StoreFullCookNodeName)" Requires="$(FullCookNodeName)">
+			<!-- Copy all cook artifacts to the compare_incremental dir. Copy rather than move so they remain
+			     behind for the next instance of this job to use in an incremental cook. -->
+			<Log Message="$(NewLine)========== IncrementalCompare.xml: Copy =========="/>
+			<Log Message="Delete Directories=$(CompareDirRecook)"/>
+			<Delete Directories="$(CompareDirRecook)"/>
+			<Log Message="Copy From=$(PlatformCookedDir)/*.* To=$(CompareDirRecook)"/>
+			<Copy From="$(PlatformCookedDir)/*.*" To="$(CompareDirRecook)" ErrorIfNotFound="true"/>
+
+			<!-- Write the files from forcerecook's zenstore out into CompareDirRecook. -->
+			<Log Message="$(NewLine)========== IncrementalCompare.xml: ZenMirror ForceRecook Cook =========="/>
+			<Log Message="ZenMirror Project=$(ProjectFile) Platform=$(TargetPlatform) DestinationFileDir=$(CompareDirRecook)"/>
+			<ZenMirror Project="$(ProjectFile)" Platform="$(TargetPlatform)" DestinationFileDir="$(CompareDirRecook)"/>
+		</Node>
+		<Node Name="$(DiffCooksNodeName)" Requires="$(StoreFullCookNodeName)">
+			<!-- Compare the incremental cook with the non-incremental cook. -->
+			<Log Message="$(NewLine)========== IncrementalCompare.xml: DiffCook =========="/>
+			<Commandlet Name="DiffCook" Project="$(ProjectFile)" Arguments="-base=$(CompareDirRecook) -comp=$(CompareDirIncremental) -targetplatform=$(TargetPlatform) -addedverbosity=error -removedverbosity=error -modifiedverbosity=warning -hide=modifiedpackages"/>
+		</Node>
+	</Agent>
+
+	<Aggregate Name="IncrementalCompare" Requires="$(DiffCooksNodeName)"/>
+
+</BuildGraph>
\ No newline at end of file
diff --git a/Engine/Build/Graph/Tests/IncrementalValidate.xml b/Engine/Build/Graph/Tests/IncrementalValidate.xml
new file mode 100644
--- /dev/null
+++ b/Engine/Build/Graph/Tests/IncrementalValidate.xml
@@ -0,0 +1,90 @@
+<?xml version='1.0' ?>
+<BuildGraph xmlns="http://www.epicgames.com/BuildGraph" xmlns:xsi="http://www.w3.org/2001/XMLSchema-instance" xsi:schemaLocation="http://www.epicgames.com/BuildGraph ../Schema.xsd" >
+
+	<!-- Project Options -->
+	<Option Name="ProjectFile" DefaultValue="_Unknown_" Restrict=".+" Description="The projects to cook"/>
+	<Option Name="ProjectName" DefaultValue="_Unknown_" Restrict=".+" Description="The projects to cook"/>
+	<Option Name="EditorTarget" DefaultValue="$(ProjectName)Editor" Description="Compile name of the editor target"/>
+
+	<!-- Note we can not make the default='', because the default is used for some validation. -->
+	<Error Message="No project file specified" If="'$(ProjectFile)' == '_Unknown_'"/>
+	<Error Message="No project name specified" If="'$(ProjectName)' == '_Unknown_'"/>
+
+	<!-- Platform Options -->
+	<Option Name="TargetPlatform" DefaultValue="Windows" Restrict=".+" Description="The platform to cook"/>
+
+	<!-- Agent Options -->
+	<Option Name="AgentType" DefaultValue="" Description="The platform to cook"/>
+	<Property Name="AgentType" Value="IterativeValidateAgent_$(ProjectName)_$(TargetPlatform)" If="'$(AgentType)' == ''"/>
+
+	<Option Name="WithUBA" Restrict="true|false" DefaultValue="true" Description="Whether to use UBA compile or not"/>
+	<Property Name="UBACompileArgs" Value=""/>
+	<Property Name="UBACompileArgs" Value="-UBA -ActionStallReportTime=300 -UBALinkRemote" If="$(WithUBA)"/>
+	
+	<!-- IncrementalValidate Options -->
+	<Option Name="IncrementalMode" DefaultValue="incremental" Restrict="legacyiterative|incremental" Description="Whether to cook using Incremental or LegacyIterative"/>
+	<Option Name="ValidateMode" DefaultValue="onephase" Restrict="onephase|twophase" Description="Whether to validate results in onephase (which reports indeterminism as a FalsePositiveIncrementalSkip error) or two phase (which detects indeterminism and downgrades it to a display message)"/>
+	<Property Name="IncrementalArgs" Value="-zenstore -cookincremental" If="'$(IncrementalMode)' == 'incremental'"/>
+	<Property Name="IncrementalArgs" Value="-LegacyIterative -ini:Editor:[CookSettings]:CookIncremental=false -ini:Editor:[CookSettings]:LegacyIterativeUseClassFilters=true" If="'$(IncrementalMode)' == 'legacyiterative'"/>
+
+	<!-- Profiling Options -->
+	<Option Name="ProfilingFilesPublishRootPath" DefaultValue="" Description="Where to publish profiling data files."/>
+	<Option Name="WithCookUTRACECapture" DefaultValue="false" Restrict="true|false" Description="Whether to capture UTRACE data during cooks. Only true or false"/>
+
+	<!-- Config Options -->
+	<Option Name="CustomConfig" DefaultValue="" Description="What custom configuration to apply (see -CustomConfig commandline argument)"/>
+
+	<!-- Agent=IncrementalValidate variables -->
+	<Property Name="EditorPlatform" Value="Win64"/>
+	<Property Name="TraceArgs" Value=""/>
+	<Property Name="TraceArgs" Value="-trace=cpu,loadtime,savetime,cook,contextswitch" If="$(WithCookUTRACECapture) and '$(ProfilingFilesPublishRootPath)' != ''"/>
+	<Property Name="ProfilingFilesLocalRootPath" Value="$(RootDir)/Engine/Intermediate/Profiling"/>
+
+	<Property Name="CommonCookArgs" Value="-targetplatform=$(TargetPlatform) $(IncrementalArgs)"/>
+	<Property Name="CommonCookArgs" Value="$(CommonCookArgs) -CustomConfig=$(CustomConfig)" If="'$(CustomConfig)' != ''"/>
+
+	<Property Name="CompileNodeName" Value="Compile $(EditorPlatform) $(ProjectName)"/>
+	<Property Name="IncrementalValidateSinglePhaseName" Value="Cook IncrementalValidateSinglePhase $(ProjectName) $(TargetPlatform)"/>
+	<Property Name="IncrementalValidatePhase1Name" Value="Cook IncrementalValidatePhase1 $(ProjectName) $(TargetPlatform)"/>
+	<Property Name="IncrementalValidatePhase2Name" Value="Cook IncrementalValidatePhase2 $(ProjectName) $(TargetPlatform)"/>
+
+	<Agent Name="IncrementalValidateAgent" Type="$(AgentType)">
+		<Node Name="$(CompileNodeName)">
+			<Compile Target="ShaderCompileWorker" Platform="$(EditorPlatform)" Project="$(ProjectFile)" Configuration="Development" Arguments="$(UBACompileArgs)"/>
+			<Compile Target="$(EditorTarget)" Platform="$(EditorPlatform)" Project="$(ProjectFile)" Configuration="Development" Arguments="$(UBACompileArgs)"/>
+		</Node>
+		<Node Name="$(IncrementalValidateSinglePhaseName)" Requires="$(CompileNodeName)">
+			<Property Name="TraceArgsFinal" Value=""/>
+			<Property Name="TraceArgsFinal" Value="$(TraceArgs) -tracefile=&quot;$(ProfilingFilesLocalRootPath)/$(ProjectName)_IncrementalValidate.utrace&quot;" If="$(WithCookUTRACECapture) and '$(ProfilingFilesPublishRootPath)' != ''"/>
+			<Do If="'$(TraceArgsFinal)' != '' and Exists('$(ProfilingFilesLocalRootPath)')">
+				<Delete Files="$(ProfilingFilesLocalRootPath)/..."/>
+			</Do>
+			<Commandlet Name="Cook" Project="$(ProjectFile)" Arguments="$(CommonCookArgs) $(TraceArgsFinal) -IncrementalValidate -IncrementalValidateAllowWrite"/>
+			<Do If="'$(TraceArgsFinal)' != ''">
+				<Zip FromDir="$(ProfilingFilesLocalRootPath)" Files="*.utrace" ZipFile="$(ProfilingFilesLocalRootPath)/$(ProjectName)_IncrementalValidate.zip"/>
+				<Copy Files="$(ProjectName)_IncrementalValidate.zip" From="$(ProfilingFilesLocalRootPath)" To="$(ProfilingFilesPublishRootPath)\$(EscapedBranch)-CL-$(Change)"/>
+			</Do>
+		</Node>
+		<Node Name="$(IncrementalValidatePhase1Name)" Requires="$(CompileNodeName)">
+			<Property Name="TraceArgsFinal" Value=""/>
+			<Property Name="TraceArgsFinal" Value="$(TraceArgs) -tracefile=&quot;$(ProfilingFilesLocalRootPath)/$(ProjectName)_IncrementalValidatePhase1.utrace&quot;" If="$(WithCookUTRACECapture) and '$(ProfilingFilesPublishRootPath)' != ''"/>
+			<Do If="'$(TraceArgsFinal)' != '' and Exists('$(ProfilingFilesLocalRootPath)')">
+				<Delete Files="$(ProfilingFilesLocalRootPath)/..."/>
+			</Do>
+			<Commandlet Name="Cook" Project="$(ProjectFile)" Arguments="$(CommonCookArgs) $(TraceArgsFinal) -IncrementalValidatePhase1 -IncrementalValidateAllowWrite"/>
+		</Node>
+		<Node Name="$(IncrementalValidatePhase2Name)" Requires="$(IncrementalValidatePhase1Name)">
+			<Property Name="TraceArgsFinal" Value=""/>
+			<Property Name="TraceArgsFinal" Value="$(TraceArgs) -tracefile=&quot;$(ProfilingFilesLocalRootPath)/$(ProjectName)_IncrementalValidatePhase2.utrace&quot;" If="$(WithCookUTRACECapture) and '$(ProfilingFilesPublishRootPath)' != ''"/>
+			<Commandlet Name="Cook" Project="$(ProjectFile)" Arguments="$(CommonCookArgs) $(TraceArgsFinal) -IncrementalValidatePhase2 -IncrementalValidateAllowWrite"/>
+			<Do If="'$(TraceArgsFinal)' != ''">
+				<Zip FromDir="$(ProfilingFilesLocalRootPath)" Files="*.utrace" ZipFile="$(ProfilingFilesLocalRootPath)/$(ProjectName)_IncrementalValidate.zip"/>
+				<Copy Files="$(ProjectName)_IncrementalValidate.zip" From="$(ProfilingFilesLocalRootPath)" To="$(ProfilingFilesPublishRootPath)\$(EscapedBranch)-CL-$(Change)"/>
+			</Do>
+		</Node>
+	</Agent>
+
+	<Aggregate Name="IncrementalValidate" Requires="$(IncrementalValidateSinglePhaseName)" If="'$(ValidateMode)' == 'onephase'"/>
+	<Aggregate Name="IncrementalValidate" Requires="$(IncrementalValidatePhase2Name)" If="'$(ValidateMode)' == 'twophase'"/>
+
+</BuildGraph>
\ No newline at end of file
diff --git a/Engine/Build/Graph/Tests/IterativeCompare.xml b/Engine/Build/Graph/Tests/IterativeCompare.xml
deleted file mode 100644
diff --git a/Engine/Build/Graph/Tests/IterativeValidate.xml b/Engine/Build/Graph/Tests/IterativeValidate.xml
deleted file mode 100644
diff --git a/Engine/Source/Editor/UnrealEd/Private/CookOnTheFlyServer.cpp b/Engine/Source/Editor/UnrealEd/Private/CookOnTheFlyServer.cpp
--- a/Engine/Source/Editor/UnrealEd/Private/CookOnTheFlyServer.cpp
+++ b/Engine/Source/Editor/UnrealEd/Private/CookOnTheFlyServer.cpp
@@ -40,8 +40,8 @@
 #include "Cooker/CookTypes.h"
 #include "Cooker/CookWorkerClient.h"
 #include "Cooker/DiffPackageWriter.h"
+#include "Cooker/IncrementalValidatePackageWriter.h"
 #include "Cooker/IoStoreCookOnTheFlyRequestManager.h"
-#include "Cooker/IterativeValidatePackageWriter.h"
 #include "Cooker/LooseCookedPackageWriter.h"
 #include "Cooker/MPCollector.h"
 #include "Cooker/NetworkFileCookOnTheFlyRequestManager.h"
diff --git a/Engine/Source/Editor/UnrealEd/Private/Cooker/IncrementalValidatePackageWriter.cpp b/Engine/Source/Editor/UnrealEd/Private/Cooker/IncrementalValidatePackageWriter.cpp
new file mode 100644
--- /dev/null
+++ b/Engine/Source/Editor/UnrealEd/Private/Cooker/IncrementalValidatePackageWriter.cpp
@@ -0,0 +1,1259 @@
+// Copyright Epic Games, Inc. All Rights Reserved.
+
+#include "Cooker/IncrementalValidatePackageWriter.h"
+#include "Cooker/CompactBinaryTCP.h"
+#include "Cooker/CookGenerationHelper.h"
+#include "Cooker/CookPackageData.h"
+#include "Cooker/MPCollector.h"
+#include "CookOnTheSide/CookLog.h"
+#include "CookOnTheSide/CookOnTheFlyServer.h"
+#include "HAL/FileManager.h"
+#include "Logging/LogMacros.h"
+#include "Misc/CommandLine.h"
+#include "Misc/ConfigCacheIni.h"
+#include "Misc/OutputDeviceHelper.h"
+#include "Misc/Parse.h"
+#include "Misc/Paths.h"
+#include "Serialization/NameAsStringProxyArchive.h"
+#include "Templates/UniquePtr.h"
+
+DEFINE_LOG_CATEGORY_STATIC(LogIncrementalValidate, Log, All);
+
+constexpr FStringView IncrementalValidateFilename(TEXTVIEW("IncrementalValidate.bin"));
+
+class FIncrementalValidateMPCollector : public UE::Cook::IMPCollector
+{
+public:
+	FIncrementalValidateMPCollector(FIncrementalValidatePackageWriter* InOwner) : Owner(InOwner) {}
+
+	virtual FGuid GetMessageType() const { return MessageType; }
+	virtual const TCHAR* GetDebugName() const { return TEXT("IncrementalValidateMPCollector"); }
+
+	virtual void ServerTick(UE::Cook::FMPCollectorServerTickContext& Context) override;
+	virtual void ClientTickPackage(UE::Cook::FMPCollectorClientTickPackageContext& Context) override;
+	virtual void ServerTickPackage(UE::Cook::FMPCollectorServerTickPackageContext& Context);
+
+	virtual void ClientReceiveMessage(UE::Cook::FMPCollectorClientMessageContext& Context, FCbObjectView Message) override;
+	virtual void ServerReceiveMessage(UE::Cook::FMPCollectorServerMessageContext& Context, FCbObjectView Message) override;
+
+public:
+	static FGuid MessageType;
+
+private:
+	enum class EMessageSubtype : uint8
+	{
+		ServerToClient_WorkerStartup,
+		ClientToServer_ReplIsAnotherSaveNeeded,
+		ServerToClient_ReplUpdatePackageModificationStatus,
+		Invalid
+	};
+
+	bool TryWritePackageStatus(FCbWriter& Writer, FName PackageName);
+	void ReadAndSyncPackageStatus(FCbObjectView Message, FName PackageName);
+	FIncrementalValidatePackageWriter* Owner;
+};
+
+FCbWriter& operator<<(FCbWriter& Writer, const FIncrementalValidatePackageWriter::FMessage& Message)
+{
+	Writer.BeginObject();
+	Writer << "Text" << Message.Text;
+
+	static_assert(sizeof(Message.Verbosity) == sizeof(uint8));
+	Writer << "Verbosity" << (uint8)Message.Verbosity;
+
+	Writer.EndObject();
+	return Writer;
+}
+
+bool LoadFromCompactBinary(FCbFieldView Field, FIncrementalValidatePackageWriter::FMessage& Message)
+{
+	bool bOk = !Field.HasError();
+	if (bOk && LoadFromCompactBinary(Field["Text"], Message.Text))
+	{
+		uint8 Verbosity = ELogVerbosity::NumVerbosity;
+		bOk = LoadFromCompactBinary(Field["Verbosity"], Verbosity) && Verbosity < ELogVerbosity::NumVerbosity;
+		if (bOk)
+		{
+			Message.Verbosity = (ELogVerbosity::Type)Verbosity;
+		}
+	}
+	if (!bOk)
+	{
+		Message = FIncrementalValidatePackageWriter::FMessage();
+	}
+	return bOk;
+}
+
+FCbWriter& operator<<(FCbWriter& Writer, FIncrementalValidatePackageWriter::EPackageStatus Status)
+{
+	Writer << (uint8)Status;
+	return Writer;
+}
+
+bool LoadFromCompactBinary(FCbFieldView Field, FIncrementalValidatePackageWriter::EPackageStatus& Status)
+{
+	uint8 StatusInteger = (uint8)FIncrementalValidatePackageWriter::EPackageStatus::Count;
+	if (!LoadFromCompactBinary(Field, StatusInteger))
+	{
+		UE_LOG(LogIncrementalValidate, Error, TEXT("Failed to deserialize package status."));
+	}
+	else if (StatusInteger >= (uint8)FIncrementalValidatePackageWriter::EPackageStatus::Count)
+	{
+		UE_LOG(LogIncrementalValidate, Error, TEXT("Unexpected package status deserialized: %d"), StatusInteger);
+	}
+	else
+	{
+		Status = (FIncrementalValidatePackageWriter::EPackageStatus)StatusInteger;
+		return true;
+	}
+	return false;
+}
+
+FArchive& operator<<(FArchive& Ar, FIncrementalValidatePackageWriter::FPackageStatusInfo& Info)
+{
+	Ar << Info.Status;
+	bool bHasAssetClass = Info.AssetClass.IsValid();
+	Ar << bHasAssetClass;
+	if (bHasAssetClass)
+	{
+		Ar << Info.AssetClass;
+	}
+	return Ar;
+}
+
+FCbWriter& operator<<(FCbWriter& Writer, const FIncrementalValidatePackageWriter::FPackageStatusInfo& Info)
+{
+	Writer.BeginArray();
+	Writer << Info.Status;
+	if (Info.AssetClass.IsValid())
+	{
+		Writer << Info.AssetClass;
+	}
+	Writer.EndArray();
+	return Writer;
+}
+
+bool LoadFromCompactBinary(FCbFieldView Field, FIncrementalValidatePackageWriter::FPackageStatusInfo& Info)
+{
+	FCbArrayView ArrayField = Field.AsArrayView();
+	if (ArrayField.Num() < 1)
+	{
+		return false;
+	}
+	FCbFieldViewIterator Iter = ArrayField.CreateViewIterator();
+	bool bOk = true;
+	bOk = LoadFromCompactBinary(*(Iter++), Info.Status) & bOk;
+	if (ArrayField.Num() >= 2)
+	{
+		bOk = LoadFromCompactBinary(*(Iter++), Info.AssetClass) & bOk;
+	}
+	return bOk;
+}
+
+void FIncrementalValidateMPCollector::ServerTick(UE::Cook::FMPCollectorServerTickContext& Context)
+{
+	static_assert(sizeof(EMessageSubtype) == sizeof(uint8));
+
+	if (Context.GetEventType() == UE::Cook::FMPCollectorServerTickContext::EServerEventType::WorkerStartup)
+	{
+		FCbWriter Writer;
+		Writer.BeginObject();
+		Writer << "MessageSubtype" << (uint8)EMessageSubtype::ServerToClient_WorkerStartup;
+		Writer << "PackageStatusMap" << Owner->PackageStatusMap;
+		Writer << "PackageMessageMap" << Owner->PackageMessageMap;
+		Writer.EndObject();
+		Context.AddMessage(Writer.Save().AsObject());
+	}
+}
+
+void FIncrementalValidateMPCollector::ClientTickPackage(UE::Cook::FMPCollectorClientTickPackageContext& Context)
+{
+	static_assert(sizeof(EMessageSubtype) == sizeof(uint8));
+
+	FCbWriter Writer;
+	Writer.BeginObject();
+	Writer << "MessageSubtype" << (uint8)EMessageSubtype::ClientToServer_ReplIsAnotherSaveNeeded;
+
+	const FName PackageName = Context.GetPackageName();
+	if (PackageName.IsNone())
+	{
+		UE_LOG(LogCook, Error, TEXT("Context does not contain a valid package name."))
+		// It's safe to continue because TryWritePackageStatus will return false if the name is none. 
+		// The error is logged here to make the call site clear
+	}
+
+	if (TryWritePackageStatus(Writer, PackageName))
+	{
+		Writer.EndObject();
+		Context.AddMessage(Writer.Save().AsObject());
+	}
+}
+
+void FIncrementalValidateMPCollector::ServerTickPackage(UE::Cook::FMPCollectorServerTickPackageContext& Context)
+{
+	static_assert(sizeof(EMessageSubtype) == sizeof(uint8));
+
+	FCbWriter Writer;
+	Writer.BeginObject();
+	Writer << "MessageSubtype" << (uint8)EMessageSubtype::ServerToClient_ReplUpdatePackageModificationStatus;
+
+	const FName PackageName = Context.GetPackageName();
+	if (PackageName.IsNone())
+	{
+		UE_LOG(LogCook, Error, TEXT("Context does not contain a valid package name."))
+		// It's safe to continue because TryWritePackageStatus will return false if the name is none. 
+		// The error is logged here to make the call site clear
+	}
+
+	if (TryWritePackageStatus(Writer, PackageName))
+	{
+		Writer.EndObject();
+		Context.AddMessage(Writer.Save().AsObject());
+	}
+}
+
+void FIncrementalValidateMPCollector::ClientReceiveMessage(UE::Cook::FMPCollectorClientMessageContext& Context, FCbObjectView Message)
+{
+	uint8 MessageSubtypeAsInteger = (uint8)EMessageSubtype::Invalid;
+	if (LoadFromCompactBinary(Message["MessageSubtype"], MessageSubtypeAsInteger))
+	{
+		switch ((EMessageSubtype)MessageSubtypeAsInteger)
+		{
+			case EMessageSubtype::ServerToClient_WorkerStartup:
+			{
+				bool bOk = LoadFromCompactBinary(Message["PackageStatusMap"], Owner->PackageStatusMap);
+				bOk = bOk && LoadFromCompactBinary(Message["PackageMessageMap"], Owner->PackageMessageMap);
+				check(bOk); // If we fail this, we will fail to get anything right during the rest of the validation. Better to terminate quickly.
+				break;
+			}
+			case EMessageSubtype::ServerToClient_ReplUpdatePackageModificationStatus:
+			{
+				FName PackageName = Context.GetPackageName();
+				if (!PackageName.IsNone())
+				{
+					ReadAndSyncPackageStatus(Message, PackageName);
+				}
+				else
+				{
+					UE_LOG(LogCook, Error, TEXT("Cannot process ServerToClient_ReplUpdatePackageModificationStatus without a valid package name in the current context."));
+				}
+				break;
+			}
+			default:
+			{
+				UE_LOG(LogCook, Error, TEXT("Unexpected message type: %d"), MessageSubtypeAsInteger);
+				break;
+			}
+		};
+	}
+}
+
+void FIncrementalValidateMPCollector::ServerReceiveMessage(UE::Cook::FMPCollectorServerMessageContext& Context, FCbObjectView Message)
+{
+	uint8 MessageSubtypeAsInteger = (uint8)EMessageSubtype::Invalid;
+	if (LoadFromCompactBinary(Message["MessageSubtype"], MessageSubtypeAsInteger))
+	{
+		FName PackageName = Context.GetPackageName();
+		if (PackageName.IsNone())
+		{
+			UE_LOG(LogCook, Error, TEXT("Cannot process messages on server without a valid package name in the current context."));
+		}
+		else if (MessageSubtypeAsInteger == (uint8)EMessageSubtype::ClientToServer_ReplIsAnotherSaveNeeded)
+		{
+			ReadAndSyncPackageStatus(Message, PackageName);
+			Owner->MarkPackageCompletedOnDirector(PackageName, Context.GetWorkerId());
+		}
+		else
+		{
+			UE_LOG(LogCook, Error, TEXT("Unexpected message received. MessageSubtype == %d"), MessageSubtypeAsInteger);
+		}
+	}
+	else
+	{
+		UE_LOG(LogCook, Error, TEXT("Invalid message received. No MessageSubtype field available."));
+	}
+}
+
+FGuid FIncrementalValidateMPCollector::MessageType(TEXT("5E56C5D96F3B455E9452C15ADA601A71"));
+
+bool FIncrementalValidateMPCollector::TryWritePackageStatus(FCbWriter& Writer, FName PackageName)
+{
+	const FIncrementalValidatePackageWriter::FPackageStatusInfo* PackageStatus
+		= Owner->PackageStatusMap.Find(PackageName);
+
+	if (PackageStatus && PackageStatus->Status != FIncrementalValidatePackageWriter::EPackageStatus::NotYetProcessed)
+	{
+		Writer << "Status" << *PackageStatus;
+
+		const TArray<FIncrementalValidatePackageWriter::FMessage>* Messages = Owner->PackageMessageMap.Find(PackageName);
+		if (Messages != nullptr)
+		{
+			Writer << "MessageArray" << *Messages;
+		}
+		return true;
+	}
+	return false;
+}
+
+void FIncrementalValidateMPCollector::ReadAndSyncPackageStatus(FCbObjectView Message, FName PackageName)
+{
+	FIncrementalValidatePackageWriter::FPackageStatusInfo Info;
+	bool bOk = LoadFromCompactBinary(Message["Status"], Info);
+
+	if (bOk)
+	{
+		Owner->PackageStatusMap.FindOrAdd(PackageName) = MoveTemp(Info);
+		if (Message.FindView("MessageArray").HasValue())
+		{
+			TArray<FIncrementalValidatePackageWriter::FMessage>& MessageArray = Owner->PackageMessageMap.FindOrAdd(PackageName);
+			bOk = LoadFromCompactBinary(Message["MessageArray"], MessageArray);
+		}
+	}
+	
+	if (!bOk)
+	{
+		UE_LOG(LogCook, Error, TEXT("Invalid message received in ReadAndSyncPackageStatus. Failed to load Info from Message[\"Status\"] for package \"%s\""),
+			*PackageName.ToString());
+	}
+}
+
+FIncrementalValidatePackageWriter::FIncrementalValidatePackageWriter(UCookOnTheFlyServer& InCOTFS,
+	TUniquePtr<ICookedPackageWriter>&& InInner, EPhase InPhase, const FString& ResolvedMetadataPath,
+	UE::Cook::FDeterminismManager* InDeterminismManager)
+	: FDiffPackageWriter(MoveTemp(InInner), InDeterminismManager)
+	, MetadataPath(ResolvedMetadataPath)
+	, COTFS(InCOTFS)
+	, Phase(InPhase)
+{
+	COTFS.RegisterCollector(new FIncrementalValidateMPCollector(this));
+
+	Indent = FCString::Spc(FOutputDeviceHelper::FormatLogLine(ELogVerbosity::Warning,
+		LogIncrementalValidate.GetCategoryName(), TEXT(""), GPrintLogTimes).Len());
+
+	TArray<FString> IncrementalValidatePackageIgnoreList;
+	GConfig->GetArray(TEXT("IncrementalValidate"), TEXT("PackageIgnoreList"),
+		IncrementalValidatePackageIgnoreList, GEditorIni);
+	PackageIgnoreList.Reserve(IncrementalValidatePackageIgnoreList.Num());
+	for (const FString& PackageName : IncrementalValidatePackageIgnoreList)
+	{
+		PackageIgnoreList.Add(FName(FStringView(PackageName)));
+	}
+
+	LoggingSoftMaximum = -1;
+	GConfig->GetValue(TEXT("IncrementalValidate"), TEXT("LoggingSoftMaximum"), LoggingSoftMaximum, GEditorIni);
+}
+
+void FIncrementalValidatePackageWriter::BeginPackage(const FBeginPackageInfo& Info)
+{
+	bPackageFirstPass = true;
+	switch (Phase)
+	{
+	case EPhase::AllInOnePhase:
+		if (GetPackageStatus(Info.PackageName) == EPackageStatus::DeclaredUnmodified_NotYetProcessed)
+		{
+			// Save to memory and look for diffs before saving it out to disk
+			SaveAction = ESaveAction::CheckForDiffs;
+			Super::BeginPackage(Info);
+		}
+		else
+		{
+			// Not incrementally skippable, so we expect it to change. No need to look for diffs, just save to disk.
+			if (bReadOnly)
+			{
+				SaveAction = ESaveAction::IgnoreResults;
+			}
+			else
+			{
+				SaveAction = ESaveAction::SaveToInner;
+				Inner->BeginPackage(Info);
+			}
+		}
+		break;
+	case EPhase::Phase1:
+		SaveAction = ESaveAction::CheckForDiffs;
+		Super::BeginPackage(Info);
+		break;
+	case EPhase::Phase2:
+		{
+			EPackageStatus PackageStatus = GetPackageStatus(Info.PackageName);
+			if (PackageStatus == EPackageStatus::DeclaredUnmodified_ConfirmedUnmodified ||
+				PackageStatus == EPackageStatus::DeclaredUnmodified_FoundModified_OnIgnoreList)
+			{
+				// Already saved in Phase 1; no need to diff it or save it now
+				SaveAction = ESaveAction::IgnoreResults;
+			}
+			else if (PackageStatus == EPackageStatus::DeclaredUnmodified_FoundModified_IndeterminismOrFalsePositive)
+			{
+				SaveAction = ESaveAction::CheckForDiffs;
+				Super::BeginPackage(Info);
+			}
+			else
+			{
+				// This is an IncrementallyModified package. It was found during Phase1 to be modified and would in a
+				// normal incremental cook be resaved rather than incrementally skipped. Resave it as normal.
+				SaveAction = ESaveAction::SaveToInner;
+				Inner->BeginPackage(Info);
+			}
+			break;
+		}
+	default:
+		checkNoEntry();
+		break;
+	}
+}
+
+void FIncrementalValidatePackageWriter::CommitPackage(FCommitPackageInfo&& Info)
+{
+	switch (SaveAction)
+	{
+	case ESaveAction::CheckForDiffs:
+		Super::CommitPackage(MoveTemp(Info));
+		break;
+	case ESaveAction::SaveToInner:
+		Inner->CommitPackage(MoveTemp(Info));
+		break;
+	case ESaveAction::IgnoreResults:
+		break;
+	default:
+		checkNoEntry();
+		break;
+	}
+}
+
+void FIncrementalValidatePackageWriter::WritePackageData(const FPackageInfo& Info, FLargeMemoryWriter& ExportsArchive,
+	const TArray<FFileRegion>& FileRegions)
+{
+	switch (SaveAction)
+	{
+	case ESaveAction::CheckForDiffs:
+		Super::WritePackageData(Info, ExportsArchive, FileRegions);
+		break;
+	case ESaveAction::SaveToInner:
+		Inner->WritePackageData(Info, ExportsArchive, FileRegions);
+		break;
+	case ESaveAction::IgnoreResults:
+		break;
+	default:
+		checkNoEntry();
+		break;
+	}
+}
+
+void FIncrementalValidatePackageWriter::WriteBulkData(const FBulkDataInfo& Info, const FIoBuffer& BulkData,
+	const TArray<FFileRegion>& FileRegions)
+{
+	switch (SaveAction)
+	{
+	case ESaveAction::CheckForDiffs:
+		Super::WriteBulkData(Info, BulkData, FileRegions);
+		break;
+	case ESaveAction::SaveToInner:
+		Inner->WriteBulkData(Info, BulkData, FileRegions);
+		break;
+	case ESaveAction::IgnoreResults:
+		break;
+	default:
+		checkNoEntry();
+		break;
+	}
+}
+
+void FIncrementalValidatePackageWriter::WriteAdditionalFile(const FAdditionalFileInfo& Info, const FIoBuffer& FileData)
+{
+	switch (SaveAction)
+	{
+	case ESaveAction::CheckForDiffs:
+		Super::WriteAdditionalFile(Info, FileData);
+		break;
+	case ESaveAction::SaveToInner:
+		Inner->WriteAdditionalFile(Info, FileData);
+		break;
+	case ESaveAction::IgnoreResults:
+		break;
+	default:
+		checkNoEntry();
+		break;
+	}
+}
+void FIncrementalValidatePackageWriter::WriteLinkerAdditionalData(const FLinkerAdditionalDataInfo& Info,
+	 const FIoBuffer& Data, const TArray<FFileRegion>& FileRegions)
+{
+	switch (SaveAction)
+	{
+	case ESaveAction::CheckForDiffs:
+		Super::WriteLinkerAdditionalData(Info, Data, FileRegions);
+		break;
+	case ESaveAction::SaveToInner:
+		Inner->WriteLinkerAdditionalData(Info, Data, FileRegions);
+		break;
+	case ESaveAction::IgnoreResults:
+		break;
+	default:
+		checkNoEntry();
+		break;
+	}
+}
+
+void FIncrementalValidatePackageWriter::WritePackageTrailer(const FPackageTrailerInfo& Info, const FIoBuffer& Data)
+{
+	switch (SaveAction)
+	{
+	case ESaveAction::CheckForDiffs:
+		Super::WritePackageTrailer(Info, Data);
+		break;
+	case ESaveAction::SaveToInner:
+		Inner->WritePackageTrailer(Info, Data);
+		break;
+	case ESaveAction::IgnoreResults:
+		break;
+	default:
+		checkNoEntry();
+		break;
+	}
+}
+
+int64 FIncrementalValidatePackageWriter::GetExportsFooterSize()
+{
+	switch (SaveAction)
+	{
+	case ESaveAction::CheckForDiffs:
+		return Super::GetExportsFooterSize();
+	case ESaveAction::SaveToInner:
+		return Inner->GetExportsFooterSize();
+	case ESaveAction::IgnoreResults:
+		return 0;
+	default:
+		checkNoEntry();
+		return 0;
+	}
+}
+
+TUniquePtr<FLargeMemoryWriter> FIncrementalValidatePackageWriter::CreateLinkerArchive(FName PackageName, UObject* Asset, uint16 MultiOutputIndex)
+{
+	if (Asset)
+	{
+		PackageStatusMap.FindOrAdd(PackageName).AssetClass = Asset->GetClass()->GetClassPathName();
+	}
+
+	switch (SaveAction)
+	{
+	case ESaveAction::CheckForDiffs:
+		return Super::CreateLinkerArchive(PackageName, Asset, MultiOutputIndex);
+	case ESaveAction::SaveToInner:
+		return Inner->CreateLinkerArchive(PackageName, Asset, MultiOutputIndex);
+	case ESaveAction::IgnoreResults:
+		return MakeUnique<FLargeMemoryWriter>();
+	default:
+		checkNoEntry();
+		return TUniquePtr<FLargeMemoryWriter>();
+	}
+}
+
+TUniquePtr<FLargeMemoryWriter> FIncrementalValidatePackageWriter::CreateLinkerExportsArchive(FName PackageName, UObject* Asset, uint16 MultiOutputIndex)
+{
+	switch (SaveAction)
+	{
+	case ESaveAction::CheckForDiffs:
+		return Super::CreateLinkerExportsArchive(PackageName, Asset, MultiOutputIndex);
+	case ESaveAction::SaveToInner:
+		return Inner->CreateLinkerExportsArchive(PackageName, Asset, MultiOutputIndex);
+	case ESaveAction::IgnoreResults:
+		return MakeUnique<FLargeMemoryWriter>();
+	default:
+		checkNoEntry();
+		return TUniquePtr<FLargeMemoryWriter>();
+	}
+}
+
+bool FIncrementalValidatePackageWriter::IsPreSaveCompleted() const
+{
+	return !bPackageFirstPass;
+}
+
+ICookedPackageWriter::FCookCapabilities FIncrementalValidatePackageWriter::GetCookCapabilities() const
+{
+	FCookCapabilities Result = Super::GetCookCapabilities();
+	Result.bReadOnly = bReadOnly;
+	Result.bOverridesPackageModificationStatus = true;
+	return Result;
+}
+
+void FIncrementalValidatePackageWriter::Initialize(const FCookInfo& CookInfo)
+{
+	switch (Phase)
+	{
+	case EPhase::AllInOnePhase:
+		if (CookInfo.bFullBuild)
+		{
+			UE_CLOG(COTFS.GetCookMode() != ECookMode::CookWorker,
+				LogIncrementalValidate, Display,
+				TEXT("The cook is running non-incrementally. All packages are reported \"modified\" and will be resaved as in a normal cook."));
+			bReadOnly = false;
+		}
+		else
+		{
+			bReadOnly = !FParse::Param(FCommandLine::Get(), TEXT("IncrementalValidateAllowWrite"));
+		}
+		break;
+	case EPhase::Phase1:
+		if (CookInfo.bFullBuild)
+		{
+			UE_CLOG(COTFS.GetCookMode() != ECookMode::CookWorker,
+				LogIncrementalValidate, Display,
+				TEXT("The cook is running non-incrementally. All packages are reported \"modified\" and will be resaved during the final IncrementalValidate phase."));
+		}
+		bReadOnly = false;
+		break;
+	case EPhase::Phase2:
+		if (CookInfo.bFullBuild)
+		{
+			UE_CLOG(COTFS.GetCookMode() != ECookMode::CookWorker, 
+				LogIncrementalValidate, Display,
+				TEXT("The cook is running non-incrementally. Packages that were incrementally skipped and found valid will be resaved anyway."));
+		}
+		bReadOnly = false;
+		break;
+	default:
+		checkNoEntry();
+		break;
+	}
+	Super::Initialize(CookInfo);
+}
+
+void FIncrementalValidatePackageWriter::UpdatePackageModificationStatus(FName PackageName, bool bIncrementallyUnmodified,
+	bool& bInOutShouldIncrementallySkip)
+{
+	// We need to not skip previously cooked generator packages, if they were modified and we're read only we still
+	// need to cook them so we can investigate their generated packages. Look up whether they were a generator in
+	// the previous cook results.
+	bool bKnownGenerator = false;
+	if (bIncrementallyUnmodified)
+	{
+		// GenerationHelpers were created for previously cooked generators by UCOTFS::PopulateCookedPackages.
+		UE::Cook::FPackageData* PackageData = COTFS.PackageDatas->FindPackageDataByPackageName(PackageName);
+		if (PackageData)
+		{
+			TRefCountPtr<UE::Cook::FGenerationHelper> GenerationHelper = PackageData->GetGenerationHelper();
+			if (GenerationHelper)
+			{
+				bKnownGenerator = true;
+			}
+		}
+	}
+
+	switch (Phase)
+	{
+	case EPhase::AllInOnePhase:		
+		// Save the input value for bIncrementallyUnmodified, and report skippable for the modified if possible
+		if (bIncrementallyUnmodified)
+		{
+			SetPackageStatus(PackageName, EPackageStatus::DeclaredUnmodified_NotYetProcessed);
+			bInOutShouldIncrementallySkip = false;
+		}
+		else
+		{
+			SetPackageStatus(PackageName, EPackageStatus::DeclaredModified_WillNotVerify);
+			if (!bKnownGenerator && bReadOnly)
+			{
+				bInOutShouldIncrementallySkip = true;
+			}
+		}
+		break;
+	case EPhase::Phase1:
+		// Save the incrementally unmodified packages to verify their diffs. Skip the non-generator incrementally
+		// modified packages. Do not skip generator packages; we need to save them to test their generated packages.
+		bInOutShouldIncrementallySkip = !bIncrementallyUnmodified && !bKnownGenerator;
+		if (!bIncrementallyUnmodified)
+		{
+			SetPackageStatus(PackageName, EPackageStatus::DeclaredModified_WillNotVerify);
+		}
+		else
+		{
+			SetPackageStatus(PackageName, EPackageStatus::DeclaredUnmodified_NotYetProcessed);
+		}
+		break;
+	case EPhase::Phase2:
+	{
+		// Ignore the Unmodified flag from this cook phase. Skip the packages that were found to 
+		// be IncrementallyValidated from Phase1. Save the packages that phase1 found modified; this phase
+		// is responsible for getting those resaved. Reexecute save for the packages that were IncrementalFailed
+		// from phase1, so we can test whether they are indeterministic. Always save generators so we can
+		// test their generated packages.
+		EPackageStatus PackageStatus = GetPackageStatus(PackageName);
+		bInOutShouldIncrementallySkip =
+			(PackageStatus == EPackageStatus::DeclaredUnmodified_ConfirmedUnmodified
+				|| PackageStatus == EPackageStatus::DeclaredUnmodified_FoundModified_OnIgnoreList)
+			&& !bKnownGenerator;
+		break;
+	}
+	default:
+		checkNoEntry();
+		break;
+	}
+
+	bool bInnerIncrementallyUnmodified = bInOutShouldIncrementallySkip;
+	bool bInnerInOutShouldIncrementallySkip = bInOutShouldIncrementallySkip;
+	Inner->UpdatePackageModificationStatus(PackageName, bInnerIncrementallyUnmodified, bInnerInOutShouldIncrementallySkip);
+	checkf(bInnerInOutShouldIncrementallySkip == bInnerIncrementallyUnmodified,
+		TEXT("FIncrementalValidatePackageWriter is not supported with an Inner that modifies bInOutShouldIncrementallySkip."));
+}
+
+void FIncrementalValidatePackageWriter::BeginCook(const FCookInfo& Info)
+{
+	switch (Phase)
+	{
+	case EPhase::AllInOnePhase:
+		if (bReadOnly)
+		{
+			UE_CLOG(COTFS.GetCookMode() != ECookMode::CookWorker,
+				LogIncrementalValidate, Display,
+				TEXT("-IncrementalValidateAllowWrite not present, read-only mode. Running -diffonly on all packages that were found to be incrementally unmodified."));
+		}
+		else
+		{
+			UE_CLOG(COTFS.GetCookMode() != ECookMode::CookWorker,
+				LogIncrementalValidate, Display,
+				TEXT("-IncrementalValidateAllowWrite is present, writable mode. Resaving packages as in a normal cook, but also running -diffonly on all packages that were found to be incrementally unmodified."));
+		}
+		if (Info.bFullBuild)
+		{
+			UE_CLOG(COTFS.GetCookMode() != ECookMode::CookWorker,
+				LogIncrementalValidate, Error,
+				TEXT("IncrementalValidate was bypassed on this run; it is a full cook and all packages are marked incrementally modified."));
+		}
+		break;
+	case EPhase::Phase1:
+		UE_CLOG(COTFS.GetCookMode() != ECookMode::CookWorker,
+			LogIncrementalValidate, Display,
+			TEXT("Phase1: running -diffonly and a resave on all packages discovered to be incrementally unmodified."));
+		if (Info.bFullBuild)
+		{
+			UE_CLOG(COTFS.GetCookMode() != ECookMode::CookWorker,
+				LogIncrementalValidate, Error,
+				TEXT("IncrementalValidate was bypassed on this run; it is a full cook and all packages are marked incrementally modified."));
+		}
+		break;
+	case EPhase::Phase2:
+		{
+			Load();
+			FStatusCounts StatusCounts = CountPackagesByStatus();
+			UE_CLOG(COTFS.GetCookMode() != ECookMode::CookWorker,
+				LogIncrementalValidate, Display,
+				TEXT("Phase2: %d packages were found during Phase1 to be incrementally unmodified but had differences. "
+					"Running -diffonly on them again to check whether the differences are due to indeterminism or to FalsePositiveIncrementalSkips."), 
+					StatusCounts[EPackageStatus::DeclaredUnmodified_FoundModified_IndeterminismOrFalsePositive]);
+			UE_CLOG(COTFS.GetCookMode() != ECookMode::CookWorker,
+				LogIncrementalValidate, Display,
+				TEXT("%d packages were found during Phase1 to be modified or new and will be resaved."),
+				StatusCounts[EPackageStatus::DeclaredModified_WillNotVerify]);
+			break;
+		}
+	default:
+		checkNoEntry();
+		break;
+	}
+	Super::BeginCook(Info);
+}
+
+void FIncrementalValidatePackageWriter::EndCook(const FCookInfo& Info)
+{
+	Super::EndCook(Info);
+	FStatusCounts StatusCounts = CountPackagesByStatus();
+
+	switch (Phase)
+	{
+	case EPhase::AllInOnePhase:
+	{
+		int32 DetectedUnmodified = StatusCounts[EPackageStatus::DeclaredUnmodified_ConfirmedUnmodified]
+			+ StatusCounts[EPackageStatus::DeclaredUnmodified_FoundModified_FalsePositive]
+			+ StatusCounts[EPackageStatus::DeclaredUnmodified_FoundModified_OnIgnoreList];
+		UE_CLOG(COTFS.GetCookMode() != ECookMode::CookWorker,
+			LogIncrementalValidate, Display,
+			TEXT("Modified: %d. DetectedUnmodified: %d. ValidatedUnmodified: %d. IncrementalSkipFalsePositive: %d."),
+			StatusCounts[EPackageStatus::DeclaredModified_WillNotVerify], 
+			DetectedUnmodified,
+			StatusCounts[EPackageStatus::DeclaredUnmodified_ConfirmedUnmodified]
+			+ StatusCounts[EPackageStatus::DeclaredUnmodified_FoundModified_OnIgnoreList],
+			StatusCounts[EPackageStatus::DeclaredUnmodified_FoundModified_FalsePositive]);
+		FString Message = FString::Printf(TEXT("Packages Incrementally Skipped: %d: IncrementalSkipFalsePositive: %d."),
+			DetectedUnmodified,
+			StatusCounts[EPackageStatus::DeclaredUnmodified_FoundModified_FalsePositive]);
+		if (StatusCounts[EPackageStatus::DeclaredUnmodified_FoundModified_FalsePositive] > 0)
+		{
+			TStringBuilder<1024> MessageWithDiagnostics;
+			MessageWithDiagnostics << Message;
+			TMap<FTopLevelAssetPath, TArray<FName>> ClassFalsePositiveCounts = GetClassStatusSummary(
+				EPackageStatus::DeclaredUnmodified_FoundModified_FalsePositive);
+			int32 NumClassesPrinted = 0;
+			constexpr int32 MaxNumClassesPrinted = 25;
+			for (const TPair<FTopLevelAssetPath, TArray<FName>>& Pair : ClassFalsePositiveCounts)
+			{
+				MessageWithDiagnostics << TEXT("\n\t") << WriteToString<256>(Pair.Key) << TEXT(": ") << Pair.Value.Num();
+
+				int32 NumPackagesPrinted = 0;
+				constexpr int32 MaxNumPackagesPrinted = 10;
+				for (FName PackageName : Pair.Value)
+				{
+					MessageWithDiagnostics << TEXT("\n\t\t") << PackageName;
+					if (++NumPackagesPrinted >= MaxNumPackagesPrinted)
+					{
+						break;
+					}
+				}
+				if (Pair.Value.Num() > NumPackagesPrinted)
+				{
+					MessageWithDiagnostics << TEXT("\n\t\t...");
+				}
+
+				if (++NumClassesPrinted >= MaxNumClassesPrinted)
+				{
+					break;
+				}
+			}
+			if (ClassFalsePositiveCounts.Num() > NumClassesPrinted)
+			{
+				MessageWithDiagnostics << TEXT("\n\t...");
+			}
+			UE_CLOG(COTFS.GetCookMode() != ECookMode::CookWorker,
+				LogIncrementalValidate, Error, TEXT("%s"), *MessageWithDiagnostics);
+		}
+		else if (StatusCounts[EPackageStatus::DeclaredUnmodified_FoundModified_OnIgnoreList] > 0)
+		{
+			Message = FString::Printf(TEXT("Packages Incrementally Skipped: %d: IncrementalSkipFalsePositive (Ignored): %d."),
+				DetectedUnmodified,
+				StatusCounts[EPackageStatus::DeclaredUnmodified_FoundModified_OnIgnoreList]);
+			UE_CLOG(COTFS.GetCookMode() != ECookMode::CookWorker,
+				LogIncrementalValidate, Warning, TEXT("%s"), *Message);
+		}
+		else
+		{
+			UE_CLOG(COTFS.GetCookMode() != ECookMode::CookWorker,
+				LogIncrementalValidate, Display, TEXT("%s"), *Message);
+		}
+		break;
+	}
+	case EPhase::Phase1:
+	{
+		int32 DetectedUnmodified = StatusCounts[EPackageStatus::DeclaredUnmodified_ConfirmedUnmodified]
+			+ StatusCounts[EPackageStatus::DeclaredUnmodified_FoundModified_IndeterminismOrFalsePositive]
+			+ StatusCounts[EPackageStatus::DeclaredUnmodified_FoundModified_OnIgnoreList];
+		UE_CLOG(COTFS.GetCookMode() != ECookMode::CookWorker,
+			LogIncrementalValidate, Display,
+			TEXT("Modified: %d. DetectedUnmodified: %d. ValidatedUnmodified: %d. IncrementalSkipFalsePositiveOrIndeterminism: %d."),
+			StatusCounts[EPackageStatus::DeclaredModified_WillNotVerify],
+			DetectedUnmodified,
+			StatusCounts[EPackageStatus::DeclaredUnmodified_ConfirmedUnmodified]
+			+ StatusCounts[EPackageStatus::DeclaredUnmodified_FoundModified_OnIgnoreList],
+			StatusCounts[EPackageStatus::DeclaredUnmodified_FoundModified_IndeterminismOrFalsePositive]);
+		Save();
+		break;
+	}
+	case EPhase::Phase2:
+	{
+		int32 DetectedUnmodified = StatusCounts[EPackageStatus::DeclaredUnmodified_ConfirmedUnmodified]
+			+ StatusCounts[EPackageStatus::DeclaredUnmodified_FoundModified_IndeterminismOrFalsePositive]
+			+ StatusCounts[EPackageStatus::DeclaredUnmodified_FoundModified_Indeterminism]
+			+ StatusCounts[EPackageStatus::DeclaredUnmodified_FoundModified_FalsePositive]
+			+ StatusCounts[EPackageStatus::DeclaredUnmodified_FoundModified_OnIgnoreList];
+		UE_CLOG(COTFS.GetCookMode() != ECookMode::CookWorker,
+			LogIncrementalValidate, Display,
+			TEXT("Modified: %d. DetectedUnmodified: %d. ValidatedUnmodified: %d. Indeterminism: %d. IncrementalSkipFalsePositive: %d."),
+			StatusCounts[EPackageStatus::DeclaredModified_WillNotVerify],
+			DetectedUnmodified,
+			StatusCounts[EPackageStatus::DeclaredUnmodified_ConfirmedUnmodified]
+			+ StatusCounts[EPackageStatus::DeclaredUnmodified_FoundModified_OnIgnoreList],
+			StatusCounts[EPackageStatus::DeclaredUnmodified_FoundModified_Indeterminism],
+			StatusCounts[EPackageStatus::DeclaredUnmodified_FoundModified_FalsePositive]
+			);
+		FString Message = FString::Printf(TEXT("Packages Incrementally Skipped: %d: IncrementalSkipFalsePositive: %d."),
+			DetectedUnmodified,
+			StatusCounts[EPackageStatus::DeclaredUnmodified_FoundModified_FalsePositive]);
+		if (StatusCounts[EPackageStatus::DeclaredUnmodified_FoundModified_FalsePositive] > 0)
+		{
+			UE_CLOG(COTFS.GetCookMode() != ECookMode::CookWorker,
+				LogIncrementalValidate, Error, TEXT("%s"), *Message);
+		}
+		else if (StatusCounts[EPackageStatus::DeclaredUnmodified_FoundModified_OnIgnoreList] > 0)
+		{
+			Message = FString::Printf(TEXT("Packages Incrementally Skipped: %d: IncrementalSkipFalsePositive (Ignored): %d."),
+				DetectedUnmodified,
+				StatusCounts[EPackageStatus::DeclaredUnmodified_FoundModified_FalsePositive]);
+			UE_CLOG(COTFS.GetCookMode() != ECookMode::CookWorker,
+				LogIncrementalValidate, Warning, TEXT("%s"), *Message);
+		}
+		else
+		{
+			UE_CLOG(COTFS.GetCookMode() != ECookMode::CookWorker,
+				LogIncrementalValidate, Display, TEXT("%s"), *Message);
+		}
+		break;
+	}
+	default:
+		checkNoEntry();
+		break;
+	}
+}
+
+void FIncrementalValidatePackageWriter::UpdateSaveArguments(FSavePackageArgs& SaveArgs)
+{
+	switch (SaveAction)
+	{
+	case ESaveAction::CheckForDiffs:
+		Super::UpdateSaveArguments(SaveArgs);
+		break;
+	case ESaveAction::SaveToInner:
+		Inner->UpdateSaveArguments(SaveArgs);
+		break;
+	case ESaveAction::IgnoreResults:
+		break;
+	default:
+		checkNoEntry();
+		break;
+	}
+}
+
+bool FIncrementalValidatePackageWriter::IsAnotherSaveNeeded(FSavePackageResultStruct& PreviousResult, FSavePackageArgs& SaveArgs)
+{
+	bool bResult = IsAnotherSaveNeededInternal(PreviousResult, SaveArgs);
+	if (!bResult && COTFS.GetCookMode() != ECookMode::CookWorker)
+	{
+		MarkPackageCompletedOnDirector(BeginInfo.PackageName, UE::Cook::FWorkerId::Local());
+	}
+	return bResult;
+
+}
+
+bool FIncrementalValidatePackageWriter::IsAnotherSaveNeededInternal(FSavePackageResultStruct& PreviousResult,
+	FSavePackageArgs& SaveArgs)
+{
+	bPackageFirstPass = false;
+	checkf(!Inner->IsAnotherSaveNeeded(PreviousResult, SaveArgs),
+		TEXT("FIncrementalValidatePackageWriter does not support an Inner that needs multiple saves."));
+	if (PreviousResult == ESavePackageResult::Timeout)
+	{
+		return false;
+	}
+	switch (SaveAction)
+	{
+	case ESaveAction::CheckForDiffs:
+		break;
+	case ESaveAction::SaveToInner:
+		// The SaveToInner pass, if present, is the last pass in a phase
+		return false;
+	case ESaveAction::IgnoreResults:
+		// The IgnoreResults pass, if present, is the last pass in a phase
+		return false;
+	default:
+		checkNoEntry();
+		break;
+	}
+
+	switch (Phase)
+	{
+	case EPhase::AllInOnePhase:
+		check(GetPackageStatus(BeginInfo.PackageName) == EPackageStatus::DeclaredUnmodified_NotYetProcessed); // Otherwise we would have set SaveAction=SaveToInner or IgnoreResults and early exited above
+		if (Super::IsAnotherSaveNeeded(PreviousResult, SaveArgs))
+		{
+			return true;
+		}
+		else
+		{
+			// Once our superclass has finished looking for differences, finish it off and start a SaveToInner pass
+			FCommitPackageInfo CommitInfo;
+			CommitInfo.Status = IPackageWriter::ECommitStatus::Success;
+			CommitInfo.PackageName = BeginInfo.PackageName;
+			CommitInfo.WriteOptions = EWriteOptions::None;
+			Super::CommitPackage(MoveTemp(CommitInfo));
+
+			if (bIsDifferent && !bNewPackage)
+			{
+				if (!PackageIgnoreList.Contains(BeginInfo.PackageName))
+				{
+					SetPackageStatus(BeginInfo.PackageName, EPackageStatus::DeclaredUnmodified_FoundModified_FalsePositive);
+				}
+				else
+				{
+					SetPackageStatus(BeginInfo.PackageName, EPackageStatus::DeclaredUnmodified_FoundModified_OnIgnoreList);
+				}
+			}
+			else if (!bNewPackage)
+			{
+
+				SetPackageStatus(BeginInfo.PackageName, EPackageStatus::DeclaredUnmodified_ConfirmedUnmodified);
+			}
+			else
+			{
+				SetPackageStatus(BeginInfo.PackageName, EPackageStatus::DeclaredModified_WillNotVerify);
+			}
+
+			if (bReadOnly)
+			{
+				SaveAction = ESaveAction::IgnoreResults;
+				return false;
+			}
+			else
+			{
+				Inner->BeginPackage(BeginInfo);
+				SaveAction = ESaveAction::SaveToInner;
+				return true;
+			}
+		}
+	case EPhase::Phase1:
+		if (Super::IsAnotherSaveNeeded(PreviousResult, SaveArgs))
+		{
+			return true;
+		}
+		else if (bIsDifferent && !bNewPackage)
+		{
+			// If our superclass FDiffPackageWriter found differences, when it finishes the saves it wants to do,
+			// Finish it off and start a SaveToInner pass
+			FCommitPackageInfo CommitInfo;
+			CommitInfo.Status = IPackageWriter::ECommitStatus::Success;
+			CommitInfo.PackageName = BeginInfo.PackageName;
+			CommitInfo.WriteOptions = EWriteOptions::None;
+			Super::CommitPackage(MoveTemp(CommitInfo));
+			Inner->BeginPackage(BeginInfo);
+			SaveAction = ESaveAction::SaveToInner;
+
+			if (!PackageIgnoreList.Contains(BeginInfo.PackageName))
+			{
+				// Mark that the incremental validation failed if it was not already marked by log or warning messages.
+				// We need to record it for an indeterminism test
+				SetPackageStatus(BeginInfo.PackageName, EPackageStatus::DeclaredUnmodified_FoundModified_IndeterminismOrFalsePositive);
+			}
+			else
+			{
+				SetPackageStatus(BeginInfo.PackageName, EPackageStatus::DeclaredUnmodified_FoundModified_OnIgnoreList);
+			}
+
+			return true;
+		}
+		else if (!bNewPackage)
+		{
+			// No differences found, so finish off the superclass's save during CommitPackage, without doing a
+			// SaveToInner pass
+			// Mark that the incremental validation passed
+			SetPackageStatus(BeginInfo.PackageName, EPackageStatus::DeclaredUnmodified_ConfirmedUnmodified);
+
+			TArray<FMessage> Messages;
+			PackageMessageMap.RemoveAndCopyValue(BeginInfo.PackageName, Messages);
+			for (FMessage& Message : Messages)
+			{
+				// If no differences were detected, we should not have logged any warning or error messages
+				check(Message.Verbosity > ELogVerbosity::Warning);
+			}
+			return false;
+		}
+		else
+		{
+			// New packages need to be resaved in Phase2; for our purposes they are equivalent
+			// to a package that the IncrementalCook detected as modified.
+			// Do not add an entry for it in our results for incremental packages, and do not resave it in this pass
+			SetPackageStatus(BeginInfo.PackageName, EPackageStatus::DeclaredModified_WillNotVerify);
+			return false;
+		}
+	case EPhase::Phase2:
+		LogIncrementalDifferences();
+		// No need to do the Super's second diff pass to find callstacks - just knowing whether differences exist is enough.
+		// No need to save package to disk; for these packages (packages found to be incrementally unmodified during
+		// Phase1) they were already resaved during Phase1
+		return false;
+	default:
+		checkNoEntry();
+		return false;
+	}
+}
+
+void FIncrementalValidatePackageWriter::OnDiffWriterMessage(ELogVerbosity::Type Verbosity, FStringView Message)
+{
+	PackageMessageMap.FindOrAdd(BeginInfo.PackageName).Add(FMessage{ FString(Message), Verbosity });
+}
+
+void FIncrementalValidatePackageWriter::MarkPackageCompletedOnDirector(FName PackageName,
+	UE::Cook::FWorkerId WorkerId)
+{
+	FPackageStatusInfo* Status = PackageStatusMap.Find(PackageName);
+	TArray<FMessage>* Messages = PackageMessageMap.Find(PackageName);
+
+	if (!Status || !Messages)
+	{
+		return;
+	}
+
+	int32& TotalCount = TotalStatusCounts.FindOrAdd(Status->Status, 0);
+	TArray<FName>& ClassStatusArray = ClassStatusSummary.FindOrAdd(Status->AssetClass).FindOrAdd(Status->Status);
+	++TotalCount;
+	ClassStatusArray.Add(PackageName);
+	if (LoggingSoftMaximum >= 0 && TotalCount > LoggingSoftMaximum && ClassStatusArray.Num() > 1)
+	{
+		// Suppress reported logs for this package
+		return;
+	}
+
+	switch (Phase)
+	{
+	case EPhase::AllInOnePhase:
+		for (FMessage& Message : *Messages)
+		{
+			FMsg::Logf(__FILE__, __LINE__, LogIncrementalValidate.GetCategoryName(), Message.Verbosity,
+				TEXT("%s"), *ResolveText(Message.Text));
+		}
+		break;
+	case EPhase::Phase1:
+		break;
+	case EPhase::Phase2:
+		break;
+	default:
+		checkNoEntry();
+		break;
+	}
+}
+
+void FIncrementalValidatePackageWriter::LogIncrementalDifferences()
+{
+	// This function is called during Phase2 for packages that had differences during Phase1.
+	// It is called immediately after first-pass package save that is run by our super class FDiffPackageWriter, which
+	// compared it against the version that was written to disk during Phase1. If there are differences
+	// now from Phase1, then this package has a determinism issue. We log that information at display rather
+	// than Warning because this cookmode only logs warnings for IncrementalSkipFalsePositives.
+	bool bHasDeterminismIssue = bIsDifferent;
+	if (bHasDeterminismIssue)
+	{
+		UE_LOG(LogIncrementalValidate, Display, TEXT("Could not validate %s because it has a non-deterministic save."),
+			*BeginInfo.PackageName.ToString());
+		SetPackageStatus(BeginInfo.PackageName, EPackageStatus::DeclaredUnmodified_FoundModified_Indeterminism);
+		return;
+	}
+
+	// Otherwise, no determinism issues, so the differences indicate a bug in Diff Package
+	SetPackageStatus(BeginInfo.PackageName, EPackageStatus::DeclaredUnmodified_FoundModified_FalsePositive);
+	FMsg::Logf(__FILE__, __LINE__, LogIncrementalValidate.GetCategoryName(), ELogVerbosity::Warning,
+		TEXT("IncrementalSkipFalsePositive package %s."), *BeginInfo.PackageName.ToString());
+	TArray<FMessage>& Messages = PackageMessageMap.FindOrAdd(BeginInfo.PackageName);
+	for (const FMessage& Message : Messages)
+	{
+		FMsg::Logf(__FILE__, __LINE__, LogIncrementalValidate.GetCategoryName(), Message.Verbosity,
+			TEXT("%s"), *ResolveText(Message.Text));
+	}
+}
+
+void FIncrementalValidatePackageWriter::Save()
+{
+	FString IncrementalValidatePath = GetIncrementalValidatePath();
+	TUniquePtr<FArchive> DiskArchive(IFileManager::Get().CreateFileWriter(*IncrementalValidatePath));
+
+	if (!DiskArchive)
+	{
+		UE_LOG(LogIncrementalValidate, Error,
+			TEXT("Could not write to file %s. This file is needed to store results for the -IncrementalValidate cook."),
+			*IncrementalValidatePath);
+		return;
+	}
+	FNameAsStringProxyArchive Ar(*DiskArchive);
+	Serialize(Ar);
+}
+
+void FIncrementalValidatePackageWriter::Load()
+{
+	FString IncrementalValidatePath = GetIncrementalValidatePath();
+	TUniquePtr<FArchive> DiskArchive(IFileManager::Get().CreateFileReader(*IncrementalValidatePath));
+	if (!DiskArchive)
+	{
+		UE_LOG(LogIncrementalValidate, Fatal,
+			TEXT("Could not load file %s. This file is required and should have been written by the -IncrementalValidatePhase1 cook."),
+			*IncrementalValidatePath);
+		return;
+	}
+	FNameAsStringProxyArchive Ar(*DiskArchive);
+	Serialize(Ar);
+	if (Ar.IsError())
+	{
+		UE_LOG(LogIncrementalValidate, Fatal, TEXT("Corrupt file %s"), *IncrementalValidatePath);
+	}
+}
+
+void FIncrementalValidatePackageWriter::Serialize(FArchive& Ar)
+{
+	constexpr int32 LatestVersion = 0;
+	int32 Version = LatestVersion;
+	Ar << Version;
+	if (Ar.IsLoading() && Version != LatestVersion)
+	{
+		Ar.SetError();
+		return;
+	}
+	Ar << PackageStatusMap;
+	Ar << PackageMessageMap;
+}
+
+FArchive& operator<<(FArchive& Ar, FIncrementalValidatePackageWriter::FMessage& Message)
+{
+	uint8 Verbosity = static_cast<uint8>(Message.Verbosity);
+	Ar << Verbosity << Message.Text;
+	if (Ar.IsLoading())
+	{
+		Message.Verbosity = static_cast<ELogVerbosity::Type>(Verbosity);
+	}
+	return Ar;
+}
+
+FString FIncrementalValidatePackageWriter::GetIncrementalValidatePath() const
+{
+	return FPaths::Combine(MetadataPath, FString(IncrementalValidateFilename));
+}
+
+FIncrementalValidatePackageWriter::EPackageStatus FIncrementalValidatePackageWriter::GetPackageStatus(FName PackageName) const
+{
+	if (const FPackageStatusInfo* Info = PackageStatusMap.Find(PackageName))
+	{
+		return Info->Status;
+	}
+	return EPackageStatus::NotYetProcessed;
+}
+
+void FIncrementalValidatePackageWriter::SetPackageStatus(FName PackageName, EPackageStatus NewStatus)
+{
+	FPackageStatusInfo& Info = PackageStatusMap.FindOrAdd(PackageName);
+	Info.Status = NewStatus;
+	switch (NewStatus)
+	{
+	case EPackageStatus::DeclaredUnmodified_ConfirmedUnmodified: [[fallthrough]];
+	case EPackageStatus::DeclaredModified_WillNotVerify:
+		// For the non-error cases, clear the AssetClass so that we don't waste bandwidth or diskspace to store it.
+		Info.AssetClass = FTopLevelAssetPath();
+		break;
+	default:
+		break;
+	}
+}
+
+FIncrementalValidatePackageWriter::FStatusCounts FIncrementalValidatePackageWriter::CountPackagesByStatus()
+{
+	FIncrementalValidatePackageWriter::FStatusCounts StatusCounts;
+
+	for (const TPair<FName,FPackageStatusInfo>& Pair : PackageStatusMap)
+	{
+		StatusCounts[Pair.Value.Status]++;
+	}
+
+	return StatusCounts;
+}
+
+TMap<FTopLevelAssetPath, TArray<FName>> FIncrementalValidatePackageWriter::GetClassStatusSummary(EPackageStatus PackageStatus)
+{
+	TMap<FTopLevelAssetPath, TArray<FName>> Result;
+	for (const TPair<FTopLevelAssetPath, TMap<EPackageStatus, TArray<FName>>>& Pair : ClassStatusSummary)
+	{
+		const TArray<FName>* Packages = Pair.Value.Find(PackageStatus);
+		if (Packages)
+		{
+			Result.Add(Pair.Key, *Packages);
+		}
+	}
+
+	Result.ValueSort([](const TArray<FName>& A, const TArray<FName>& B)
+		{
+			return A.Num() > B.Num();
+		});
+	return Result;
+}
diff --git a/Engine/Source/Editor/UnrealEd/Private/Cooker/IncrementalValidatePackageWriter.h b/Engine/Source/Editor/UnrealEd/Private/Cooker/IncrementalValidatePackageWriter.h
new file mode 100644
--- /dev/null
+++ b/Engine/Source/Editor/UnrealEd/Private/Cooker/IncrementalValidatePackageWriter.h
@@ -0,0 +1,145 @@
+// Copyright Epic Games, Inc. All Rights Reserved.
+
+#pragma once
+
+#include "DiffPackageWriter.h"
+
+class UCookOnTheFlyServer;
+class FIncrementalValidateMPCollector; 
+namespace UE::Cook { struct FWorkerId; }
+
+/**
+ * A CookedPackageWriter that diffs the cook results of incrementally-unmodified packages between their last cook
+ * results and the current cook.
+ */
+class FIncrementalValidatePackageWriter : public FDiffPackageWriter
+{
+public:
+	using Super = FDiffPackageWriter;
+	enum class EPhase
+	{
+		AllInOnePhase,
+		Phase1,
+		Phase2,
+	};
+	FIncrementalValidatePackageWriter(UCookOnTheFlyServer& InCOTFS, TUniquePtr<ICookedPackageWriter>&& InInner,
+		EPhase InPhase, const FString& ResolvedMetadataPath,
+		UE::Cook::FDeterminismManager* InDeterminismManager);
+
+	// IPackageWriter
+	virtual void BeginPackage(const FBeginPackageInfo& Info) override;
+	virtual void CommitPackage(FCommitPackageInfo&& Info) override;
+	virtual void WritePackageData(const FPackageInfo& Info, FLargeMemoryWriter& ExportsArchive,
+		const TArray<FFileRegion>& FileRegions) override;
+	virtual void WriteBulkData(const FBulkDataInfo& Info, const FIoBuffer& BulkData,
+		const TArray<FFileRegion>& FileRegions) override;
+	virtual void WriteAdditionalFile(const FAdditionalFileInfo& Info, const FIoBuffer& FileData) override;
+	virtual void WriteLinkerAdditionalData(const FLinkerAdditionalDataInfo& Info, const FIoBuffer& Data,
+		const TArray<FFileRegion>& FileRegions) override;
+	virtual void WritePackageTrailer(const FPackageTrailerInfo& Info, const FIoBuffer& Data) override;
+	virtual int64 GetExportsFooterSize() override;
+	virtual TUniquePtr<FLargeMemoryWriter> CreateLinkerArchive(FName PackageName, UObject* Asset, uint16 MultiOutputIndex) override;
+	virtual TUniquePtr<FLargeMemoryWriter> CreateLinkerExportsArchive(FName PackageName, UObject* Asset, uint16 MultiOutputIndex) override;
+	virtual bool IsPreSaveCompleted() const override;
+
+	// ICookedPackageWriter
+	virtual FCookCapabilities GetCookCapabilities() const override;
+	virtual void Initialize(const FCookInfo& CookInfo) override;
+	virtual void UpdatePackageModificationStatus(FName PackageName, bool bIncrementallyUnmodified,
+		bool& bInOutShouldIncrementallySkip) override;
+	virtual void BeginCook(const FCookInfo& Info) override;
+	virtual void EndCook(const FCookInfo& Info) override;
+	virtual void UpdateSaveArguments(FSavePackageArgs& SaveArgs) override;
+	virtual bool IsAnotherSaveNeeded(FSavePackageResultStruct& PreviousResult, FSavePackageArgs& SaveArgs) override;
+
+protected:
+	enum class EPackageStatus : uint8;
+
+	enum class ESaveAction : uint8
+	{
+		CheckForDiffs,
+		SaveToInner,
+		IgnoreResults,
+	};
+
+	enum class EPackageStatus : uint8
+	{
+		NotYetProcessed,
+		DeclaredUnmodified_ConfirmedUnmodified,
+		DeclaredUnmodified_FoundModified_IndeterminismOrFalsePositive,
+		DeclaredUnmodified_FoundModified_Indeterminism,
+		DeclaredUnmodified_FoundModified_FalsePositive,
+		DeclaredUnmodified_FoundModified_OnIgnoreList,
+		DeclaredUnmodified_NotYetProcessed,
+		DeclaredModified_WillNotVerify,
+		Count
+	};
+
+	struct FMessage
+	{
+		FString Text;
+		ELogVerbosity::Type Verbosity;
+	};
+	friend FArchive& operator<<(FArchive& Ar, FMessage& Message);
+
+	struct FStatusCounts
+	{
+		FStatusCounts()
+		{
+			Data = MakeUniformStaticArray<uint32, (uint32)EPackageStatus::Count>(0);
+		}
+
+		uint32& operator[](EPackageStatus PackageStatus) 
+		{ 
+			return Data[(uint32)PackageStatus]; 
+		}
+
+	private:
+		TStaticArray<uint32, (uint32)EPackageStatus::Count> Data;
+	};
+
+	struct FPackageStatusInfo
+	{
+		FTopLevelAssetPath AssetClass;
+		EPackageStatus Status;
+	};
+
+protected:
+	virtual void OnDiffWriterMessage(ELogVerbosity::Type Verbosity, FStringView Message) override;
+	void LogIncrementalDifferences();
+	void Save();
+	void Load();
+	void Serialize(FArchive& Ar);
+	FString GetIncrementalValidatePath() const;
+	EPackageStatus GetPackageStatus(FName PackageName) const;
+	void SetPackageStatus(FName PackageName, EPackageStatus NewStatus);
+	FStatusCounts CountPackagesByStatus();
+	TMap<FTopLevelAssetPath, TArray<FName>> GetClassStatusSummary(EPackageStatus PackageStatus);
+	bool IsAnotherSaveNeededInternal(FSavePackageResultStruct& PreviousResult, FSavePackageArgs& SaveArgs);
+	void MarkPackageCompletedOnDirector(FName PackageName, UE::Cook::FWorkerId WorkerId);
+
+protected:
+	TMap<FName, FPackageStatusInfo> PackageStatusMap;
+	TMap<FName, TArray<FMessage>> PackageMessageMap;
+	TMap<EPackageStatus, int32> TotalStatusCounts; // Only populated on director
+	TMap<FTopLevelAssetPath, TMap<EPackageStatus, TArray<FName>>> ClassStatusSummary; // Only populated on director
+	TSet<FName> PackageIgnoreList;
+	FString MetadataPath;
+	UCookOnTheFlyServer& COTFS;
+	int32 LoggingSoftMaximum = -1;
+	EPhase Phase = EPhase::AllInOnePhase;
+	ESaveAction SaveAction = ESaveAction::IgnoreResults;
+	bool bPackageFirstPass = false;
+	bool bReadOnly = true;
+
+	friend class FIncrementalValidateMPCollector;
+	friend FCbWriter& operator<<(FCbWriter& Writer, const FMessage& Message);
+	friend bool LoadFromCompactBinary(FCbFieldView Field, FMessage& Path);
+	friend FCbWriter& operator<<(FCbWriter& Writer, EPackageStatus Status);
+	friend bool LoadFromCompactBinary(FCbFieldView Field, EPackageStatus& Status);
+	friend FArchive& operator<<(FArchive& Writer, FPackageStatusInfo& Info);
+	friend FCbWriter& operator<<(FCbWriter& Writer, const FPackageStatusInfo& Info);
+	friend bool LoadFromCompactBinary(FCbFieldView Field, FPackageStatusInfo& Info);
+};
+
+
diff --git a/Engine/Source/Editor/UnrealEd/Private/Cooker/IterativeValidatePackageWriter.cpp b/Engine/Source/Editor/UnrealEd/Private/Cooker/IterativeValidatePackageWriter.cpp
deleted file mode 100644
diff --git a/Engine/Source/Editor/UnrealEd/Private/Cooker/IterativeValidatePackageWriter.h b/Engine/Source/Editor/UnrealEd/Private/Cooker/IterativeValidatePackageWriter.h
deleted file mode 100644
