Cooker: Downgrade LogGenerationHelper default from Verbose to Log. The errors that required the increased logging have not occurred for a while and it generates a lot of spam.
#rnx
#rnx

diff --git a/Engine/Source/Editor/UnrealEd/Private/Cooker/CookTypes.h b/Engine/Source/Editor/UnrealEd/Private/Cooker/CookTypes.h
--- a/Engine/Source/Editor/UnrealEd/Private/Cooker/CookTypes.h
+++ b/Engine/Source/Editor/UnrealEd/Private/Cooker/CookTypes.h
@@ -52,7 +52,7 @@ bool LoadFromCompactBinary(FCbFieldView Field, UE::Cook::FInitializeConfigSettin
 #define DEBUG_COOKONTHEFLY 0
 
 LLM_DECLARE_TAG(Cooker_CachedPlatformData);
-DECLARE_LOG_CATEGORY_EXTERN(LogCookGenerationHelper, Verbose, All);
+DECLARE_LOG_CATEGORY_EXTERN(LogCookGenerationHelper, Log, All);
 
 inline constexpr uint32 ExpectedMaxNumPlatforms = 32;
 
