CookCorrectness: Do not log content missing from cook warnings when running incrementalvalidate or other diffs. Not all packages are saved during incrementalvalidate.
#rnx

diff --git a/Engine/Source/Editor/UnrealEd/Private/CookOnTheFlyServer.cpp b/Engine/Source/Editor/UnrealEd/Private/CookOnTheFlyServer.cpp
--- a/Engine/Source/Editor/UnrealEd/Private/CookOnTheFlyServer.cpp
+++ b/Engine/Source/Editor/UnrealEd/Private/CookOnTheFlyServer.cpp
@@ -5755,6 +5755,11 @@ public:
 		bInitialized = true;
 	}
 
+	bool IsDiffModeActive() const
+	{
+		return DiffMode != EDiffMode::None;
+	}
+
 private:
 	bool bInitialized = false;
 	EDiffMode DiffMode = EDiffMode::None;
@@ -11948,7 +11953,7 @@ void FEDLMPCollector::ServerReceiveMessage(FMPCollectorServerMessageContext& Con
 
 bool UCookOnTheFlyServer::ShouldVerifyEDLCookInfo() const
 {
-	return CookByTheBookOptions->DlcName.IsEmpty() && !bCookFilter;
+	return CookByTheBookOptions->DlcName.IsEmpty() && !bCookFilter && !DiffModeHelper->IsDiffModeActive();
 }
 
 void UCookOnTheFlyServer::BeginCookEDLCookInfo(FBeginCookContext& BeginContext)
