Cooker: Turn off logging of UObject counts after GC when LLM is not enabled. Periodic SoftGC has been working well to reduce memory use so we are not using this information anymore, and it generates a lot of spam.
#rnx

diff --git a/Engine/Source/Editor/UnrealEd/Private/Cooker/CookGarbageCollect.cpp b/Engine/Source/Editor/UnrealEd/Private/Cooker/CookGarbageCollect.cpp
--- a/Engine/Source/Editor/UnrealEd/Private/Cooker/CookGarbageCollect.cpp
+++ b/Engine/Source/Editor/UnrealEd/Private/Cooker/CookGarbageCollect.cpp
@@ -1092,6 +1092,14 @@ void UCookOnTheFlyServer::EvaluateGarbageCollectionResults(bool bWasDueToOOM, bo
 		);
 	};
 
+	// Only show diagnostics if LLM is on, because they are somewhat expensive. We could add a separate setting
+	// for this, but it's more convenient to combine it with the LLM enabled setting
+#if ENABLE_LOW_LEVEL_MEM_TRACKER
+	bool bShowExtendedDiagnostics = FLowLevelMemTracker::Get().IsEnabled();
+#else
+	constexpr bool bShowExtendedDiagnostics = false;
+#endif
+
 	if (!bWasSoftGC)
 	{
 		bool bWasImpactful =
@@ -1151,25 +1159,39 @@ void UCookOnTheFlyServer::EvaluateGarbageCollectionResults(bool bWasDueToOOM, bo
 			*AllocatorStatsToString(AllocatorStatsAfterGC));
 
 
-		UE_LOG(LogCook, Display, TEXT("See log for memory use information for UObject classes and LLM tags."));
-
-		{
-			TGuardValue<bool> SoftGCGuard(UPackage::bSupportCookerSoftGC, true);
-			ConstructSoftGCPackageToObjectList(this->SoftGCPackageToObjectListBuffer);
-			UE::Cook::DumpObjClassList(CookByTheBookOptions->SessionStartupObjects);
-			UPackage::SoftGCPackageToObjectList.Empty();
-			SoftGCPackageToObjectListBuffer.Empty();
-		}
-		GLog->Logf(TEXT("Memory Analysis: LLM Tags:"));
 #if ENABLE_LOW_LEVEL_MEM_TRACKER
-		if (FLowLevelMemTracker::Get().IsEnabled())
+		bool bExtendedMemoryAnalysisEnabled = FLowLevelMemTracker::Get().IsEnabled();
+#else
+		constexpr bool bExtendedMemoryAnalysisEnabled = false;
+#endif
+
+		// Only show diagnostics if LLM is on, because they are somewhat expensive. We could add a separate setting
+		// for this, but it's more convenient to combine it with the LLM enabled setting
+		if (!bShowExtendedDiagnostics)
 		{
-			FLowLevelMemTracker::Get().DumpToLog();
+			UE_LOG(LogCook, Display, TEXT("Extended memory diagnostics are disabled. Run with -llm or -trace=memtag to log information for UObject classes and LLM tags."));
 		}
 		else
-#endif
 		{
-			GLog->Logf(TEXT("LLM Tags are not displayed because llm is disabled. Run with -llm or -trace=memtag to see llm tags."));
+			UE_LOG(LogCook, Display, TEXT("See log for memory use information for UObject classes and LLM tags."));
+			{
+				TGuardValue<bool> SoftGCGuard(UPackage::bSupportCookerSoftGC, true);
+				ConstructSoftGCPackageToObjectList(this->SoftGCPackageToObjectListBuffer);
+				UE::Cook::DumpObjClassList(CookByTheBookOptions->SessionStartupObjects);
+				UPackage::SoftGCPackageToObjectList.Empty();
+				SoftGCPackageToObjectListBuffer.Empty();
+			}
+			GLog->Logf(TEXT("Memory Analysis: LLM Tags:"));
+#if ENABLE_LOW_LEVEL_MEM_TRACKER
+			if (FLowLevelMemTracker::Get().IsEnabled())
+			{
+				FLowLevelMemTracker::Get().DumpToLog();
+			}
+			else
+#endif
+			{
+				GLog->Logf(TEXT("LLM Tags are not displayed because llm is disabled. Run with -llm or -trace=memtag to see llm tags."));
+			}
 		}
 	}
 	else
@@ -1180,14 +1202,7 @@ void UCookOnTheFlyServer::EvaluateGarbageCollectionResults(bool bWasDueToOOM, bo
 		// get loaded again.
 		PackageTracker->AddExpectedNeverLoadPackages(ExpectedFreedPackageNames);
 
-		// Only show diagnostics if LLM is on, because they are somewhat expensive. We could add a separate setting
-		// for this, but it's more convenient to combine it with the LLM enabled setting
-#if ENABLE_LOW_LEVEL_MEM_TRACKER
-		bool bShowDiagnostics = FLowLevelMemTracker::Get().IsEnabled();
-#else
-		constexpr bool bShowDiagnostics = false;
-#endif
-		if (bShowDiagnostics)
+		if (bShowExtendedDiagnostics)
 		{
 			// If some packages we expected to be freed were not freed, show the reference chains for why
 			// they were not freed.
