[AssetThumbnail] Scale down Play indicator and status after certain size

diff --git a/Engine/Source/Editor/ContentBrowser/Private/AssetViewWidgets.cpp b/Engine/Source/Editor/ContentBrowser/Private/AssetViewWidgets.cpp
--- a/Engine/Source/Editor/ContentBrowser/Private/AssetViewWidgets.cpp
+++ b/Engine/Source/Editor/ContentBrowser/Private/AssetViewWidgets.cpp
@@ -2041,7 +2041,7 @@ int32 SAssetListItem::GetWidgetSwitcherIndex() const
 
 EVisibility SAssetListItem::GetDirtyBadgeVisibility() const
 {
-	return CurrentThumbnailSize.IsSet() && CurrentThumbnailSize.Get() == EThumbnailSize::Tiny ? EVisibility::Visible : EVisibility::Collapsed;
+	return CurrentThumbnailSize.IsSet() && CurrentThumbnailSize.Get() <= EThumbnailSize::Medium ? EVisibility::Visible : EVisibility::Collapsed;
 }
 
 TSharedRef<SWidget> SAssetListItem::GenerateWidgetForColumn(const FName& ColumnName, FIsSelected InIsSelected)
@@ -2159,7 +2159,7 @@ int32 SAssetListItem::GetGenericThumbnailSize() const
 
 bool SAssetListItem::ShouldAllowAssetStatusOverlay(bool bForceGenericThumbnail) const
 {
-	return !bForceGenericThumbnail && CurrentThumbnailSize.Get() >= EThumbnailSize::XLarge;
+	return !bForceGenericThumbnail && CurrentThumbnailSize.Get() >= EThumbnailSize::Large;
 }
 
 bool SAssetListItem::ShouldAllowAssetActionOverlay(bool bForceGenericThumbnail) const
@@ -2169,7 +2169,7 @@ bool SAssetListItem::ShouldAllowAssetActionOverlay(bool bForceGenericThumbnail)
 
 bool SAssetListItem::ShouldAllowAssetActionOverlayIndicator(bool bForceGenericThumbnail) const
 {
-	return !bForceGenericThumbnail && CurrentThumbnailSize.Get() >= EThumbnailSize::Large;
+	return !bForceGenericThumbnail && CurrentThumbnailSize.Get() >= EThumbnailSize::Small;
 }
 
 bool SAssetListItem::ShouldUseAssetChip() const
@@ -2236,43 +2236,48 @@ TSharedRef<SWidget> SAssetListItem::GenerateWidgetForColumnNewStyle(const FName&
 				.WidthOverride( this, &SAssetListItem::GetThumbnailBoxSize )
 				.HeightOverride( this, &SAssetListItem::GetThumbnailBoxSize )
 				[
-					SNew(SWidgetSwitcher)
-					.WidgetIndex(this, &SAssetListItem::GetWidgetSwitcherIndex)
+					SNew(SOverlay)
 
-					+ SWidgetSwitcher::Slot()
-					.VAlign(VAlign_Fill)
-					.HAlign(HAlign_Fill)
+					+ SOverlay::Slot()
 					[
-						SNew(SOverlay)
+						SNew(SWidgetSwitcher)
+						.WidgetIndex(this, &SAssetListItem::GetWidgetSwitcherIndex)
 
-						+ SOverlay::Slot()
+						+ SWidgetSwitcher::Slot()
+						.VAlign(VAlign_Fill)
+						.HAlign(HAlign_Fill)
 						[
-							SNew(SImage)
-							.Image(IconBrush)
-							.ColorAndOpacity(this, &SAssetViewItem::GetAssetColor)
-						]
+							SNew(SOverlay)
 
-						// Dirty state
-						+SOverlay::Slot()
-						.HAlign(HAlign_Left)
-						.VAlign(VAlign_Bottom)
-						[
-							SNew(SBox)
-							.WidthOverride(IconOverlaySize)
-							.HeightOverride(IconOverlaySize)
-							.Visibility(this, &SAssetListItem::GetDirtyBadgeVisibility)
+							+ SOverlay::Slot()
 							[
 								SNew(SImage)
-								.Image(this, &SAssetListItem::GetDirtyImage)
+								.Image(IconBrush)
+								.ColorAndOpacity(this, &SAssetViewItem::GetAssetColor)
 							]
 						]
+
+						+ SWidgetSwitcher::Slot()
+						.VAlign(VAlign_Fill)
+						.HAlign(HAlign_Fill)
+						[
+							ItemContents
+						]
 					]
 
-					+ SWidgetSwitcher::Slot()
-					.VAlign(VAlign_Fill)
-					.HAlign(HAlign_Fill)
+					// Dirty state
+					+SOverlay::Slot()
+					.HAlign(HAlign_Left)
+					.VAlign(VAlign_Bottom)
 					[
-						ItemContents
+						SNew(SBox)
+						.WidthOverride(IconOverlaySize)
+						.HeightOverride(IconOverlaySize)
+						.Visibility(this, &SAssetListItem::GetDirtyBadgeVisibility)
+						[
+							SNew(SImage)
+							.Image(this, &SAssetListItem::GetDirtyImage)
+						]
 					]
 				]
 			]
diff --git a/Engine/Source/Editor/EditorStyle/Private/StarshipStyle.cpp b/Engine/Source/Editor/EditorStyle/Private/StarshipStyle.cpp
--- a/Engine/Source/Editor/EditorStyle/Private/StarshipStyle.cpp
+++ b/Engine/Source/Editor/EditorStyle/Private/StarshipStyle.cpp
@@ -1255,6 +1255,7 @@ void FStarshipEditorStyle::FStyle::SetupGeneralStyles()
 		Set( "AssetThumbnail.ClassBackground", new IMAGE_BRUSH( "Common/ClassBackground_64x", Icon64x64, FLinearColor(0.75f, 0.75f, 0.75f, 1.0f) ) );
 		Set( "AssetThumbnail.Font", DEFAULT_FONT( "Regular", 10 ) );
 		Set( "AssetThumbnail.StatusOverflowFont", DEFAULT_FONT( "Bold", 9 ) );
+		Set( "AssetThumbnail.StatusOverflowFontSmall", DEFAULT_FONT( "Regular", 5 ) );
 		Set( "AssetThumbnail.FontSmall", DEFAULT_FONT( "Regular", 7 ) );
 		Set( "AssetThumbnail.ColorAndOpacity", FLinearColor(1.75f, 1.75f, 1.75f, 1.f) );
 		Set( "AssetThumbnail.HintFont", DEFAULT_FONT( "Regular", 8 ) );
diff --git a/Engine/Source/Editor/UnrealEd/Private/AssetThumbnail.cpp b/Engine/Source/Editor/UnrealEd/Private/AssetThumbnail.cpp
--- a/Engine/Source/Editor/UnrealEd/Private/AssetThumbnail.cpp
+++ b/Engine/Source/Editor/UnrealEd/Private/AssetThumbnail.cpp
@@ -457,6 +457,7 @@ public:
 				.Padding(StatusPadding)
 				[
 					SNew(SBorder)
+					.Padding(StatusBorderPadding)
 					.BorderImage(FAppStyle::GetBrush(Style, ".AssetThumbnailStatusBar"))
 					.Visibility(this, &SAssetThumbnail::GetStatusBorderVisibility)
 					[
@@ -478,8 +479,6 @@ public:
 					if (OutThumbnailInfo.ActionImageWidget.IsValid())
 					{
 						constexpr float PaddingFromTopLeftBorder = 4.f;
-						constexpr float TopLeftImageSize = 20.f;
-						constexpr float BorderPaddingToActionImageWidget = 4.f;
 
 						OverlayWidget->AddSlot()
 						.ZOrder(OverlayZOrder)
@@ -488,12 +487,12 @@ public:
 						.Padding(PaddingFromTopLeftBorder, PaddingFromTopLeftBorder, 0.f, 0.f)
 						[
 							SNew(SBox)
-							.WidthOverride(TopLeftImageSize)
-							.HeightOverride(TopLeftImageSize)
+							.WidthOverride(this, &SAssetThumbnail::GetPlayIndicatorSize)
+							.HeightOverride(this, &SAssetThumbnail::GetPlayIndicatorSize)
 							[
 								SNew(SBorder)
 								.BorderImage(FAppStyle::GetBrush(Style, ".AssetThumbnailBar"))
-								.Padding(BorderPaddingToActionImageWidget)
+								.Padding(this, &SAssetThumbnail::GetPlayIndicatorPadding)
 								.Visibility(this, &SAssetThumbnail::GetActionIconOverlayVisibility)
 								[
 									OutThumbnailInfo.ActionImageWidget.ToSharedRef()
@@ -532,7 +531,7 @@ public:
 						// Set the default style and padding for the Button
 						OutThumbnailInfo.ActionButtonArgs
 							.ButtonStyle(&FAppStyle::GetWidgetStyle<FButtonStyle>(Style, ".Action.Button"))
-							.ContentPadding(8.f);
+							.ContentPadding(this, &SAssetThumbnail::GetPlayButtonContentPadding);
 
 						TAttribute<EVisibility> ActionVisibility = TAttribute<EVisibility>::CreateSP(this, &SAssetThumbnail::GetActionButtonVisibility);
 						TSharedRef<SButton> ActionButton = SNew(SButton);
@@ -693,9 +692,59 @@ public:
 				HintTextBlock->SetFont( GetHintTextFont() );
 				HintTextBlock->SetWrapTextAt( GetTextWrapWidth() );
 			}
+
+			if (!OverlayInfo.IsEmpty())
+			{
+				const float ThumbnailWidth = WidthLastFrame - (StatusPadding * 2) - (StatusBorderPadding * 2) - (GetAssetThumbnailBorderPadding() * 2);
+				const int32 MaxShownStatus = FMath::FloorToInt32(ThumbnailWidth / DefaultStatusSize);
+				constexpr int32 CutoffNumberBeforeResizing = 3;
+
+				if (MaxShownStatus < CutoffNumberBeforeResizing)
+				{
+					StatusSize = ThumbnailWidth / 3.f;
+				}
+				else
+				{
+					StatusSize = DefaultStatusSize;
+				}
+
+				StatusSize = FMath::FloorToFloat(StatusSize);
+			}
+
+			if (WidthLastFrame < PlayIndicatorMaxSizeThreshold)
+			{
+				PlayIndicatorSize = (PlayIndicatorDefaultSize * WidthLastFrame) / PlayIndicatorMaxSizeThreshold;
+				PlayIndicatorPadding = (PlayIndicatorDefaultPadding * WidthLastFrame) / PlayIndicatorMaxSizeThreshold;
+				PlayButtonContentPadding = (PlayButtonContentDefaultPadding * WidthLastFrame) / PlayIndicatorMaxSizeThreshold;
+			}
+			else
+			{
+				PlayIndicatorSize = PlayIndicatorDefaultSize;
+				PlayIndicatorPadding = PlayIndicatorDefaultPadding;
+				PlayButtonContentPadding = PlayButtonContentDefaultPadding;
+			}
+
+			PlayButtonContentPadding = FMath::FloorToFloat(PlayButtonContentDefaultPadding);
+			PlayIndicatorPadding = FMath::FloorToFloat(PlayIndicatorPadding);
+			PlayIndicatorSize = FMath::FloorToFloat(PlayIndicatorSize);
 		}
 	}
 
+	FOptionalSize GetPlayIndicatorSize() const
+	{
+		return FOptionalSize(PlayIndicatorSize);
+	}
+
+	FMargin GetPlayIndicatorPadding() const
+	{
+		return FMargin(PlayIndicatorPadding);
+	}
+
+	FMargin GetPlayButtonContentPadding() const
+	{
+		return FMargin(PlayButtonContentPadding);
+	}
+
 	TSharedRef<SDocumentationToolTip> GetDefaultToolTip() const
 	{
 		const FAssetData& AssetData = AssetThumbnail->GetAssetData();
@@ -1045,13 +1094,16 @@ private:
 
 	TSharedRef<SWidget> CreateStatusWidget(int32 StatusIndex, const FAssetDisplayInfo& InStatusInfo) const
 	{
-		TSharedRef<SLayeredImage> StatusLayeredImage = SNew(SLayeredImage).Image(InStatusInfo.StatusIcon);
+		TSharedRef<SLayeredImage> StatusLayeredImage = SNew(SLayeredImage)
+			.Image(InStatusInfo.StatusIcon)
+			.DesiredSizeOverride(this, &SAssetThumbnail::GetStatusSizeForImage);
 		StatusLayeredImage->AddLayer(InStatusInfo.StatusIconOverlay);
 
 		return SNew(SBox)
-			.WidthOverride(StatusSize)
-			.HeightOverride(StatusSize)
+			.WidthOverride(this, &SAssetThumbnail::GetStatusSize)
+			.HeightOverride(this, &SAssetThumbnail::GetStatusSize)
 			.Visibility(this, &SAssetThumbnail::GetStatusVisibility, StatusIndex)
+			.HAlign(HAlign_Center)
 			.VAlign(VAlign_Center)
 			[
 				StatusLayeredImage
@@ -1074,22 +1126,53 @@ private:
 	TSharedRef<SWidget> CreateStatusOverflowWidget() const
 	{
 		return SNew(SBox)
-			.WidthOverride(StatusSize)
-			.HeightOverride(StatusSize)
+			.WidthOverride(this, &SAssetThumbnail::GetStatusSize)
+			.HeightOverride(this, &SAssetThumbnail::GetStatusSize)
 			.HAlign(HAlign_Center)
 			.VAlign(VAlign_Center)
 			.Visibility(this, &SAssetThumbnail::GetStatusOverflowVisibility)
 			[
 				SNew(STextBlock)
-				.Font(FAppStyle::GetFontStyle(Style, ".StatusOverflowFont"))
+				.Font(this, &SAssetThumbnail::GetStatusOverflowFont)
 				.Text(this, &SAssetThumbnail::GetStatusOverflowText)
+				.ColorAndOpacity(FStyleColors::Foreground)
 			];
 	}
 
+	FSlateFontInfo GetStatusOverflowFont() const
+	{
+		FString FontStyle = TEXT(".StatusOverflowFont");
+		if (StatusSize != DefaultStatusSize)
+		{
+			FontStyle = TEXT(".StatusOverflowFontSmall");
+		}
+		return FAppStyle::GetFontStyle(Style, StringCast<ANSICHAR>(*FontStyle).Get());
+	}
+
+	FOptionalSize GetStatusSize() const
+	{
+		return FOptionalSize(StatusSize);
+	}
+
+	TOptional<FVector2D> GetStatusSizeForImage() const
+	{
+		return TOptional(FVector2D(StatusSize));
+	}
+
+	float GetAssetThumbnailBorderPadding() const
+	{
+		constexpr float SmallBorderPadding = 1.f;
+		constexpr float BorderPadding = 2.f;
+		if (SmallAssetChipBorder.IsSet())
+		{
+			return SmallAssetChipBorder.Get() ? SmallBorderPadding : BorderPadding;
+		}
+		return 0.f;
+	}
 	EVisibility GetStatusOverflowVisibility() const
 	{
 		const FGeometry ThumbnailGeometry = GetPaintSpaceGeometry();
-		const float ThumbnailWidth = ThumbnailGeometry.GetAbsoluteSize().X - (StatusPadding * 2);
+		const float ThumbnailWidth = ThumbnailGeometry.GetAbsoluteSize().X - (StatusPadding * 2) - (StatusBorderPadding * 2) - (GetAssetThumbnailBorderPadding() * 2);
 		const int32 MaxShownStatus = FMath::FloorToInt32(ThumbnailWidth / StatusSize);
 
 		int32 ShownStatus = 0;
@@ -1106,7 +1189,7 @@ private:
 	FText GetStatusOverflowText() const
 	{
 		FGeometry ThumbnailGeometry = GetPaintSpaceGeometry();
-		const float ThumbnailWidth = ThumbnailGeometry.GetAbsoluteSize().X - (StatusPadding * 2);
+		const float ThumbnailWidth = ThumbnailGeometry.GetAbsoluteSize().X - (StatusPadding * 2) - (StatusBorderPadding * 2) - (GetAssetThumbnailBorderPadding() * 2);
 		const int32 MaxShownStatus = FMath::FloorToInt32(ThumbnailWidth / StatusSize);
 
 		int32 ShownStatus = 0;
@@ -1133,7 +1216,7 @@ private:
 		}
 
 		const FGeometry ThumbnailGeometry = GetPaintSpaceGeometry();
-		const float ThumbnailWidth = ThumbnailGeometry.GetAbsoluteSize().X - (StatusPadding * 2);
+		const float ThumbnailWidth = ThumbnailGeometry.GetAbsoluteSize().X - (StatusPadding * 2) - (StatusBorderPadding * 2) - (GetAssetThumbnailBorderPadding() * 2);
 		const int32 StatusIndexConsideringHiddenOnes = (InStatusIndex - CollapsedStatusBeforeThis);
 		const int32 ShownStatus = FMath::FloorToInt32(ThumbnailWidth / StatusSize);
 		const int32 StatusIndexConsideringClippedStatusIfVisible = StatusOverflowWidget.IsValid() && StatusOverflowWidget->GetVisibility().IsVisible() ? StatusIndexConsideringHiddenOnes + 1 : StatusIndexConsideringHiddenOnes;
@@ -1575,8 +1658,17 @@ private:
 	TAttribute<int32> GenericThumbnailSize;
 	EThumbnailColorStripOrientation ColorStripOrientation;
 	TSharedPtr<SWidget> StatusOverflowWidget;
-	const float StatusSize = 16.f;
+	float StatusSize = 16.f;
+	const float DefaultStatusSize = 16.f;
 	const float StatusPadding = 4.f;
+	const float StatusBorderPadding = 2.f;
+	float PlayIndicatorPadding = 4.f;
+	const float PlayIndicatorDefaultPadding = 4.f;
+	float PlayButtonContentPadding = 8.f;
+	const float PlayButtonContentDefaultPadding = 8.f;
+	float PlayIndicatorSize = 20.f;
+	const float PlayIndicatorMaxSizeThreshold = 64.f;
+	const float PlayIndicatorDefaultSize = 20.f;
 	TArray<FAssetDisplayInfo> OverlayInfo;
 	TArray<TSharedPtr<SWidget>> Statuses;
 	TSharedPtr<IAssetSystemInfoProvider> AssetSystemInfoProvider;
