Fix possible crashes when importing offline TTF.
The code now have error handling for allocation error, and prevent out of bound access to a temporary buffer: in UTrueTypeFontFactory::CreateTextureFromBitmap and CreateTextureFromDC, BitmapData was potentially accessed based on the texture size that could be different from the bitmap one, as the height was rounded to a power of two.
Also fixed other issues found when investigating.
-Removed some logic redundancy.
-Fixed typo in drop shadow handling.
-Fixed infinite loop when trying to generate a distance field with a height inferior to 16, and fixed progress bar incorrect values (was going to 200%).
-Fixed crashes with some values when generating distance fields.
-Fixed incorrect handling of padding.
-Round resolution to multiple of 4 to ensure that textures will be compressed properly (otherwise, compressed format cannot be used and it will silently use an uncompress format, which could be a lot bigger in memory).

#tests Imported offline font with various settings (size, antialiasing,height...) => all generations suceeded and texture content was as expected.
[RN] Slate: fix various possible crashes and bugs when importing offline fonts.

diff --git a/Engine/Source/Editor/UnrealEd/Classes/Factories/TrueTypeFontFactory.h b/Engine/Source/Editor/UnrealEd/Classes/Factories/TrueTypeFontFactory.h
--- a/Engine/Source/Editor/UnrealEd/Classes/Factories/TrueTypeFontFactory.h
+++ b/Engine/Source/Editor/UnrealEd/Classes/Factories/TrueTypeFontFactory.h
@@ -67,12 +67,12 @@ class UTrueTypeFontFactory : public UTextureFactory, public FReimportHandler
 	 *
 	 * @return Returns the newly created texture, if successful, otherwise NULL
 	 */
-	UNREALED_API UTexture2D* CreateTextureFromDC( UFont* Font, Windows::HDC dc, int32 RowHeight, int32 TextureNum );
+	UNREALED_API UTexture2D* CreateTextureFromDC( UFont* Font, Windows::HDC dc, int32 BitmapWidth, int32 BitmapHeight, int32 TextureNum );
 #endif
 
 #if PLATFORM_WINDOWS || PLATFORM_MAC || PLATFORM_LINUX
 	UNREALED_API void* LoadFontFace( void* FTLibrary, int32 Height, FFeedbackContext* Warn, void** OutFontData );
-	UNREALED_API UTexture2D* CreateTextureFromBitmap( UFont* Font, uint8* BitmapData, int32 Height, int32 TextureNum );
+	UNREALED_API UTexture2D* CreateTextureFromBitmap( UFont* Font, uint8* BitmapData, int32 BitmapWidth, int32 BitmapHeight, int32 TextureNum );
 	UNREALED_API bool CreateFontTexture( UFont* Font, FFeedbackContext* Warn, const int32 NumResolutions, const int32 CharsPerPage,
 		const TMap<TCHAR,TCHAR>& InverseMap, const TArray< float >& ResHeights );
 
diff --git a/Engine/Source/Editor/UnrealEd/Private/Factories/TTFontImport.cpp b/Engine/Source/Editor/UnrealEd/Private/Factories/TTFontImport.cpp
--- a/Engine/Source/Editor/UnrealEd/Private/Factories/TTFontImport.cpp
+++ b/Engine/Source/Editor/UnrealEd/Private/Factories/TTFontImport.cpp
@@ -48,6 +48,25 @@ DEFINE_LOG_CATEGORY_STATIC(LogTTFontImport, Log, All);
 
 #define LOCTEXT_NAMESPACE "TTFontImport"
 
+namespace
+{
+	int32 FromHex(TCHAR Ch)
+	{
+		if (Ch >= '0' && Ch <= '9')
+			return Ch - '0';
+		else if (Ch >= 'a' && Ch <= 'f')
+			return 10 + Ch - 'a';
+		else if (Ch >= 'A' && Ch <= 'F')
+			return 10 + Ch - 'A';
+		UE_LOG(LogTTFontImport, Fatal, TEXT("Expecting digit, got character %c"), Ch);
+		return 0;
+	}
+
+	int32 RoundUpToMultipleOf4(int32 X)
+	{
+		return (X + 3) & ~3;
+	}
+}
 
 UTrueTypeFontFactory::UTrueTypeFontFactory(const FObjectInitializer& ObjectInitializer)
 	: Super(ObjectInitializer)
@@ -69,18 +88,6 @@ void UTrueTypeFontFactory::PostInitProperties()
 	}
 }
 
-int32 FromHex( TCHAR Ch )
-{
-	if( Ch>='0' && Ch<='9' )
-		return Ch-'0';
-	else if( Ch>='a' && Ch<='f' )
-		return 10+Ch-'a';
-	else if( Ch>='A' && Ch<='F' )
-		return 10+Ch-'A';
-	UE_LOG(LogTTFontImport, Fatal, TEXT("Expecting digit, got character %c"), Ch);
-	return 0;
-}
-
 #if !PLATFORM_WINDOWS // !!! FIXME
 #define FW_NORMAL 400
 #endif
@@ -181,6 +188,9 @@ UObject* UTrueTypeFontFactory::FactoryCreateNew(
 		ImportOptions->Data.YPadding = FMath::Max<int32>(4,ImportOptions->Data.YPadding);
 	}
 
+	ImportOptions->Data.TexturePageWidth = RoundUpToMultipleOf4(ImportOptions->Data.TexturePageWidth);
+	ImportOptions->Data.TexturePageMaxHeight = RoundUpToMultipleOf4(ImportOptions->Data.TexturePageMaxHeight);
+
 	// Copy the import settings into the font for later reference
 	Font->ImportOptions = ImportOptions->Data;
 	
@@ -344,7 +354,6 @@ class FTextureAlphaToDistanceField
 		 * @param InNumRowsToProcess    How many rows this task should process
 		 * @param InScanRadius          Radius of area to scan for nearest border
 		 * @param InScaleFactor         The input image and resulting distance field
-		 * @param WorkRemainingCounter  Counter used for providing feedback (i.e. updating the progress bar)
 		 */
 		FBuildDistanceFieldTask(
 			FThreadSafeCounter* InThreadScaleCounter,
@@ -354,8 +363,7 @@ class FTextureAlphaToDistanceField
 			int32 InDstRowWidth,
 			int32 InNumRowsToProcess,
 			int32 InScanRadius,
-			int32 InScaleFactor,
-			FThreadSafeCounter* InWorkRemainingCounter)
+			int32 InScaleFactor)
 			: ThreadScaleCounter(InThreadScaleCounter)
 			, SignedDistances(*OutSignedDistances)
 			, TaskData(InTaskData)
@@ -364,7 +372,6 @@ class FTextureAlphaToDistanceField
 			, NumRowsToProcess(InNumRowsToProcess)
 			, ScanRadius(InScanRadius)
 			, ScaleFactor(InScaleFactor)
-			, WorkRemainingCounter(InWorkRemainingCounter)
 		{
 		}
 
@@ -400,7 +407,6 @@ class FTextureAlphaToDistanceField
 		const int32 NumRowsToProcess;
 		const int32 ScanRadius;
 		const int32 ScaleFactor;
-		FThreadSafeCounter* WorkRemainingCounter;
 	};
 	
 public:
@@ -506,20 +512,21 @@ void FTextureAlphaToDistanceField::Generate(int32 ScaleFactor, int32 ScanRadius)
 	SignedDistance.Empty(DstSizeX*DstSizeY);
 	SignedDistance.AddUninitialized(DstSizeX*DstSizeY);
 	
-	// We want to run the distance field computation as 16 tasks for a speed boost on multi-core machines.
-	const int32 NumTasks = 16;
+	// We want to run the distance field computation as a maximum of 16 tasks for a speed boost on multi-core machines.
+	const int32 NumTasks = FMath::Min(DstSizeY, 16); //ensure we don't split the work so much that we end up with a null DstStripHeight... and an infinite loop.
 	FThreadSafeCounter BuildTasksCounter;
-	const int32 DstStripHeight = DstSizeY / NumTasks;
+	int32 DstStripHeight = DstSizeY / NumTasks;
 	
 	// We need to report the progress, and all the threads must be able to update it.
 	const int32 TotalDistFieldWork = DstStripHeight * NumTasks;
-	FThreadSafeCounter WorkProgressCounter;
-	
+
 
 	int32 TotalProgress = DstSizeY-1;	
 	// calculate distances
 	for (int32 y=0; y < DstSizeY; y+=DstStripHeight)
 	{
+		DstStripHeight = FMath::Min(DstStripHeight, DstSizeY - y); //Don't go out of bound.
+
 		// The tasks will clean themselves up when they are completed; no need to call delete elsewhere.
 		BuildTasksCounter.Increment();
 		(new FAutoDeleteAsyncTask<FBuildDistanceFieldTask>(
@@ -530,14 +537,13 @@ void FTextureAlphaToDistanceField::Generate(int32 ScaleFactor, int32 ScanRadius)
 			DstSizeX,
 			DstStripHeight,
 			ScanRadius,
-			ScaleFactor,
-			&WorkProgressCounter ))->StartBackgroundTask();
+			ScaleFactor))->StartBackgroundTask();
 	}
 
 	while( BuildTasksCounter.GetValue() > 0 )
 	{
 		// Waiting for Distance Field to finish generating.
-		GWarn->UpdateProgress(WorkProgressCounter.GetValue(),TotalDistFieldWork);	
+		GWarn->UpdateProgress(TotalDistFieldWork - BuildTasksCounter.GetValue(),TotalDistFieldWork);
 		FPlatformProcess::Sleep(.1f);
 	}
 	
@@ -741,12 +747,6 @@ void FTextureAlphaToDistanceField::FBuildDistanceFieldTask::DoWork(void)
 	// Build the distance field for the strip specified for this task.
 	for (int32 y=StartRow; y < (StartRow+NumRowsToProcess); y++)
 	{
-		if ( y%16 == 0 )
-		{
-			// Update the user about our progress
-			WorkRemainingCounter->Add(16);
-		}
-
 		// Build distance field for a single line
 		for (int32 x=0; x < DstRowWidth; x++)
 		{
@@ -763,7 +763,7 @@ void FTextureAlphaToDistanceField::FBuildDistanceFieldTask::DoWork(void)
 
 #if PLATFORM_WINDOWS
 
-UTexture2D* UTrueTypeFontFactory::CreateTextureFromDC( UFont* Font, HDC dc, int32 Height, int32 TextureNum )
+UTexture2D* UTrueTypeFontFactory::CreateTextureFromDC( UFont* Font, HDC dc, const int32 BitmapWidth, const int32 BitmapHeight, int32 TextureNum )
 {
 	FString TextureString = FString::Printf( TEXT("%s_Page"), *Font->GetName() );
 	if( TextureNum < 26 )
@@ -780,15 +780,6 @@ UTexture2D* UTrueTypeFontFactory::CreateTextureFromDC( UFont* Font, HDC dc, int3
 		UE_LOG(LogTTFontImport, Warning, TEXT("A texture named %s already exists!"), *TextureString );
 	}
 	
-	int32 BitmapWidth = ImportOptions->Data.TexturePageWidth;
-	int32 BitmapHeight = FMath::RoundUpToPowerOfTwo(Height);
-	if( ImportOptions->Data.bUseDistanceFieldAlpha )
-	{
-		// scale original bitmap width by scale factor to generate high res font
-		// note that height is already scaled during font bitmap generation
-		BitmapWidth *= ImportOptions->Data.DistanceFieldScaleFactor;
-	}
-
 	// Create texture for page.
 	auto Texture = NewObject<UTexture2D>(Font, *TextureString);
 
@@ -872,7 +863,7 @@ UTexture2D* UTrueTypeFontFactory::CreateTextureFromDC( UFont* Font, HDC dc, int3
 				int32 CharAlpha = SourceData[i+j*BitmapWidth];
 				int32 DropShadowAlpha;
 
-				if( ImportOptions->Data.bEnableDropShadow && i > 0 && j >> 0 )
+				if( ImportOptions->Data.bEnableDropShadow && i > 0 && j > 0 )
 				{
 					DropShadowAlpha = ( (i-1)+(j-1)*BitmapWidth);
 				}
@@ -1087,19 +1078,11 @@ bool UTrueTypeFontFactory::CreateFontTexture(
 		HBITMAP BitmapHandle;
 		void* BitmapDataPtr = NULL;
 		
-		int32 BitmapWidth = ImportOptions->Data.TexturePageWidth;
-		int32 BitmapHeight = ImportOptions->Data.TexturePageMaxHeight;
-		int32 BitmapPaddingX = ImportOptions->Data.XPadding;
-		int32 BitmapPaddingY = ImportOptions->Data.YPadding;
-
-		// scale up bitmap dimensions by for distance field generation
-		if( ImportOptions->Data.bUseDistanceFieldAlpha )
-		{
-			BitmapWidth *= ImportOptions->Data.DistanceFieldScaleFactor;
-			BitmapHeight *= ImportOptions->Data.DistanceFieldScaleFactor;
-			BitmapPaddingX *= ImportOptions->Data.DistanceFieldScaleFactor;
-			BitmapPaddingY *= ImportOptions->Data.DistanceFieldScaleFactor;
-		}
+		const int32 ActualDistanceFieldScaleFactor = ImportOptions->Data.bUseDistanceFieldAlpha ? ImportOptions->Data.DistanceFieldScaleFactor : 1;
+		const int32 BitmapWidth = ImportOptions->Data.TexturePageWidth * ActualDistanceFieldScaleFactor;
+		const int32 BitmapHeight = ImportOptions->Data.TexturePageMaxHeight * ActualDistanceFieldScaleFactor;
+		const int32 BitmapPaddingX = ImportOptions->Data.XPadding * ActualDistanceFieldScaleFactor;
+		const int32 BitmapPaddingY = ImportOptions->Data.YPadding * ActualDistanceFieldScaleFactor;
 
 		if( ImportOptions->Data.bEnableAntialiasing )
 		{
@@ -1269,17 +1252,29 @@ bool UTrueTypeFontFactory::CreateFontTexture(
 
 			// If the character is bigger than our texture size, then this isn't going to work!  The user
 			// will need to specify a larger texture resolution
-			if( CharWidth > BitmapWidth ||
-				CharHeight > BitmapHeight )
+			const int32 RequiredWidth = CharWidth + BitmapPaddingX * 2;
+			const int32 RequiredHeight = CharHeight + BitmapPaddingY * 2;
+			const bool bisBitmapTooNarrow = RequiredWidth > BitmapWidth;
+			const bool bisBitmapTooShort = RequiredHeight > BitmapHeight;
+			if (bisBitmapTooNarrow || bisBitmapTooShort)
 			{
-				UE_LOG(LogTTFontImport, Warning, TEXT( "At the specified font size, at least one font glyph would be larger than the maximum texture size you specified.") );
+				UE_LOG(LogTTFontImport, Error, TEXT("At the specified font size, at least one font glyph would be larger than the maximum texture size you specified."));
+				if (bisBitmapTooNarrow)
+				{
+					UE_LOG(LogTTFontImport, Error, TEXT("You should increase TexturePageWidth to at least %d."), RequiredWidth / ActualDistanceFieldScaleFactor);
+				}
+				else
+				{
+					UE_LOG(LogTTFontImport, Error, TEXT("You should increase TexturePageMaxHeight to at least %d."), RequiredHeight / ActualDistanceFieldScaleFactor);
+				}
+
 				DeleteDC( DCHandle );
 				DeleteObject( BitmapHandle );
 				return false;
 			}
 
 			// If it doesn't fit right here, advance to next line.
-			if( CharWidth + X + 2 > BitmapWidth)
+			if( (CharWidth + X + BitmapPaddingX) > BitmapWidth)
 			{
 				X=BitmapPaddingX;
 				Y = Y + RowHeight + BitmapPaddingY;
@@ -1292,9 +1287,11 @@ bool UTrueTypeFontFactory::CreateFontTexture(
 			}
 
 			// new page
-			if( Y+RowHeight > BitmapHeight )
+			if( (Y + RowHeight + BitmapPaddingY) > BitmapHeight )
 			{
-				Font->Textures.Add( CreateTextureFromDC( Font, DCHandle, Y+OldRowHeight, CurrentTexture ) );
+				const int32 RoundedHeight = RoundUpToMultipleOf4(Y + OldRowHeight);
+				check(RoundedHeight <= BitmapHeight);
+				Font->Textures.Add( CreateTextureFromDC( Font, DCHandle, BitmapWidth, RoundedHeight, CurrentTexture ) );
 				CurrentTexture++;
 
 				// blank out DC
@@ -1413,7 +1410,9 @@ bool UTrueTypeFontFactory::CreateFontTexture(
 			X = X + CharWidth + BitmapPaddingX;
 		}
 		// save final page
-		Font->Textures.Add( CreateTextureFromDC( Font, DCHandle, Y+RowHeight, CurrentTexture ) );
+		const int32 RoundedHeight = RoundUpToMultipleOf4(Y + RowHeight);
+		check(RoundedHeight <= BitmapHeight);
+		Font->Textures.Add( CreateTextureFromDC( Font, DCHandle, BitmapWidth, RoundedHeight, CurrentTexture ) );
 		CurrentTexture++;
 
 		DeleteDC( DCHandle );
@@ -1676,7 +1675,7 @@ void* UTrueTypeFontFactory::LoadFontFace(void* FTLibrary, int32 Height, FFeedbac
 #error "Unknown platform"
 #endif
 
-UTexture2D* UTrueTypeFontFactory::CreateTextureFromBitmap( UFont* Font, uint8* BitmapData, int32 Height, int32 TextureNum )
+UTexture2D* UTrueTypeFontFactory::CreateTextureFromBitmap( UFont* Font, uint8* BitmapData, const int32 BitmapWidth, const int32 BitmapHeight, int32 TextureNum )
 {
 	FString TextureString = FString::Printf( TEXT("%s_Page"), *Font->GetName() );
 	if( TextureNum < 26 )
@@ -1693,15 +1692,6 @@ UTexture2D* UTrueTypeFontFactory::CreateTextureFromBitmap( UFont* Font, uint8* B
 		UE_LOG(LogTTFontImport, Warning, TEXT("A texture named %s already exists!"), *TextureString );
 	}
 	
-	int32 BitmapWidth = ImportOptions->Data.TexturePageWidth;
-	int32 BitmapHeight = FMath::RoundUpToPowerOfTwo(Height);
-	if( ImportOptions->Data.bUseDistanceFieldAlpha )
-	{
-		// scale original bitmap width by scale factor to generate high res font
-		// note that height is already scaled during font bitmap generation
-		BitmapWidth *= ImportOptions->Data.DistanceFieldScaleFactor;
-	}
-
 	// Create texture for page.
 	UTexture2D* Texture = NewObject<UTexture2D>(Font, *TextureString);
 
@@ -1743,7 +1733,7 @@ UTexture2D* UTrueTypeFontFactory::CreateTextureFromBitmap( UFont* Font, uint8* B
 				int32 CharAlpha = SourceData[i+j*BitmapWidth];
 				int32 DropShadowAlpha;
 
-				if( ImportOptions->Data.bEnableDropShadow && i > 0 && j >> 0 )
+				if( ImportOptions->Data.bEnableDropShadow && i > 0 && j > 0 )
 				{
 					DropShadowAlpha = ( (i-1)+(j-1)*BitmapWidth);
 				}
@@ -1876,6 +1866,22 @@ bool UTrueTypeFontFactory::CreateFontTexture(
 	// Zero out the Texture Index
 	int32 CurrentTexture = 0;
 
+	// Create bitmap
+	const int32 ActualDistanceFieldScaleFactor = ImportOptions->Data.bUseDistanceFieldAlpha ? ImportOptions->Data.DistanceFieldScaleFactor : 1;
+	const int32 BitmapWidth = ImportOptions->Data.TexturePageWidth * ActualDistanceFieldScaleFactor;
+	const int32 BitmapHeight = ImportOptions->Data.TexturePageMaxHeight * ActualDistanceFieldScaleFactor;
+	const int32 BitmapPaddingX = ImportOptions->Data.XPadding * ActualDistanceFieldScaleFactor;
+	const int32 BitmapPaddingY = ImportOptions->Data.YPadding * ActualDistanceFieldScaleFactor;
+	const int32 BitmapBytesPerPixel = 4;
+	const int32 BitmapDataSize = BitmapWidth * BitmapHeight * BitmapBytesPerPixel;
+	uint8* BitmapDataPtr = (uint8*)FMemory::Malloc(BitmapDataSize);
+	if (BitmapDataPtr == nullptr)
+	{
+		UE_LOG(LogTTFontImport, Error, TEXT("Unable to allocate bitmap buffer for font texture. Please verify the texture page size and distance field scale factor settings."));
+		FT_Done_FreeType(FTLibrary);
+		return false;
+	}
+
 	for( int32 Page = 0; Page < NumResolutions; ++Page )
 	{
 		int32 Height = FMath::RoundToInt( ResHeights[Page] );
@@ -1904,24 +1910,6 @@ bool UTrueTypeFontFactory::CreateFontTexture(
 			return false;
 		}
 
-		// Create bitmap
-		int32 BitmapWidth = ImportOptions->Data.TexturePageWidth;
-		int32 BitmapHeight = ImportOptions->Data.TexturePageMaxHeight;
-		int32 BitmapPaddingX = ImportOptions->Data.XPadding;
-		int32 BitmapPaddingY = ImportOptions->Data.YPadding;
-
-		// scale up bitmap dimensions by for distance field generation
-		if( ImportOptions->Data.bUseDistanceFieldAlpha )
-		{
-			BitmapWidth *= ImportOptions->Data.DistanceFieldScaleFactor;
-			BitmapHeight *= ImportOptions->Data.DistanceFieldScaleFactor;
-			BitmapPaddingX *= ImportOptions->Data.DistanceFieldScaleFactor;
-			BitmapPaddingY *= ImportOptions->Data.DistanceFieldScaleFactor;
-		}
-
-		const int32 BitmapBytesPerPixel = 4;
-		const int32 BitmapDataSize = BitmapWidth * BitmapHeight * BitmapBytesPerPixel;
-		uint8* BitmapDataPtr = (uint8*)FMemory::Malloc( BitmapDataSize );
 		FMemory::Memzero( BitmapDataPtr, BitmapDataSize );
 
 		float EmScale = 1024.f/Height;
@@ -2003,17 +1991,29 @@ bool UTrueTypeFontFactory::CreateFontTexture(
 
 			// If the character is bigger than our texture size, then this isn't going to work!  The user
 			// will need to specify a larger texture resolution
-			if( CharWidth > BitmapWidth ||
-				CharHeight > BitmapHeight )
+			const int32 RequiredWidth = CharWidth + BitmapPaddingX * 2;
+			const int32 RequiredHeight = CharHeight + BitmapPaddingY * 2;
+			const bool bisBitmapTooNarrow = RequiredWidth > BitmapWidth;
+			const bool bisBitmapTooShort = RequiredHeight > BitmapHeight;
+			if (bisBitmapTooNarrow || bisBitmapTooShort )
 			{
-				UE_LOG(LogTTFontImport, Warning, TEXT( "At the specified font size, at least one font glyph would be larger than the maximum texture size you specified.") );
+				UE_LOG(LogTTFontImport, Error, TEXT( "At the specified font size, at least one font glyph would be larger than the maximum texture size you specified.") );
+				if (bisBitmapTooNarrow)
+				{
+					UE_LOG(LogTTFontImport, Error, TEXT("You should increase TexturePageWidth to at least %d."), RequiredWidth / ActualDistanceFieldScaleFactor);
+				}
+				else
+				{
+					UE_LOG(LogTTFontImport, Error, TEXT("You should increase TexturePageMaxHeight to at least %d."), RequiredHeight / ActualDistanceFieldScaleFactor);
+				}
+				
 				FMemory::Free( BitmapDataPtr );
 				FT_Done_FreeType( FTLibrary );
 				return false;
 			}
 
 			// If it doesn't fit right here, advance to next line.
-			if( CharWidth + X + 2 > BitmapWidth)
+			if( (CharWidth + X + BitmapPaddingX) > BitmapWidth)
 			{
 				X=BitmapPaddingX;
 				Y = Y + RowHeight + BitmapPaddingY;
@@ -2026,9 +2026,11 @@ bool UTrueTypeFontFactory::CreateFontTexture(
 			}
 
 			// new page
-			if( Y+RowHeight > BitmapHeight )
+			if( (Y + RowHeight + BitmapPaddingY) > BitmapHeight )
 			{
-				Font->Textures.Add( CreateTextureFromBitmap( Font, BitmapDataPtr, Y+OldRowHeight, CurrentTexture ) );
+				const int32 RoundedHeight = RoundUpToMultipleOf4(Y + OldRowHeight);
+				check(RoundedHeight <= BitmapHeight);
+				Font->Textures.Add( CreateTextureFromBitmap( Font, BitmapDataPtr, BitmapWidth, RoundedHeight, CurrentTexture ) );
 				CurrentTexture++;
 
 				FMemory::Memzero( BitmapDataPtr, BitmapDataSize );
@@ -2091,11 +2093,10 @@ bool UTrueTypeFontFactory::CreateFontTexture(
 			X = X + CharWidth + BitmapPaddingX;
 		}
 		// save final page
-		Font->Textures.Add( CreateTextureFromBitmap( Font, BitmapDataPtr, Y+RowHeight, CurrentTexture ) );
+		const int32 RoundedHeight = RoundUpToMultipleOf4(Y + RowHeight);
+		Font->Textures.Add( CreateTextureFromBitmap( Font, BitmapDataPtr, BitmapWidth, RoundedHeight, CurrentTexture ) );
 		CurrentTexture++;
 
-		FMemory::Free( BitmapDataPtr );
-
 		if( FontData )
 		{
 			FMemory::Free( FontData );
@@ -2109,6 +2110,7 @@ bool UTrueTypeFontFactory::CreateFontTexture(
 
 	GWarn->UpdateProgress( TotalProgress, TotalProgress );
 
+	FMemory::Free(BitmapDataPtr);
 	FT_Done_FreeType( FTLibrary );
 
 	return true;
