Temporary fix to provide appropriate spacing around the Misc widgets so that the overflow menu doesn't look too close.

#rnx

diff --git a/Engine/Source/Editor/UnrealEd/Private/Toolkits/AssetEditorToolkit.cpp b/Engine/Source/Editor/UnrealEd/Private/Toolkits/AssetEditorToolkit.cpp
--- a/Engine/Source/Editor/UnrealEd/Private/Toolkits/AssetEditorToolkit.cpp
+++ b/Engine/Source/Editor/UnrealEd/Private/Toolkits/AssetEditorToolkit.cpp
@@ -1502,7 +1502,7 @@ void FAssetEditorToolkit::GenerateToolbar()
 			SNew(SBorder)
 			.VAlign(VAlign_Center)
 			.BorderImage(FAppStyle::Get().GetBrush("AssetEditorToolbar.Background"))
-			.Padding(FMargin(0.0f))
+			.Padding(MiscWidgets == SNullWidget::NullWidget ? FMargin(0.0f) : FMargin(4.f, 0.f, 0.f, 0.f))
 			[
 				MiscWidgets
 			]
