[ContentBrowser] Integrate old GenerateExtraStateTooltipWidget with the new Status system

diff --git a/Engine/Source/Editor/ContentBrowser/Private/AssetViewWidgets.cpp b/Engine/Source/Editor/ContentBrowser/Private/AssetViewWidgets.cpp
--- a/Engine/Source/Editor/ContentBrowser/Private/AssetViewWidgets.cpp
+++ b/Engine/Source/Editor/ContentBrowser/Private/AssetViewWidgets.cpp
@@ -785,12 +785,22 @@ TSharedRef<SWidget> SAssetViewItem::GenerateExtraStateTooltipWidget() const
 						];
 				}
 			}
+
+			if (Content->NumSlots() == 0)
+			{
+				return SNullWidget::NullWidget;
+			}
 			return Content.ToSharedRef();
 		}
 	}
 	return SNullWidget::NullWidget;
 }
 
+TSharedPtr<SWidget> SAssetViewItem::GenerateExtraStateTooltipWidgetNewStyle() const
+{
+	return GenerateExtraStateTooltipWidget().ToSharedPtr();
+}
+
 EVisibility SAssetViewItem::GetThumbnailEditModeUIVisibility() const
 {
 	return !IsFolder() && ThumbnailEditMode.Get() ? EVisibility::Visible : EVisibility::Collapsed;
@@ -1876,6 +1886,7 @@ void SAssetListItem::Construct( const FArguments& InArgs )
 
 		if (UE::Editor::ContentBrowser::IsNewStyleEnabled())
 		{
+			ThumbnailConfig.AdditionalTooltipInSmallView = MakeAttributeSP(this, &SAssetListItem::GenerateExtraStateTooltipWidgetNewStyle);
 			ThumbnailConfig.IsEditModeVisible = MakeAttributeSP(this, &SAssetListItem::GetThumbnailEditModeUIVisibility);
 			ThumbnailConfig.AllowAssetSpecificThumbnailOverlay = MakeAttributeSP(this, &SAssetListItem::ShouldAllowAssetActionOverlay, ThumbnailConfig.bForceGenericThumbnail);
 			ThumbnailConfig.AllowAssetSpecificThumbnailOverlayIndicator = MakeAttributeSP(this, &SAssetListItem::ShouldAllowAssetActionOverlayIndicator, ThumbnailConfig.bForceGenericThumbnail);
@@ -2433,6 +2444,7 @@ void SAssetTileItem::Construct( const FArguments& InArgs )
 
 		if (UE::Editor::ContentBrowser::IsNewStyleEnabled())
 		{
+			ThumbnailConfig.AdditionalTooltipInSmallView = MakeAttributeSP(this, &SAssetTileItem::GenerateExtraStateTooltipWidgetNewStyle);
 			ThumbnailConfig.AssetSystemInfoProvider = MakeShared<FAssetSystemContentBrowserInfoProvider>(AssetItem);
 			ThumbnailConfig.AllowAssetStatusThumbnailOverlay = true;
 			ThumbnailConfig.ShowAssetChip = true;
diff --git a/Engine/Source/Editor/ContentBrowser/Public/AssetViewWidgets.h b/Engine/Source/Editor/ContentBrowser/Public/AssetViewWidgets.h
--- a/Engine/Source/Editor/ContentBrowser/Public/AssetViewWidgets.h
+++ b/Engine/Source/Editor/ContentBrowser/Public/AssetViewWidgets.h
@@ -217,6 +217,9 @@ protected:
 	/** Generate a widget to inject extra external state indicator on the asset tooltip. */
 	TSharedRef<SWidget> GenerateExtraStateTooltipWidget() const;
 
+	/** Generate a widget to inject extra external state indicator on the asset tooltip. Need to be TSharedPtr for AssetThumbnail */
+	TSharedPtr<SWidget> GenerateExtraStateTooltipWidgetNewStyle() const;
+
 	/** Gets the visibility for the thumbnail edit mode UI */
 	EVisibility GetThumbnailEditModeUIVisibility() const;
 
diff --git a/Engine/Source/Editor/UnrealEd/Private/AssetThumbnail.cpp b/Engine/Source/Editor/UnrealEd/Private/AssetThumbnail.cpp
--- a/Engine/Source/Editor/UnrealEd/Private/AssetThumbnail.cpp
+++ b/Engine/Source/Editor/UnrealEd/Private/AssetThumbnail.cpp
@@ -94,6 +94,7 @@ public:
 		, _Padding(0)
 		, _GenericThumbnailSize(64)
 		, _AllowAssetStatusThumbnailOverlay(false)
+		, _AdditionalTooltipInSmallView(SNullWidget::NullWidget)
 		, _ShowAssetChip(false)
 		, _UseRoundAssetChip(true)
 		, _SmallAssetChipBorder(false)
@@ -120,7 +121,7 @@ public:
 		SLATE_ATTRIBUTE(int32, GenericThumbnailSize)
 		SLATE_ARGUMENT(TSharedPtr<IAssetSystemInfoProvider>, AssetSystemInfoProvider)
 		SLATE_ATTRIBUTE(bool, AllowAssetStatusThumbnailOverlay)
-		SLATE_ATTRIBUTE(TSharedPtr<IToolTip>, TooltipOverride)
+		SLATE_ATTRIBUTE(TSharedPtr<SWidget>, AdditionalTooltipInSmallView)
 		SLATE_ATTRIBUTE(bool, ShowAssetChip)
 		SLATE_ATTRIBUTE(bool, UseRoundAssetChip)
 		SLATE_ATTRIBUTE(bool, SmallAssetChipBorder)
@@ -159,6 +160,7 @@ public:
 		SmallAssetChipBorder = InArgs._SmallAssetChipBorder;
 		EditModeVisibility = InArgs._IsEditModeVisible;
 		AlwaysExpandTooltip = InArgs._AlwaysExpandTooltip;
+		AdditionalTooltipInSmallView = InArgs._AdditionalTooltipInSmallView;
 		AssetThumbnail->OnAssetDataChanged().AddSP(this, &SAssetThumbnail::OnAssetDataChanged);
 		const FAssetData& AssetData = AssetThumbnail->GetAssetData();
 
@@ -185,17 +187,10 @@ public:
 
 		if (UE::Editor::ContentBrowser::IsNewStyleEnabled())
 		{
-			if (InArgs._TooltipOverride.IsSet())
-			{
-				SetToolTip(InArgs._TooltipOverride);
-			}
-			else
-			{
-				// Set our tooltip - this will refresh each time it's opened to make sure it's up-to-date
-				SetToolTip(SNew(SAssetThumbnailToolTip)
-					.AssetThumbnail(SharedThis(this))
-					.AlwaysExpandTooltip(AlwaysExpandTooltip));
-			}
+			// Set our tooltip - this will refresh each time it's opened to make sure it's up-to-date
+			SetToolTip(SNew(SAssetThumbnailToolTip)
+				.AssetThumbnail(SharedThis(this))
+				.AlwaysExpandTooltip(AlwaysExpandTooltip));
 		}
 
 		UpdateThumbnailClass(AssetTypeActions.Get());
@@ -829,6 +824,18 @@ public:
 				]
 			];
 
+			if (AdditionalTooltipInSmallView.IsSet())
+			{
+				OverallTooltipVBox->AddSlot()
+				[
+					SNew(SBox)
+					.Padding(this, &SAssetThumbnail::GetAdditionalSmallViewTooltipMargin)
+					[
+						AdditionalTooltipInSmallView.Get().ToSharedRef()
+					]
+				];
+			}
+
 			TSharedRef<SVerticalBox> StatusVerticalBox = SNew(SVerticalBox);
 
 			for (const FAssetDisplayInfo& AssetStatusInfo : OverlayInfo)
@@ -949,6 +956,18 @@ public:
 	}
 
 private:
+
+	FMargin GetAdditionalSmallViewTooltipMargin() const
+	{
+		if (AdditionalTooltipInSmallView.IsSet())
+		{
+			const TSharedPtr<SWidget> AdditionalTooltip = AdditionalTooltipInSmallView.Get();
+			const bool bIsValidAndNotNullWidget = AdditionalTooltip.IsValid() && AdditionalTooltip != SNullWidget::NullWidget;
+			return bIsValidAndNotNullWidget && AdditionalTooltip->GetVisibility().IsVisible() ? FMargin(0.f, 0.f, 0.f, 6.f) : FMargin(0.f);
+		}
+		return FMargin(0.f);
+	}
+
 	FSlateWidgetRun::FWidgetRunInfo OnCreateWidgetDecoratorWidget(const FTextRunInfo& InRunInfo, const ISlateStyle* InStyle) const
 	{
 		TSharedRef<SWidget> CtrlAltWidget = SNew(SBox)
@@ -1678,6 +1697,7 @@ private:
 	TAttribute<const FSlateBrush*> AssetChipBorderImageOverride;
 	TAttribute<EVisibility> EditModeVisibility;
 	TAttribute<bool> AlwaysExpandTooltip;
+	TAttribute<TSharedPtr<SWidget>> AdditionalTooltipInSmallView;
 	TAttribute<bool> AllowAssetSpecificThumbnailOverlay;
 	TAttribute<bool> AllowAssetSpecificThumbnailOverlayIndicator;
 	TAttribute<bool> AllowAssetStatusThumbnailOverlay;
@@ -1887,7 +1907,7 @@ TSharedRef<SWidget> FAssetThumbnail::MakeThumbnailWidget( const FAssetThumbnailC
 			.Padding(InConfig.Padding)
 			.GenericThumbnailSize(InConfig.GenericThumbnailSize)
 			.AssetSystemInfoProvider(InConfig.AssetSystemInfoProvider)
-			.TooltipOverride(InConfig.TooltipOverride)
+			.AdditionalTooltipInSmallView(InConfig.AdditionalTooltipInSmallView)
 			.AllowAssetStatusThumbnailOverlay(InConfig.AllowAssetStatusThumbnailOverlay)
 			.ShowAssetChip(InConfig.ShowAssetChip)
 			.UseRoundAssetChip(InConfig.UseRoundAssetChip)
diff --git a/Engine/Source/Editor/UnrealEd/Public/AssetThumbnail.h b/Engine/Source/Editor/UnrealEd/Public/AssetThumbnail.h
--- a/Engine/Source/Editor/UnrealEd/Public/AssetThumbnail.h
+++ b/Engine/Source/Editor/UnrealEd/Public/AssetThumbnail.h
@@ -8,6 +8,7 @@
 #include "AssetRegistry/AssetData.h"
 #include "Rendering/RenderingCommon.h"
 #include "Widgets/SWidget.h"
+#include "Widgets/SNullWidget.h"
 #include "TickableEditorObject.h"
 
 class AActor;
@@ -71,6 +72,7 @@ struct FAssetThumbnailConfig
 		, AssetTypeColorOverride()
 		, Padding(0)
 		, AssetSystemInfoProvider(nullptr)
+		, AdditionalTooltipInSmallView(SNullWidget::NullWidget)
 		, AllowAssetStatusThumbnailOverlay( false )
 		, ShowAssetChip(true)
 		, UseRoundAssetChip(true)
@@ -108,7 +110,8 @@ struct FAssetThumbnailConfig
 	TAttribute<int32> GenericThumbnailSize = 64;
 	EThumbnailColorStripOrientation ColorStripOrientation = EThumbnailColorStripOrientation::HorizontalBottomEdge;
 	TSharedPtr<IAssetSystemInfoProvider> AssetSystemInfoProvider;
-	TAttribute<TSharedPtr<IToolTip>> TooltipOverride;
+	/** Additional tooltip placed between AssetType and Status */
+	TAttribute<TSharedPtr<SWidget>> AdditionalTooltipInSmallView;
 	/** Whether to show the asset status or not */
 	TAttribute< bool > AllowAssetStatusThumbnailOverlay;
 	/** Whether to show the asset chip or not */
