Engine / SourceControl: Clean up 'files deleted' delegate remnants.

Changelist #22455420 removed the OnSourceControlFilesDeleted delegate but left the handle and the unregister call.

#rnx

diff --git a/Engine/Source/Editor/UnrealEd/Classes/Editor/UnrealEdEngine.h b/Engine/Source/Editor/UnrealEd/Classes/Editor/UnrealEdEngine.h
--- a/Engine/Source/Editor/UnrealEd/Classes/Editor/UnrealEdEngine.h
+++ b/Engine/Source/Editor/UnrealEd/Classes/Editor/UnrealEdEngine.h
@@ -961,9 +961,6 @@ private:
 	/** Internal function to filter and add visualizers to a specific list */
 	UNREALED_API void AddVisualizers(AActor* Actor, TArray<FCachedComponentVisualizer>& Visualizers, TFunctionRef<bool(const TSharedPtr<FComponentVisualizer>&)> Condition);
 
-	/** Delegate Called after files have been deleted to perform necessary cleanups. */
-	FDelegateHandle SourceControlFilesDeletedHandle;
-
 	/** Keep track of need to stop stall detector thread */
 	bool bRequiresStallDetectorShutdown = false;
 };
diff --git a/Engine/Source/Editor/UnrealEd/Private/UnrealEdEngine.cpp b/Engine/Source/Editor/UnrealEd/Private/UnrealEdEngine.cpp
--- a/Engine/Source/Editor/UnrealEd/Private/UnrealEdEngine.cpp
+++ b/Engine/Source/Editor/UnrealEd/Private/UnrealEdEngine.cpp
@@ -511,11 +511,6 @@ void UUnrealEdEngine::FinishDestroy()
 		delete PerformanceMonitor;
 	}
 
-	if (ISourceControlModule* SourceControlModule = FModuleManager::GetModulePtr<ISourceControlModule>(FName("SourceControl")))
-	{
-		SourceControlModule->UnregisterFilesDeleted(SourceControlFilesDeletedHandle);
-	}
-
 	FPackageName::OnContentPathMounted().RemoveAll(this);
 	FPackageName::OnContentPathDismounted().RemoveAll(this);
 	UPackage::PackageDirtyStateChangedEvent.RemoveAll(this);
