Avoid creating new localization keys when reinitializing an asset defined Structure

[RN] User defined structs with Text fields with non trivial default values no longer alter themselves on load

diff --git a/Engine/Source/Editor/UnrealEd/Private/UserDefinedStructEditorData.cpp b/Engine/Source/Editor/UnrealEd/Private/UserDefinedStructEditorData.cpp
--- a/Engine/Source/Editor/UnrealEd/Private/UserDefinedStructEditorData.cpp
+++ b/Engine/Source/Editor/UnrealEd/Private/UserDefinedStructEditorData.cpp
@@ -249,7 +249,7 @@ void UUserDefinedStructEditorData::ReinitializeDefaultInstance(FString* OutLog)
 				FStructVariableDescription* VarDesc = VariablesDescriptions.FindByPredicate(FStructureEditorUtils::FFindByGuidHelper<FStructVariableDescription>(VarGuid));
 				if (VarDesc && !VarDesc->CurrentDefaultValue.IsEmpty())
 				{
-					if (!FBlueprintEditorUtils::PropertyValueFromString(Property, VarDesc->CurrentDefaultValue, StructData, ScriptStruct))
+					if (!FBlueprintEditorUtils::PropertyValueFromString(Property, VarDesc->CurrentDefaultValue, StructData, ScriptStruct, PPF_DuplicateVerbatim))
 					{
 						const FString Message = FString::Printf(TEXT("Cannot parse value. Property: %s String: \"%s\" ")
 							, *Property->GetDisplayNameText().ToString()
