Changed FGenerationHelper::FinishGeneratedPlatformSave to report the pacakge dependencies of the ExternalPackages it's depending upon so their dependencies don't get lost. The reason they'd get lost is that ExternalPackages gets included in the generated package and arent't cooked independently,.

diff --git a/Engine/Source/Editor/UnrealEd/Private/Cooker/CookGenerationHelper.cpp b/Engine/Source/Editor/UnrealEd/Private/Cooker/CookGenerationHelper.cpp
--- a/Engine/Source/Editor/UnrealEd/Private/Cooker/CookGenerationHelper.cpp
+++ b/Engine/Source/Editor/UnrealEd/Private/Cooker/CookGenerationHelper.cpp
@@ -1654,6 +1654,32 @@ void FGenerationHelper::FinishGeneratedPlatformSave(FPackageData& PackageData, c
 	// These dependencies are used in the AssetRegistry entry and runtime dependencies for the generated package
 	OutDependencies = Info->PackageDependencies;
 
+	// For ExternalPackage (that'll get stripped from the cook), append their dependencies gathered from 
+	// the AssetRegistry to PackageDependencies so that the GeneratedPackage dependencies are complete
+	TArray<FAssetDependency> PackageDependencies;
+	TSet<FAssetDependency> ExternalPackageDependencies;
+
+	// External packages are detected by their presence in the NeverCookPackageList 
+	FThreadSafeSet<FName>& NeverCookPackageList = GetOwner().GetPackageDatas().GetCookOnTheFlyServer().PackageTracker->NeverCookPackageList; 
+
+	for (const FAssetDependency& Dependency : Info->PackageDependencies)
+	{
+		// Verify it's an external package before adding its dependencies
+		if (NeverCookPackageList.Contains(Dependency.AssetId.PackageName))
+		{
+			PackageDependencies.Reset();
+			IAssetRegistry::GetChecked().GetDependencies(FAssetIdentifier(Dependency.AssetId.PackageName), PackageDependencies,
+				UE::AssetRegistry::EDependencyCategory::Package);
+
+			for(const FAssetDependency& PackageDependency : PackageDependencies)
+			{				
+				ExternalPackageDependencies.Add(FAssetDependency::PackageDependency(PackageDependency.AssetId.PackageName, PackageDependency.Properties));
+			}
+		}		
+	}
+
+	OutDependencies.Append(ExternalPackageDependencies.Array());
+
 	// Add on the custom cook BuildResultDependencies declared for the generated package by the CookPackageSplitter.
 	OutBuildResultDependencies = Info->BuildResultDependencies;
 
diff --git a/Engine/Source/Runtime/AssetRegistry/Public/AssetRegistry/IAssetRegistry.h b/Engine/Source/Runtime/AssetRegistry/Public/AssetRegistry/IAssetRegistry.h
--- a/Engine/Source/Runtime/AssetRegistry/Public/AssetRegistry/IAssetRegistry.h
+++ b/Engine/Source/Runtime/AssetRegistry/Public/AssetRegistry/IAssetRegistry.h
@@ -189,6 +189,16 @@ struct FAssetDependency
 			Properties };
 	}
 
+	friend inline uint32 GetTypeHash(const FAssetDependency& InHash)
+	{
+		uint32 Hash = GetTypeHash(InHash.AssetId);
+		Hash = HashCombine(Hash, GetTypeHash(InHash.Category));
+		Hash = HashCombine(Hash, GetTypeHash(InHash.Properties));
+
+		return Hash;
+	}
+
+
 	ASSETREGISTRY_API void WriteCompactBinary(FCbWriter& Writer) const;
 private:
 	friend FCbWriter& operator<<(FCbWriter& Writer, const FAssetDependency& Dependency)
