Fixed an issue where thumbnails for blueprinnts that have non-uniform scales in their root components are rendered deformed.

The issue was that the thumnail renderer was re-locating the preview actor spawned to render it using SetActorTransform. That is a problem because it will also ovewrite the scale set in the blueprint. Changin the implementation to use SetActorLocation instead fixes the issue.

#tests Tested in editor that regenerating thumbnails of existing blueprints they still look ok. And that thumbnails from affected blueprints that would previously deform, now render correctly

diff --git a/Engine/Source/Editor/UnrealEd/Private/ThumbnailHelpers.cpp b/Engine/Source/Editor/UnrealEd/Private/ThumbnailHelpers.cpp
--- a/Engine/Source/Editor/UnrealEd/Private/ThumbnailHelpers.cpp
+++ b/Engine/Source/Editor/UnrealEd/Private/ThumbnailHelpers.cpp
@@ -1122,9 +1122,9 @@ void FClassActorThumbnailScene::SpawnPreviewActor(UClass* InClass)
 		{
 			const FBoxSphereBounds Bounds = GetPreviewActorBounds();
 			const float BoundsZOffset = GetBoundsZOffset(Bounds);
-			const FTransform Transform(-Bounds.Origin + FVector(0, 0, BoundsZOffset));
+			const FVector TargetLocation(-Bounds.Origin + FVector(0, 0, BoundsZOffset));
 
-			PreviewActor->SetActorTransform(Transform);
+			PreviewActor->SetActorLocation(TargetLocation);
 		}
 	}
 }
