Recreate instancing graph when just reintantiating objects that were created from an archetype (sub-object of a CDO).

#rnx

diff --git a/Engine/Source/Editor/Kismet/Private/BlueprintCompilationManager.cpp b/Engine/Source/Editor/Kismet/Private/BlueprintCompilationManager.cpp
--- a/Engine/Source/Editor/Kismet/Private/BlueprintCompilationManager.cpp
+++ b/Engine/Source/Editor/Kismet/Private/BlueprintCompilationManager.cpp
@@ -3141,6 +3141,15 @@ void FBlueprintCompilationManagerImpl::ReinstanceBatch(TArray<FReinstancingJob>&
 					CacheManager,
 					ReinstancingJob.ArchetypeInstancesWithCachedArchetypes);
 
+				UObject* Template = nullptr;
+				if (UObject** NewArchetypeTemplate = OldArchetypeInfo.ArchetypeTemplate ? OldArchetypeToNewArchetype.Find(OldArchetypeInfo.ArchetypeTemplate) : nullptr)
+				{
+					Template = *NewArchetypeTemplate;
+				}
+
+				FObjectInstancingGraph InstancingGraph;
+				const bool bUseInstancingGraph = FBlueprintCompileReinstancer::PrePopulateInstancingGraphForArchetype(InstancingGraph, OldInstance, Template, &OldArchetypeToNewArchetype);
+
 				UObject* Destination = GetTransientOuterForRename(OldInstance->GetClass());
 				OldInstance->Rename(
 					nullptr,
@@ -3154,12 +3163,7 @@ void FBlueprintCompilationManagerImpl::ReinstanceBatch(TArray<FReinstancingJob>&
 				FMakeClassSpawnableOnScope TemporarilySpawnable(NewClass);
 				const EObjectFlags FlagMask = RF_Public | RF_ArchetypeObject | RF_Transactional | RF_Transient | RF_TextExportTransient | RF_InheritableComponentTemplate | RF_Standalone; //TODO: what about RF_RootSet?
 
-				UObject* Template = nullptr;
-				if (UObject** NewArchetypeTemplate = OldArchetypeInfo.ArchetypeTemplate ? OldArchetypeToNewArchetype.Find(OldArchetypeInfo.ArchetypeTemplate) : nullptr)
-				{
-					Template = *NewArchetypeTemplate;
-				}
-				UObject* NewArchetype = NewObject<UObject>(OriginalOuter, NewClass, OriginalName, OriginalFlags & FlagMask, Template);
+				UObject* NewArchetype = NewObject<UObject>(OriginalOuter, NewClass, OriginalName, OriginalFlags & FlagMask, Template, false, bUseInstancingGraph ? &InstancingGraph : nullptr);
 
 				OldArchetypeToNewArchetype.Add(OldInstance, NewArchetype);
 				if (OldToNewTemplatesForClass)
diff --git a/Engine/Source/Editor/UnrealEd/Private/Kismet2/KismetReinstanceUtilities.cpp b/Engine/Source/Editor/UnrealEd/Private/Kismet2/KismetReinstanceUtilities.cpp
--- a/Engine/Source/Editor/UnrealEd/Private/Kismet2/KismetReinstanceUtilities.cpp
+++ b/Engine/Source/Editor/UnrealEd/Private/Kismet2/KismetReinstanceUtilities.cpp
@@ -2166,7 +2166,7 @@ bool FBlueprintCompileReinstancer::IsReinstClass(const UClass* Class)
 	return Class && Class->GetFName().ToString().StartsWith(ReinstPrefix);
 }
 
-static void ReplaceObjectHelper(UObject*& OldObject, UClass* OldClass, UObject*& NewUObject, UClass* NewClass, TMap<UObject*, UObject*>& OldToNewInstanceMap, const TMap<UClass*, UClass*>& OldToNewClassMap, TMap<UObject*, FName>& OldToNewNameMap, int32 OldObjIndex, TArray<UObject*>& ObjectsToReplace, TArray<UObject*>& PotentialEditorsForRefreshing, TSet<AActor*>& OwnersToRerunConstructionScript, TFunctionRef<TArray<TObjectPtr<USceneComponent>>&(USceneComponent*)> GetAttachChildrenArray, bool bIsComponent, bool bArchetypesAreUpToDate)
+void FBlueprintCompileReinstancer::ReplaceObjectHelper(UObject*& OldObject, UClass* OldClass, UObject*& NewUObject, UClass* NewClass, TMap<UObject*, UObject*>& OldToNewInstanceMap, const TMap<UClass*, UClass*>& OldToNewClassMap, TMap<UObject*, FName>& OldToNewNameMap, int32 OldObjIndex, TArray<UObject*>& ObjectsToReplace, TArray<UObject*>& PotentialEditorsForRefreshing, TSet<AActor*>& OwnersToRerunConstructionScript, TFunctionRef<TArray<TObjectPtr<USceneComponent>>&(USceneComponent*)> GetAttachChildrenArray, bool bIsComponent, bool bArchetypesAreUpToDate)
 {
 	SCOPED_LOADTIMER_ASSET_TEXT(*WriteToString<256>(TEXT("ReplaceObjectHelper "), *GetPathNameSafe(OldObject)));
 	// If the old object was spawned from an archetype (i.e. not the CDO), we must use the new version of that archetype as the template object when constructing the new instance.
@@ -2280,8 +2280,11 @@ static void ReplaceObjectHelper(UObject*& OldObject, UClass* OldClass, UObject*&
 			}
 		}
 
+		FObjectInstancingGraph InstancingGraph;
+		const bool bUseInstancingGraph = PrePopulateInstancingGraphForArchetype(InstancingGraph, OldObject, NewArchetype, &OldToNewInstanceMap);
+		
 		FMakeClassSpawnableOnScope TemporarilySpawnable(NewClass);
-		NewUObject = NewObject<UObject>(DestinationOuter, NewClass, NewName, RF_NoFlags, NewArchetype, false, nullptr, OldExternalPackage);
+		NewUObject = NewObject<UObject>(DestinationOuter, NewClass, NewName, RF_NoFlags, NewArchetype, false, bUseInstancingGraph ? &InstancingGraph : nullptr, OldExternalPackage);
 	}
 
 	check(NewUObject != nullptr);
@@ -2290,7 +2293,7 @@ static void ReplaceObjectHelper(UObject*& OldObject, UClass* OldClass, UObject*&
 
 	TMap<UObject*, UObject*> CreatedInstanceMap;
 	TArray< TTuple<UObject*, UObject*>> OrderedListOfObjectToCopy;
-	FBlueprintCompileReinstancer::PreCreateSubObjectsForReinstantiation(OldToNewClassMap, OldObject, NewUObject, CreatedInstanceMap, &OldToNewInstanceMap, &OrderedListOfObjectToCopy);
+	PreCreateSubObjectsForReinstantiation(OldToNewClassMap, OldObject, NewUObject, CreatedInstanceMap, &OldToNewInstanceMap, &OrderedListOfObjectToCopy);
 	OldToNewInstanceMap.Append(CreatedInstanceMap);
 
 	// Copy property values
@@ -3277,6 +3280,76 @@ void FBlueprintCompileReinstancer::PreCreateSubObjectsForReinstantiation(
 	}
 }
 
+bool FBlueprintCompileReinstancer::PrePopulateInstancingGraphForArchetype(FObjectInstancingGraph& InstancingGraph, TNotNull<UObject*> OldInstance, TNotNull<UObject*> InArchetype, TMap<UObject*, UObject*>* OldToNewInstanceMap)
+{
+	checkf(!InstancingGraph.HasDestinationRoot(), TEXT("Can only pre-populate an empty InstancingGraph"));
+
+	// Check if we are re-instancing from an archetype, otherwise nothing to do
+	if (InArchetype->HasAnyFlags(RF_ClassDefaultObject))
+	{
+		return false;
+	}
+
+	// Only needed for objects that are using OS for now
+	if (!FOverridableManager::Get().IsEnabled(OldInstance))
+	{
+		return false;
+	}
+
+	// Find the root of the instantiation
+	UObject* InstancingGraphRoot = OldInstance;
+	{
+		if (UObject** NewInstancingGraphRoot = OldToNewInstanceMap ? OldToNewInstanceMap->Find(InstancingGraphRoot) : nullptr)
+		{
+			InstancingGraphRoot = *NewInstancingGraphRoot;
+		}
+		InstancingGraphRoot = InstancingGraphRoot->GetOuter();
+	} while (InstancingGraphRoot && !InstancingGraphRoot->GetArchetype()->HasAnyFlags(RF_ClassDefaultObject));
+
+	// Sanitization of the found instancing root
+	if (!InstancingGraphRoot || InstancingGraphRoot->IsA<UPackage>())
+	{
+		UE_LOG(LogClass, Warning, TEXT("Object %s was created from an archetype we could not find the root of the instancing graph."), *GetPathNameSafe(OldInstance));
+		return false;
+	}
+
+	auto GatherMatchingSubObjects = [&InstancingGraph, OldInstance](UObject* Instance, UObject* Template, auto GatherMatchingSubObjectsRecurse) -> void
+	{
+		InstancingGraph.AddNewObject(Instance, Template);
+
+		ForEachObjectWithOuter(Instance, [OldInstance, Template, GatherMatchingSubObjectsRecurse](UObject* SubObject)
+		{
+			// Skip the object that is about to be re-instantiated
+			if (SubObject == OldInstance)
+			{
+				return;
+			}
+
+			// Just handle subobject that are created from archetype, skip the others
+			UObject* SubObjectArchetype = SubObject->GetArchetype();
+			if (SubObjectArchetype->HasAnyFlags(RF_ClassDefaultObject))
+			{
+				return;
+			}
+
+			// Sanitization of the archetype
+			if (SubObjectArchetype->GetOuter() != Template)
+			{
+				UE_LOG(LogClass, Warning, TEXT("Missmatched while trying to rebuild instancing graph for '%s', expecting the outer of archetype %s to be %s but was %s"), 
+						*GetPathNameSafe(OldInstance), *GetPathNameSafe(SubObjectArchetype), *GetPathNameSafe(Template), *GetPathNameSafe(SubObjectArchetype->GetOuter()));
+				return;
+			}
+
+			GatherMatchingSubObjectsRecurse(SubObject, SubObjectArchetype, GatherMatchingSubObjectsRecurse);
+
+		}, /*bIncludeNestedObjects*/false);
+	};
+
+	GatherMatchingSubObjects(InstancingGraphRoot, InstancingGraphRoot->GetArchetype(), GatherMatchingSubObjects );
+
+	return true;
+}
+
 void FBlueprintCompileReinstancer::PreCreateSubObjectsForReinstantiation_Inner(
 	const TSet<UObject*>* OldInstancedSubObjects, 
 	const TMap<UClass*, UClass*>& OldToNewClassMap, 
diff --git a/Engine/Source/Editor/UnrealEd/Public/Kismet2/KismetReinstanceUtilities.h b/Engine/Source/Editor/UnrealEd/Public/Kismet2/KismetReinstanceUtilities.h
--- a/Engine/Source/Editor/UnrealEd/Public/Kismet2/KismetReinstanceUtilities.h
+++ b/Engine/Source/Editor/UnrealEd/Public/Kismet2/KismetReinstanceUtilities.h
@@ -11,6 +11,7 @@
 class FReinstanceFinalizer;
 class UBlueprint;
 class FEditorCacheArchetypeManager;
+struct FObjectInstancingGraph;
 
 DECLARE_STATS_GROUP(TEXT("Kismet Reinstancer"), STATGROUP_KismetReinstancer, STATCAT_Advanced);
 
@@ -328,6 +329,22 @@ public:
 	static UNREALED_API void PreCreateSubObjectsForReinstantiation(const TMap<UClass*, UClass*>& OldToNewClassMap, UObject* OldObject, UObject* NewUObject, TMap<UObject*, UObject*>& CreatedInstanceMap, const TMap<UObject*, UObject*>* OldToNewInstanceMap = nullptr, TArray< TTuple<UObject*, UObject*>>* OrderedListOfObjectToCopy = nullptr);
 
 private:
+
+	/** Helper to replace an instance of an object*/
+	static UNREALED_API void ReplaceObjectHelper(UObject*& OldObject, UClass* OldClass, UObject*& NewUObject, UClass* NewClass, TMap<UObject*, UObject*>& OldToNewInstanceMap, const TMap<UClass*, UClass*>& OldToNewClassMap, TMap<UObject*, FName>& OldToNewNameMap, int32 OldObjIndex, TArray<UObject*>& ObjectsToReplace, TArray<UObject*>& PotentialEditorsForRefreshing, TSet<AActor*>& OwnersToRerunConstructionScript, TFunctionRef<TArray<TObjectPtr<USceneComponent>>&(USceneComponent*)> GetAttachChildrenArray, bool bIsComponent, bool bArchetypesAreUpToDate);
+
+	/**
+	 * Pre-populate instancing graph for a specified instance if needed
+	 *
+	 * @param InstancingGraph The instancing graph to populate
+	 * @param OldInstance Object that is about to be reallocated
+	 * @param InArchetype Archetype of ObjectInstance
+	 * @param OldToNewInstanceMap Possible instance mapping of instances
+	 * @return false if the instancing graph is not needed
+	 */
+	static UNREALED_API bool PrePopulateInstancingGraphForArchetype(FObjectInstancingGraph& InstancingGraph, TNotNull<UObject*> OldInstance, TNotNull<UObject*> InArchetype, TMap<UObject*, UObject*>* OldToNewInstanceMap = nullptr);
+
+
 	/** Handles the sub object pre-creation recursively */
 	static UNREALED_API void PreCreateSubObjectsForReinstantiation_Inner(const TSet<UObject*>* OldInstancedSubObjects, const TMap<UClass*, UClass*>& OldToNewClassMap, UObject* OldObject, UObject* NewUObject, TMap<UObject*, UObject*>& CreatedInstanceMap, const TMap<UObject*, UObject*>* OldToNewInstanceMap, TArray< TTuple<UObject*, UObject*>>* OrderedListOfObjectToCopy);
 
diff --git a/Engine/Source/Runtime/Engine/Classes/Engine/BlueprintGeneratedClass.h b/Engine/Source/Runtime/Engine/Classes/Engine/BlueprintGeneratedClass.h
--- a/Engine/Source/Runtime/Engine/Classes/Engine/BlueprintGeneratedClass.h
+++ b/Engine/Source/Runtime/Engine/Classes/Engine/BlueprintGeneratedClass.h
@@ -441,6 +441,13 @@ public:
 	UPROPERTY()
 	uint8 bHasCookedComponentInstancingData:1;
 
+	/** **** Experimental ****
+	 * This will enable dynamic subobject instancing, which will use default data to determine which
+	 * properties need instancing at construction/serialization time rather than reflected property
+	 * 'Instanced' semantics. */
+	UPROPERTY()
+	uint8 bSupportsDynamicInstancedReference:1;
+
 #if WITH_EDITORONLY_DATA
 	/** Used to check if this class has sparse data that can be serialized. This will be false when loading the data if it hasn't already been saved out. */
 	UPROPERTY()
@@ -629,6 +636,7 @@ protected:
 	ENGINE_API virtual FName FindPropertyNameFromGuid(const FGuid& PropertyGuid) const override;
 	ENGINE_API virtual FGuid FindPropertyGuidFromName(const FName InName) const override;
 	ENGINE_API virtual bool ArePropertyGuidsAvailable() const override;
+	ENGINE_API virtual bool SupportsDynamicInstancedReference() const override;
 	//~ End UClass interface
 
 	/**
diff --git a/Engine/Source/Runtime/Engine/Private/BlueprintGeneratedClass.cpp b/Engine/Source/Runtime/Engine/Private/BlueprintGeneratedClass.cpp
--- a/Engine/Source/Runtime/Engine/Private/BlueprintGeneratedClass.cpp
+++ b/Engine/Source/Runtime/Engine/Private/BlueprintGeneratedClass.cpp
@@ -402,6 +402,7 @@ UBlueprintGeneratedClass::UBlueprintGeneratedClass(const FObjectInitializer& Obj
 {
 	NumReplicatedProperties = 0;
 	bHasCookedComponentInstancingData = false;
+	bSupportsDynamicInstancedReference = false;
 	bCustomPropertyListForPostConstructionInitialized = false;
 #if WITH_EDITORONLY_DATA
 	bIsSparseClassDataSerializable = false;
@@ -2887,6 +2888,11 @@ bool UBlueprintGeneratedClass::ArePropertyGuidsAvailable() const
 	return false;
 }
 
+bool UBlueprintGeneratedClass::SupportsDynamicInstancedReference() const
+{
+	return bSupportsDynamicInstancedReference;
+}
+
 FName UBlueprintGeneratedClass::FindPropertyNameFromGuid(const FGuid& PropertyGuid) const
 {
 	// Check parent first as it may have renamed a property since this class was last saved
