Added support for showing Verse paths to various message boxes, dialogs and slow tasks.
Added various utility functions to IAssetTools to help with displaying Verse paths in straight forward circumstances.

diff --git a/Engine/Plugins/Editor/DataValidation/Source/DataValidation/Private/EditorValidatorSubsystem.cpp b/Engine/Plugins/Editor/DataValidation/Source/DataValidation/Private/EditorValidatorSubsystem.cpp
--- a/Engine/Plugins/Editor/DataValidation/Source/DataValidation/Private/EditorValidatorSubsystem.cpp
+++ b/Engine/Plugins/Editor/DataValidation/Source/DataValidation/Private/EditorValidatorSubsystem.cpp
@@ -18,6 +18,7 @@
 #include "EditorValidatorBase.h"
 #include "Engine/BlueprintGeneratedClass.h"
 #include "Engine/Level.h"
+#include "IAssetTools.h"
 #include "IDirectoryWatcher.h"
 #include "ISourceControlModule.h"
 #include "LoadExternalObjectsForValidation.h"
@@ -555,7 +556,7 @@ EDataValidationResult UEditorValidatorSubsystem::ValidateAssetsInternal(
 
 		if (!InSettings.bSilent)
 		{
-			SlowTask.EnterProgressFrame(1.0f, FText::Format(LOCTEXT("DataValidation.ValidatingFilename", "Validating {0}"), FText::FromString(Data.GetFullName())));
+			SlowTask.EnterProgressFrame(1.0f, FText::Format(LOCTEXT("DataValidation.ValidatingFilename", "Validating {0}"), FText::FromString(IAssetTools::Get().GetUserFacingFullName(Data))));
 		}
 		
 		if (OutResults.NumChecked >= InSettings.MaxAssetsToValidate)
diff --git a/Engine/Source/Developer/AssetTools/Private/AssetTools.cpp b/Engine/Source/Developer/AssetTools/Private/AssetTools.cpp
--- a/Engine/Source/Developer/AssetTools/Private/AssetTools.cpp
+++ b/Engine/Source/Developer/AssetTools/Private/AssetTools.cpp
@@ -6047,10 +6047,8 @@ void UAssetToolsImpl::InitAdvancedCopyFromCopyParams(FAdvancedCopyParams CopyPar
 	}
 	else
 	{
-		const FText ReportMessage = FText::FromString(CopyParams.GetDropLocationForAdvancedCopy());
-
 		SAdvancedCopyReportDialog::FOnReportConfirmed OnReportConfirmed = SAdvancedCopyReportDialog::FOnReportConfirmed::CreateUObject(this, &UAssetToolsImpl::AdvancedCopyPackages_ReportConfirmed);
-		SAdvancedCopyReportDialog::OpenPackageReportDialog(CopyParams, ReportMessage, CompleteDestinationMap, CompleteDependencyMap, OnReportConfirmed);
+		SAdvancedCopyReportDialog::OpenPackageReportDialog(CopyParams, CopyParams.GetDropLocationForAdvancedCopy(), CompleteDestinationMap, CompleteDependencyMap, OnReportConfirmed);
 	}
 }
 
@@ -6504,6 +6502,298 @@ UE::AssetTools::FShowingContentVersePath& UAssetToolsImpl::GetOnShowingContentVe
 	return ShowingContentVersePathDelegate;
 }
 
+FString UAssetToolsImpl::GetUserFacingLongPackagePath(FStringView LongPackagePath) const
+{
+	if (!LongPackagePath.IsEmpty() && ShowingContentVersePath())
+	{
+		bool bHadClassesPrefix = false;
+		const FStringView MountPointName = FPathViews::GetMountPointNameFromPath(LongPackagePath, &bHadClassesPrefix);
+		if (!MountPointName.IsEmpty() && !bHadClassesPrefix)
+		{
+			// If the object isn't mounted under a plugin, the object doesn't have a Verse path.
+			if (TSharedPtr<IPlugin> Plugin = IPluginManager::Get().FindPlugin(MountPointName))
+			{
+				// If the plugin doesn't have a root vpath, the object doesn't have a Verse path.
+				const FString& PluginVersePath = Plugin->GetVersePath();
+				if (!PluginVersePath.IsEmpty())
+				{
+					// Don't FPaths::Combine if getting the Verse path for the root plugin path, FPaths::Combine will add a trailing '/' making the Verse path invalid.
+					const int32 CharCount = MountPointName.Len() + 1;
+					UE::Core::FVersePath VersePath;
+					if (VersePath.TryMake(VersePath, LongPackagePath.Len() == CharCount ? PluginVersePath : FPaths::Combine(PluginVersePath, LongPackagePath.RightChop(CharCount))))
+					{
+						return MoveTemp(VersePath).ToString();
+					}
+				}
+			}
+		}
+	}
+
+	return FString(LongPackagePath);
+}
+
+FString UAssetToolsImpl::GetUserFacingLongPackageName(FName PackageName) const
+{
+	if (!PackageName.IsNone() && ShowingContentVersePath())
+	{
+		TMap<FName, FAssetData> Assets;
+		UE::AssetRegistry::GetAssetForPackages(TConstArrayView<FName>(&PackageName, 1), Assets);
+
+		if (const FAssetData* AssetData = Assets.Find(PackageName))
+		{
+			UE::Core::FVersePath VersePath = AssetData->GetVersePath();
+			if (VersePath.IsValid())
+			{
+				return MoveTemp(VersePath).ToString();
+			}
+		}
+	}
+
+	return PackageName.ToString();
+}
+
+TArray<FString> UAssetToolsImpl::GetUserFacingLongPackageNames(TConstArrayView<FName> PackageNames) const
+{
+	if (PackageNames.IsEmpty())
+	{
+		return {};
+	}
+
+	TArray<FString> UserFacingPaths;
+	UserFacingPaths.Reserve(PackageNames.Num());
+	
+	if (ShowingContentVersePath())
+	{
+		TMap<FName, FAssetData> Assets;
+		UE::AssetRegistry::GetAssetForPackages(PackageNames, Assets);
+
+		for (const FName& PackageName : PackageNames)
+		{
+			if (!PackageName.IsNone())
+			{
+				if (const FAssetData* AssetData = Assets.Find(PackageName))
+				{
+					UE::Core::FVersePath VersePath = AssetData->GetVersePath();
+					if (VersePath.IsValid())
+					{
+						UserFacingPaths.Add(MoveTemp(VersePath).ToString());
+						continue;
+					}
+				}
+			}
+
+			UserFacingPaths.Add(PackageName.ToString());
+		}
+	}
+	else
+	{
+		for (const FName& PackageName : PackageNames)
+		{
+			UserFacingPaths.Add(PackageName.ToString());
+		}
+	}
+
+	return UserFacingPaths;
+}
+
+FString UAssetToolsImpl::GetUserFacingObjectPath(const FSoftObjectPath& ObjectPath) const
+{
+	// If the object path is just a package name, get the path for the package.
+	if (ObjectPath.IsAsset() && ObjectPath.GetAssetFName().IsNone())
+	{
+		return GetUserFacingLongPackageName(ObjectPath.GetLongPackageFName());
+	}
+
+	if (ShowingContentVersePath())
+	{
+		FAssetRegistryModule& AssetRegistryModule = FModuleManager::LoadModuleChecked<FAssetRegistryModule>(TEXT("AssetRegistry"));
+		FAssetData AssetData;
+		if (AssetRegistryModule.Get().TryGetAssetByObjectPath(ObjectPath, AssetData) == UE::AssetRegistry::EExists::Exists)
+		{
+			UE::Core::FVersePath VersePath = AssetData.GetVersePath();
+			if (VersePath.IsValid())
+			{
+				return MoveTemp(VersePath).ToString();
+			}
+		}
+	}
+
+	return ObjectPath.ToString();
+}
+
+TArray<FString> UAssetToolsImpl::GetUserFacingObjectPaths(TConstArrayView<FSoftObjectPath> ObjectPaths) const
+{
+	if (ObjectPaths.IsEmpty())
+	{
+		return {};
+	}
+
+	TArray<FString> UserFacingPaths;
+	UserFacingPaths.Reserve(ObjectPaths.Num());
+
+	if (ShowingContentVersePath())
+	{
+		FAssetRegistryModule& AssetRegistryModule = FModuleManager::LoadModuleChecked<FAssetRegistryModule>(TEXT("AssetRegistry"));
+
+		// Get the primary asset for any object paths that are just package names.
+		TMap<FName, FAssetData> Assets;
+		{
+			TArray<FName> PackageNames;
+
+			for (const FSoftObjectPath& ObjectPath : ObjectPaths)
+			{
+				if (ObjectPath.IsAsset() && ObjectPath.GetAssetFName().IsNone())
+				{
+					PackageNames.Add(ObjectPath.GetLongPackageFName());
+				}
+			}
+
+			if (!PackageNames.IsEmpty())
+			{
+				UE::AssetRegistry::GetAssetForPackages(PackageNames, Assets);
+			}
+		}
+
+		for (const FSoftObjectPath& ObjectPath : ObjectPaths)
+		{
+			UE::Core::FVersePath VersePath;
+
+			// If the object path is just a package name, get the Verse path for the primary asset.
+			if (ObjectPath.IsAsset() && ObjectPath.GetAssetFName().IsNone())
+			{
+				if (const FAssetData* AssetData = Assets.Find(ObjectPath.GetLongPackageFName()))
+				{
+					VersePath = AssetData->GetVersePath();
+				}
+			}
+			else
+			{
+				FAssetData AssetData;
+				if (AssetRegistryModule.Get().TryGetAssetByObjectPath(ObjectPath, AssetData) == UE::AssetRegistry::EExists::Exists)
+				{
+					VersePath = AssetData.GetVersePath();
+				}
+			}
+
+			if (VersePath.IsValid())
+			{
+				UserFacingPaths.Add(MoveTemp(VersePath).ToString());
+			}
+			else
+			{
+				UserFacingPaths.Add(ObjectPath.ToString());
+			}
+		}
+	}
+	else
+	{
+		for (const FSoftObjectPath& ObjectPath : ObjectPaths)
+		{
+			UserFacingPaths.Add(ObjectPath.ToString());
+		}
+	}
+
+	return UserFacingPaths;
+}
+
+FString UAssetToolsImpl::GetUserFacingLongPackageName(const UPackage& Package) const
+{
+	return GetUserFacingLongPackageName(Package.GetFName());
+}
+
+TArray<FString> UAssetToolsImpl::GetUserFacingLongPackageNames(TConstArrayView<const UPackage*> Packages) const
+{
+	if (Packages.IsEmpty())
+	{
+		return {};
+	}
+
+	TArray<FName, TInlineAllocator<8>> PackageNames;
+	PackageNames.Reserve(Packages.Num());
+
+	for (const UPackage* Package : Packages)
+	{
+		if (ensure(Package))
+		{
+			PackageNames.Add(Package->GetFName());
+		}
+		else
+		{
+			PackageNames.Emplace();
+		}
+	}
+
+	return GetUserFacingLongPackageNames(PackageNames);
+}
+
+FString UAssetToolsImpl::GetUserFacingFullName(const UObject& Object) const
+{
+	TStringBuilder<256> StringBuilder;
+
+	if (ShowingContentVersePath())
+	{
+		const UE::Core::FVersePath VersePath = Object.GetVersePath();
+		if (VersePath.IsValid())
+		{
+			Object.GetClass()->GetFName().AppendString(StringBuilder);
+			StringBuilder.AppendChar(TEXT(' '));
+			StringBuilder.Append(VersePath.ToString());
+			return StringBuilder.ToString();
+		}
+	}
+
+	Object.GetFullName(StringBuilder);
+	return StringBuilder.ToString();
+}
+
+FString UAssetToolsImpl::GetUserFacingPathName(const UObject& Object) const
+{
+	if (ShowingContentVersePath())
+	{
+		UE::Core::FVersePath VersePath = Object.GetVersePath();
+		if (VersePath.IsValid())
+		{
+			return MoveTemp(VersePath).ToString();
+		}
+	}
+
+	return Object.GetPathName();
+}
+
+FString UAssetToolsImpl::GetUserFacingFullName(const FAssetData& AssetData) const
+{
+	TStringBuilder<256> StringBuilder;
+
+	if (ShowingContentVersePath())
+	{
+		const UE::Core::FVersePath VersePath = AssetData.GetVersePath();
+		if (VersePath.IsValid())
+		{
+			AssetData.AssetClassPath.AppendString(StringBuilder);
+			StringBuilder.AppendChar(TEXT(' '));
+			StringBuilder.Append(VersePath.ToString());
+			return StringBuilder.ToString();
+		}
+	}
+
+	AssetData.GetFullName(StringBuilder);
+	return StringBuilder.ToString();
+}
+
+FString UAssetToolsImpl::GetUserFacingPathName(const FAssetData& AssetData) const
+{
+	if (ShowingContentVersePath())
+	{
+		UE::Core::FVersePath VersePath = AssetData.GetVersePath();
+		if (VersePath.IsValid())
+		{
+			return MoveTemp(VersePath).ToString();
+		}
+	}
+
+	return AssetData.GetObjectPathString();
+}
+
 bool UAssetToolsImpl::AllPassWritableFolderFilter(const TArray<FString>& InPaths) const
 {
 	if (WritableFolderPermissionList->HasFiltering())
diff --git a/Engine/Source/Developer/AssetTools/Private/AssetTools.h b/Engine/Source/Developer/AssetTools/Private/AssetTools.h
--- a/Engine/Source/Developer/AssetTools/Private/AssetTools.h
+++ b/Engine/Source/Developer/AssetTools/Private/AssetTools.h
@@ -149,6 +149,17 @@ public:
 	virtual FAssetData FindAssetByVersePath(const UE::Core::FVersePath& VersePath, bool bIncludeOnlyOnDiskAssets = false, bool bSkipARFilteredAssets = true) const override;
 	virtual bool ShowingContentVersePath() const override;
 	virtual UE::AssetTools::FShowingContentVersePath& GetOnShowingContentVersePath() override;
+	virtual FString GetUserFacingLongPackagePath(FStringView LongPackagePath) const override;
+	virtual FString GetUserFacingLongPackageName(FName PackageName) const override;
+	virtual TArray<FString> GetUserFacingLongPackageNames(TConstArrayView<FName> PackageNames) const override;
+	virtual FString GetUserFacingObjectPath(const FSoftObjectPath& ObjectPath) const override;
+	virtual TArray<FString> GetUserFacingObjectPaths(TConstArrayView<FSoftObjectPath> ObjectPaths) const override;
+	virtual FString GetUserFacingLongPackageName(const UPackage& Package) const override;
+	virtual TArray<FString> GetUserFacingLongPackageNames(TConstArrayView<const UPackage*> Packages) const override;
+	virtual FString GetUserFacingFullName(const UObject& Object) const override;
+	virtual FString GetUserFacingPathName(const UObject& Object) const override;
+	virtual FString GetUserFacingFullName(const FAssetData& AssetData) const override;
+	virtual FString GetUserFacingPathName(const FAssetData& AssetData) const override;
 	virtual bool AllPassWritableFolderFilter(const TArray<FString>& InPaths) const override;
 	virtual void NotifyBlockedByWritableFolderFilter() const;
 	virtual bool IsNameAllowed(const FString& Name, FText* OutErrorMessage = nullptr) const override;
diff --git a/Engine/Source/Developer/AssetTools/Private/SAdvancedCopyReportDialog.cpp b/Engine/Source/Developer/AssetTools/Private/SAdvancedCopyReportDialog.cpp
--- a/Engine/Source/Developer/AssetTools/Private/SAdvancedCopyReportDialog.cpp
+++ b/Engine/Source/Developer/AssetTools/Private/SAdvancedCopyReportDialog.cpp
@@ -137,14 +137,14 @@ bool FAdvancedCopyReportNode::AddPackage_Recursive(const FString& InSource, cons
 }
 
 BEGIN_SLATE_FUNCTION_BUILD_OPTIMIZATION
-void SAdvancedCopyReportDialog::Construct( const FArguments& InArgs, const FAdvancedCopyParams& InParams, const FText& InReportMessage, const TArray<TMap<FString, FString>>& InDestinationMap, const TArray<TMap<FName, FName>>& DependencyMap, const SAdvancedCopyReportDialog::FOnReportConfirmed& InOnReportConfirmed )
+void SAdvancedCopyReportDialog::Construct( const FArguments& InArgs, const FAdvancedCopyParams& InParams, const FString& InDestination, const TArray<TMap<FString, FString>>& InDestinationMap, const TArray<TMap<FName, FName>>& DependencyMap, const SAdvancedCopyReportDialog::FOnReportConfirmed& InOnReportConfirmed )
 {
 	OnReportConfirmed = InOnReportConfirmed;
 	InitialDestinationMap = InDestinationMap;
 	CloneSet = MakeShared<TSet<FString>>();
 	PackageReportRootNode.IncludedSet = CloneSet;
 	CurrentCopyParams = InParams;
-	ReportString = InReportMessage.ToString();
+	Destination = InDestination;
 	ConstructNodeTree(InDestinationMap, DependencyMap);
 
 	TSharedRef< SHeaderRow > HeaderRowWidget = SNew(SHeaderRow);
@@ -176,7 +176,7 @@ void SAdvancedCopyReportDialog::Construct( const FArguments& InArgs, const FAdva
 			.Padding(0, 4)
 			[
 				SNew(STextBlock)
-				.Text(this, &SAdvancedCopyReportDialog::GetHeaderText, InReportMessage)
+				.Text(this, &SAdvancedCopyReportDialog::GetHeaderText, FText::FromString(IAssetTools::Get().GetUserFacingLongPackagePath(Destination)))
 				.TextStyle( FAppStyle::Get(), "PackageMigration.DialogTitle" )
 				.AutoWrapText(true)
 			]
@@ -314,11 +314,11 @@ void SAdvancedCopyReportDialog::ToggleGeneratingDependencies(ECheckBoxState NewS
 {
 	CurrentCopyParams.bShouldCheckForDependencies = NewState == ECheckBoxState::Checked;
 	CloseDialog();
-	IAssetTools& AssetTools = FModuleManager::GetModuleChecked<FAssetToolsModule>("AssetTools").Get();
+	IAssetTools& AssetTools = IAssetTools::Get();
 	AssetTools.InitAdvancedCopyFromCopyParams(CurrentCopyParams);
 }
 
-void SAdvancedCopyReportDialog::OpenPackageReportDialog(const FAdvancedCopyParams& InParams, const FText& ReportMessage, const TArray<TMap<FString, FString>>& InDestinationMap, const TArray<TMap<FName, FName>>& DependencyMap, const SAdvancedCopyReportDialog::FOnReportConfirmed& InOnReportConfirmed)
+void SAdvancedCopyReportDialog::OpenPackageReportDialog(const FAdvancedCopyParams& InParams, const FString& Destination, const TArray<TMap<FString, FString>>& InDestinationMap, const TArray<TMap<FName, FName>>& DependencyMap, const SAdvancedCopyReportDialog::FOnReportConfirmed& InOnReportConfirmed)
 {
 	TSharedRef<SWindow> ReportWindow = SNew(SWindow)
 		.Title(LOCTEXT("AdvancedCopyReportWindowTitle", "Advanced Copy Asset Report"))
@@ -326,7 +326,7 @@ void SAdvancedCopyReportDialog::OpenPackageReportDialog(const FAdvancedCopyParam
 		.SupportsMaximize(true)
 		.SupportsMinimize(true)
 		[
-			SNew(SAdvancedCopyReportDialog, InParams, ReportMessage, InDestinationMap, DependencyMap, InOnReportConfirmed)
+			SNew(SAdvancedCopyReportDialog, InParams, Destination, InDestinationMap, DependencyMap, InOnReportConfirmed)
 		];
 
 	IMainFrameModule& MainFrameModule = FModuleManager::LoadModuleChecked<IMainFrameModule>(TEXT("MainFrame"));
@@ -414,10 +414,10 @@ FReply SAdvancedCopyReportDialog::OkClicked()
 					FPackageName::SplitLongPackageName(It.Value(), OutPackageRoot, OutPackagePath, OutPackageName);
 
 					FString CroppedDestination = It.Value();
-					const bool bDidRemovePrefix = CroppedDestination.RemoveFromStart(ReportString);
+					const bool bDidRemovePrefix = CroppedDestination.RemoveFromStart(Destination);
 					if (bDidRemovePrefix)
 					{
-						It.Value() = ReportString / CroppedDestination.Replace(*FindString, *ReplaceString);
+						It.Value() = Destination / CroppedDestination.Replace(*FindString, *ReplaceString);
 					}
 				}
 			}
@@ -548,13 +548,22 @@ const TSharedRef< SWidget > SAdvancedCopyColumn::ConstructRowWidget(TSharedPtr<s
 	}
 	else if (ColumnName == SourceColumnLabel)
 	{
+		FString Source;
+		if (FPackageName::DoesPackageExist(TreeItem.Get()->Source))
+		{
+			Source = IAssetTools::Get().GetUserFacingLongPackageName(FName(TreeItem.Get()->Source));
+		}
+		else
+		{
+			Source = IAssetTools::Get().GetUserFacingLongPackagePath(TreeItem.Get()->Source);
+		}
 		return SNew(STextBlock)
-			.Text(FText::FromString(TreeItem.Get()->Source));
+			.Text(FText::FromString(MoveTemp(Source)));
 	}
 	else if (ColumnName == RelativeDestinationColumnLabel)
 	{
 		FString CroppedDestination = TreeItem.Get()->Destination;
-		const bool bDidRemovePrefix = CroppedDestination.RemoveFromStart(Row.GetReportDialog()->GetReportString());
+		const bool bDidRemovePrefix = CroppedDestination.RemoveFromStart(Row.GetReportDialog()->GetDestination());
 		if (bDidRemovePrefix)
 		{
 			CroppedDestination = TEXT(".") + CroppedDestination;
diff --git a/Engine/Source/Developer/AssetTools/Private/SAdvancedCopyReportDialog.h b/Engine/Source/Developer/AssetTools/Private/SAdvancedCopyReportDialog.h
--- a/Engine/Source/Developer/AssetTools/Private/SAdvancedCopyReportDialog.h
+++ b/Engine/Source/Developer/AssetTools/Private/SAdvancedCopyReportDialog.h
@@ -79,21 +79,21 @@ public:
 	SLATE_END_ARGS()
 
 	/** Constructs this widget with InArgs */
-	void Construct( const FArguments& InArgs, const FAdvancedCopyParams& InParams, const FText& InReportMessage, const TArray<TMap<FString, FString>>& DestinationMap, const TArray<TMap<FName, FName>>& DependencyMap, const FOnReportConfirmed& InOnReportConfirmed );
+	void Construct( const FArguments& InArgs, const FAdvancedCopyParams& InParams, const FString& InDestination, const TArray<TMap<FString, FString>>& DestinationMap, const TArray<TMap<FName, FName>>& DependencyMap, const FOnReportConfirmed& InOnReportConfirmed );
 		
 	ECheckBoxState IsGeneratingDependencies() const;
 
 	void ToggleGeneratingDependencies(ECheckBoxState NewState);
 
 	/** Opens the dialog in a new window */
-	static void OpenPackageReportDialog(const FAdvancedCopyParams& InParams, const FText& ReportMessage, const TArray<TMap<FString, FString>>& DestinationMap, const TArray<TMap<FName, FName>>& DependencyMap, const FOnReportConfirmed& InOnReportConfirmed);
+	static void OpenPackageReportDialog(const FAdvancedCopyParams& InParams, const FString& Destination, const TArray<TMap<FString, FString>>& DestinationMap, const TArray<TMap<FName, FName>>& DependencyMap, const FOnReportConfirmed& InOnReportConfirmed);
 
 	/** Closes the dialog. */
 	void CloseDialog();
 
 	TMap<FName, TSharedPtr<SAdvancedCopyColumn>> GetColumns() const { return Columns; }
 
-	FString GetReportString() const { return ReportString; }
+	const FString& GetDestination() const { return Destination; }
 
 private:
 	/** Constructs the node tree given the list of package names */
@@ -123,7 +123,7 @@ private:
 	TSharedPtr<TSet<FString>> CloneSet;
 	/** Map of columns that are shown on this report. */
 	TMap<FName, TSharedPtr<SAdvancedCopyColumn>> Columns;
-	FString ReportString;
+	FString Destination;
 	FAdvancedCopyParams CurrentCopyParams;
 	TArray<TMap<FString, FString>> InitialDestinationMap;
 	FString FindString;
diff --git a/Engine/Source/Developer/AssetTools/Public/IAssetTools.h b/Engine/Source/Developer/AssetTools/Public/IAssetTools.h
--- a/Engine/Source/Developer/AssetTools/Public/IAssetTools.h
+++ b/Engine/Source/Developer/AssetTools/Public/IAssetTools.h
@@ -677,6 +677,75 @@ public:
 	/** Delegate to control whether we are showing content Verse paths or not */
 	UE_INTERNAL virtual UE::AssetTools::FShowingContentVersePath& GetOnShowingContentVersePath() = 0;
 
+	/**
+	 * Return the path to display to the user in editor UI for the given package path.
+	 * Attempts to convert the folder to a Verse path if ShowingContentVersePath is true.
+	 */
+	UE_INTERNAL virtual FString GetUserFacingLongPackagePath(FStringView LongPackagePath) const = 0;
+
+	/**
+	 * Return the path to display to the user in editor UI for the given long package name.
+	 * Attempts to convert the primary asset in the package to a Verse path if ShowingContentVersePath is true.
+	 */
+	UE_INTERNAL virtual FString GetUserFacingLongPackageName(FName PackageName) const = 0;
+
+	/**
+	 * Return the paths to display to the user in editor UI for the given long package names.
+	 * Attempts to convert the primary asset in each package to a Verse path if ShowingContentVersePath is true.
+	 * @note The array return value is guaranteed to have the same size as PackageNames.
+	 */
+	UE_INTERNAL virtual TArray<FString> GetUserFacingLongPackageNames(TConstArrayView<FName> PackageNames) const = 0;
+
+	/**
+	* Return the path to display to the user in editor UI for the given soft object path.
+	* Attempts to convert the object path to a Verse path if ShowingContentVersePath is true.
+	*/
+	UE_INTERNAL virtual FString GetUserFacingObjectPath(const FSoftObjectPath& ObjectPath) const = 0;
+
+	/**
+	 * Return the paths to display to the user in editor UI for the given soft object paths.
+	 * Attempts to convert each object path to a Verse path if ShowingContentVersePath is true.
+	 * @note The array return value is guaranteed to have the same size as ObjectPaths.
+	 */
+	UE_INTERNAL virtual TArray<FString> GetUserFacingObjectPaths(TConstArrayView<FSoftObjectPath> ObjectPaths) const = 0;
+
+	/**
+	 * Return the path to display to the user in editor UI for the given package.
+	 * Attempts to convert the primary asset in the package to a Verse path if ShowingContentVersePath is true.
+	 */
+	UE_INTERNAL virtual FString GetUserFacingLongPackageName(const UPackage& Package) const = 0;
+
+	/**
+	 * Return the paths to display to the user in editor UI for the given packages.
+	 * Attempts to convert the primary asset in each package to a Verse path if ShowingContentVersePath is true.
+	 * @note The array return value is guaranteed to have the same size as Packages.
+	 */
+	UE_INTERNAL virtual TArray<FString> GetUserFacingLongPackageNames(TConstArrayView<const UPackage*> Packages) const = 0;
+
+	/**
+	 * Return the full name to display to the user in editor UI for the given object.
+	 * Attempts to convert the object path portion of the full name to a Verse path if ShowingContentVersePath is true.
+	 */
+	UE_INTERNAL virtual FString GetUserFacingFullName(const UObject& Object) const = 0;
+
+	/**
+	 * Return the path to display to the user in editor UI for the given asset data.
+	 * Attempts to convert the object path to a Verse path if ShowingContentVersePath is true.
+	 */
+	UE_INTERNAL virtual FString GetUserFacingPathName(const UObject& Object) const = 0;
+
+	/**
+	 * Return the full name to display to the user in editor UI for the given object.
+	 * Attempts to convert the asset path portion of the full name to a Verse path if ShowingContentVersePath is true.
+	 */
+	UE_INTERNAL virtual FString GetUserFacingFullName(const FAssetData& AssetData) const = 0;
+
+	/**
+	 * Return the path to display to the user in editor UI for the given asset data.
+	 * Attempts to convert the asset path to a Verse path if ShowingContentVersePath is true.
+	 */
+	UE_INTERNAL virtual FString GetUserFacingPathName(const FAssetData& AssetData) const = 0;
+
 	/** Returns true if all in list pass writable folder filter */
 	virtual bool AllPassWritableFolderFilter(const TArray<FString>& InPaths) const = 0;
 
diff --git a/Engine/Source/Editor/UnrealEd/Private/ConsolidateWindow.cpp b/Engine/Source/Editor/UnrealEd/Private/ConsolidateWindow.cpp
--- a/Engine/Source/Editor/UnrealEd/Private/ConsolidateWindow.cpp
+++ b/Engine/Source/Editor/UnrealEd/Private/ConsolidateWindow.cpp
@@ -3,6 +3,7 @@
 #include "ConsolidateWindow.h"
 #include "Modules/ModuleManager.h"
 #include "UObject/ObjectRedirector.h"
+#include "IAssetTools.h"
 #include "InputCoreTypes.h"
 #include "Layout/Visibility.h"
 #include "Input/Reply.h"
@@ -91,9 +92,9 @@ public:
 		}
 
 		/** @return		The full name of the object this item represents */
-		FString GetObjectName() const
+		FText GetObjectName() const
 		{
-			return Object->GetFullName();
+			return FText::FromString(IAssetTools::Get().GetUserFacingFullName(*Object));
 		}
 		/** @return		The UObject this item represents */	
 		const UObject* GetObject() const
@@ -465,7 +466,7 @@ TSharedRef<ITableRow> SConsolidateToolWidget::OnGenerateRowForList( TSharedPtr<F
 			.OnCheckStateChanged( ListItemPtr.ToSharedRef(), &FListItem::OnAssetSelected )
 			[
 				SNew (STextBlock)
-				.Text( FText::FromString(ListItemPtr->GetObjectName()) )
+				.Text( ListItemPtr.ToSharedRef(), &FListItem::GetObjectName )
 			]
 		];
 }
diff --git a/Engine/Source/Editor/UnrealEd/Private/FileHelpers.cpp b/Engine/Source/Editor/UnrealEd/Private/FileHelpers.cpp
--- a/Engine/Source/Editor/UnrealEd/Private/FileHelpers.cpp
+++ b/Engine/Source/Editor/UnrealEd/Private/FileHelpers.cpp
@@ -8,6 +8,7 @@
 #include "HAL/FileManager.h"
 #include "Algo/Copy.h"
 #include "Algo/Transform.h"
+#include "Containers/VersePath.h"
 #include "Misc/CommandLine.h"
 #include "Misc/Paths.h"
 #include "Misc/ConfigCacheIni.h"
@@ -2081,9 +2082,9 @@ void FEditorFileUtils::MakePackagesWritable(const TArray<UPackage*>& PackagesToM
 	if (PackagesMadeWritableFailure.Num() > 0)
 	{
 		FString PkgsWhichFailedWritable;
-		for (UPackage* Package : PackagesMadeWritableFailure)
+		for (const FString& UserFacingPath : IAssetTools::Get().GetUserFacingLongPackageNames(PackagesMadeWritableFailure))
 		{
-			PkgsWhichFailedWritable += FString::Printf(TEXT("\n%s"), *Package->GetName());
+			PkgsWhichFailedWritable += FString::Printf(TEXT("\n%s"), *UserFacingPath);
 		}
 
 		FFormatNamedArguments Arguments;
@@ -2136,10 +2137,8 @@ ECommandResult::Type FEditorFileUtils::CheckoutPackages(const TArray<UPackage*>&
 		// Print out all the packages and set the check out result
 		auto FailedIntermediateOperations = [&CheckOutResult, &PkgsWhichFailedCheckout, &PkgsToCheckOut]()
 		{
-			for (auto PkgsToCheckOutIter = PkgsToCheckOut.CreateConstIterator(); PkgsToCheckOutIter; ++PkgsToCheckOutIter)
+			for (const FString& PackageToCheckOutName : IAssetTools::Get().GetUserFacingLongPackageNames(PkgsToCheckOut))
 			{
-				UPackage* PackageToCheckOut = *PkgsToCheckOutIter;
-				const FString PackageToCheckOutName = PackageToCheckOut->GetName();
 				PkgsWhichFailedCheckout += FString::Printf(TEXT("\n%s"), *PackageToCheckOutName);
 			}
 			CheckOutResult = ECommandResult::Failed;
@@ -2233,7 +2232,7 @@ ECommandResult::Type FEditorFileUtils::CheckoutPackages(const TArray<UPackage*>&
 					// If the package couldn't be checked out, log it so the list of failures can be displayed afterwards
 					if (bShowCheckoutError)
 					{
-						const FString PackageToCheckOutName = PackageToCheckOut->GetName();
+						const FString PackageToCheckOutName = IAssetTools::Get().GetUserFacingLongPackageName(*PackageToCheckOut);
 						PkgsWhichFailedCheckout += FString::Printf(TEXT("\n%s"), *PackageToCheckOutName);
 						CheckOutResult = ECommandResult::Failed;
 					}
@@ -2298,7 +2297,7 @@ ECommandResult::Type FEditorFileUtils::CheckoutPackages(const TArray<UPackage*>&
 					}
 					else
 					{
-						PkgsWhichFailedCheckout += FString::Printf(TEXT("\n%s"), *CurPackageName);
+						PkgsWhichFailedCheckout += FString::Printf(TEXT("\n%s"), *IAssetTools::Get().GetUserFacingLongPackageName(*CurPackage));
 						CheckOutResult = ECommandResult::Failed;
 					}
 				}
@@ -2418,7 +2417,7 @@ ECommandResult::Type FEditorFileUtils::CheckoutPackages(const TArray<FString>& P
 			// If the package couldn't be checked out, log it so the list of failures can be displayed afterwards
 			if(bShowCheckoutError)
 			{
-				PkgsWhichFailedCheckout += FString::Printf( TEXT("\n%s"), *PackageToCheckOutName );
+				PkgsWhichFailedCheckout += FString::Printf( TEXT("\n%s"), *IAssetTools::Get().GetUserFacingLongPackageName(FName(PackageToCheckOutName)) );
 			}
 		}
 	}
@@ -2498,7 +2497,7 @@ ECommandResult::Type FEditorFileUtils::CheckoutPackages(const TArray<FString>& P
 				}
 				else
 				{
-					PkgsWhichFailedCheckout += FString::Printf( TEXT("\n%s"), *CurPackageName );
+					PkgsWhichFailedCheckout += FString::Printf( TEXT("\n%s"), *IAssetTools::Get().GetUserFacingLongPackageName(FName(CurPackageName)) );
 				}
 			}
 		}
@@ -3318,7 +3317,7 @@ EAutosaveContentPackagesResult::Type FEditorFileUtils::AutosaveContentPackagesEx
 
 	for (UPackage* CurPackage : PackagesToSave)
 	{
-		SlowTask.DefaultMessage = FText::Format(LOCTEXT("AutoSavingPackage", "Saving package {0}"), FText::FromString(CurPackage->GetName()));
+		SlowTask.DefaultMessage = FText::Format(LOCTEXT("AutoSavingPackage", "Saving package {0}"), FText::FromString(IAssetTools::Get().GetUserFacingLongPackageName(*CurPackage)));
 		SlowTask.EnterProgressFrame();
 
 		// In order to save, the package must be fully-loaded first
@@ -3675,7 +3674,7 @@ static InternalSavePackageResult InternalSavePackage(UPackage* PackageToSave, bo
 			const FText ErrorPrompt = GEditor->IsPlayingOnLocalPCSession() ?
 				NSLOCTEXT("UnrealEd", "Prompt_41", "The asset '{0}' ({1}) cannot be saved as the package is locked because you are in play on PC mode.\n\nCancel: Stop saving all assets and return to the editor.\nRetry: Attempt to save the asset again.\nContinue: Skip saving this asset only." ) :
 				NSLOCTEXT("UnrealEd", "Prompt_26", "The asset '{0}' ({1}) failed to save.\n\nCancel: Stop saving all assets and return to the editor.\nRetry: Attempt to save the asset again.\nContinue: Skip saving this asset only." );
-			EAppReturnType::Type DialogCode = FMessageDialog::Open( EAppMsgType::CancelRetryContinue, EAppReturnType::Continue, FText::Format(ErrorPrompt, FText::FromString(PackageName), FText::FromString(FinalPackageFilename)) );
+			EAppReturnType::Type DialogCode = FMessageDialog::Open( EAppMsgType::CancelRetryContinue, EAppReturnType::Continue, FText::Format(ErrorPrompt, FText::FromString(IAssetTools::Get().GetUserFacingLongPackageName(*PackageToSave)), FText::FromString(FinalPackageFilename)) );
 
 			switch (DialogCode)
 			{
@@ -3700,7 +3699,7 @@ static InternalSavePackageResult InternalSavePackage(UPackage* PackageToSave, bo
 		else if ( !bWasSuccessful )
 		{
 			// We failed at saving because we are in bIsUnattended mode, there is no need to attempt to save again
-			FText FailureReason = FText::Format(NSLOCTEXT("UnrealEd", "SaveAssetFailed", "The asset '{0}' ({1}) failed to save."), FText::FromString(PackageName), FText::FromString(FinalPackageFilename));
+			FText FailureReason = FText::Format(NSLOCTEXT("UnrealEd", "SaveAssetFailed", "The asset '{0}' ({1}) failed to save."), FText::FromString(IAssetTools::Get().GetUserFacingLongPackageName(*PackageToSave)), FText::FromString(FinalPackageFilename));
 			FMessageDialog::Open( EAppMsgType::Ok, FailureReason );
 			bAttemptSave = false;
 			ReturnCode = InternalSavePackageResult::Error;
@@ -3727,9 +3726,9 @@ static void InternalWarnUserAboutFailedSave( const TArray<UPackage*>& InFailedPa
 	if ( InFailedPackages.Num() > 0 )
 	{
 		FString FailedPackages;
-		for ( TArray<UPackage*>::TConstIterator FailedIter( InFailedPackages ); FailedIter; ++FailedIter )
+		for (const FString& UserFacingPath : IAssetTools::Get().GetUserFacingLongPackageNames(InFailedPackages))
 		{
-			FailedPackages += FString::Printf( TEXT("\n%s"), *( (*FailedIter)->GetName() ) );
+			FailedPackages += FString::Printf( TEXT("\n%s"), *UserFacingPath );
 		}
 
 		FFormatNamedArguments Arguments;
@@ -3826,11 +3825,21 @@ static bool InternalSavePackagesFast(const TArray<UPackage*>& PackagesToSave, bo
 				}
 
 				const UWorld* const AssociatedWorld = UWorld::FindWorldInPackage(CurPackage);
-				const bool bIsMapPackage = AssociatedWorld != nullptr;
 
-				const FText SavingPackageText = (bIsMapPackage)
-					? FText::Format(NSLOCTEXT("UnrealEd", "SavingMapf", "Saving map {0}"), FText::FromString(CurPackage->GetName()))
-					: FText::Format(NSLOCTEXT("UnrealEd", "SavingAssetf", "Saving asset {0}"), FText::FromString(CurPackage->GetName()));
+				FText SavingPackageText;
+				if (AssociatedWorld != nullptr)
+				{
+					UE::Core::FVersePath VersePath;
+					if (IAssetTools::Get().ShowingContentVersePath())
+					{
+						VersePath = AssociatedWorld->GetVersePath();
+					}
+					SavingPackageText = FText::Format(NSLOCTEXT("UnrealEd", "SavingMapf", "Saving map {0}"), FText::FromString(VersePath.IsValid() ? MoveTemp(VersePath).ToString() : CurPackage->GetName()));
+				}
+				else
+				{
+					SavingPackageText = FText::Format(NSLOCTEXT("UnrealEd", "SavingAssetf", "Saving asset {0}"), FText::FromString(IAssetTools::Get().GetUserFacingLongPackageName(*CurPackage)));
+				}
 
 				GWarn->StatusForceUpdate(PkgIter.GetIndex(), PackagesToSave.Num(), SavingPackageText);
 
@@ -4275,11 +4284,21 @@ FEditorFileUtils::EPromptReturnCode InternalPromptForCheckoutAndSave(const TArra
 		for (UPackage* Package : PackagesToSave)
 		{
 			const UWorld* const AssociatedWorld = UWorld::FindWorldInPackage(Package);
-			const bool bIsMapPackage = AssociatedWorld != nullptr;
 
-			const FText SavingPackageText = (bIsMapPackage)
-				? FText::Format(NSLOCTEXT("UnrealEd", "SavingMapf", "Saving map {0}"), FText::FromString(Package->GetName()))
-				: FText::Format(NSLOCTEXT("UnrealEd", "SavingAssetf", "Saving asset {0}"), FText::FromString(Package->GetName()));
+			FText SavingPackageText;
+			if (AssociatedWorld != nullptr)
+			{
+				UE::Core::FVersePath VersePath;
+				if (IAssetTools::Get().ShowingContentVersePath())
+				{
+					VersePath = AssociatedWorld->GetVersePath();
+				}
+				SavingPackageText = FText::Format(NSLOCTEXT("UnrealEd", "SavingMapf", "Saving map {0}"), FText::FromString(VersePath.IsValid() ? MoveTemp(VersePath).ToString() : Package->GetName()));
+			}
+			else
+			{
+				SavingPackageText = FText::Format(NSLOCTEXT("UnrealEd", "SavingAssetf", "Saving asset {0}"), FText::FromString(IAssetTools::Get().GetUserFacingLongPackageName(*Package)));
+			}
 
 			SlowTask.EnterProgressFrame(1, SavingPackageText);
 
@@ -4337,10 +4356,10 @@ FEditorFileUtils::EPromptReturnCode InternalPromptForCheckoutAndSave(const TArra
 	if (!FUncontrolledChangelistsModule::Get().IsEnabled() && (WritablePackageFiles.Num() > 0))
 	{
 		FString WritableFiles;
-		for (UPackage* PackageIter : WritablePackageFiles)
+		for (const FString& UserFacingPath : IAssetTools::Get().GetUserFacingLongPackageNames(WritablePackageFiles))
 		{
 			// A warning message was created.  Try and show it.
-			WritableFiles += FString::Printf(TEXT("\n%s"), *PackageIter->GetName());
+			WritableFiles += FString::Printf(TEXT("\n%s"), *UserFacingPath);
 		}
 
 		const FText WritableFileWarning = FText::Format(NSLOCTEXT("UnrealEd", "Warning_WritablePackagesNotCheckedOut", "The following assets are writable on disk but not checked out from revision control:{0}"),
diff --git a/Engine/Source/Editor/UnrealEd/Private/ObjectTools.cpp b/Engine/Source/Editor/UnrealEd/Private/ObjectTools.cpp
--- a/Engine/Source/Editor/UnrealEd/Private/ObjectTools.cpp
+++ b/Engine/Source/Editor/UnrealEd/Private/ObjectTools.cpp
@@ -1972,10 +1972,9 @@ namespace ObjectTools
 			}
 
 			FString DirtiedPackageNames;
-			for ( TArray<UPackage*>::TConstIterator DirtyPkgIter( DirtiedPackages ); DirtyPkgIter; ++DirtyPkgIter )
+			for (const FString& UserFacingPath : IAssetTools::Get().GetUserFacingLongPackageNames(DirtiedPackages))
 			{
-				const UPackage* CurDirtyPkg = *DirtyPkgIter;
-				DirtiedPackageNames += CurDirtyPkg->GetName() + TEXT("\n");
+				DirtiedPackageNames += UserFacingPath + TEXT("\n");
 			}
 
 			FFormatNamedArguments Arguments;
@@ -4473,7 +4472,6 @@ namespace ObjectTools
 			// Make sure there is no async compilation outstanding in the package we're going to unload
 			UPackageTools::FlushAsyncCompilation( { Object->GetPackage() });
 
-			FString OldObjectFullName = Object->GetFullName();
 			FString OldObjectPathName = Object->GetPathName();
 			GEditor->RenameObject( Object, NewPackage, *ObjName, bLeaveRedirector ? REN_None : REN_DontCreateRedirectors );
 
@@ -4518,11 +4516,11 @@ namespace ObjectTools
 			if(bMoveRedirectorFailed)
 			{
 				ErrorMessage += FText::Format( NSLOCTEXT("UnrealEd", "Error_CouldntRenameObjectRedirectorF", "Couldn't rename '{0}' object because there is an object redirector of the same name, please fixup redirect from editor by enabling Show Redirects in content browser.\n"),
-					FText::FromString(Object->GetFullName()) ).ToString();
+					FText::FromString(IAssetTools::Get().GetUserFacingFullName(*Object)) ).ToString();
 			}
 			else
 			{
-				ErrorMessage += FText::Format( NSLOCTEXT("UnrealEd", "Error_CouldntRenameObjectF", "Couldn't rename '{0}'.\n"), FText::FromString(Object->GetFullName()) ).ToString();
+				ErrorMessage += FText::Format( NSLOCTEXT("UnrealEd", "Error_CouldntRenameObjectF", "Couldn't rename '{0}'.\n"), FText::FromString(IAssetTools::Get().GetUserFacingFullName(*Object)) ).ToString();
 			}
 
 			// @todo asset: Find an appropriate place for localized sounds
diff --git a/Engine/Source/Editor/UnrealEd/Private/PackageTools.cpp b/Engine/Source/Editor/UnrealEd/Private/PackageTools.cpp
--- a/Engine/Source/Editor/UnrealEd/Private/PackageTools.cpp
+++ b/Engine/Source/Editor/UnrealEd/Private/PackageTools.cpp
@@ -22,6 +22,7 @@
 #include "Misc/PackageName.h"
 #include "Editor/Transactor.h"
 #include "Editor/EditorPerProjectUserSettings.h"
+#include "IAssetTools.h"
 #include "ISourceControlOperation.h"
 #include "SourceControlOperations.h"
 #include "ISourceControlProvider.h"
@@ -260,12 +261,12 @@ UPackageTools::UPackageTools(const FObjectInitializer& ObjectInitializer)
 			{	
 				// Ask user to fully load or suppress the message and just fully load
 				if(bSuppress || EAppReturnType::Yes == FMessageDialog::Open( EAppMsgType::YesNo, EAppReturnType::Yes, FText::Format(
-					NSLOCTEXT("UnrealEd", "NeedsToFullyLoadPackageF", "Package {0} is not fully loaded. Do you want to fully load it? Not doing so will abort the '{1}' operation."),
-					FText::FromString(TopLevelPackage->GetName()), OperationText ) ) )
+					NSLOCTEXT("UnrealEd", "NeedsToFullyLoadPackageF", "Asset {0} is not fully loaded. Do you want to fully load it? Not doing so will abort the '{1}' operation."),
+					FText::FromString(IAssetTools::Get().GetUserFacingLongPackageName(*TopLevelPackage)), OperationText)))
 				{
 					// Fully load package.
 					const FScopedBusyCursor BusyCursor;
-					GWarn->BeginSlowTask( NSLOCTEXT("UnrealEd", "FullyLoadingPackages", "Fully loading packages"), true );
+					GWarn->BeginSlowTask( NSLOCTEXT("UnrealEd", "FullyLoadingPackages", "Fully loading assets"), true );
 					TopLevelPackage->FullyLoad();
 					GWarn->EndSlowTask();
 					bNeedsUpdate = true;
@@ -388,9 +389,9 @@ UPackageTools::UPackageTools(const FObjectInitializer& ObjectInitializer)
 		if ( DirtyPackages.Num() > 0 )
 		{
 			FString DirtyPackagesList;
-			for (UPackage* DirtyPackage : DirtyPackages)
+			for (const FString& UserFacingPath : IAssetTools::Get().GetUserFacingLongPackageNames(DirtyPackages))
 			{
-				DirtyPackagesList += FString::Printf(TEXT("\n    %s"), *(DirtyPackage->GetName()));
+				DirtyPackagesList += FString::Printf(TEXT("\n    %s"), *UserFacingPath);
 			}
 
 			FFormatNamedArguments Args;
@@ -467,7 +468,7 @@ UPackageTools::UPackageTools(const FObjectInitializer& ObjectInitializer)
 			int32 PackageIndex = 0;
 			for (UPackage* PackageBeingUnloaded : PackagesToUnload)
 			{
-				GWarn->StatusUpdate(PackageIndex++, PackagesToUnload.Num(), FText::Format(NSLOCTEXT("UnrealEd", "Unloadingf", "Unloading {0}..."), FText::FromString(PackageBeingUnloaded->GetName()) ) );
+				GWarn->StatusUpdate(PackageIndex++, PackagesToUnload.Num(), FText::Format(NSLOCTEXT("UnrealEd", "Unloadingf", "Unloading {0}..."), FText::FromString(IAssetTools::Get().GetUserFacingLongPackageName(*PackageBeingUnloaded)) ) );
 
 				// Flush all pending render commands, as unloading the package may invalidate render resources.
 				FlushRenderingCommands();
@@ -701,16 +702,16 @@ UPackageTools::UPackageTools(const FObjectInitializer& ObjectInitializer)
 		if (InteractionMode == EReloadPackagesInteractionMode::Interactive)
 		{
 			FTextBuilder ReloadDirtyPackagesMsgBuilder;
-			ReloadDirtyPackagesMsgBuilder.AppendLine(NSLOCTEXT("UnrealEd", "ShouldReloadDirtyPackagesHeader", "The following packages have been modified:"));
+			ReloadDirtyPackagesMsgBuilder.AppendLine(NSLOCTEXT("UnrealEd", "ShouldReloadDirtyPackagesHeader", "The following assets have been modified:"));
 			{
 				ReloadDirtyPackagesMsgBuilder.Indent();
-				for (UPackage* DirtyPackage : Filtered.DirtyPackages)
+				for (const FString& UserFacingPath : IAssetTools::Get().GetUserFacingLongPackageNames(Filtered.DirtyPackages))
 				{
-					ReloadDirtyPackagesMsgBuilder.AppendLine(DirtyPackage->GetFName());
+					ReloadDirtyPackagesMsgBuilder.AppendLine(UserFacingPath);
 				}
 				ReloadDirtyPackagesMsgBuilder.Unindent();
 			}
-			ReloadDirtyPackagesMsgBuilder.AppendLine(NSLOCTEXT("UnrealEd", "ShouldReloadDirtyPackagesFooter", "Would you like to reload these packages? This will revert any changes you have made."));
+			ReloadDirtyPackagesMsgBuilder.AppendLine(NSLOCTEXT("UnrealEd", "ShouldReloadDirtyPackagesFooter", "Would you like to reload these assets? This will revert any changes you have made."));
 
 			ReloadDirtyPackagesResult = FMessageDialog::Open(EAppMsgType::YesNo, ReloadDirtyPackagesMsgBuilder.ToText());
 		}
@@ -742,16 +743,16 @@ UPackageTools::UPackageTools(const FObjectInitializer& ObjectInitializer)
 			ErrorMessageBuilder.AppendLine();
 		}
 
-		ErrorMessageBuilder.AppendLine(NSLOCTEXT("UnrealEd", "Error_ReloadDirtyPackagesHeader", "The following packages have been modified and cannot be reloaded:"));
+		ErrorMessageBuilder.AppendLine(NSLOCTEXT("UnrealEd", "Error_ReloadDirtyPackagesHeader", "The following assets have been modified and cannot be reloaded:"));
 		{
 			ErrorMessageBuilder.Indent();
-			for (UPackage* DirtyPackage : Filtered.DirtyPackages)
+			for (const FString& UserFacingPath : IAssetTools::Get().GetUserFacingLongPackageNames(Filtered.DirtyPackages))
 			{
-				ErrorMessageBuilder.AppendLine(DirtyPackage->GetFName());
+				ErrorMessageBuilder.AppendLine(UserFacingPath);
 			}
 			ErrorMessageBuilder.Unindent();
 		}
-		ErrorMessageBuilder.AppendLine(NSLOCTEXT("UnrealEd", "Error_ReloadDirtyPackagesFooter", "Saving these packages will allow them to be reloaded."));
+		ErrorMessageBuilder.AppendLine(NSLOCTEXT("UnrealEd", "Error_ReloadDirtyPackagesFooter", "Saving these assets will allow them to be reloaded."));
 	}
 
 	void InformUsersAboutInMemoryPackages(const FFilteredPackages& Filtered, FTextBuilder& ErrorMessageBuilder)
@@ -764,12 +765,12 @@ UPackageTools::UPackageTools(const FObjectInitializer& ObjectInitializer)
 		{
 			ErrorMessageBuilder.AppendLine();
 		}
-		ErrorMessageBuilder.AppendLine(NSLOCTEXT("UnrealEd", "Error_ReloadInMemoryPackagesHeader", "The following packages are in-memory only and cannot be reloaded:"));
+		ErrorMessageBuilder.AppendLine(NSLOCTEXT("UnrealEd", "Error_ReloadInMemoryPackagesHeader", "The following assets are in-memory only and cannot be reloaded:"));
 		{
 			ErrorMessageBuilder.Indent();
-			for (UPackage* InMemoryPackage : Filtered.InMemoryPackages)
+			for (const FString& UserFacingPath : IAssetTools::Get().GetUserFacingLongPackageNames(Filtered.InMemoryPackages))
 			{
-				ErrorMessageBuilder.AppendLine(InMemoryPackage->GetFName());
+				ErrorMessageBuilder.AppendLine(UserFacingPath);
 			}
 			ErrorMessageBuilder.Unindent();
 		}
@@ -950,12 +951,12 @@ UPackageTools::UPackageTools(const FObjectInitializer& ObjectInitializer)
 					ErrorMessageBuilder.AppendLine();
 				}
 
-				ErrorMessageBuilder.AppendLine(NSLOCTEXT("UnrealEd", "Error_ReloadFailedPackagesHeader", "The following packages failed to reload:"));
+				ErrorMessageBuilder.AppendLine(NSLOCTEXT("UnrealEd", "Error_ReloadFailedPackagesHeader", "The following assets failed to reload:"));
 				{
 					ErrorMessageBuilder.Indent();
-					for (UPackage* FailedPackage : FailedPackages)
+					for (const FString& UserFacingPath : IAssetTools::Get().GetUserFacingLongPackageNames(FailedPackages))
 					{
-						ErrorMessageBuilder.AppendLine(FailedPackage->GetFName());
+						ErrorMessageBuilder.AppendLine(UserFacingPath);
 					}
 					ErrorMessageBuilder.Unindent();
 				}
diff --git a/Engine/Source/Runtime/Engine/Private/LevelInstance/LevelInstanceSubsystem.cpp b/Engine/Source/Runtime/Engine/Private/LevelInstance/LevelInstanceSubsystem.cpp
--- a/Engine/Source/Runtime/Engine/Private/LevelInstance/LevelInstanceSubsystem.cpp
+++ b/Engine/Source/Runtime/Engine/Private/LevelInstance/LevelInstanceSubsystem.cpp
@@ -15,6 +15,8 @@
 #include "GameFramework/ActorPrimitiveColorHandler.h"
 
 #if WITH_EDITOR
+#include "AssetToolsModule.h"
+#include "IAssetTools.h"
 #include "Settings/LevelEditorMiscSettings.h"
 #include "LevelInstance/LevelInstanceEditorLevelStreaming.h"
 #include "LevelInstance/LevelInstanceEditorPropertyOverrideLevelStreaming.h"
@@ -340,7 +342,9 @@ void ULevelInstanceSubsystem::UpdateStreamingStateInternal()
 		for (const TPair<ILevelInstanceInterface*, bool>& Pair : LevelInstancesToLoadOrUpdateCopy)
 		{
 #if WITH_EDITOR
-			SlowTask.EnterProgressFrame(1.f, FText::Format(LOCTEXT("LoadingLevelInstance", "Loading Level Instance {0}"), FText::FromString(Pair.Key->GetWorldAsset().ToString())));
+			SlowTask.EnterProgressFrame(1.f, FText::Format(
+				LOCTEXT("LoadingLevelInstance", "Loading Level Instance {0}"),
+				FText::FromString(FAssetToolsModule::GetModule().Get().GetUserFacingObjectPath(Pair.Key->GetWorldAsset().ToSoftObjectPath()))));
 #endif
 			LoadLevelInstance(Pair.Key);
 		}
