AxisDisplayInfo adds specific DisplayName vs ToolTip accessors
* Axis (ie. translations) changed to Left, Up, Forward and the (was X, was -Y, was Z) moved into the tooltip
* Follow up to change usages in various areas of editor

diff --git a/Engine/Source/Editor/UnrealEd/Private/EditorAxisDisplayInfo.cpp b/Engine/Source/Editor/UnrealEd/Private/EditorAxisDisplayInfo.cpp
--- a/Engine/Source/Editor/UnrealEd/Private/EditorAxisDisplayInfo.cpp
+++ b/Engine/Source/Editor/UnrealEd/Private/EditorAxisDisplayInfo.cpp
@@ -48,50 +48,70 @@ EAxisList::Type FEditorAxisDisplayInfo::GetAxisDisplayCoordinateSystem() const
 	return AxisDisplayCoordinateSystem.Get(EAxisList::XYZ/*DefaultValue*/);
 }
 
+FText FEditorAxisDisplayInfo::GetAxisToolTip(EAxisList::Type Axis) const
+{
+	static const FText XToolTip = LOCTEXT("XDisplayName", "X");
+	static const FText YToolTip = LOCTEXT("YDisplayName", "Y");
+	static const FText ZToolTip = LOCTEXT("ZDisplayName", "Z");
+	static const FText LeftToolTip = LOCTEXT("LeftToolTip", "Left (was -Y)");
+	static const FText UpDisplayName = LOCTEXT("UpToolTip", "Up (was Z)");
+	static const FText ForwardDisplayName = LOCTEXT("ForwardToolTip", "Forward (was X)");
+
+	Axis = MapAxis(Axis);
+	
+	if (Axis == EAxisList::X)
+	{
+		return XToolTip;
+	}
+	else if (Axis == EAxisList::Y)
+	{
+		return YToolTip;
+	}
+	else if (Axis == EAxisList::Z)
+	{
+		return ZToolTip;
+	}
+	else if (Axis == EAxisList::Left)
+	{
+		return LeftToolTip;
+	}
+	else if (Axis == EAxisList::Up)
+	{
+		return UpDisplayName;
+	}
+	else if (Axis == EAxisList::Forward)
+	{
+		return ForwardDisplayName;
+	}
+
+	ensureMsgf(false, TEXT("Unsupported Axis: %d"), Axis);
+
+	return LOCTEXT("UnsupportedDisplayName", "Unsupported");
+	
+}
+
 FText FEditorAxisDisplayInfo::GetAxisDisplayName(EAxisList::Type Axis)
 {
 	static const FText XDisplayName = LOCTEXT("XDisplayName", "X");
 	static const FText YDisplayName = LOCTEXT("YDisplayName", "Y");
 	static const FText ZDisplayName = LOCTEXT("ZDisplayName", "Z");
-	static const FText RightDisplayName = LOCTEXT("RightDisplayName", "Right (was Y)");
-	static const FText LeftDisplayName = LOCTEXT("LeftDisplayName", "Left (was -Y)");
-	static const FText UpDisplayName = LOCTEXT("UpDisplayName", "Up (was Z)");
-	static const FText ForwardDisplayName = LOCTEXT("ForwardDisplayName", "Forward (was X)");
+	static const FText LeftDisplayName = LOCTEXT("LeftDisplayName", "Left");
+	static const FText UpDisplayName = LOCTEXT("UpDisplayName", "Up");
+	static const FText ForwardDisplayName = LOCTEXT("ForwardDisplayName", "Forward");
 	
 	Axis = MapAxis(Axis);
 
 	if (Axis == EAxisList::X)
 	{
-		if (UseForwardRightUpDisplayNames())
-		{
-			return ForwardDisplayName;
-		}
-		else
-		{
-			return XDisplayName;
-		}
+		return XDisplayName;
 	}
 	else if (Axis == EAxisList::Y)
 	{
-		if (UseForwardRightUpDisplayNames())
-		{
-			return RightDisplayName;
-		}
-		else
-		{
-			return YDisplayName;
-		}
+		return YDisplayName;
 	}
 	else if (Axis == EAxisList::Z)
 	{
-		if (UseForwardRightUpDisplayNames())
-		{
-			return UpDisplayName;
-		}
-		else
-		{
-			return ZDisplayName;
-		}
+		return ZDisplayName;
 	}
 	else if (Axis == EAxisList::Left)
 	{
@@ -116,7 +136,6 @@ FText FEditorAxisDisplayInfo::GetAxisDisplayNameShort(EAxisList::Type Axis)
 	static const FText XDisplayNameShort = LOCTEXT("XDisplayNameShort", "X");
 	static const FText YDisplayNameShort = LOCTEXT("YDisplayNameShort", "Y");
 	static const FText ZDisplayNameShort = LOCTEXT("ZDisplayNameShort", "Z");
-	static const FText RightDisplayNameShort = LOCTEXT("RightDisplayNameShort", "R");
 	static const FText LeftDisplayNameShort = LOCTEXT("LeftDisplayNameShort", "L");
 	static const FText UpDisplayNameShort = LOCTEXT("UpDisplayNameShort", "U");
 	static const FText ForwardDisplayNameShort = LOCTEXT("ForwardDisplayNameShort", "F");
@@ -125,36 +144,15 @@ FText FEditorAxisDisplayInfo::GetAxisDisplayNameShort(EAxisList::Type Axis)
 
 	if (Axis == EAxisList::X)
 	{
-		if (UseForwardRightUpDisplayNames())
-		{
-			return ForwardDisplayNameShort;
-		}
-		else
-		{
-			return XDisplayNameShort;
-		}
+		return XDisplayNameShort;
 	}
 	else if (Axis == EAxisList::Y)
 	{
-		if (UseForwardRightUpDisplayNames())
-		{
-			return RightDisplayNameShort;
-		}
-		else
-		{
-			return YDisplayNameShort;
-		}
+		return YDisplayNameShort;
 	}
 	else if (Axis == EAxisList::Z)
 	{
-		if (UseForwardRightUpDisplayNames())
-		{
-			return UpDisplayNameShort;
-		}
-		else
-		{
-			return ZDisplayNameShort;
-		}
+		return ZDisplayNameShort;
 	}
 	else if (Axis == EAxisList::Left)
 	{
@@ -219,9 +217,14 @@ bool FEditorAxisDisplayInfo::UseForwardRightUpDisplayNames()
 	return bUseForwardRightUpDisplayNames.Get(false/*DefaultValue*/);
 }
 
+FText FEditorAxisDisplayInfo::GetRotationAxisToolTip(EAxisList::Type Axis) const
+{
+	return const_cast<FEditorAxisDisplayInfo*>(this)->GetRotationAxisNameShort(Axis);
+}
+
 FText FEditorAxisDisplayInfo::GetRotationAxisName(EAxisList::Type Axis)
 {
-	return FEditorAxisDisplayInfo::GetRotationAxisNameShort(Axis);
+	return GetRotationAxisNameShort(Axis);
 }
 
 FText FEditorAxisDisplayInfo::GetRotationAxisNameShort(EAxisList::Type Axis)
@@ -258,7 +261,7 @@ FText FEditorAxisDisplayInfo::GetRotationAxisNameShort(EAxisList::Type Axis)
 	return LOCTEXT("UnsupportedDisplayNameShort", "?");
 }
 
-EAxisList::Type FEditorAxisDisplayInfo::MapAxis(EAxisList::Type Axis)
+EAxisList::Type FEditorAxisDisplayInfo::MapAxis(EAxisList::Type Axis) const
 {
 	if (GetAxisDisplayCoordinateSystem() == EAxisList::XYZ)
 	{
diff --git a/Engine/Source/Editor/UnrealEd/Public/EditorAxisDisplayInfo.h b/Engine/Source/Editor/UnrealEd/Public/EditorAxisDisplayInfo.h
--- a/Engine/Source/Editor/UnrealEd/Public/EditorAxisDisplayInfo.h
+++ b/Engine/Source/Editor/UnrealEd/Public/EditorAxisDisplayInfo.h
@@ -12,15 +12,17 @@ public:
 	virtual ~FEditorAxisDisplayInfo() = default;
 
 	virtual EAxisList::Type GetAxisDisplayCoordinateSystem() const override;
-
+	
+	virtual FText GetAxisToolTip(EAxisList::Type Axis) const override;
 	virtual FText GetAxisDisplayName(EAxisList::Type Axis) override;
-
 	virtual FText GetAxisDisplayNameShort(EAxisList::Type Axis) override;
+	
 
 	virtual FLinearColor GetAxisColor(EAxisList::Type Axis) override;
 
 	virtual bool UseForwardRightUpDisplayNames() override;
 
+	virtual FText GetRotationAxisToolTip(EAxisList::Type Axis) const override;
 	virtual FText GetRotationAxisName(EAxisList::Type Axis) override;
 	virtual FText GetRotationAxisNameShort(EAxisList::Type Axis) override;
 	
@@ -28,7 +30,7 @@ public:
 
 private:
 	// Maps the given axis from FLU -> XYZ if the AxisDisplayCoordinateSystem is XYZ
-	EAxisList::Type MapAxis(EAxisList::Type Axis);
+	EAxisList::Type MapAxis(EAxisList::Type Axis) const;
 
 	// Inits info stored in settings, like the axis colors
 	void InitSettingsInfo(double EditorStartupTime);
diff --git a/Engine/Source/Runtime/Core/Public/Misc/AxisDisplayInfo.h b/Engine/Source/Runtime/Core/Public/Misc/AxisDisplayInfo.h
--- a/Engine/Source/Runtime/Core/Public/Misc/AxisDisplayInfo.h
+++ b/Engine/Source/Runtime/Core/Public/Misc/AxisDisplayInfo.h
@@ -94,6 +94,11 @@ public:
 	 */
 	virtual EAxisList::Type GetAxisDisplayCoordinateSystem() const = 0;
 
+	/**
+	 * Gets the tool tip text used for axis components of widgets describing translations and scales
+	 */
+	virtual FText GetAxisToolTip(EAxisList::Type Axis) const = 0;
+
 	/**
 	 * Gets the display name to use for the given axis.
 	 *
@@ -118,7 +123,19 @@ public:
 	 */
 	virtual FLinearColor GetAxisColor(EAxisList::Type Axis) = 0;
 
+	/**
+	 * Gets the tool tip text used for axis components of widgets describing rotations
+	 */
+	virtual FText GetRotationAxisToolTip(EAxisList::Type Axis) const = 0;
+	/**
+	 * Gets the display name use for axis components of widgets describing rotations
+	 * Typically this would be used as a label or title
+	 */
 	virtual FText GetRotationAxisName(EAxisList::Type Axis) = 0;
+	/**
+	 * Gets a short display name use for axis components of widgets describing rotations
+	 * Typically this would be used as a label or title where little room is available
+	 */
 	virtual FText GetRotationAxisNameShort(EAxisList::Type Axis) = 0;
 
 	/**
