Fix static analysis from previous check in

#rnx

diff --git a/Engine/Source/Editor/UnrealEd/Private/Kismet2/KismetReinstanceUtilities.cpp b/Engine/Source/Editor/UnrealEd/Private/Kismet2/KismetReinstanceUtilities.cpp
--- a/Engine/Source/Editor/UnrealEd/Private/Kismet2/KismetReinstanceUtilities.cpp
+++ b/Engine/Source/Editor/UnrealEd/Private/Kismet2/KismetReinstanceUtilities.cpp
@@ -3298,6 +3298,7 @@ bool FBlueprintCompileReinstancer::PrePopulateInstancingGraphForArchetype(FObjec
 
 	// Find the root of the instantiation
 	UObject* InstancingGraphRoot = OldInstance;
+	do
 	{
 		if (UObject** NewInstancingGraphRoot = OldToNewInstanceMap ? OldToNewInstanceMap->Find(InstancingGraphRoot) : nullptr)
 		{
