Fix inverted ortho pan bug caused by removal of L axis negation

diff --git a/Engine/Source/Editor/UnrealEd/Private/EditorViewportClient.cpp b/Engine/Source/Editor/UnrealEd/Private/EditorViewportClient.cpp
--- a/Engine/Source/Editor/UnrealEd/Private/EditorViewportClient.cpp
+++ b/Engine/Source/Editor/UnrealEd/Private/EditorViewportClient.cpp
@@ -4823,11 +4823,6 @@ FVector FEditorViewportClient::TranslateDelta( FKey InKey, float InDelta, bool I
 					case LVT_Perspective:
 						break;
 					}
-
-					if (IsOrtho() && AxisDisplayInfo::GetAxisDisplayCoordinateSystem() == EAxisList::LeftUpForward)
-					{
-						vec.Y = -vec.Y;
-					}
 				}
 			}
 		}
