Minor parallelization and allocation optimizations in UCookOnTheFlyServer.

#rnx

diff --git a/Engine/Source/Editor/UnrealEd/Private/CookOnTheFlyServer.cpp b/Engine/Source/Editor/UnrealEd/Private/CookOnTheFlyServer.cpp
--- a/Engine/Source/Editor/UnrealEd/Private/CookOnTheFlyServer.cpp
+++ b/Engine/Source/Editor/UnrealEd/Private/CookOnTheFlyServer.cpp
@@ -9728,11 +9728,13 @@ void UCookOnTheFlyServer::CookByTheBookFinishedInternal()
 				}
 			}
 
+			CookedPackageNames.Reserve(CookedPackageDatas.Num());
 			for (FPackageData* PackageData : CookedPackageDatas)
 			{
 				CookedPackageNames.Add(PackageData->GetPackageName());
 			}
 
+			IgnorePackageNames.Reserve(IgnorePackageDatas.Num());
 			for (FPackageData* PackageData : IgnorePackageDatas)
 			{
 				IgnorePackageNames.Add(PackageData->GetPackageName());
@@ -12523,13 +12525,21 @@ void UCookOnTheFlyServer::GenerateLocalizationReferences()
 
 	UE_LOG(LogCook, Display, TEXT("Found %d localized assets"), AssetDataForCultures.Num());
 
-	for (const FAssetData& AssetData : AssetDataForCultures)
+	TArray<TPair<FName, FName>> FlatData;
+	FlatData.SetNum(AssetDataForCultures.Num());
+
+	ParallelFor(AssetDataForCultures.Num(), [&AssetDataForCultures, &FlatData](int32 Index)
 	{
+		const FAssetData& AssetData = AssetDataForCultures[Index];
 		const FName LocalizedPackageName = AssetData.PackageName;
 		const FName SourcePackageName = *FPackageName::GetSourcePackagePath(LocalizedPackageName.ToString());
-
-		TArray<FName>& LocalizedPackageNames = CookByTheBookOptions->SourceToLocalizedPackageVariants.FindOrAdd(SourcePackageName);
-		LocalizedPackageNames.AddUnique(LocalizedPackageName);
+		FlatData[Index] = {SourcePackageName, LocalizedPackageName};
+	});
+	
+	for (const TPair<FName, FName>& Pair: FlatData)
+	{
+		TArray<FName>& LocalizedPackageNames = CookByTheBookOptions->SourceToLocalizedPackageVariants.FindOrAdd(Pair.Get<0>());
+		LocalizedPackageNames.AddUnique(Pair.Get<1>());
 	}
 }
 
