Support the fact that archetype can be null in the ReplaceObjectHelper

diff --git a/Engine/Source/Editor/UnrealEd/Private/Kismet2/KismetReinstanceUtilities.cpp b/Engine/Source/Editor/UnrealEd/Private/Kismet2/KismetReinstanceUtilities.cpp
--- a/Engine/Source/Editor/UnrealEd/Private/Kismet2/KismetReinstanceUtilities.cpp
+++ b/Engine/Source/Editor/UnrealEd/Private/Kismet2/KismetReinstanceUtilities.cpp
@@ -3280,12 +3280,12 @@ void FBlueprintCompileReinstancer::PreCreateSubObjectsForReinstantiation(
 	}
 }
 
-bool FBlueprintCompileReinstancer::PrePopulateInstancingGraphForArchetype(FObjectInstancingGraph& InstancingGraph, TNotNull<UObject*> OldInstance, TNotNull<UObject*> InArchetype, TMap<UObject*, UObject*>* OldToNewInstanceMap)
+bool FBlueprintCompileReinstancer::PrePopulateInstancingGraphForArchetype(FObjectInstancingGraph& InstancingGraph, TNotNull<UObject*> OldInstance, UObject* InArchetype, TMap<UObject*, UObject*>* OldToNewInstanceMap)
 {
 	checkf(!InstancingGraph.HasDestinationRoot(), TEXT("Can only pre-populate an empty InstancingGraph"));
 
 	// Check if we are re-instancing from an archetype, otherwise nothing to do
-	if (InArchetype->HasAnyFlags(RF_ClassDefaultObject))
+	if (!InArchetype || InArchetype->HasAnyFlags(RF_ClassDefaultObject))
 	{
 		return false;
 	}
diff --git a/Engine/Source/Editor/UnrealEd/Public/Kismet2/KismetReinstanceUtilities.h b/Engine/Source/Editor/UnrealEd/Public/Kismet2/KismetReinstanceUtilities.h
--- a/Engine/Source/Editor/UnrealEd/Public/Kismet2/KismetReinstanceUtilities.h
+++ b/Engine/Source/Editor/UnrealEd/Public/Kismet2/KismetReinstanceUtilities.h
@@ -342,7 +342,7 @@ private:
 	 * @param OldToNewInstanceMap Possible instance mapping of instances
 	 * @return false if the instancing graph is not needed
 	 */
-	static UNREALED_API bool PrePopulateInstancingGraphForArchetype(FObjectInstancingGraph& InstancingGraph, TNotNull<UObject*> OldInstance, TNotNull<UObject*> InArchetype, TMap<UObject*, UObject*>* OldToNewInstanceMap = nullptr);
+	static UNREALED_API bool PrePopulateInstancingGraphForArchetype(FObjectInstancingGraph& InstancingGraph, TNotNull<UObject*> OldInstance, UObject* InArchetype, TMap<UObject*, UObject*>* OldToNewInstanceMap = nullptr);
 
 
 	/** Handles the sub object pre-creation recursively */
