EditorDomain: Disable save/load before AssetRegistryGather has finished. It was already a performance problem, and now we have discovered that it deadlocks because of a bug in the AssetRegistry. Keep it disabled at least until UE-281930 is fixed.
#rnx

diff --git a/Engine/Source/Editor/UnrealEd/Private/EditorDomain/EditorDomain.cpp b/Engine/Source/Editor/UnrealEd/Private/EditorDomain/EditorDomain.cpp
--- a/Engine/Source/Editor/UnrealEd/Private/EditorDomain/EditorDomain.cpp
+++ b/Engine/Source/Editor/UnrealEd/Private/EditorDomain/EditorDomain.cpp
@@ -129,6 +129,10 @@ FEditorDomain::FEditorDomain(EEditorDomainEnabled EnableLevel)
 	Locks = TRefCountPtr<FLocks>(new FLocks(*this));
 	Workspace.Reset(MakePackageResourceManagerFile());
 	GConfig->GetBool(TEXT("CookSettings"), TEXT("EditorDomainExternalSave"), bExternalSave, GEditorIni);
+	// Load / Save from EditorDomain before AssetRegistry has finished its gather can cause performance problems due to
+	// blocking on AssetRegistry, and can also cause deadlocks with the AssetRegistry's gather thread, UE-281930.
+	// Disable it until we have fixed those problems.
+	bEnabledBeforeAssetRegistryGather = false;
 	if (bExternalSave)
 	{
 		SaveClient.Reset(new FEditorDomainSaveClient());
@@ -165,12 +169,24 @@ FEditorDomain::FEditorDomain(EEditorDomainEnabled EnableLevel)
 		}
 		FCoreUObjectDelegates::OnEndLoadPackage.AddRaw(this, &FEditorDomain::OnEndLoadPackage);
 	}
+	if (AssetRegistry->IsLoadingAssets() || !AssetRegistry->IsSearchAllAssets())
+	{
+		AssetRegistry->OnFilesLoaded().AddRaw(this, &FEditorDomain::OnAssetRegistryInitialScanComplete);
+	}
+	else
+	{
+		bAssetRegistryInitialScanComplete.store(true, std::memory_order_relaxed);
+	}
 	UPackage::PackageSavedWithContextEvent.AddRaw(this, &FEditorDomain::OnPackageSavedWithContext);
 	AssetRegistry->OnAssetUpdatedOnDisk().AddRaw(this, &FEditorDomain::OnAssetUpdatedOnDisk);
 }
 
 FEditorDomain::~FEditorDomain()
 {
+	// We have a pointer to AssetRegistry, but it might be dangling due to the AssetRegistry module already shutting down.
+	// Refetch it from the singleton pointer, and handle the case that it is now null.
+	AssetRegistry = IAssetRegistry::Get();
+
 	TUniquePtr<UE::DerivedData::FRequestOwner> LocalBatchDownloadOwner;
 	{
 		FScopeLock ScopeLock(&Locks->Lock);
@@ -185,6 +201,10 @@ FEditorDomain::~FEditorDomain()
 	UPackage::PackageSavedWithContextEvent.RemoveAll(this);
 	FCoreUObjectDelegates::OnEndLoadPackage.RemoveAll(this);
 	FCoreDelegates::OnPostEngineInit.RemoveAll(this);
+	if (AssetRegistry)
+	{
+		AssetRegistry->OnFilesLoaded().RemoveAll(this);
+	}
 	Locks->Owner = nullptr;
 	AssetRegistry = nullptr;
 	Workspace.Reset();
@@ -392,7 +412,8 @@ int64 FEditorDomain::FileSize(const FPackagePath& PackagePath, FBulkDataCookedIn
 
 	checkf(CookedIndex.IsDefault(), TEXT("Editor Domain does not support bulkdata cooked indices"));
 
-	if (PackageSegment != EPackageSegment::Header)
+	if (PackageSegment != EPackageSegment::Header
+		|| (!bEnabledBeforeAssetRegistryGather && !IsAssetRegistryInitialScanComplete()))
 	{
 		return Workspace->FileSize(PackagePath, FBulkDataCookedIndex::Default, PackageSegment, OutUpdatedPath);
 	}
@@ -484,7 +505,8 @@ FOpenPackageResult FEditorDomain::OpenReadPackage(const FPackagePath& PackagePat
 	checkf(CookedIndex.IsDefault(), TEXT("Editor Domain does not support bulkdata cooked indices"));
 
 	FScopeLock ScopeLock(&Locks->Lock);
-	if (PackageSegment != EPackageSegment::Header)
+	if (PackageSegment != EPackageSegment::Header
+		|| (!bEnabledBeforeAssetRegistryGather && !IsAssetRegistryInitialScanComplete()))
 	{
 		return Workspace->OpenReadPackage(PackagePath, CookedIndex, PackageSegment, OutUpdatedPath);
 	}
@@ -543,7 +565,8 @@ FOpenAsyncPackageResult FEditorDomain::OpenAsyncReadPackage(const FPackagePath&
 	checkf(CookedIndex.IsDefault(), TEXT("Editor Domain does not support bulkdata cooked indices"));
 
 	FScopeLock ScopeLock(&Locks->Lock);
-	if (PackageSegment != EPackageSegment::Header)
+	if (PackageSegment != EPackageSegment::Header
+		|| (!bEnabledBeforeAssetRegistryGather && !IsAssetRegistryInitialScanComplete()))
 	{
 		return Workspace->OpenAsyncReadPackage(PackagePath, CookedIndex, PackageSegment);
 	}
@@ -657,6 +680,11 @@ void FEditorDomain::OnEndLoadPackage(const FEndLoadPackageContext& Context)
 	{
 		return;
 	}
+	if (!bEnabledBeforeAssetRegistryGather && !IsAssetRegistryInitialScanComplete())
+	{
+		return;
+	}
+
 	TArray<UPackage*> PackagesToSave;
 	{
 		FScopeLock ScopeLock(&Locks->Lock);
@@ -689,6 +717,11 @@ void FEditorDomain::OnPostEngineInit()
 		}
 	}
 
+	if (!bEnabledBeforeAssetRegistryGather && !IsAssetRegistryInitialScanComplete())
+	{
+		return;
+	}
+
 	TArray<UPackage*> PackagesToSave;
 	FString PackageName;
 	for (TObjectIterator<UPackage> It; It; ++It)
@@ -712,6 +745,16 @@ void FEditorDomain::OnPostEngineInit()
 	}
 }
 
+void FEditorDomain::OnAssetRegistryInitialScanComplete()
+{
+	bAssetRegistryInitialScanComplete.store(true, std::memory_order_relaxed);
+}
+
+bool FEditorDomain::IsAssetRegistryInitialScanComplete() const
+{
+	return bAssetRegistryInitialScanComplete.load(std::memory_order_relaxed);
+}
+
 void FEditorDomain::FilterKeepPackagesToSave(TArray<UPackage*>& InOutPackagesToSave)
 {
 	FPackagePath PackagePath;
diff --git a/Engine/Source/Editor/UnrealEd/Private/EditorDomain/EditorDomain.h b/Engine/Source/Editor/UnrealEd/Private/EditorDomain/EditorDomain.h
--- a/Engine/Source/Editor/UnrealEd/Private/EditorDomain/EditorDomain.h
+++ b/Engine/Source/Editor/UnrealEd/Private/EditorDomain/EditorDomain.h
@@ -23,6 +23,8 @@
 #include "UObject/PackageResourceManager.h"
 #include "UObject/UnrealNames.h"
 
+#include <atomic>
+
 class FArchive;
 class FEditorDomainSaveClient;
 class FObjectPostSaveContext;
@@ -247,6 +249,10 @@ private:
 	void MarkLoadedFromEditorDomain(const FPackagePath& PackagePath, TRefCountPtr<FPackageSource>& PackageSource);
 	/** Callback for PostEngineInit, to handle saving of packages which we could not save before then. */
 	void OnPostEngineInit();
+	/** Callback for AssetRegistry's InitialScanComplete, to enable save/load if we disable before initial scan. */
+	void OnAssetRegistryInitialScanComplete();
+	/** Report whether we have set our flag from OnAssetRegistryInitialScanComplete. */
+	bool IsAssetRegistryInitialScanComplete() const;
 	/** EndLoad callback to handle saving the EditorDomain version of the package. */
 	void OnEndLoadPackage(const FEndLoadPackageContext& Context);
 	/** For each of the now-loaded packages, if we had to load from workspace domain, save into the editor domain. */
@@ -281,6 +287,10 @@ private:
 	bool bExternalSave = false;
 	/** Marker for whether our PostEngineInit callback has been called */
 	bool bHasPassedPostEngineInit = false;
+	/** Marker for whether AssetRegistry initial scan has finished. */
+	std::atomic<bool> bAssetRegistryInitialScanComplete{ false };
+	/** Configuration setting for whether save/load is allowed before the AssetRegistry initial scan has finished. */
+	bool bEnabledBeforeAssetRegistryGather = true;
 	/**
 	 * Configuration value for whether saves into the EditorDomain should be skipped for packages loaded before
 	 * the DDC server reports whether they exist in the upstream DDC.
