Scope Reenable GarbageElimination upon Actor/Component deletion.

#rnx

diff --git a/Engine/Source/Editor/UnrealEd/Private/EditorActor.cpp b/Engine/Source/Editor/UnrealEd/Private/EditorActor.cpp
--- a/Engine/Source/Editor/UnrealEd/Private/EditorActor.cpp
+++ b/Engine/Source/Editor/UnrealEd/Private/EditorActor.cpp
@@ -87,6 +87,22 @@
 
 DEFINE_LOG_CATEGORY_STATIC(LogEditorActor, Log, All);
 
+struct FGarbageEliminationScopeGuard
+{
+	bool PreviousValue;
+
+	FGarbageEliminationScopeGuard(bool bNewValue)
+		: PreviousValue(UObjectBaseUtility::IsGarbageEliminationEnabled())
+	{
+		UObjectBaseUtility::SetGarbageEliminationEnabled(bNewValue);
+	}
+
+	~FGarbageEliminationScopeGuard()
+	{
+		UObjectBaseUtility::SetGarbageEliminationEnabled(PreviousValue);
+	}
+};
+
 static int32 RecomputePoly( ABrush* InOwner, FPoly* Poly )
 {
 	// force recalculation of normal, and texture U and V coordinates in FPoly::Finalize()
@@ -834,7 +850,10 @@ bool UUnrealEdEngine::DeleteComponents(const TArray<UActorComponent*>& InCompone
 		LevelEditor.BroadcastComponentsEdited();
 
 		// Remove all references to destroyed components once at the end, instead of once for each component destroyed
-		CollectGarbage(GARBAGE_COLLECTION_KEEPFLAGS);
+		{
+			FGarbageEliminationScopeGuard Guard(true);
+			CollectGarbage(GARBAGE_COLLECTION_KEEPFLAGS);
+		}
 	}
 
 	// Update the editor component selection if possible
@@ -1239,7 +1258,10 @@ bool UUnrealEdEngine::DeleteActors(const TArray<AActor*>& InActorsToDelete, UWor
 	}
 
 	// Remove all references to destroyed actors once at the end, instead of once for each Actor destroyed
-	CollectGarbage(GARBAGE_COLLECTION_KEEPFLAGS);
+	{
+		FGarbageEliminationScopeGuard Guard(true);
+		CollectGarbage(GARBAGE_COLLECTION_KEEPFLAGS);
+	}
 
 	LegacySyncBatch.Reset();
 	NoteSelectionChange();
