[WorldPartition] [PCG] Fix regression where PCG Builder's 'bLoadActivatedRuntimeDataLayers' wasn't properly mapped to the WorldPartitionBuilder arguments
- Add a bIncludeInitiallyActivatedRuntimeDataLayers in the base Builder class
- Set it from PCG instead of bExcludeNonInitiallyActivatedRuntimeDataLayers

diff --git a/Engine/Plugins/PCG/Source/PCGEditor/Private/WorldPartitionBuilder/PCGWorldPartitionBuilder.cpp b/Engine/Plugins/PCG/Source/PCGEditor/Private/WorldPartitionBuilder/PCGWorldPartitionBuilder.cpp
--- a/Engine/Plugins/PCG/Source/PCGEditor/Private/WorldPartitionBuilder/PCGWorldPartitionBuilder.cpp
+++ b/Engine/Plugins/PCG/Source/PCGEditor/Private/WorldPartitionBuilder/PCGWorldPartitionBuilder.cpp
@@ -548,7 +548,7 @@ bool UPCGWorldPartitionBuilder::PreRun(UWorld* World, FPackageSourceControlHelpe
 	// Setup DataLayers
 	// Set Base Builder Class fields
 	bLoadNonDynamicDataLayers = Args.bLoadEditorOnlyDataLayers;
-	bExcludeNonInitiallyActivatedRuntimeDataLayer = Args.bLoadActivatedRuntimeDataLayers;
+	bIncludeInitiallyActivatedRuntimeDataLayers = Args.bLoadActivatedRuntimeDataLayers;
 
 	IncludedDataLayers.Append(Args.IncludedDataLayers);
 	ExcludedDataLayers.Append(Args.ExcludedDataLayers);
diff --git a/Engine/Source/Editor/UnrealEd/Private/WorldPartition/WorldPartitionBuilder.cpp b/Engine/Source/Editor/UnrealEd/Private/WorldPartition/WorldPartitionBuilder.cpp
--- a/Engine/Source/Editor/UnrealEd/Private/WorldPartition/WorldPartitionBuilder.cpp
+++ b/Engine/Source/Editor/UnrealEd/Private/WorldPartition/WorldPartitionBuilder.cpp
@@ -332,10 +332,11 @@ void UWorldPartitionBuilder::LoadDataLayers(UWorld* InWorld)
 		{
 			// Forced inclusion
 			const bool bForceIncludeEditorDataLayer = bLoadNonDynamicDataLayers && !DataLayer->IsRuntime();
-			const bool bIsForceIncluded = IsForceIncludedDataLayer(DataLayer) || bForceIncludeEditorDataLayer;
+			const bool bForceIncludeInitiallyActivatedRuntimeDataLayer = bIncludeInitiallyActivatedRuntimeDataLayers && DataLayer->IsRuntime() && DataLayer->GetInitialRuntimeState() == EDataLayerRuntimeState::Activated;
+			const bool bIsForceIncluded = IsForceIncludedDataLayer(DataLayer) || bForceIncludeEditorDataLayer || bForceIncludeInitiallyActivatedRuntimeDataLayer;
 
 			// Forced exclusion
-			const bool bForceExcludeNonInitiallyActivatedRuntimeDataLayer = (bExcludeNonInitiallyActivatedRuntimeDataLayer && DataLayer->IsRuntime()) ? DataLayer->GetInitialRuntimeState() != EDataLayerRuntimeState::Activated : false;
+			const bool bForceExcludeNonInitiallyActivatedRuntimeDataLayer = (bExcludeNonInitiallyActivatedRuntimeDataLayers && DataLayer->IsRuntime()) ? DataLayer->GetInitialRuntimeState() != EDataLayerRuntimeState::Activated : false;
 			const bool bIsForceExcluded = IsForceExcludedDataLayer(DataLayer) || bForceExcludeNonInitiallyActivatedRuntimeDataLayer;
 
 			// Default state
diff --git a/Engine/Source/Editor/UnrealEd/Public/WorldPartition/WorldPartitionBuilder.h b/Engine/Source/Editor/UnrealEd/Public/WorldPartition/WorldPartitionBuilder.h
--- a/Engine/Source/Editor/UnrealEd/Public/WorldPartition/WorldPartitionBuilder.h
+++ b/Engine/Source/Editor/UnrealEd/Public/WorldPartition/WorldPartitionBuilder.h
@@ -159,9 +159,10 @@ protected:
 	TSet<TObjectPtr<UDataLayerAsset>> ExcludedDataLayers;
 
 	bool bLoadNonDynamicDataLayers = true;
-	bool bExcludeNonInitiallyActivatedRuntimeDataLayer = false;
+	bool bIncludeInitiallyActivatedRuntimeDataLayers = false;
+	bool bExcludeNonInitiallyActivatedRuntimeDataLayers = false;
 
-	UE_DEPRECATED(5.6, "This has been renamed to bExcludeNonInitiallyActivatedRuntimeDataLayer")
+	UE_DEPRECATED(5.6, "This has been renamed to bExcludeNonInitiallyActivatedRuntimeDataLayers")
 	bool bLoadInitiallyActiveDataLayers = true;
 
 	FModifiedFilesHandler ModifiedFilesHandler;
