Trivial change of "check" to "ensure" in UnrealEdMisc.cpp.  Follow up to review comment.

#rnx

diff --git a/Engine/Source/Editor/UnrealEd/Private/UnrealEdMisc.cpp b/Engine/Source/Editor/UnrealEd/Private/UnrealEdMisc.cpp
--- a/Engine/Source/Editor/UnrealEd/Private/UnrealEdMisc.cpp
+++ b/Engine/Source/Editor/UnrealEd/Private/UnrealEdMisc.cpp
@@ -1643,13 +1643,13 @@ void FUnrealEdMisc::SelectActorFromMessageToken(AActor* InActor)
 
 UNREALED_API void FUnrealEdMisc::OnMultiSelectActorFromMessageTokenBegin()
 {
-	check(bActorMultiSelect == false);
+	ensure(bActorMultiSelect == false);
 	bActorMultiSelect = true;
 }
 
 UNREALED_API void FUnrealEdMisc::OnMultiSelectActorFromMessageTokenEnd()
 {
-	check(bActorMultiSelect == true);
+	ensure(bActorMultiSelect == true);
 	bActorMultiSelect = false;
 
 	// Clear selection
