WaitForAsyncFileWrites before returning from CookAsset, this is a safer default than having writes in flight

#rnx

diff --git a/Engine/Source/Editor/UnrealEd/Private/Cooker/CookFunctionLibrary.cpp b/Engine/Source/Editor/UnrealEd/Private/Cooker/CookFunctionLibrary.cpp
--- a/Engine/Source/Editor/UnrealEd/Private/Cooker/CookFunctionLibrary.cpp
+++ b/Engine/Source/Editor/UnrealEd/Private/Cooker/CookFunctionLibrary.cpp
@@ -250,6 +250,8 @@ void UCookFunctionLibrary::CookAsset(UObject* Object, const FString& ForPlatform
 		// persist cook dependencies along side the loose files we wrote above..
 		const FString DepsFilename = UE::Private::GetDepsFilename(Object, DestinationSubfolder + "_" + ForPlatform);
 		UE::Private::SaveDepsToFile(Deps, DepsFilename);
+
+		UPackage::WaitForAsyncFileWrites();
 	}
 
 	if (bWasFilterEditorOnly)
