Focus and scroll the Output Log tab when ProfileGPU is run
- The new GPU profiler only outputs the profiling results to the output log, so it is helpful to bring that tab to the foreground for the user.

diff --git a/Engine/Source/Developer/OutputLog/Private/OutputLogModule.cpp b/Engine/Source/Developer/OutputLog/Private/OutputLogModule.cpp
--- a/Engine/Source/Developer/OutputLog/Private/OutputLogModule.cpp
+++ b/Engine/Source/Developer/OutputLog/Private/OutputLogModule.cpp
@@ -444,6 +444,21 @@ void FOutputLogModule::FocusOutputLog()
 	}
 }
 
+void FOutputLogModule::FocusOutputLogAndScrollToEnd()
+{
+	FocusOutputLog();
+
+	if (TSharedPtr<SOutputLog> Log = OutputLog.Pin())
+	{
+		Log->OnConsoleCommandExecuted();
+	}
+
+	if (TSharedPtr<SOutputLog> Log = OutputLogDrawer.Pin())
+	{
+		Log->OnConsoleCommandExecuted();
+	}
+}
+
 void FOutputLogModule::UpdateOutputLogFilter(const TArray<FName>& CategoriesToShow, TOptional<bool> bShowErrors, TOptional<bool> bShowWarnings, TOptional<bool> bShowLogs)
 {
 	FOutputFilterParams Params;
diff --git a/Engine/Source/Developer/OutputLog/Private/SOutputLog.h b/Engine/Source/Developer/OutputLog/Private/SOutputLog.h
--- a/Engine/Source/Developer/OutputLog/Private/SOutputLog.h
+++ b/Engine/Source/Developer/OutputLog/Private/SOutputLog.h
@@ -414,9 +414,11 @@ protected:
 	 */
 	void OnUserScrolled(float ScrollOffset);
 
+public:
 	/** Called when a console command is entered for this output log */
 	void OnConsoleCommandExecuted();
 
+protected:
 	/** Request we immediately force scroll to the bottom of the log */
 	void RequestForceScroll(bool bIfUserHasNotScrolledUp = false);
 
diff --git a/Engine/Source/Developer/OutputLog/Public/OutputLogModule.h b/Engine/Source/Developer/OutputLog/Public/OutputLogModule.h
--- a/Engine/Source/Developer/OutputLog/Public/OutputLogModule.h
+++ b/Engine/Source/Developer/OutputLog/Public/OutputLogModule.h
@@ -81,6 +81,9 @@ public:
 	/** Opens and focuses on the Output Log Drawer if the status bar exists, otherwise opens and focuses on the Output Log Tab. */
 	void OUTPUTLOG_API FocusOutputLog();
 
+	/** Same as FocusOutputLog(), but also scrolls the output log window to the end. */
+	virtual void FocusOutputLogAndScrollToEnd();
+
 	const TSharedPtr<SDockTab> GetOutputLogTab() const { return OutputLogTab.Pin(); }
 
 	struct FOutputFilterParams
diff --git a/Engine/Source/Editor/UnrealEd/Private/UnrealEdMisc.cpp b/Engine/Source/Editor/UnrealEd/Private/UnrealEdMisc.cpp
--- a/Engine/Source/Editor/UnrealEd/Private/UnrealEdMisc.cpp
+++ b/Engine/Source/Editor/UnrealEd/Private/UnrealEdMisc.cpp
@@ -102,6 +102,7 @@
 #include "Factories/FbxAssetImportData.h"
 #include "Factories/FbxSkeletalMeshImportData.h"
 #include "Factories/FbxStaticMeshImportData.h"
+#include "OutputLogModule.h"
 
 
 #define USE_UNIT_TESTS 0
@@ -599,6 +600,14 @@ void FUnrealEdMisc::OnInit()
 	FEditorDelegates::OnEditorInitialized.Broadcast(TotalEditorStartupTime);
 	
 	GShaderCompilingManager->PrintStats();
+
+#if RHI_NEW_GPU_PROFILER
+	OnProfileGPUHandle = UE::RHI::GPUProfiler::OnProfileGPU.AddLambda([this]()
+	{
+		FOutputLogModule& OutputLogModule = FModuleManager::Get().LoadModuleChecked<FOutputLogModule>("OutputLog");
+		OutputLogModule.FocusOutputLogAndScrollToEnd();
+	});
+#endif
 }
 
 FString FUnrealEdMisc::FindMapFileFromPartialName(const FString& PartialMapName)
@@ -1054,6 +1063,13 @@ void FUnrealEdMisc::OnExit()
 		NavigationBuildingNotificationHandler = NULL;
 	}
 
+#if RHI_NEW_GPU_PROFILER
+	if (OnProfileGPUHandle.IsValid())
+	{
+		UE::RHI::GPUProfiler::OnProfileGPU.Remove(OnProfileGPUHandle);
+	}
+#endif
+
 	// Report session maximum window and tab counts to engine analytics, if available
 	if (FEngineAnalytics::IsAvailable())
 	{
diff --git a/Engine/Source/Editor/UnrealEd/Public/UnrealEdMisc.h b/Engine/Source/Editor/UnrealEd/Public/UnrealEdMisc.h
--- a/Engine/Source/Editor/UnrealEd/Public/UnrealEdMisc.h
+++ b/Engine/Source/Editor/UnrealEd/Public/UnrealEdMisc.h
@@ -356,4 +356,7 @@ private:
 	bool bActorMultiSelect = false;
 	TArray<AActor*> ActorsSelected;
 	TArray<FGuid> UnloadedActorsSelected;
+
+	/** Handle used to register the "OnProfileGPU" delegate in the RHI. */
+	FDelegateHandle OnProfileGPUHandle;
 };
diff --git a/Engine/Source/Runtime/RHI/Private/GPUProfiler.cpp b/Engine/Source/Runtime/RHI/Private/GPUProfiler.cpp
--- a/Engine/Source/Runtime/RHI/Private/GPUProfiler.cpp
+++ b/Engine/Source/Runtime/RHI/Private/GPUProfiler.cpp
@@ -722,6 +722,8 @@ bool FGPUTiming::GAreGlobalsInitialized = false;
 
 namespace UE::RHI::GPUProfiler
 {
+	RHI_API FRHIOnProfileGPU OnProfileGPU;
+
 	TLockFreePointerListUnordered<void, PLATFORM_CACHE_LINE_SIZE> FEventStream::FChunk::MemoryPool;
 
 	static TArray<FEventSink*>& GetSinks()
@@ -2471,6 +2473,11 @@ namespace UE::RHI::GPUProfiler
 		FConsoleCommandWithArgsDelegate::CreateLambda([](const TArray<FString>& Args)
 		{
 			GGPUProfilerSink_StatSystem.bTriggerProfile = true;
+
+			if (OnProfileGPU.IsBound())
+			{
+				OnProfileGPU.Broadcast();
+			}
 		}));
 #endif // WITH_PROFILEGPU
 }
diff --git a/Engine/Source/Runtime/RHI/Public/GPUProfiler.h b/Engine/Source/Runtime/RHI/Public/GPUProfiler.h
--- a/Engine/Source/Runtime/RHI/Public/GPUProfiler.h
+++ b/Engine/Source/Runtime/RHI/Public/GPUProfiler.h
@@ -21,6 +21,9 @@
 
 namespace UE::RHI::GPUProfiler
 {
+	DECLARE_MULTICAST_DELEGATE(FRHIOnProfileGPU);
+	extern RHI_API FRHIOnProfileGPU OnProfileGPU;
+
 	struct FQueue
 	{
 		enum class EType : uint8
