[Curve Table] Setting constant as the deafult curve interpolation type option in the import/reimport Curve Table ComboBox

diff --git a/Engine/Source/Editor/UnrealEd/Private/SCSVImportOptions.cpp b/Engine/Source/Editor/UnrealEd/Private/SCSVImportOptions.cpp
--- a/Engine/Source/Editor/UnrealEd/Private/SCSVImportOptions.cpp
+++ b/Engine/Source/Editor/UnrealEd/Private/SCSVImportOptions.cpp
@@ -49,6 +49,11 @@ void SCSVImportOptions::Construct(const FArguments& InArgs)
 	ImportTypes.Add(MakeShareable(new ECSVImportType(ECSVImportType::ECSV_CurveVector)));
 	ImportTypes.Add(MakeShareable(new ECSVImportType(ECSVImportType::ECSV_CurveLinearColor)));
 
+	// Populate the valid interpolation modes
+	CurveInterpModes.Add( MakeShareable( new ERichCurveInterpMode(ERichCurveInterpMode::RCIM_Constant) ) );
+	CurveInterpModes.Add( MakeShareable( new ERichCurveInterpMode(ERichCurveInterpMode::RCIM_Linear) ) );
+	CurveInterpModes.Add( MakeShareable( new ERichCurveInterpMode(ERichCurveInterpMode::RCIM_Cubic) ) );
+
 	// Create properties view
 	FPropertyEditorModule & EditModule = FModuleManager::Get().GetModuleChecked<FPropertyEditorModule>("PropertyEditor");
 	FDetailsViewArgs DetailsViewArgs;
@@ -159,6 +164,7 @@ void SCSVImportOptions::Construct(const FArguments& InArgs)
 				SAssignNew(CurveInterpCombo, SComboBox<CurveInterpModePtr>)
 				.OptionsSource( &CurveInterpModes )
 				.OnGenerateWidget( this, &SCSVImportOptions::MakeCurveTypeWidget )
+				.InitiallySelectedItem( CurveInterpModes[0] )
 				.Visibility( this, &SCSVImportOptions::GetCurveTypeVis )
 				[
 					SNew(STextBlock)
@@ -215,15 +221,6 @@ void SCSVImportOptions::Construct(const FArguments& InArgs)
 	// set-up selection
 	ImportTypeCombo->SetSelectedItem(bReimportingCurveTable ? CurveTableTypePtr : DataTableTypePtr);
 	PropertyView->SetObject(TempImportDataTable.Get());
-
-	// Populate the valid interpolation modes
-	{
-		CurveInterpModes.Add( MakeShareable( new ERichCurveInterpMode(ERichCurveInterpMode::RCIM_Constant) ) );
-		CurveInterpModes.Add( MakeShareable( new ERichCurveInterpMode(ERichCurveInterpMode::RCIM_Linear) ) );
-		CurveInterpModes.Add( MakeShareable( new ERichCurveInterpMode(ERichCurveInterpMode::RCIM_Cubic) ) );
-	}
-
-	// NB: Both combo boxes default to first item in their options lists as initially selected item
 }
 
 	/** If we should import */
