Fixed unlocalized type name in material creation message

#rnx

diff --git a/Engine/Source/Editor/UnrealEd/Private/LevelEditorViewport.cpp b/Engine/Source/Editor/UnrealEd/Private/LevelEditorViewport.cpp
--- a/Engine/Source/Editor/UnrealEd/Private/LevelEditorViewport.cpp
+++ b/Engine/Source/Editor/UnrealEd/Private/LevelEditorViewport.cpp
@@ -849,10 +849,8 @@ UObject* FLevelEditorViewportClient::GetOrCreateMaterialFromTexture(UTexture* Un
 
 	CreatedMaterialInterface->ForceRecompileForRendering();
 
-	const FString MaterialTypeCreated = bCreateMaterialInstance ? TEXT("MaterialInstance") : TEXT("Material");
-
 	// Warn users that a new material has been created
-	FNotificationInfo Info( FText::Format( LOCTEXT( "DropTextureMaterialCreated", "{0} '{1}' Created" ), FText::FromString(MaterialTypeCreated), FText::FromString(MaterialFullName)));
+	FNotificationInfo Info( FText::Format( LOCTEXT( "DropTextureMaterialCreated", "{0} '{1}' Created" ), CreatedMaterialInterface->GetClass()->GetDisplayNameText(), FText::FromString(MaterialFullName)));
 	Info.ExpireDuration = 4.0f;
 	Info.bUseLargeFont = true;
 	Info.bUseSuccessFailIcons = false;
