Containers: Remove the FastPointerKeyFuncs used by Cooker, LLM, and AssetRegistry; use DefaultKeyFuncs instead. DefaultKeyFuncs uses GetTypeHash(T*), and GetTypeHash(T*) has been optimized and is now faster than the hash function used by FastPointerKeyFuncs.
#rnx

diff --git a/Engine/Source/Editor/UnrealEd/Private/Cooker/CookPackageData.h b/Engine/Source/Editor/UnrealEd/Private/Cooker/CookPackageData.h
--- a/Engine/Source/Editor/UnrealEd/Private/Cooker/CookPackageData.h
+++ b/Engine/Source/Editor/UnrealEd/Private/Cooker/CookPackageData.h
@@ -1240,7 +1240,7 @@ public:
 	FRequestQueue& GetRequestQueue();
 
 	/** Return the Set that holds unordered all PackageDatas that are in the AssignedToWorker state. */
-	TFastPointerSet<FPackageData*>& GetAssignedToWorkerSet();
+	TSet<FPackageData*>& GetAssignedToWorkerSet();
 
 	/**
 	 * Return the LoadQueue used by CookOnTheFlyServer. Container for packages in the Load state and for various
@@ -1254,7 +1254,7 @@ public:
 	FPackageDataQueue& GetSaveQueue();
 
 	/** Return the Set that holds unordered all PackageDatas that are in one of the SaveStalled states. */
-	TFastPointerSet<FPackageData*>& GetSaveStalledSet();
+	TSet<FPackageData*>& GetSaveStalledSet();
 
 	/**
 	Return the PackageData for the given PackageName and FileName; no validation is done on the names.
@@ -1462,8 +1462,8 @@ private:
 	uint32 NextLeafToRootRank = 0;
 	int32 PendingCookedPlatformDataNum = 0;
 	FRequestQueue RequestQueue;
-	TFastPointerSet<FPackageData*> AssignedToWorkerSet;
-	TFastPointerSet<FPackageData*> SaveStalledSet;
+	TSet<FPackageData*> AssignedToWorkerSet;
+	TSet<FPackageData*> SaveStalledSet;
 	FLoadQueue LoadQueue;
 	FPackageDataQueue SaveQueue;
 	UCookOnTheFlyServer& CookOnTheFlyServer;
@@ -2021,7 +2021,7 @@ inline FRequestQueue& FPackageDatas::GetRequestQueue()
 	return RequestQueue;
 }
 
-inline TFastPointerSet<FPackageData*>& FPackageDatas::GetAssignedToWorkerSet()
+inline TSet<FPackageData*>& FPackageDatas::GetAssignedToWorkerSet()
 {
 	return AssignedToWorkerSet;
 }
@@ -2031,7 +2031,7 @@ inline FLoadQueue& FPackageDatas::GetLoadQueue()
 	return LoadQueue;
 }
 
-inline TFastPointerSet<FPackageData*>& FPackageDatas::GetSaveStalledSet()
+inline TSet<FPackageData*>& FPackageDatas::GetSaveStalledSet()
 {
 	return SaveStalledSet;
 }
diff --git a/Engine/Source/Editor/UnrealEd/Private/Cooker/CookPlatformManager.cpp b/Engine/Source/Editor/UnrealEd/Private/Cooker/CookPlatformManager.cpp
--- a/Engine/Source/Editor/UnrealEd/Private/Cooker/CookPlatformManager.cpp
+++ b/Engine/Source/Editor/UnrealEd/Private/Cooker/CookPlatformManager.cpp
@@ -339,7 +339,7 @@ TMap<ITargetPlatform*, ITargetPlatform*> FPlatformManager::RemapTargetPlatforms(
 	TMap<ITargetPlatform*, ITargetPlatform*> Remap;
 
 	ITargetPlatformManagerModule* PlatformManager = GetTargetPlatformManager();
-	TFastPointerMap<const ITargetPlatform*, FPlatformData*> NewPlatformDatas;
+	TMap<const ITargetPlatform*, FPlatformData*> NewPlatformDatas;
 
 	{
 		FRWScopeLock PlatformDatasScopeLock(PlatformDatasLock, FRWScopeLockType::SLT_Write);
diff --git a/Engine/Source/Editor/UnrealEd/Private/Cooker/CookPlatformManager.h b/Engine/Source/Editor/UnrealEd/Private/Cooker/CookPlatformManager.h
--- a/Engine/Source/Editor/UnrealEd/Private/Cooker/CookPlatformManager.h
+++ b/Engine/Source/Editor/UnrealEd/Private/Cooker/CookPlatformManager.h
@@ -197,7 +197,7 @@ namespace Cook
 		static uint32 IsInPlatformsLockTLSSlot;
 
 		/** A collection of flags and other data we maintain for each platform we have encountered in any session. */
-		TFastPointerMap<FPlatformId, FPlatformData*> PlatformDatas;
+		TMap<FPlatformId, FPlatformData*> PlatformDatas;
 		TMap<FName, FPlatformData*> PlatformDatasByName;
 
 		/*
@@ -237,10 +237,10 @@ namespace Cook
 }
 
 template <typename Value>
-void RemapMapKeys(TFastPointerMap<const ITargetPlatform*, Value>& Map, const TMap<ITargetPlatform*,
+void RemapMapKeys(TMap<const ITargetPlatform*, Value>& Map, const TMap<ITargetPlatform*,
 	ITargetPlatform*>& Remap)
 {
-	TFastPointerMap<const ITargetPlatform*, Value> NewMap;
+	TMap<const ITargetPlatform*, Value> NewMap;
 	NewMap.Reserve(Map.Num());
 	for (TPair<const ITargetPlatform*, Value>& OldPair : Map)
 	{
diff --git a/Engine/Source/Editor/UnrealEd/Private/Cooker/CookTypes.h b/Engine/Source/Editor/UnrealEd/Private/Cooker/CookTypes.h
--- a/Engine/Source/Editor/UnrealEd/Private/Cooker/CookTypes.h
+++ b/Engine/Source/Editor/UnrealEd/Private/Cooker/CookTypes.h
@@ -56,46 +56,6 @@ DECLARE_LOG_CATEGORY_EXTERN(LogCookGenerationHelper, Log, All);
 
 inline constexpr uint32 ExpectedMaxNumPlatforms = 32;
 
-/** A BaseKeyFuncs for Maps and Sets with a quicker hash function for pointers than TDefaultMapKeyFuncs */
-template<typename KeyType>
-struct TFastPointerSetKeyFuncs : public DefaultKeyFuncs<KeyType>
-{
-	using typename DefaultKeyFuncs<KeyType>::KeyInitType;
-	static FORCEINLINE uint32 GetKeyHash(KeyInitType Key)
-	{
-#if PLATFORM_64BITS
-		static_assert(sizeof(UPTRINT) == sizeof(uint64), "Expected pointer size to be 64 bits");
-		// Ignoring the lower 4 bits since they are likely zero anyway.
-		const uint64 ImportantBits = reinterpret_cast<uint64>(Key) >> 4;
-		return GetTypeHash(ImportantBits);
-#else
-		static_assert(sizeof(UPTRINT) == sizeof(uint32), "Expected pointer size to be 32 bits");
-		return static_cast<uint32>(reinterpret_cast<uint32>(Key));
-#endif
-	}
-};
-
-template<typename KeyType, typename ValueType, bool bInAllowDuplicateKeys>
-struct TFastPointerMapKeyFuncs : public TDefaultMapKeyFuncs<KeyType, ValueType, bInAllowDuplicateKeys>
-{
-	using typename TDefaultMapKeyFuncs<KeyType, ValueType, bInAllowDuplicateKeys>::KeyInitType;
-	static FORCEINLINE uint32 GetKeyHash(KeyInitType Key)
-	{
-		return TFastPointerSetKeyFuncs<KeyType>::GetKeyHash(Key);
-	}
-};
-
-/** A TMap which uses TFastPointerMapKeyFuncs instead of TDefaultMapKeyFuncs */
-template<typename KeyType, typename ValueType, typename SetAllocator = FDefaultSetAllocator>
-class TFastPointerMap : public TMap<KeyType, ValueType, SetAllocator,
-	TFastPointerMapKeyFuncs<KeyType, ValueType, false>>
-{};
-
-/** A TSet which uses TFastPointerSetKeyFuncs instead of DefaultKeyFuncs */
-template<typename KeyType, typename SetAllocator = FDefaultSetAllocator>
-class TFastPointerSet : public TSet<KeyType, TFastPointerSetKeyFuncs<KeyType>, SetAllocator>
-{};
-
 namespace UE::Cook
 {
 	class FDeterminismManager;
@@ -105,14 +65,14 @@ namespace UE::Cook
 	/** A function that is called when a requested package finishes cooking (when successful, failed, or skipped) */
 	typedef TUniqueFunction<void(FPackageData*)> FCompletionCallback;
 
-	class FPackageDataSet : public TFastPointerSet<FPackageData*>
+	class FPackageDataSet : public TSet<FPackageData*>
 	{
-		using TFastPointerSet<FPackageData*>::TFastPointerSet;
+		using TSet<FPackageData*>::TSet;
 	};
 	template <typename ValueType>
-	class TPackageDataMap : public TFastPointerMap<FPackageData*, ValueType>
+	class TPackageDataMap : public TMap<FPackageData*, ValueType>
 	{
-		using TFastPointerMap<FPackageData*, ValueType>::TFastPointerMap;
+		using TMap<FPackageData*, ValueType>::TMap;
 	};
 
 	/**
@@ -506,7 +466,7 @@ namespace UE::Cook
 		bool bCookIncremental = true;
 		bool bCookIncrementalAllowAllClasses = false;
 		TArray<FName> NeverCookPackageList;
-		TFastPointerMap<const ITargetPlatform*, TSet<FName>> PlatformSpecificNeverCookPackages;
+		TMap<const ITargetPlatform*, TSet<FName>> PlatformSpecificNeverCookPackages;
 
 		friend FCbWriter& ::operator<<(FCbWriter& Writer, const UE::Cook::FBeginCookConfigSettings& Value);
 		friend bool ::LoadFromCompactBinary(FCbFieldView Field, UE::Cook::FBeginCookConfigSettings& Value);
diff --git a/Engine/Source/Editor/UnrealEd/Private/Cooker/PackageTracker.h b/Engine/Source/Editor/UnrealEd/Private/Cooker/PackageTracker.h
--- a/Engine/Source/Editor/UnrealEd/Private/Cooker/PackageTracker.h
+++ b/Engine/Source/Editor/UnrealEd/Private/Cooker/PackageTracker.h
@@ -268,7 +268,7 @@ public:
 
 	/** Packages to never cook - entries are LongPackageNames. */
 	FThreadSafeSet<FName> NeverCookPackageList;
-	TFastPointerMap<const ITargetPlatform*, TSet<FName>> PlatformSpecificNeverCookPackages;
+	TMap<const ITargetPlatform*, TSet<FName>> PlatformSpecificNeverCookPackages;
 
 	// Thread-safe enumeration of loaded package. 
 	// A lock is held during enumeration, keep code simple and optimal so the lock is released as fast as possible.
@@ -332,7 +332,7 @@ private:
 	// Begin Data Guarded by Lock
 
 	/** All packages currently loaded by the engine. */
-	TFastPointerSet<UPackage*> LoadedPackages;
+	TSet<UPackage*> LoadedPackages;
 	/** List of packages from COTFS that we expect will never be loaded again; log a warning if we see them. */
 	TSet<FName> ExpectedNeverLoadPackages;
 	/* Events about UPackages created since the last call to GetPackageStream, plus related data such as the
diff --git a/Engine/Source/Runtime/AssetRegistry/Private/AssetRegistry.cpp b/Engine/Source/Runtime/AssetRegistry/Private/AssetRegistry.cpp
--- a/Engine/Source/Runtime/AssetRegistry/Private/AssetRegistry.cpp
+++ b/Engine/Source/Runtime/AssetRegistry/Private/AssetRegistry.cpp
@@ -9229,32 +9229,7 @@ struct FSetManageReferencesNodeData
 	EDependencyProperty CurrentManagerProperties = EDependencyProperty::None;
 };
 
-// TODO: Replace the three existing copies of TFastPointerMap/TFastPointerSet with a shared implementation,
-// or move the behavior into DefaultKeyFuncs.
-template<typename KeyType, typename ValueType, bool bInAllowDuplicateKeys>
-struct TFastPointerMapKeyFuncs : public TDefaultMapKeyFuncs<KeyType, ValueType, bInAllowDuplicateKeys>
-{
-	using typename TDefaultMapKeyFuncs<KeyType, ValueType, bInAllowDuplicateKeys>::KeyInitType;
-	static FORCEINLINE uint32 GetKeyHash(KeyInitType Key)
-	{
-#if PLATFORM_64BITS
-		static_assert(sizeof(UPTRINT) == sizeof(uint64), "Expected pointer size to be 64 bits");
-		// Ignoring the lower 4 bits since they are likely zero anyway.
-		const uint64 ImportantBits = reinterpret_cast<uint64>(Key) >> 4;
-		return GetTypeHash(ImportantBits);
-#else
-		static_assert(sizeof(UPTRINT) == sizeof(uint32), "Expected pointer size to be 32 bits");
-		return static_cast<uint32>(reinterpret_cast<uint32>(Key));
-#endif
-	}
-};
-
-template<typename KeyType, typename ValueType, typename SetAllocator = FDefaultSetAllocator>
-class TFastPointerMap : public TMap<KeyType, ValueType, SetAllocator,
-	TFastPointerMapKeyFuncs<KeyType, ValueType, false /* bInAllowDuplicateKeys */>>
-{};
-
-using FSetManageReferencesNodeDataMap = TFastPointerMap<FDependsNode*, TUniquePtr<FSetManageReferencesNodeData>>;
+using FSetManageReferencesNodeDataMap = TMap<FDependsNode*, TUniquePtr<FSetManageReferencesNodeData>>;
 struct FSetManageReferencesScratch
 {
 	FSetManageReferencesNodeDataMap NodeData;
diff --git a/Engine/Source/Runtime/Core/Private/HAL/LowLevelMemTracker.cpp b/Engine/Source/Runtime/Core/Private/HAL/LowLevelMemTracker.cpp
--- a/Engine/Source/Runtime/Core/Private/HAL/LowLevelMemTracker.cpp
+++ b/Engine/Source/Runtime/Core/Private/HAL/LowLevelMemTracker.cpp
@@ -1321,7 +1321,7 @@ private:
 		UE::LLM::ESizeParams SizeParams);
 
 	FConstTagDataArray Columns;
-	TFastPointerLLMSet<const FTagData*> ExistingColumns;
+	TLLMSet<const FTagData*> ExistingColumns;
 	FArchive* Archive;
 	double LastWriteTime;
 	int32 WriteCount;
@@ -1345,7 +1345,7 @@ private:
 	void SendTagDeclaration(const FTagData* TagData);
 
 	ELLMTracker Tracker;
-	TFastPointerLLMSet<const FTagData*> DeclaredTags;
+	TLLMSet<const FTagData*> DeclaredTags;
 	bool bTrackerSpecSent = false;
 	bool bTagSetSpecSent = false;
 };
@@ -1364,7 +1364,7 @@ protected:
 	void RecordTagToCsv(int32 CsvCategoryIndex, const FTagData* TagData, int64 TagSize);
 #if LLM_CSV_PROFILER_WRITER_ENABLED 
 	ELLMTracker Tracker;
-	TFastPointerLLMMap<const FTagData*, FName> TagDataToCsvStatName;
+	TLLMMap<const FTagData*, FName> TagDataToCsvStatName;
 #endif
 };
 
diff --git a/Engine/Source/Runtime/Core/Private/HAL/LowLevelMemTrackerPrivate.h b/Engine/Source/Runtime/Core/Private/HAL/LowLevelMemTrackerPrivate.h
--- a/Engine/Source/Runtime/Core/Private/HAL/LowLevelMemTrackerPrivate.h
+++ b/Engine/Source/Runtime/Core/Private/HAL/LowLevelMemTrackerPrivate.h
@@ -219,7 +219,7 @@ namespace LLMPrivate
 			SizeInSnapshot = 0;
 		}
 	};
-	typedef TFastPointerLLMMap<const FTagData*, FTrackerTagSizeData> FTrackerTagSizeMap;
+	typedef TLLMMap<const FTagData*, FTrackerTagSizeData> FTrackerTagSizeMap;
 
 	/**
 	 * Size information stored on each threadstate for a tag
diff --git a/Engine/Source/Runtime/Core/Private/HAL/LowLevelMemoryUtils.h b/Engine/Source/Runtime/Core/Private/HAL/LowLevelMemoryUtils.h
--- a/Engine/Source/Runtime/Core/Private/HAL/LowLevelMemoryUtils.h
+++ b/Engine/Source/Runtime/Core/Private/HAL/LowLevelMemoryUtils.h
@@ -1049,49 +1049,12 @@ class FDefaultBitArrayLLMAllocator : public TInlineAllocator<4, FDefaultLLMAlloc
 class FDefaultSparseArrayLLMAllocator : public TSparseArrayAllocator<FDefaultLLMAllocator, FDefaultBitArrayLLMAllocator> { public: typedef TSparseArrayAllocator<FDefaultLLMAllocator, FDefaultBitArrayLLMAllocator> Typedef; };
 class FDefaultSetLLMAllocator : public TSetAllocator<FDefaultSparseArrayLLMAllocator, TInlineAllocator<1, FDefaultLLMAllocator>> { public: typedef TSetAllocator<FDefaultSparseArrayLLMAllocator, TInlineAllocator<1, FDefaultLLMAllocator>>         Typedef; };
 
-// Providing a fast hash function for pointer maps and sets; this fast hash function just uses the pointer
-// cast to an int rather than mixing the bits of the pointer. This adds a performance vulnerability to
-// clustered data, but that usually is not a problem with pointers allocated from our FLLMAllocator.
-namespace UE
-{
-namespace LLMPrivate
-{
-	template<typename KeyType>
-	struct TFastPointerSetKeyFuncs : public DefaultKeyFuncs<KeyType>
-	{
-		using typename DefaultKeyFuncs<KeyType>::KeyInitType;
-		static FORCEINLINE uint32 GetKeyHash(KeyInitType Key)
-		{
-#if PLATFORM_64BITS
-			static_assert(sizeof(UPTRINT) == sizeof(uint64), "Expected pointer size to be 64 bits");
-			// Ignoring the lower 4 bits since they are likely zero anyway.
-			const uint64 ImportantBits = reinterpret_cast<uint64>(Key) >> 4;
-			return GetTypeHash(ImportantBits);
-#else
-			static_assert(sizeof(UPTRINT) == sizeof(uint32), "Expected pointer size to be 32 bits");
-			return static_cast<uint32>(reinterpret_cast<uint32>(Key));
-#endif
-		}
-	};
-
-	template<typename KeyType, typename ValueType, bool bInAllowDuplicateKeys>
-	struct TFastPointerMapKeyFuncs : public TDefaultMapKeyFuncs<KeyType, ValueType, bInAllowDuplicateKeys>
-	{
-		using typename TDefaultMapKeyFuncs<KeyType, ValueType, bInAllowDuplicateKeys>::KeyInitType;
-		static FORCEINLINE uint32 GetKeyHash(KeyInitType Key)
-		{
-			return TFastPointerSetKeyFuncs<KeyType>::GetKeyHash(Key);
-		}
-	};
-}
-}
-
 template<typename KeyType, typename ValueType>
-class TFastPointerLLMMap : public TMap<KeyType, ValueType, FDefaultSetLLMAllocator, UE::LLMPrivate::TFastPointerMapKeyFuncs<KeyType, ValueType, false>>
+class TLLMMap : public TMap<KeyType, ValueType, FDefaultSetLLMAllocator>
 {};
 
 template<typename KeyType>
-class TFastPointerLLMSet : public TSet<KeyType, UE::LLMPrivate::TFastPointerSetKeyFuncs<KeyType>, FDefaultSetLLMAllocator>
+class TLLMSet : public TSet<KeyType, DefaultKeyFuncs<KeyType>, FDefaultSetLLMAllocator>
 {};
 
 
