[HLOD] [Experimental] Add better support for SubWP HLODs by introducing Standalone HLODs

* New World Partition property - "Build Standalone HLOD". When toggled HLOD Builds result in Standalone HLOD being built
* With Standalone HLODs, HLOD actors are placed in separate levels, instead of in the main level, which allows for independent streaming when using the level as SubWP
* Each HLOD Setup is placed in a separate level, which allows for independent streaming
* Standalone HLOD Levels are placed in HLOD folder and are named based on the main level name - <Level>_HLOD<N>
* Standalone HLODs are automatically loaded when using the level as main world or when using it as SubWP
* Requires WP and SubWP to have the same HLOD Setups
* Supports Editor, PIE and Cooked builds
* Supports FastGeo plugin
* Supports Data Layers
* Supports nested SubWPs

Current limitations:
* Multiple instances of the same level used as SubWP are not supported
* External Data Layers are not supported
* Distributed Builds are supported but they require the Standalone HLOD levels to be available in P4, so at least one non-distributed HLOD Setup run is required before distributed builds can be executed

diff --git a/Engine/Plugins/Editor/WorldPartitionHLODUtilities/Source/Private/WorldPartition/HLOD/Utilities/WorldPartitionHLODUtilities.cpp b/Engine/Plugins/Editor/WorldPartitionHLODUtilities/Source/Private/WorldPartition/HLOD/Utilities/WorldPartitionHLODUtilities.cpp
--- a/Engine/Plugins/Editor/WorldPartitionHLODUtilities/Source/Private/WorldPartition/HLOD/Utilities/WorldPartitionHLODUtilities.cpp
+++ b/Engine/Plugins/Editor/WorldPartitionHLODUtilities/Source/Private/WorldPartition/HLOD/Utilities/WorldPartitionHLODUtilities.cpp
@@ -208,7 +208,14 @@ TArray<AWorldPartitionHLOD*> FWorldPartitionHLODUtilities::CreateHLODActors(FHLO
 			FActorSpawnParameters SpawnParams;
 			SpawnParams.Name = HLODActorName;
 			SpawnParams.NameMode = FActorSpawnParameters::ESpawnActorNameMode::Required_Fatal;
-			HLODActor = InCreationParams.WorldPartition->GetWorld()->SpawnActor<AWorldPartitionHLOD>(HLODLayer->GetHLODActorClass(), SpawnParams);
+			if (InCreationParams.bIsStandalone)
+			{
+				HLODActor = InCreationParams.TargetWorld->SpawnActor<AWorldPartitionHLOD>(HLODLayer->GetHLODActorClass(), SpawnParams);
+			}
+			else
+			{
+				HLODActor = InCreationParams.WorldPartition->GetWorld()->SpawnActor<AWorldPartitionHLOD>(HLODLayer->GetHLODActorClass(), SpawnParams);
+			}
 
 			check(HLODActor->GetContentBundleGuid() == InCreationParams.ContentBundleGuid);
 			check(HLODActor->GetExternalDataLayerAsset() == InCreationParams.GetExternalDataLayerAsset());
@@ -250,6 +257,18 @@ TArray<AWorldPartitionHLOD*> FWorldPartitionHLODUtilities::CreateHLODActors(FHLO
 			SourceActorsMap.KeySort(TLess<FGuid>());
 			SourceActorsMap.GenerateValueArray(SourceActorsArray);
 
+			if (InCreationParams.bIsStandalone)
+			{
+				for (FWorldPartitionRuntimeCellObjectMapping& SourceActor : SourceActorsArray)
+				{
+					if (SourceActor.ContainerID.IsMainContainer())
+					{
+						// Set ContainerID for Standalone HLOD Source Actors to allow them to be loaded in the context of Standalone HLOD Worlds for HLOD build
+						SourceActor.ContainerID = FActorContainerID(FActorContainerID::GetMainContainerID(), FGuid::NewDeterministicGuid(SourceActor.ContainerPackage.ToString()));
+					}
+				}
+			}
+
 			bool bSourceActorsChanged = HLODSourceActors->GetActors().Num() != SourceActorsArray.Num();
 			if (!bSourceActorsChanged)
 			{
@@ -332,6 +351,12 @@ TArray<AWorldPartitionHLOD*> FWorldPartitionHLODUtilities::CreateHLODActors(FHLO
 			DirtyReason = TEXT("MinVisibleDistance");
 		}
 
+		if (HLODActor->IsStandalone() != InCreationParams.bIsStandalone)
+		{
+			HLODActor->SetIsStandalone(InCreationParams.bIsStandalone);
+			DirtyReason = TEXT("Standalone");
+		}
+
 		// If any change was performed, mark HLOD package as dirty
 		if (DirtyReason)
 		{
diff --git a/Engine/Plugins/Experimental/FastGeoStreaming/Source/FastGeoStreaming/Private/FastGeoHLOD.cpp b/Engine/Plugins/Experimental/FastGeoStreaming/Source/FastGeoStreaming/Private/FastGeoHLOD.cpp
--- a/Engine/Plugins/Experimental/FastGeoStreaming/Source/FastGeoStreaming/Private/FastGeoHLOD.cpp
+++ b/Engine/Plugins/Experimental/FastGeoStreaming/Source/FastGeoStreaming/Private/FastGeoHLOD.cpp
@@ -31,6 +31,7 @@ void FFastGeoHLOD::Serialize(FArchive& Ar)
 
 	Ar << bRequireWarmup;
 	Ar << SourceCellGuid;
+	Ar << StandaloneHLODGuid;
 }
 
 ULevel* FFastGeoHLOD::GetHLODLevel() const
@@ -101,6 +102,16 @@ const FGuid& FFastGeoHLOD::GetSourceCellGuid() const
 	return SourceCellGuid;
 }
 
+bool FFastGeoHLOD::IsStandalone() const
+{
+	return StandaloneHLODGuid.IsValid();
+}
+
+const FGuid& FFastGeoHLOD::GetStandaloneHLODGuid() const
+{
+	return StandaloneHLODGuid;
+}
+
 void FFastGeoHLOD::OnRegister()
 {
 	GetLevel()->GetWorld()->GetSubsystem<UWorldPartitionHLODRuntimeSubsystem>()->RegisterHLODObject(this);
@@ -122,6 +133,11 @@ void FFastGeoHLOD::SetRequireWarmup(bool bInRequireWarmup)
 	bRequireWarmup = bInRequireWarmup;
 }
 
+void FFastGeoHLOD::SetStandaloneHLODGuid(const FGuid& InStandaloneHLODGuid)
+{
+	StandaloneHLODGuid = InStandaloneHLODGuid;
+}
+
 FFastGeoComponent& FFastGeoHLOD::AddComponent(FFastGeoElementType InComponentType)
 {
 	FFastGeoComponent& NewComponent = Super::AddComponent(InComponentType);
diff --git a/Engine/Plugins/Experimental/FastGeoStreaming/Source/FastGeoStreaming/Private/FastGeoHLOD.h b/Engine/Plugins/Experimental/FastGeoStreaming/Source/FastGeoStreaming/Private/FastGeoHLOD.h
--- a/Engine/Plugins/Experimental/FastGeoStreaming/Source/FastGeoStreaming/Private/FastGeoHLOD.h
+++ b/Engine/Plugins/Experimental/FastGeoStreaming/Source/FastGeoStreaming/Private/FastGeoHLOD.h
@@ -32,6 +32,8 @@ public:
 	virtual TSet<UObject*> GetAssetsToWarmup() const override;
 	virtual void SetVisibility(bool bIsVisible) override;
 	virtual const FGuid& GetSourceCellGuid() const override;
+	virtual bool IsStandalone() const override;
+	virtual const FGuid& GetStandaloneHLODGuid() const override;
 	// ~End IWorldPartitionHLODObject interface
 
 #if WITH_EDITOR
@@ -42,6 +44,7 @@ public:
 
 	void SetSourceCellGuid(const FGuid& InSourceCellGuid);
 	void SetRequireWarmup(bool bInRequireWarmup);
+	void SetStandaloneHLODGuid(const FGuid& InStandaloneHLODGuid);
 #endif
 
 private:
@@ -52,4 +55,5 @@ private:
 	// Persistent Data
 	bool bRequireWarmup;
 	FGuid SourceCellGuid;
+	FGuid StandaloneHLODGuid;
 };
\ No newline at end of file
diff --git a/Engine/Plugins/Experimental/FastGeoStreaming/Source/FastGeoStreaming/Private/FastGeoWorldPartitionRuntimeCellTransformer.cpp b/Engine/Plugins/Experimental/FastGeoStreaming/Source/FastGeoStreaming/Private/FastGeoWorldPartitionRuntimeCellTransformer.cpp
--- a/Engine/Plugins/Experimental/FastGeoStreaming/Source/FastGeoStreaming/Private/FastGeoWorldPartitionRuntimeCellTransformer.cpp
+++ b/Engine/Plugins/Experimental/FastGeoStreaming/Source/FastGeoStreaming/Private/FastGeoWorldPartitionRuntimeCellTransformer.cpp
@@ -301,6 +301,7 @@ void UFastGeoWorldPartitionRuntimeCellTransformer::Transform(ULevel* InLevel)
 				FastGeoHLOD.Reset(new FFastGeoHLOD(FastGeo, *FString::Printf(TEXT("FastGeoHLOD_%s"), *Actor->GetName())));
 				FastGeoHLOD->SetSourceCellGuid(HLODActor->GetSourceCellGuid());
 				FastGeoHLOD->SetRequireWarmup(HLODActor->DoesRequireWarmup());
+				FastGeoHLOD->SetStandaloneHLODGuid(HLODActor->GetStandaloneHLODGuid());
 				CurrentComponentCluster = FastGeoHLOD.Get();
 			}
 
diff --git a/Engine/Source/Editor/UnrealEd/Classes/Commandlets/WorldPartitionBuilderCommandlet.h b/Engine/Source/Editor/UnrealEd/Classes/Commandlets/WorldPartitionBuilderCommandlet.h
--- a/Engine/Source/Editor/UnrealEd/Classes/Commandlets/WorldPartitionBuilderCommandlet.h
+++ b/Engine/Source/Editor/UnrealEd/Classes/Commandlets/WorldPartitionBuilderCommandlet.h
@@ -31,4 +31,6 @@ private:
 	FString AutoSubmitTags;
 
 	TMap<FString, TArray<FString>> AutoSubmitFiles;
+
+	TArray<FString> AdditionalWorldPackagesToProcess;
 };
\ No newline at end of file
diff --git a/Engine/Source/Editor/UnrealEd/Private/Commandlets/WorldPartitionBuilderCommandlet.cpp b/Engine/Source/Editor/UnrealEd/Private/Commandlets/WorldPartitionBuilderCommandlet.cpp
--- a/Engine/Source/Editor/UnrealEd/Private/Commandlets/WorldPartitionBuilderCommandlet.cpp
+++ b/Engine/Source/Editor/UnrealEd/Private/Commandlets/WorldPartitionBuilderCommandlet.cpp
@@ -193,6 +193,19 @@ int32 UWorldPartitionBuilderCommandlet::Main(const FString& Params)
 			UE_LOG(LogWorldPartitionBuilderCommandlet, Warning, TEXT("Failed to execute %s on map %s."), *BuilderClassName, *MapPackageName);
 			Result = 1;
 		}
+
+		for (int32 AdditionalPackageIndex = 0; AdditionalPackageIndex < AdditionalWorldPackagesToProcess.Num(); AdditionalPackageIndex++)
+		{
+			const FString& AdditionalPackageName = AdditionalWorldPackagesToProcess[AdditionalPackageIndex];
+			UE_LOG(LogWorldPartitionBuilderCommandlet, Display, TEXT("[%d.%d / %d] Executing %s on additional map %s..."), PackageIndex, AdditionalPackageIndex + 1, PackageCount, *BuilderClassName, *AdditionalPackageName);
+
+			if (!RunBuilder(BuilderClass, AdditionalPackageName))
+			{
+				UE_LOG(LogWorldPartitionBuilderCommandlet, Warning, TEXT("Failed to execute %s on additional map %s."), *BuilderClassName, *MapPackageName);
+				Result = 1;
+			}
+		}
+		AdditionalWorldPackagesToProcess.Empty();
 	}
 
 	// Autosubmit
@@ -289,6 +302,12 @@ bool UWorldPartitionBuilderCommandlet::RunBuilder(TSubclassOf<UWorldPartitionBui
 		SaveConfig(CPF_Config, *WorldConfigFilename);
 	}
 
+	TArray<FString> WorldPackages;
+	if (Builder->ShouldProcessAdditionalWorlds(World, WorldPackages))
+	{
+		AdditionalWorldPackagesToProcess.Append(WorldPackages);
+	}
+
 	return bResult;
 }
 
diff --git a/Engine/Source/Editor/UnrealEd/Private/WorldPartition/WorldPartitionHLODsBuilder.cpp b/Engine/Source/Editor/UnrealEd/Private/WorldPartition/WorldPartitionHLODsBuilder.cpp
--- a/Engine/Source/Editor/UnrealEd/Private/WorldPartition/WorldPartitionHLODsBuilder.cpp
+++ b/Engine/Source/Editor/UnrealEd/Private/WorldPartition/WorldPartitionHLODsBuilder.cpp
@@ -13,6 +13,7 @@
 #include "UObject/SavePackage.h"
 
 #include "ActorFolder.h"
+#include "EditorWorldUtils.h"
 #include "Engine/Engine.h"
 #include "Engine/LevelStreamingGCHelper.h"
 #include "EngineUtils.h"
@@ -21,6 +22,8 @@
 
 #include "AssetRegistry/AssetRegistryModule.h"
 #include "DerivedDataCacheInterface.h"
+#include "DirectoryWatcherModule.h"
+#include "IDirectoryWatcher.h"
 
 #include "WorldPartition/DataLayer/DataLayerInstance.h"
 #include "WorldPartition/DataLayer/ExternalDataLayerEngineSubsystem.h"
@@ -32,7 +35,9 @@
 #include "WorldPartition/HLOD/HLODActorDesc.h"
 #include "WorldPartition/HLOD/HLODLayer.h"
 #include "WorldPartition/HLOD/HLODProviderInterface.h"
+#include "WorldPartition/HLOD/StandaloneHLODSubsystem.h"
 #include "WorldPartition/IWorldPartitionEditorModule.h"
+#include "WorldPartition/RuntimeHashSet/WorldPartitionRuntimeHashSet.h"
 
 #define LOCTEXT_NAMESPACE "WorldPartitionHLODsBuilder"
 
@@ -49,6 +54,7 @@ UWorldPartitionHLODsBuilder::UWorldPartitionHLODsBuilder(const FObjectInitialize
 	: Super(ObjectInitializer)
 	, BuilderIdx(INDEX_NONE)
 	, BuilderCount(INDEX_NONE)
+	, bBuildingStandaloneHLOD(false)
 {
 	if (!IsTemplate())
 	{
@@ -178,6 +184,37 @@ bool UWorldPartitionHLODsBuilder::ShouldProcessWorld(UWorld* InWorld) const
 	return bShouldProcessWorld;
 }
 
+bool UWorldPartitionHLODsBuilder::ShouldProcessAdditionalWorlds(UWorld* InWorld, TArray<FString>& OutPackageNames) const
+{
+	// If during Build step and if building standalone HLOD, we want to run the builder on standalone HLOD levels,
+	// so that the HLOD Actors, which were created in those levels can be built
+	if (ShouldRunStep(EHLODBuildStep::HLOD_Build) == false)
+	{
+		return false;
+	}
+
+	UWorldPartition* WP = InWorld->GetWorldPartition();
+	if (WP && WP->HasStandaloneHLOD())
+	{
+		if (const UWorldPartitionRuntimeHashSet* WorldPartitionRuntimeHashSet = Cast<UWorldPartitionRuntimeHashSet>(WP->RuntimeHash))
+		{
+			FString FolderPath, PackagePrefix;
+			UWorldPartitionStandaloneHLODSubsystem::GetStandaloneHLODFolderPathAndPackagePrefix(InWorld->GetPackage()->GetName(), FolderPath, PackagePrefix);
+
+			WorldPartitionRuntimeHashSet->ForEachHLODLayer([&FolderPath, &PackagePrefix, &OutPackageNames](FName RuntimePartitionName, FName HLODSetupName, int32 HLODSetupIndex)
+			{
+				const FString HLODLevelPackageName = FString::Printf(TEXT("%s/%s%d"), *FolderPath, *PackagePrefix, HLODSetupIndex);
+				OutPackageNames.Add(HLODLevelPackageName);
+
+				return true;
+			});
+
+			return true;
+		}
+	}
+	return false;
+}
+
 bool UWorldPartitionHLODsBuilder::PreWorldInitialization(UWorld* InWorld, FPackageSourceControlHelper& PackageHelper)
 {
 	if (bDistributedBuild)
@@ -214,6 +251,7 @@ bool UWorldPartitionHLODsBuilder::RunInternal(UWorld* InWorld, const FCellInfo&
 {
 	World = InWorld;
 	WorldPartition = World->GetWorldPartition();
+	bBuildingStandaloneHLOD = WorldPartition->HasStandaloneHLOD();
 	
 	// Allows HLOD Streaming levels to be GCed properly
 	FLevelStreamingGCHelper::EnableForCommandlet();
@@ -258,14 +296,14 @@ bool UWorldPartitionHLODsBuilder::RunInternal(UWorld* InWorld, const FCellInfo&
 
 bool UWorldPartitionHLODsBuilder::SetupHLODActors()
 {
-	// No setup needed for non partitioned worlds
-	if (WorldPartition)
+	// No setup needed for non partitioned worlds and standalone HLOD worlds
+	if (WorldPartition && !WorldPartition->IsStandaloneHLODWorld())
 	{
 		auto ActorFolderAddedDelegateHandle = GEngine->OnActorFolderAdded().AddLambda([this](UActorFolder* InActorFolder)
 		{
 			UPackage* ActorFolderPackage = InActorFolder->GetPackage();
 			const bool bIsTempPackage = FPackageName::IsTempPackage(ActorFolderPackage->GetName());
-			if (!bIsTempPackage)
+			if (!bIsTempPackage && InActorFolder->IsInitiallyExpanded())
 			{
 				// We don't want the HLOD folders to be expanded by default
 				InActorFolder->SetIsInitiallyExpanded(false);
@@ -284,6 +322,37 @@ bool UWorldPartitionHLODsBuilder::SetupHLODActors()
 
 		WorldPartition->SetupHLODActors(SetupHLODActorsParams);
 
+		if (bBuildingStandaloneHLOD)
+		{
+			// Retrieve additional Standalone HLOD levels that have to be processed
+			AdditionalWorldPartitionsForStandaloneHLOD = MoveTemp(SetupHLODActorsParams.OutAdditionalWorldPartitionsForStandaloneHLOD);
+			if (IsDistributedBuild())
+			{
+				// Generate working dirs for additional Standalone HLOD levels
+				StandaloneHLODWorkingDirs.SetNum(AdditionalWorldPartitionsForStandaloneHLOD.Num());
+				for (int32 Index = 0; Index < StandaloneHLODWorkingDirs.Num(); Index++)
+				{
+					StandaloneHLODWorkingDirs[Index] = GetDistributedBuildWorkingDir(AdditionalWorldPartitionsForStandaloneHLOD[Index]->GetWorld());
+				}
+			}
+
+			// Refresh Asset Registry to include Standalone HLOD levels that were created during SetupHLODActors
+			FDirectoryWatcherModule& DirectoryWatcherModule = FModuleManager::Get().LoadModuleChecked<FDirectoryWatcherModule>(TEXT("DirectoryWatcher"));
+			DirectoryWatcherModule.Get()->Tick(-1.0f);
+
+			FAssetRegistryModule& AssetRegistryModule = FModuleManager::LoadModuleChecked<FAssetRegistryModule>(TEXT("AssetRegistry"));
+			IAssetRegistry& AssetRegistry = AssetRegistryModule.Get();
+
+			for (const TObjectPtr<UWorldPartition>& AdditionalWorldPartition : AdditionalWorldPartitionsForStandaloneHLOD)
+			{
+				FString WorldName = AdditionalWorldPartition->GetWorld()->GetPackage()->GetName();
+				TArray<FString> ExternalObjectsPaths = ULevel::GetExternalObjectsPaths(WorldName);
+
+				AssetRegistry.ScanModifiedAssetFiles({WorldName});
+				AssetRegistry.ScanPathsSynchronous(ExternalObjectsPaths, true);
+			}
+		}
+
 		// When performing a distributed build, ensure our work folder is empty
 		if (IsDistributedBuild())
 		{
@@ -293,14 +362,27 @@ bool UWorldPartitionHLODsBuilder::SetupHLODActors()
 		UE_LOG(LogWorldPartitionHLODsBuilder, Display, TEXT("#### World HLOD actors ####"));
 
 		int32 NumActors = 0;
-		for (FActorDescContainerInstanceCollection::TIterator<AWorldPartitionHLOD> HLODIterator(WorldPartition); HLODIterator; ++HLODIterator)
+		TFunction<void(UWorldPartition*)> ListHLODActors = [&NumActors](UWorldPartition* WorldPartitionToProcess)
 		{
-			FWorldPartitionActorDescInstance* HLODActorDescInstance = *HLODIterator;
-			FString PackageName = HLODActorDescInstance->GetActorPackage().ToString();
+			for (FActorDescContainerInstanceCollection::TIterator<AWorldPartitionHLOD> HLODIterator(WorldPartitionToProcess); HLODIterator; ++HLODIterator)
+			{
+				FWorldPartitionActorDescInstance* HLODActorDescInstance = *HLODIterator;
+				FString PackageName = HLODActorDescInstance->GetActorPackage().ToString();
+
+				UE_LOG(LogWorldPartitionHLODsBuilder, Display, TEXT("    [%d] %s"), NumActors, *PackageName);
+
+				NumActors++;
+			}
+		};
 
-			UE_LOG(LogWorldPartitionHLODsBuilder, Display, TEXT("    [%d] %s"), NumActors, *PackageName);
+		ListHLODActors(WorldPartition);
 
-			NumActors++;
+		if (bBuildingStandaloneHLOD)
+		{
+			for (const TObjectPtr<UWorldPartition>& AdditionalWorldPartition : AdditionalWorldPartitionsForStandaloneHLOD)
+			{
+				ListHLODActors(AdditionalWorldPartition);
+			}
 		}
 
 		UE_LOG(LogWorldPartitionHLODsBuilder, Display, TEXT("#### World contains %d HLOD actors ####"), NumActors);
@@ -308,8 +390,10 @@ bool UWorldPartitionHLODsBuilder::SetupHLODActors()
 
 	if (IsUsingBuildManifest())
 	{
-		TMap<FString, int32> FilesToBuilderMap;
-		bool bGenerated = GenerateBuildManifest(FilesToBuilderMap);
+		// With Standalone HLOD levels we might be generating work for multiple builders across multiple worlds,
+		// keep track of Builder Index and World Index for each file
+		TMap<FString, TPair<int32, int32>> FilesToBuilderAndWorldIndexMap;
+		bool bGenerated = GenerateBuildManifest(FilesToBuilderAndWorldIndexMap);
 		if (!bGenerated)
 		{
 			return false;
@@ -321,29 +405,52 @@ bool UWorldPartitionHLODsBuilder::SetupHLODActors()
 			// Ensure we don't hold on to packages of always loaded actors
 			// When running distributed builds, we wanna leave the machine clean, so added files are deleted, check'd out files are reverted
 			// and deleted files are restored.
+			bool bCollectGarbage = false;
 			if (WorldPartition)
 			{
 				WorldPartition->Uninitialize();
+				bCollectGarbage = true;
+			}
+			
+			// Clean up Standalone HLOD levels
+			for (const TObjectPtr<UWorldPartition>& AdditionalWorldPartition : AdditionalWorldPartitionsForStandaloneHLOD)
+			{
+				if (AdditionalWorldPartition)
+				{
+					AdditionalWorldPartition->Uninitialize();
+					bCollectGarbage = true;
+				}
+			}
+			AdditionalWorldPartitionsForStandaloneHLOD.Empty();
+			
+			if (bCollectGarbage)
+			{
 				FWorldPartitionHelpers::DoCollectGarbage();
 			}
 
-			TArray<FBuilderModifiedFiles> BuildersFiles;
-			BuildersFiles.SetNum(BuilderCount);
+			TArray<TArray<FBuilderModifiedFiles>> BuildersFilesPerWorld;
+			BuildersFilesPerWorld.SetNum(bBuildingStandaloneHLOD ? StandaloneHLODWorkingDirs.Num() : 1);
+			for (TArray<FBuilderModifiedFiles>& BuildersFiles : BuildersFilesPerWorld)
+			{
+				BuildersFiles.SetNum(BuilderCount);
+			}
 
 			for (int32 i = 0; i < FBuilderModifiedFiles::EFileOperation::NumFileOperations; i++)
 			{
 				FBuilderModifiedFiles::EFileOperation FileOp = (FBuilderModifiedFiles::EFileOperation)i;
 				for (const FString& ModifiedFile : ModifiedFiles.Get(FileOp))
 				{
-					int32* Idx = FilesToBuilderMap.Find(ModifiedFile);
+					// Key - Builder Index
+					// Value - World Index
+					TPair<int32, int32>* Idx = FilesToBuilderAndWorldIndexMap.Find(ModifiedFile);
 					if (Idx)
 					{
-						BuildersFiles[*Idx].Add(FileOp, ModifiedFile);
+						BuildersFilesPerWorld[Idx->Value][Idx->Key].Add(FileOp, ModifiedFile);
 					}
 					else
 					{
-						// Add general files to the last builder
-						BuildersFiles.Last().Add(FileOp, ModifiedFile);
+						// Add general files to the last builder, first world
+						BuildersFilesPerWorld[0].Last().Add(FileOp, ModifiedFile);
 					}
 				}
 			}
@@ -352,11 +459,24 @@ bool UWorldPartitionHLODsBuilder::SetupHLODActors()
 			TArray<FString> BuildProducts;
 
 			// Copy files that will be handled by the different builders
-			for (int32 Idx = 0; Idx < BuilderCount; Idx++)
+			for (int32 WorldIndex = 0; WorldIndex < BuildersFilesPerWorld.Num(); WorldIndex++)
 			{
-				if (!CopyFilesToWorkingDir(GetHLODBuilderFolderName(Idx), BuildersFiles[Idx], BuildProducts))
+				FString WorkingDir;
+				if (bBuildingStandaloneHLOD)
 				{
-					return false;
+					WorkingDir = StandaloneHLODWorkingDirs[WorldIndex];
+				}
+				else
+				{
+					WorkingDir = DistributedBuildWorkingDir;
+				}
+
+				for (int32 Idx = 0; Idx < BuilderCount; Idx++)
+				{
+					if (!CopyFilesToWorkingDir(GetHLODBuilderFolderName(Idx), BuildersFilesPerWorld[WorldIndex][Idx], WorkingDir, BuildProducts))
+					{
+						return false;
+					}
 				}
 			}
 
@@ -371,6 +491,17 @@ bool UWorldPartitionHLODsBuilder::SetupHLODActors()
 		}
 	}
 
+	// Clean up Standalone HLOD levels if not cleaned up before
+	for (const TObjectPtr<UWorldPartition>& AdditionalWorldPartition : AdditionalWorldPartitionsForStandaloneHLOD)
+	{
+		if (AdditionalWorldPartition)
+		{
+			AdditionalWorldPartition->Uninitialize();
+		}
+	}
+	AdditionalWorldPartitionsForStandaloneHLOD.Empty();
+	FWorldPartitionHelpers::DoCollectGarbage();
+
 	return true;
 }
 
@@ -402,7 +533,9 @@ bool UWorldPartitionHLODsBuilder::BuildHLODActors()
 			return false;
 		}
 
-		if (!ValidateWorkload(HLODActorsToBuild))
+		FHLODWorkload WorkloadToValidate;
+		WorkloadToValidate.PerWorldHLODWorkloads.Add(HLODActorsToBuild);
+		if (!ValidateWorkload(WorkloadToValidate, /*bShouldConsiderExternalHLODActors=*/false))
 		{
 			return false;
 		}
@@ -502,7 +635,7 @@ bool UWorldPartitionHLODsBuilder::BuildHLODActors()
 
 		TArray<FString> BuildProducts;
 
-		if (!CopyFilesToWorkingDir("ToSubmit", ModifiedFiles, BuildProducts))
+		if (!CopyFilesToWorkingDir("ToSubmit", ModifiedFiles, DistributedBuildWorkingDir, BuildProducts))
 		{
 			return false;
 		}
@@ -528,6 +661,52 @@ bool UWorldPartitionHLODsBuilder::DeleteHLODActors()
 	};
 
 	TArray<FString> PackagesToDelete;
+
+	if (bBuildingStandaloneHLOD)
+	{
+		// Find all Stadalone HLOD levels and delete them and all their external actors
+		UWorld* SourceWorld = WorldPartition->GetWorld();
+		FString FolderPath, PackagePrefix;
+		UWorldPartitionStandaloneHLODSubsystem::GetStandaloneHLODFolderPathAndPackagePrefix(SourceWorld->GetPackage()->GetName(), FolderPath, PackagePrefix);
+
+		TArray<FString> Packages;
+		FPackageName::FindPackagesInDirectory(Packages, FolderPath);
+
+		for (const FString& Package : Packages)
+		{
+			if (!Package.Contains(PackagePrefix))
+			{
+				continue;
+			}
+
+			FString PackageName = FPackageName::FilenameToLongPackageName(Package);
+			PackagesToDelete.Add(PackageName);
+
+			const TArray<FString> ExternalObjectsPaths = ULevel::GetExternalObjectsPaths(PackageName);
+			for (const FString& ExternaObjectsPath : ExternalObjectsPaths)
+			{
+				FString ExternalObjectsDirectoryPath = FPackageName::LongPackageNameToFilename(ExternaObjectsPath);
+				if (IFileManager::Get().DirectoryExists(*ExternalObjectsDirectoryPath))
+				{
+					const bool bSuccess = IFileManager::Get().IterateDirectoryRecursively(*ExternalObjectsDirectoryPath, [&PackagesToDelete](const TCHAR* FilenameOrDirectory, bool bIsDirectory)
+					{
+						if (!bIsDirectory)
+						{
+							PackagesToDelete.Add(FilenameOrDirectory);
+						}
+
+						return true;
+					});
+
+					if (!bSuccess)
+					{
+						UE_LOG(LogWorldPartitionHLODsBuilder, Log, TEXT("Failed to iterate external package folder: %s"), *ExternalObjectsDirectoryPath);
+					}
+				}
+			}
+		}
+	}
+
 	for (FActorDescContainerInstanceCollection::TIterator<> Iterator(WorldPartition); Iterator; ++Iterator)
 	{
 		if (HLODActorClasses.FindByPredicate([ActorClass = Iterator->GetActorNativeClass()](const UClass* HLODClass) { return ActorClass->IsChildOf(HLODClass); }))
@@ -626,47 +805,82 @@ bool UWorldPartitionHLODsBuilder::GetHLODActorsToBuild(TArray<FGuid>& HLODActors
 	}
 	else
 	{
-		TArray<TArray<FGuid>> HLODWorkloads = GetHLODWorkloads(1);
-		HLODActorsToBuild = MoveTemp(HLODWorkloads[0]);
+		// When getting HLOD Workloads during Build step, we don't want to consider Standalone HLOD Actors in Standalone HLOD Levels,
+		// as they'll be considered when the builder runs directly on those levels
+		TArray<FHLODWorkload> HLODWorkloads = GetHLODWorkloads(1, /*bShouldConsiderExternalHLODActors=*/false);
+		HLODActorsToBuild = MoveTemp(HLODWorkloads[0].PerWorldHLODWorkloads[0]);
 	}
 
 	return bRet;
 }
 
-TArray<TArray<FGuid>> UWorldPartitionHLODsBuilder::GetHLODWorkloads(int32 NumWorkloads) const
+TArray<UWorldPartitionHLODsBuilder::FHLODWorkload> UWorldPartitionHLODsBuilder::GetHLODWorkloads(int32 NumWorkloads, bool bShouldConsiderExternalHLODActors) const
 {
 	if (!WorldPartition)
 	{
-		return { { FGuid() } };
+		return { FHLODWorkload() };
 	}
 
+	// Build a mapping of HLODActor to WorldPartition Index to be used when splitting actors into workloads
+	// 0 - World Partition currently processed by the builder
+	// 1 ... N - World Partitions from AdditionalWorldPartitionsForStandaloneHLOD array
+	TMap<FGuid, uint32> HLODActorToWorldPartitionIndex;
+
 	// Build a mapping of 1 HLOD[Level] -> N HLOD[Level - 1]
 	TMap<FGuid, TArray<FGuid>>	HLODParenting;
-	for (FActorDescContainerInstanceCollection::TIterator<AWorldPartitionHLOD> HLODIterator(WorldPartition); HLODIterator; ++HLODIterator)
+	TFunction<void(UWorldPartition*, uint32)> ProcessWorldPartition = [this, &HLODParenting, &HLODActorToWorldPartitionIndex, bShouldConsiderExternalHLODActors](UWorldPartition* WorldPartitionToProcess, uint32 WorldPartitionIndex)
 	{
-		const FHLODActorDesc& HLODActorDesc = *(FHLODActorDesc*)HLODIterator->GetActorDesc();
-		// Filter by HLOD actor
-		if (!HLODActorToBuild.IsNone() && HLODActorDesc.GetActorLabel() != HLODActorToBuild)
+		for (FActorDescContainerInstanceCollection::TIterator<AWorldPartitionHLOD> HLODIterator(WorldPartitionToProcess); HLODIterator; ++HLODIterator)
 		{
-			continue;
-		}
+			const FHLODActorDesc& HLODActorDesc = *(FHLODActorDesc*)HLODIterator->GetActorDesc();
+			// Filter by HLOD actor
+			if (!HLODActorToBuild.IsNone() && HLODActorDesc.GetActorLabel() != HLODActorToBuild)
+			{
+				continue;
+			}
 
-		// Filter by HLOD layer
-		if (!HLODLayerToBuild.IsNone() && HLODActorDesc.GetSourceHLODLayer().GetAssetName() != HLODLayerToBuild)
-		{
-			continue;
-		}
+			// Filter by HLOD layer
+			if (!HLODLayerToBuild.IsNone() && HLODActorDesc.GetSourceHLODLayer().GetAssetName() != HLODLayerToBuild)
+			{
+				continue;
+			}
 
-		// When requested to build a single HLOD Layer, skip the child actors
-		if (HLODLayerToBuild.IsNone())
-		{
-			HLODParenting.Add(HLODIterator->GetGuid(), HLODActorDesc.GetChildHLODActors());
+			if (bBuildingStandaloneHLOD && IsDistributedBuild())
+			{
+				HLODActorToWorldPartitionIndex.Add(HLODIterator->GetGuid(), WorldPartitionIndex);
+			}
+
+			// When requested to build a single HLOD Layer, skip the child actors
+			if (HLODLayerToBuild.IsNone())
+			{
+				TArray<FGuid>& ChildActors = HLODParenting.Add(HLODIterator->GetGuid(), HLODActorDesc.GetChildHLODActors());
+
+				if (bShouldConsiderExternalHLODActors)
+				{
+					ChildActors.Append(HLODActorDesc.GetExternalChildHLODActors());
+				}
+			}
+			else
+			{
+				HLODParenting.Add(HLODIterator->GetGuid());
+			}
 		}
-		else
+	};
+
+	// In distributed builds all workloads are prepared during the Setup step, which doesn't run on Standalone HLOD Levels, so we have to generate the workloads for them as well
+	// If building Standalone HLODs, all HLOD actors are in Standalone HLOD Levels, so we can skip processing the main world
+	// In non-distributed builds, workloads are generated during the Build step, which runs on Standalone HLOD Levels
+	if (bBuildingStandaloneHLOD && IsDistributedBuild())
+	{
+		for (int32 WorldIndex = 0; WorldIndex < AdditionalWorldPartitionsForStandaloneHLOD.Num(); WorldIndex++)
 		{
-			HLODParenting.Add(HLODIterator->GetGuid());
+			ProcessWorldPartition(AdditionalWorldPartitionsForStandaloneHLOD[WorldIndex], WorldIndex);
 		}
 	}
+	else
+	{
+		ProcessWorldPartition(WorldPartition, 0);
+	}
 
 	// All child HLODs must be built before their parent HLOD
 	// Create groups to ensure those will be processed in the correct order, on the same builder
@@ -708,76 +922,134 @@ TArray<TArray<FGuid>> UWorldPartitionHLODsBuilder::GetHLODWorkloads(int32 NumWor
 	HLODGroups.ValueSort([](const TArray<FGuid>& GroupA, const TArray<FGuid>& GroupB) { return GroupA.Num() > GroupB.Num(); });
 
 	// Dispatch them in multiple lists and try to balance the workloads as much as possible
-	TArray<TArray<FGuid>> Workloads;
+	TArray<FHLODWorkload> Workloads;
 	Workloads.SetNum(NumWorkloads);
 
+	for (FHLODWorkload& Workload : Workloads)
+	{
+		if (bBuildingStandaloneHLOD && IsDistributedBuild())
+		{
+			Workload.PerWorldHLODWorkloads.SetNum(AdditionalWorldPartitionsForStandaloneHLOD.Num());
+		}
+		else
+		{
+			Workload.PerWorldHLODWorkloads.SetNum(1);
+		}
+	}
+
 	int32 Idx = 0;
 	for (const auto& Pair : HLODGroups)
 	{
-		Workloads[Idx % NumWorkloads].Append(Pair.Value);
+		int32 WorkloadNum = Idx % NumWorkloads;
+		for (const FGuid& HLODActorGuid : Pair.Value)
+		{
+			int32 WorldIndex;
+			if (bBuildingStandaloneHLOD && IsDistributedBuild())
+			{
+				// We might be generating workloads for a few Worlds at the same time. Find which one, so that we can assign actor to the right workload
+				const uint32* IndexPtr = HLODActorToWorldPartitionIndex.Find(HLODActorGuid);
+				check(IndexPtr);
+				WorldIndex = *IndexPtr;
+			}
+			else
+			{
+				WorldIndex = 0;
+			}
+			
+			Workloads[WorkloadNum].PerWorldHLODWorkloads[WorldIndex].Add(HLODActorGuid);
+		}
 		Idx++;
 	}
 
 	// Validate workloads to ensure our meshes are built in the correct order
-	for (const TArray<FGuid>& Workload : Workloads)
+	for (const FHLODWorkload& Workload : Workloads)
 	{
-		check(ValidateWorkload(Workload));
+		check(ValidateWorkload(Workload, bShouldConsiderExternalHLODActors));
 	}
 
 	return Workloads;
 }
 
-bool UWorldPartitionHLODsBuilder::ValidateWorkload(const TArray<FGuid>& Workload) const
+bool UWorldPartitionHLODsBuilder::ValidateWorkload(const FHLODWorkload& Workload, bool bShouldConsiderExternalHLODActors) const
 {
 	check(WorldPartition);
 
+	uint32 NumHLODs = 0;
+	for (const TArray<FGuid>& HLODActorArray : Workload.PerWorldHLODWorkloads)
+	{
+		NumHLODs += HLODActorArray.Num();
+	}
+
 	TSet<FGuid> ProcessedHLOD;
-	ProcessedHLOD.Reserve(Workload.Num());
+	ProcessedHLOD.Reserve(NumHLODs);
 
 	// For each HLOD entry in the workload, validate that its children are found before itself
-	for (const FGuid& HLODActorGuid : Workload)
+	for (int32 WorldIndex = 0; WorldIndex < Workload.PerWorldHLODWorkloads.Num(); WorldIndex++)
 	{
-		const FWorldPartitionActorDescInstance* ActorDescInstance = WorldPartition->GetActorDescInstance(HLODActorGuid);
-		if(!ActorDescInstance)
-		{
-			UE_LOG(LogWorldPartitionHLODsBuilder, Error, TEXT("Unknown actor guid found (\"%s\"), your HLOD actors are probably out of date. Run with -SetupHLODs to fix this. Exiting..."), *HLODActorGuid.ToString());
-			return false;
-		}
-
-		if (!ActorDescInstance->GetActorNativeClass()->IsChildOf<AWorldPartitionHLOD>())
+		UWorldPartition* CurrentWorldPartition = (bBuildingStandaloneHLOD && IsDistributedBuild()) ? AdditionalWorldPartitionsForStandaloneHLOD[WorldIndex].Get() : WorldPartition;
+		for (const FGuid& HLODActorGuid : Workload.PerWorldHLODWorkloads[WorldIndex])
 		{
-			UE_LOG(LogWorldPartitionHLODsBuilder, Error, TEXT("Unexpected actor guid found in HLOD workload (\"%s\"), exiting..."), *HLODActorGuid.ToString());
-			return false;
-		}
+			const FWorldPartitionActorDescInstance* ActorDescInstance = CurrentWorldPartition->GetActorDescInstance(HLODActorGuid);
+			if(!ActorDescInstance)
+			{
+				UE_LOG(LogWorldPartitionHLODsBuilder, Error, TEXT("Unknown actor guid found (\"%s\"), your HLOD actors are probably out of date. Run with -SetupHLODs to fix this. Exiting..."), *HLODActorGuid.ToString());
+				return false;
+			}
 
-		// When requested to build a single HLOD Layer, do not validate that child actors are included
-		if (HLODLayerToBuild.IsNone())
-		{
-			const FHLODActorDesc* HLODActorDesc = static_cast<const FHLODActorDesc*>(ActorDescInstance->GetActorDesc());
+			if (!ActorDescInstance->GetActorNativeClass()->IsChildOf<AWorldPartitionHLOD>())
+			{
+				UE_LOG(LogWorldPartitionHLODsBuilder, Error, TEXT("Unexpected actor guid found in HLOD workload (\"%s\"), exiting..."), *HLODActorGuid.ToString());
+				return false;
+			}
 
-			for (const FGuid& ChildHLODActorGuid : HLODActorDesc->GetChildHLODActors())
+			// When requested to build a single HLOD Layer, do not validate that child actors are included
+			if (HLODLayerToBuild.IsNone())
 			{
-				if (!ProcessedHLOD.Contains(ChildHLODActorGuid))
+				const FHLODActorDesc* HLODActorDesc = static_cast<const FHLODActorDesc*>(ActorDescInstance->GetActorDesc());
+
+				for (const FGuid& ChildHLODActorGuid : HLODActorDesc->GetChildHLODActors())
 				{
-					UE_LOG(LogWorldPartitionHLODsBuilder, Error, TEXT("Child HLOD actor (\"%s\") missing or out of order in HLOD workload, exiting..."), *HLODActorGuid.ToString());
-					return false;
+					if (!ProcessedHLOD.Contains(ChildHLODActorGuid))
+					{
+						UE_LOG(LogWorldPartitionHLODsBuilder, Error, TEXT("Child HLOD actor (\"%s\") missing or out of order in HLOD workload, exiting..."), *ChildHLODActorGuid.ToString());
+						return false;
+					}
+				}
+
+				// Skip checking whether external child actors are included if we're not considering them
+				if (bShouldConsiderExternalHLODActors)
+				{
+					for (const FGuid& ExternalChildHLODActorGuid : HLODActorDesc->GetExternalChildHLODActors())
+					{
+						if (!ProcessedHLOD.Contains(ExternalChildHLODActorGuid))
+						{
+							UE_LOG(LogWorldPartitionHLODsBuilder, Error, TEXT("External child HLOD actor (\"%s\") missing or out of order in HLOD workload, exiting..."), *ExternalChildHLODActorGuid.ToString());
+							return false;
+						}
+					}
 				}
 			}
-		}
 
-		ProcessedHLOD.Add(HLODActorGuid);
+			ProcessedHLOD.Add(HLODActorGuid);
+		}
 	}
 
 	return true;
 }
 
-bool UWorldPartitionHLODsBuilder::GenerateBuildManifest(TMap<FString, int32>& FilesToBuilderMap) const
+bool UWorldPartitionHLODsBuilder::GenerateBuildManifest(TMap<FString, TPair<int32, int32>>& FilesToBuilderAndWorldIndexMap) const
 {
-	TArray<TArray<FGuid>> BuildersWorkload = GetHLODWorkloads(BuilderCount);
+	// We're generating manifest for Standalone HLOD levels as well (if any), so we want to consider External HLOD actors
+	TArray<FHLODWorkload> BuildersWorkload = GetHLODWorkloads(BuilderCount, /*bShouldConsiderExternalHLODActors=*/true);
 
-	FConfigFile ConfigFile;
-	ConfigFile.SetInt64(TEXT("General"), TEXT("BuilderCount"), BuilderCount);
-	ConfigFile.SetString(TEXT("General"), TEXT("EngineVersion"), *FEngineVersion::Current().ToString());
+	// If we're generating manifest for Standalone HLOD levels, each of them needs a separate config file
+	TArray<FConfigFile> ConfigFiles;
+	ConfigFiles.SetNum(WorldPartition->HasStandaloneHLOD() ? AdditionalWorldPartitionsForStandaloneHLOD.Num() : 1);
+	for (FConfigFile& ConfigFile : ConfigFiles)
+	{
+		ConfigFile.SetInt64(TEXT("General"), TEXT("BuilderCount"), BuilderCount);
+		ConfigFile.SetString(TEXT("General"), TEXT("EngineVersion"), *FEngineVersion::Current().ToString());
+	}
 
 	// When processing multiple maps, ensure that the worldload is distributed evenly between builders.
 	// Otherwise, maps with a single HLOD would all end up being processed by the first builder, while the others would have no work.
@@ -788,43 +1060,60 @@ bool UWorldPartitionHLODsBuilder::GenerateBuildManifest(TMap<FString, int32>& Fi
 		const int32 WorkloadIndex = Idx;
 		const int32 BuilderIndex = (BuilderDispatchOffset + Idx) % BuilderCount;
 
-		if (!BuildersWorkload.IsValidIndex(WorkloadIndex) || BuildersWorkload[WorkloadIndex].IsEmpty())
+		if (!BuildersWorkload.IsValidIndex(WorkloadIndex) || BuildersWorkload[WorkloadIndex].PerWorldHLODWorkloads.IsEmpty())
 		{
 			continue;
 		}
 
 		FString SectionName = GetHLODBuilderFolderName(BuilderIndex);
 
-		for(const FGuid& ActorGuid : BuildersWorkload[WorkloadIndex])
+		for (int32 WorldIndex = 0; WorldIndex < BuildersWorkload[WorkloadIndex].PerWorldHLODWorkloads.Num(); WorldIndex++)
 		{
-			ConfigFile.AddToSection(*SectionName, TEXT("+HLODActorGuid"), ActorGuid.ToString(EGuidFormats::Digits));
-
-			if (WorldPartition)
+			UWorldPartition* CurrentWorldPartition = WorldPartition->HasStandaloneHLOD() ? AdditionalWorldPartitionsForStandaloneHLOD[WorldIndex].Get() : WorldPartition;
+			for(const FGuid& ActorGuid : BuildersWorkload[WorkloadIndex].PerWorldHLODWorkloads[WorldIndex])
 			{
-				// Track which builder is responsible to handle each actor
-				const FWorldPartitionActorDescInstance* ActorDescInstance = WorldPartition->GetActorDescInstance(ActorGuid);
-				if (!ActorDescInstance)
+				ConfigFiles[WorldIndex].AddToSection(*SectionName, TEXT("+HLODActorGuid"), ActorGuid.ToString(EGuidFormats::Digits));
+
+				if (CurrentWorldPartition)
 				{
-					UE_LOG(LogWorldPartitionHLODsBuilder, Error, TEXT("Invalid actor GUID found while generating the HLOD build manifest, exiting..."));
-					return false;
+					// Track which builder is responsible to handle each actor
+					const FWorldPartitionActorDescInstance* ActorDescInstance = CurrentWorldPartition->GetActorDescInstance(ActorGuid);
+					if (!ActorDescInstance)
+					{
+						UE_LOG(LogWorldPartitionHLODsBuilder, Error, TEXT("Invalid actor GUID found while generating the HLOD build manifest, exiting..."));
+						return false;
+					}
+					FString ActorPackageFilename = USourceControlHelpers::PackageFilename(ActorDescInstance->GetActorPackage().ToString());
+					FilesToBuilderAndWorldIndexMap.Emplace(ActorPackageFilename, TPair<int32, int32>(BuilderIndex, WorldIndex));
 				}
-				FString ActorPackageFilename = USourceControlHelpers::PackageFilename(ActorDescInstance->GetActorPackage().ToString());
-				FilesToBuilderMap.Emplace(ActorPackageFilename, BuilderIndex);
 			}
 		}
 	}
 
 	BuilderDispatchOffset++;
 
-	ConfigFile.Dirty = true;
-
-	bool bRet = ConfigFile.Write(BuildManifest);
-	if (!bRet)
+	for (int32 Index = 0; Index < ConfigFiles.Num(); Index++)
 	{
-		UE_LOG(LogWorldPartitionHLODsBuilder, Error, TEXT("Failed to write HLOD build manifest \"%s\""), *BuildManifest);
+		FString BuildManifestFile;
+		if (bBuildingStandaloneHLOD)
+		{
+			BuildManifestFile = StandaloneHLODWorkingDirs[Index] / DistributedBuildManifestName;
+		}
+		else
+		{
+			BuildManifestFile = BuildManifest;
+		}
+
+		ConfigFiles[Index].Dirty = true;
+
+		if (!ConfigFiles[Index].Write(BuildManifestFile))
+		{
+			UE_LOG(LogWorldPartitionHLODsBuilder, Error, TEXT("Failed to write HLOD build manifest \"%s\""), *BuildManifestFile);
+			return false;
+		}
 	}
 
-	return bRet;
+	return true;
 }
 
 /*
@@ -864,9 +1153,9 @@ const FName FileAction_Add(TEXT("Add"));
 const FName FileAction_Edit(TEXT("Edit"));
 const FName FileAction_Delete(TEXT("Delete"));
 
-bool UWorldPartitionHLODsBuilder::CopyFilesToWorkingDir(const FString& TargetDir, const FBuilderModifiedFiles& Files, TArray<FString>& BuildProducts)
+bool UWorldPartitionHLODsBuilder::CopyFilesToWorkingDir(const FString& TargetDir, const FBuilderModifiedFiles& Files, const FString& WorkingDir, TArray<FString>& BuildProducts)
 {
-	const FString AbsoluteTargetDir = DistributedBuildWorkingDir / TargetDir / TEXT("");
+	const FString AbsoluteTargetDir = WorkingDir / TargetDir / TEXT("");
 
 	bool bSuccess = true;
 
diff --git a/Engine/Source/Editor/UnrealEd/Public/WorldPartition/WorldPartitionBuilder.h b/Engine/Source/Editor/UnrealEd/Public/WorldPartition/WorldPartitionBuilder.h
--- a/Engine/Source/Editor/UnrealEd/Public/WorldPartition/WorldPartitionBuilder.h
+++ b/Engine/Source/Editor/UnrealEd/Public/WorldPartition/WorldPartitionBuilder.h
@@ -64,6 +64,12 @@ public:
 	DECLARE_DELEGATE_RetVal_TwoParams(bool, FModifiedFilesHandler, const TArray<FString>&, const FString&); /*bool ModifiedFileHander(ModifiedFiles, ChangeDescription) */
 	UNREALED_API void SetModifiedFilesHandler(const FModifiedFilesHandler& ModifiedFilesHandler);
 
+	/**
+	 * Overridable method for derived classes to add additional world to be process by world builder
+	 * This is called after `RunInternal` is executed.
+	 */
+	UNREALED_API virtual bool ShouldProcessAdditionalWorlds(UWorld* InWorld, TArray<FString>& OutPackageNames) const { return false; }
+
 protected:
 	/**
 	 * Overridable method for derived classes to perform operations when world builder process starts.
diff --git a/Engine/Source/Editor/UnrealEd/Public/WorldPartition/WorldPartitionHLODsBuilder.h b/Engine/Source/Editor/UnrealEd/Public/WorldPartition/WorldPartitionHLODsBuilder.h
--- a/Engine/Source/Editor/UnrealEd/Public/WorldPartition/WorldPartitionHLODsBuilder.h
+++ b/Engine/Source/Editor/UnrealEd/Public/WorldPartition/WorldPartitionHLODsBuilder.h
@@ -31,9 +31,15 @@ public:
 	UNREALED_API virtual bool PreWorldInitialization(UWorld* World, FPackageSourceControlHelper& PackageHelper) override;
 
 protected:
+	struct FHLODWorkload
+	{
+		TArray<TArray<FGuid>> PerWorldHLODWorkloads;
+	};
+
 	UNREALED_API virtual bool RunInternal(UWorld* World, const FCellInfo& InCellInfo, FPackageSourceControlHelper& PackageHelper) override;
 	virtual bool CanProcessNonPartitionedWorlds() const override { return true; }
 	UNREALED_API virtual bool ShouldProcessWorld(UWorld* World) const override;
+	UNREALED_API virtual bool ShouldProcessAdditionalWorlds(UWorld* InWorld, TArray<FString>& OutPackageNames) const override;
 	// UWorldPartitionBuilder interface end
 
 	bool IsDistributedBuild() const { return bDistributedBuild; }
@@ -46,13 +52,13 @@ protected:
 	UNREALED_API bool SubmitHLODActors();
 	UNREALED_API bool DumpStats();
 
-	UNREALED_API bool GenerateBuildManifest(TMap<FString, int32>& FilesToBuilderMap) const;
+	UNREALED_API bool GenerateBuildManifest(TMap<FString, TPair<int32, int32>>& FilesToBuilderAndWorldIndexMap) const;
 	UNREALED_API bool GetHLODActorsToBuild(TArray<FGuid>& HLODActorsToBuild) const;
 
-	UNREALED_API TArray<TArray<FGuid>> GetHLODWorkloads(int32 NumWorkloads) const;
-	UNREALED_API bool ValidateWorkload(const TArray<FGuid>& Workload) const;
+	UNREALED_API TArray<FHLODWorkload> GetHLODWorkloads(int32 NumWorkloads, bool bShouldConsiderExternalHLODActors) const;
+	UNREALED_API bool ValidateWorkload(const FHLODWorkload& Workload, bool bShouldConsiderExternalHLODActors) const;
 
-	UNREALED_API bool CopyFilesToWorkingDir(const FString& TargetDir, const FBuilderModifiedFiles& ModifiedFiles, TArray<FString>& BuildProducts);
+	UNREALED_API bool CopyFilesToWorkingDir(const FString& TargetDir, const FBuilderModifiedFiles& ModifiedFiles, const FString& WorkingDir, TArray<FString>& BuildProducts);
 	UNREALED_API bool CopyFilesFromWorkingDir(const FString& SourceDir);
 
 	UNREALED_API bool ShouldRunStep(const EHLODBuildStep BuildStep) const;
@@ -84,4 +90,8 @@ private:
 	FString DistributedBuildManifest;
 	
 	FBuilderModifiedFiles ModifiedFiles;
+
+	bool bBuildingStandaloneHLOD;
+	TArray<TObjectPtr<UWorldPartition>> AdditionalWorldPartitionsForStandaloneHLOD;
+	TArray<FString> StandaloneHLODWorkingDirs;
 };
\ No newline at end of file
diff --git a/Engine/Source/Editor/WorldPartitionEditor/Private/WorldPartition/HLOD/HLODEditorData.cpp b/Engine/Source/Editor/WorldPartitionEditor/Private/WorldPartition/HLOD/HLODEditorData.cpp
--- a/Engine/Source/Editor/WorldPartitionEditor/Private/WorldPartition/HLOD/HLODEditorData.cpp
+++ b/Engine/Source/Editor/WorldPartitionEditor/Private/WorldPartition/HLOD/HLODEditorData.cpp
@@ -54,7 +54,7 @@ void FWorldPartitionHLODEditorData::OnActorDescContainerInstanceRegistered(UActo
 		const FGuid& HLODActorGuid = HLODActorDesc.GetGuid();
 
 		TUniquePtr<FHLODSceneNode>& HLODSceneNode = ContainerInstanceHLODActorData->HLODActorNodes.Emplace(HLODActorGuid, new FHLODSceneNode());
-		HLODSceneNode->Bounds = HLODActorDesc.GetEditorBounds();
+		HLODSceneNode->Bounds = HLODActorDesc.GetEditorBounds().TransformBy(WorldPartition->GetInstanceTransform());
 		HLODSceneNode->HLODActorHandle = FWorldPartitionHandle(WorldPartition, HLODActorDesc.GetGuid());
 	}
 
@@ -155,6 +155,11 @@ static void GatherLoadedActorsBounds(TArray<FBoundsWithVolume>& OutLoadedBounds,
 	}
 }
 
+bool FWorldPartitionHLODEditorData::IsLoadedActorsStateInitialized() const
+{
+	return LastStateUpdate != INDEX_NONE;
+}
+
 void FWorldPartitionHLODEditorData::ClearLoadedActorsState()
 {
 	// Increment update counter - used to quickly find out if an HLOD actor needs be hidden without having to flag the whole hierarchy
diff --git a/Engine/Source/Editor/WorldPartitionEditor/Private/WorldPartition/HLOD/HLODEditorData.h b/Engine/Source/Editor/WorldPartitionEditor/Private/WorldPartition/HLOD/HLODEditorData.h
--- a/Engine/Source/Editor/WorldPartitionEditor/Private/WorldPartition/HLOD/HLODEditorData.h
+++ b/Engine/Source/Editor/WorldPartitionEditor/Private/WorldPartition/HLOD/HLODEditorData.h
@@ -31,6 +31,7 @@ public:
 	FWorldPartitionHLODEditorData(UWorldPartition* InWorldPartition);
 	~FWorldPartitionHLODEditorData();
 	
+	bool IsLoadedActorsStateInitialized() const;
 	void ClearLoadedActorsState();
 	void UpdateLoadedActorsState();
 	void UpdateVisibility(const FVector& InCameraLocation, double InMinDrawDistance, double InMaxDrawDistance, bool bForceVisibilityUpdate);
diff --git a/Engine/Source/Editor/WorldPartitionEditor/Private/WorldPartition/HLOD/HLODEditorSubsystem.cpp b/Engine/Source/Editor/WorldPartitionEditor/Private/WorldPartition/HLOD/HLODEditorSubsystem.cpp
--- a/Engine/Source/Editor/WorldPartitionEditor/Private/WorldPartition/HLOD/HLODEditorSubsystem.cpp
+++ b/Engine/Source/Editor/WorldPartitionEditor/Private/WorldPartition/HLOD/HLODEditorSubsystem.cpp
@@ -210,10 +210,11 @@ void UWorldPartitionHLODEditorSubsystem::OnWorldPartitionInitialized(UWorldParti
 {
 	TRACE_CPUPROFILER_EVENT_SCOPE(UWorldPartitionHLODEditorSubsystem::OnWorldPartitionInitialized);
 	
-	if (InWorldPartition->IsMainWorldPartition())
+	if (InWorldPartition->IsMainWorldPartition() || InWorldPartition->IsStandaloneHLODWorld())
 	{
 		InWorldPartition->LoaderAdapterStateChanged.AddUObject(this, &UWorldPartitionHLODEditorSubsystem::OnLoaderAdapterStateChanged);
-		HLODEditorData = MakePimpl<FWorldPartitionHLODEditorData>(InWorldPartition);
+		TPimplPtr<FWorldPartitionHLODEditorData>& HLODEditorData = WorldPartitionsHLODEditorData.Emplace(InWorldPartition, MakePimpl<FWorldPartitionHLODEditorData>(InWorldPartition));
+		HLODEditorData->ClearLoadedActorsState();
 		ForceHLODStateUpdate();
 	}
 }
@@ -222,10 +223,10 @@ void UWorldPartitionHLODEditorSubsystem::OnWorldPartitionUninitialized(UWorldPar
 {
 	TRACE_CPUPROFILER_EVENT_SCOPE(UWorldPartitionHLODEditorSubsystem::OnWorldPartitionUninitialized);
 
-	if (InWorldPartition->IsMainWorldPartition())
+	if (InWorldPartition->IsMainWorldPartition() || InWorldPartition->IsStandaloneHLODWorld())
 	{
 		InWorldPartition->LoaderAdapterStateChanged.RemoveAll(this);
-		HLODEditorData = nullptr;
+		WorldPartitionsHLODEditorData.Remove(InWorldPartition);
 	}
 }
 
@@ -248,74 +249,81 @@ void UWorldPartitionHLODEditorSubsystem::Tick(float DeltaTime)
 {
 	TRACE_CPUPROFILER_EVENT_SCOPE(UWorldPartitionHLODEditorSubsystem::Tick);
 
-	if (HLODEditorData)
+	UUnrealEditorSubsystem* UnrealEditorSubsystem = GEditor->GetEditorSubsystem<UUnrealEditorSubsystem>();
+
+	bool bCameraMoved = false;
+	bool bForceHLODVisibilityUpdate = false;
+	bool bClearLoadedActorState = false;
+
+	// Check cached global settings
+	if (IsHLODInEditorEnabled())
 	{
-		HLODEditorData->SetHLODLoadingState(IsHLODInEditorEnabled());
-		
-		if (IsHLODInEditorEnabled())
+		IWorldPartitionEditorModule* WorldPartitionEditorModule = FModuleManager::GetModulePtr<IWorldPartitionEditorModule>("WorldPartitionEditor");
+
+		// "Show HLODs over loaded region" option changed ?
+		if (WorldPartitionEditorModule->GetShowHLODsOverLoadedRegions() != bCachedShowHLODsOverLoadedRegions)
 		{
-			bool bForceHLODVisibilityUpdate = false;
+			bCachedShowHLODsOverLoadedRegions = WorldPartitionEditorModule->GetShowHLODsOverLoadedRegions();
+			bForceHLODVisibilityUpdate = true;
+			bForceHLODStateUpdate = true;
+			bClearLoadedActorState = bCachedShowHLODsOverLoadedRegions;
+		}
 
-			IWorldPartitionEditorModule* WorldPartitionEditorModule = FModuleManager::GetModulePtr<IWorldPartitionEditorModule>("WorldPartitionEditor");
+		// Min/Max draw distance for HLODs was changed ?
+		if (WorldPartitionEditorModule->GetHLODInEditorMinDrawDistance() != CachedHLODMinDrawDistance ||
+			WorldPartitionEditorModule->GetHLODInEditorMaxDrawDistance() != CachedHLODMaxDrawDistance)
+		{
+			CachedHLODMinDrawDistance = WorldPartitionEditorModule->GetHLODInEditorMinDrawDistance();
+			CachedHLODMaxDrawDistance = WorldPartitionEditorModule->GetHLODInEditorMaxDrawDistance();
+			bForceHLODVisibilityUpdate = true;
+		}
 
-			// "Show HLODs over loaded region" option changed ?
-			if (WorldPartitionEditorModule->GetShowHLODsOverLoadedRegions() != bCachedShowHLODsOverLoadedRegions)
-			{ 
-				bCachedShowHLODsOverLoadedRegions = WorldPartitionEditorModule->GetShowHLODsOverLoadedRegions();
-				bForceHLODVisibilityUpdate = true;
+		if (UnrealEditorSubsystem)
+		{
+			FVector CameraLocation;
+			FRotator CameraRotation;
 
-				// Clear loading state of actors if we are going to always show HLODs
-				if (bCachedShowHLODsOverLoadedRegions)
-				{
-					HLODEditorData->ClearLoadedActorsState();
-				}
-				else
+			if (FWorldPartitionEditorModule::GetActiveLevelViewportCameraInfo(CameraLocation, CameraRotation))
+			{
+				// Camera was moved ?
+				bCameraMoved = CameraLocation != CachedCameraLocation;
+				if (bCameraMoved)
 				{
-					bForceHLODStateUpdate = true;
+					CachedCameraLocation = CameraLocation;
 				}
 			}
+		}
+	}
 
-			// Actors or regions were loaded ?
-			if (bForceHLODStateUpdate && !bCachedShowHLODsOverLoadedRegions)
+	for (TPair<TObjectKey<UWorldPartition>, TPimplPtr<FWorldPartitionHLODEditorData>>& Pair : WorldPartitionsHLODEditorData)
+	{
+		TPimplPtr<FWorldPartitionHLODEditorData>& HLODEditorData = Pair.Value;
+
+		HLODEditorData->SetHLODLoadingState(IsHLODInEditorEnabled());
+		
+		if (IsHLODInEditorEnabled())
+		{
+			bool bNeedsInitialization = !HLODEditorData->IsLoadedActorsStateInitialized();
+
+			if (bClearLoadedActorState || (bNeedsInitialization && bCachedShowHLODsOverLoadedRegions))
 			{
-				HLODEditorData->UpdateLoadedActorsState();
-				bForceHLODVisibilityUpdate = true;
+				HLODEditorData->ClearLoadedActorsState();
 			}
-			
-			// Min/Max draw distance for HLODs was changed ?
-			if (WorldPartitionEditorModule->GetHLODInEditorMinDrawDistance() != CachedHLODMinDrawDistance ||
-				WorldPartitionEditorModule->GetHLODInEditorMaxDrawDistance() != CachedHLODMaxDrawDistance)
+
+			// Actors or regions were loaded ?
+			if ((bForceHLODStateUpdate || bNeedsInitialization) && !bCachedShowHLODsOverLoadedRegions)
 			{
-				CachedHLODMinDrawDistance = WorldPartitionEditorModule->GetHLODInEditorMinDrawDistance();
-				CachedHLODMaxDrawDistance = WorldPartitionEditorModule->GetHLODInEditorMaxDrawDistance();
+				HLODEditorData->UpdateLoadedActorsState();
 				bForceHLODVisibilityUpdate = true;
 			}
 
-			UUnrealEditorSubsystem* UnrealEditorSubsystem = GEditor->GetEditorSubsystem<UUnrealEditorSubsystem>();
-			if (UnrealEditorSubsystem)
+			if (bForceHLODVisibilityUpdate || bCameraMoved || bNeedsInitialization)
 			{
-				FVector CameraLocation;
-				FRotator CameraRotation;
-				bool bCameraMoved = false;
-
-				if (FWorldPartitionEditorModule::GetActiveLevelViewportCameraInfo(CameraLocation, CameraRotation))
-				{
-					// Camera was moved ?
-					bCameraMoved = CameraLocation != CachedCameraLocation;
-					if (bCameraMoved)
-					{
-						CachedCameraLocation = CameraLocation;
-					}
-				}
-
-				if (bForceHLODVisibilityUpdate || bCameraMoved)
-				{
-					HLODEditorData->UpdateVisibility(CameraLocation, CachedHLODMinDrawDistance, CachedHLODMaxDrawDistance, bForceHLODVisibilityUpdate);
-				}
+				HLODEditorData->UpdateVisibility(CachedCameraLocation, CachedHLODMinDrawDistance, CachedHLODMaxDrawDistance, bForceHLODVisibilityUpdate);
 			}
-			bForceHLODStateUpdate = false;
 		}
 	}
+	bForceHLODStateUpdate = false;
 }
 
 TStatId UWorldPartitionHLODEditorSubsystem::GetStatId() const
diff --git a/Engine/Source/Editor/WorldPartitionEditor/Public/WorldPartition/HLOD/HLODEditorSubsystem.h b/Engine/Source/Editor/WorldPartitionEditor/Public/WorldPartition/HLOD/HLODEditorSubsystem.h
--- a/Engine/Source/Editor/WorldPartitionEditor/Public/WorldPartition/HLOD/HLODEditorSubsystem.h
+++ b/Engine/Source/Editor/WorldPartitionEditor/Public/WorldPartition/HLOD/HLODEditorSubsystem.h
@@ -7,6 +7,7 @@
 #include "WorldPartition/WorldPartitionActorLoaderInterface.h"
 #include "WorldPartition/WorldPartitionHandle.h"
 #include "Templates/PimplPtr.h"
+#include "UObject/ObjectKey.h"
 
 #include "HLODEditorSubsystem.generated.h"
 
@@ -85,7 +86,7 @@ private:
 	bool bCachedShowHLODsOverLoadedRegions;
 	bool bForceHLODStateUpdate;
 
-	TPimplPtr<FWorldPartitionHLODEditorData> HLODEditorData;
+	TMap<TObjectKey<UWorldPartition>, TPimplPtr<FWorldPartitionHLODEditorData>> WorldPartitionsHLODEditorData;
 
 	typedef TMap<TSoftObjectPtr<UStruct>, TSet<FName>> FStructsPropertiesMap;
 	static TMap<EHLODSettingsVisibility, FStructsPropertiesMap> StructsPropertiesVisibility;
diff --git a/Engine/Source/Runtime/Core/Public/UObject/UE5ReleaseStreamObjectVersions.inl b/Engine/Source/Runtime/Core/Public/UObject/UE5ReleaseStreamObjectVersions.inl
--- a/Engine/Source/Runtime/Core/Public/UObject/UE5ReleaseStreamObjectVersions.inl
+++ b/Engine/Source/Runtime/Core/Public/UObject/UE5ReleaseStreamObjectVersions.inl
@@ -171,5 +171,8 @@ UE5_RELEASE_VERSION(BloomIndependentLensFlare, 53)
 // Add settings to IAnimationDataModel GUID generation
 UE5_RELEASE_VERSION(AnimModelGuidGenerationSettings, 54)
 
+// Add support for Standalone HLOD
+UE5_RELEASE_VERSION(WorldPartitionAddStandaloneHLODSupport, 55)
+
 // -----<new versions can be added above this line>-------------------------------------------------
 #undef UE5_RELEASE_VERSION
\ No newline at end of file
diff --git a/Engine/Source/Runtime/Engine/Classes/Engine/Level.h b/Engine/Source/Runtime/Engine/Classes/Engine/Level.h
--- a/Engine/Source/Runtime/Engine/Classes/Engine/Level.h
+++ b/Engine/Source/Runtime/Engine/Classes/Engine/Level.h
@@ -1534,6 +1534,7 @@ private:
 	friend class FWorldPartitionLevelHelper;
 	friend class UActorFolder;
 	friend class AWorldDataLayers;
+	friend class UWorldPartitionRuntimeHashSet;
 
 	/** Replace the existing LSA (if set) by spawning a new one based on this level's script blueprint */
 	void RegenerateLevelScriptActor();
diff --git a/Engine/Source/Runtime/Engine/Private/WorldPartition/HLOD/HLODActor.cpp b/Engine/Source/Runtime/Engine/Private/WorldPartition/HLOD/HLODActor.cpp
--- a/Engine/Source/Runtime/Engine/Private/WorldPartition/HLOD/HLODActor.cpp
+++ b/Engine/Source/Runtime/Engine/Private/WorldPartition/HLOD/HLODActor.cpp
@@ -93,6 +93,16 @@ const FGuid& AWorldPartitionHLOD::GetSourceCellGuid() const
 	return SourceCellGuid;
 }
 
+bool AWorldPartitionHLOD::IsStandalone() const
+{
+	return StandaloneHLODGuid.IsValid();
+}
+
+const FGuid& AWorldPartitionHLOD::GetStandaloneHLODGuid() const
+{
+	return StandaloneHLODGuid;
+}
+
 void AWorldPartitionHLOD::SetVisibility(bool bInVisible)
 {
 	ForEachComponent<USceneComponent>(false, [bInVisible](USceneComponent* SceneComponent)
@@ -511,6 +521,11 @@ void AWorldPartitionHLOD::SetSourceCellGuid(const FGuid& InSourceCellGuid)
 	SourceCellGuid = InSourceCellGuid;
 }
 
+void AWorldPartitionHLOD::SetIsStandalone(bool bInIsStandalone)
+{
+	StandaloneHLODGuid = bInIsStandalone ? GetActorGuid() : FGuid();
+}
+
 const FBox& AWorldPartitionHLOD::GetHLODBounds() const
 {
 	return HLODBounds;
diff --git a/Engine/Source/Runtime/Engine/Private/WorldPartition/HLOD/HLODActorDesc.cpp b/Engine/Source/Runtime/Engine/Private/WorldPartition/HLOD/HLODActorDesc.cpp
--- a/Engine/Source/Runtime/Engine/Private/WorldPartition/HLOD/HLODActorDesc.cpp
+++ b/Engine/Source/Runtime/Engine/Private/WorldPartition/HLOD/HLODActorDesc.cpp
@@ -17,6 +17,7 @@
 #include "WorldPartition/HLOD/HLODLayer.h"
 #include "WorldPartition/HLOD/HLODSourceActorsFromCell.h"
 #include "WorldPartition/WorldPartitionActorDescInstance.h"
+#include "WorldPartition/WorldPartitionActorDescUtils.h"
 
 FHLODActorDesc::FHLODActorDesc()
 {}
@@ -53,6 +54,14 @@ void FHLODActorDesc::Init(const AActor* InActor)
 					ChildHLODActors.Add(SubActor.ActorInstanceGuid);
 				}
 			}
+			else if (WorldPartition->IsStandaloneHLODWorld())
+			{
+				UClass* NativeClassPtr = FWorldPartitionActorDescUtils::GetActorNativeClassFromString(SubActor.NativeClass.ToString());
+				if (NativeClassPtr && NativeClassPtr->IsChildOf<AWorldPartitionHLOD>())
+				{
+					ExternalChildHLODActors.Add(SubActor.ActorInstanceGuid);
+				}
+			}
 		}
 	}
 
@@ -165,6 +174,11 @@ void FHLODActorDesc::Serialize(FArchive& Ar)
 				Ar << EditorBounds;
 			}
 		}
+
+		if (Ar.CustomVer(FUE5ReleaseStreamObjectVersion::GUID) >= FUE5ReleaseStreamObjectVersion::WorldPartitionAddStandaloneHLODSupport)
+		{
+			Ar << ExternalChildHLODActors;
+		}
 	}
 }
 
@@ -176,6 +190,7 @@ bool FHLODActorDesc::Equals(const FWorldPartitionActorDesc* Other) const
 		return SourceHLODLayer == HLODActorDesc.SourceHLODLayer &&
 			   HLODStats.OrderIndependentCompareEqual(HLODActorDesc.HLODStats) &&
 			   CompareUnsortedArrays(ChildHLODActors, HLODActorDesc.ChildHLODActors) &&
+			   CompareUnsortedArrays(ExternalChildHLODActors, HLODActorDesc.ExternalChildHLODActors) &&
 			   EditorBounds == HLODActorDesc.EditorBounds;
 	}
 	return false;
diff --git a/Engine/Source/Runtime/Engine/Private/WorldPartition/HLOD/HLODRuntimeSubsystem.cpp b/Engine/Source/Runtime/Engine/Private/WorldPartition/HLOD/HLODRuntimeSubsystem.cpp
--- a/Engine/Source/Runtime/Engine/Private/WorldPartition/HLOD/HLODRuntimeSubsystem.cpp
+++ b/Engine/Source/Runtime/Engine/Private/WorldPartition/HLOD/HLODRuntimeSubsystem.cpp
@@ -156,6 +156,15 @@ FAutoConsoleCommand UWorldPartitionHLODRuntimeSubsystem::EnableHLODCommand(
 							}
 						}
 					}
+
+					// Update visibility for Standalone HLOD actors referencing unloaded cells
+					for (const TPair<FGuid, TSet<IWorldPartitionHLODObject*>>& Pair : HLODSubSystem->StandaloneHLODObjectsReferencingUnloadedCells)
+					{
+						for (IWorldPartitionHLODObject* HLODObject : Pair.Value)
+						{
+							HLODObject->SetVisibility(UWorldPartitionHLODRuntimeSubsystem::WorldPartitionHLODEnabled);
+						}
+					}
 				}
 			}
 		}
@@ -209,11 +218,28 @@ void UWorldPartitionHLODRuntimeSubsystem::OnWorldPartitionInitialized(UWorldPart
 		// Build cell to HLOD mapping
 		if (InWorldPartition && InWorldPartition->RuntimeHash)
 		{
-			InWorldPartition->RuntimeHash->ForEachStreamingCells([&WorldPartitionHLODRuntimeData](const UWorldPartitionRuntimeCell* Cell)
+			InWorldPartition->RuntimeHash->ForEachStreamingCells([this, &WorldPartitionHLODRuntimeData, &InWorldPartition](const UWorldPartitionRuntimeCell* Cell)
 			{
 				UE_LOG(LogHLODRuntimeSubsystem, Verbose, TEXT("Registering cell %s - %s"), *Cell->GetGuid().ToString(), *Cell->GetDebugName());
 
-				WorldPartitionHLODRuntimeData.CellsData.Emplace(Cell->GetGuid());
+				FCellData& CellData = WorldPartitionHLODRuntimeData.CellsData.Emplace(Cell->GetGuid());
+
+				if (InWorldPartition->HasStandaloneHLOD())
+				{
+					StandaloneHLODCellToWorldPartitionMap.Add(Cell->GetGuid(), InWorldPartition);
+
+					// If we have Standalone HLOD Actor for this cell, attach it to the cell
+					if (TSet<IWorldPartitionHLODObject*>* StandaloneHLODObjects = StandaloneHLODObjectsReferencingUnloadedCells.Find(Cell->GetGuid()))
+					{
+						for (IWorldPartitionHLODObject* HLODObject : (*StandaloneHLODObjects))
+						{
+							CellData.LoadedHLODs.Add(HLODObject);
+							HLODObject->SetVisibility(UWorldPartitionHLODRuntimeSubsystem::WorldPartitionHLODEnabled && !CellData.bIsCellVisible);
+							UE_LOG(LogHLODRuntimeSubsystem, Verbose, TEXT("Attached Standalone HLOD %s to cell %s"), *HLODObject->GetHLODNameOrLabel(), *Cell->GetGuid().ToString());
+						}
+						StandaloneHLODObjectsReferencingUnloadedCells.Remove(Cell->GetGuid());
+					}
+				}
 				return true;
 			});
 		}
@@ -226,8 +252,30 @@ void UWorldPartitionHLODRuntimeSubsystem::OnWorldPartitionUninitialized(UWorldPa
 	{
 		check(WorldPartitionsHLODRuntimeData.Contains(InWorldPartition));
 
-#if !NO_LOGGING
 		FWorldPartitionHLODRuntimeData& WorldPartitionHLODRuntimeData = WorldPartitionsHLODRuntimeData.FindChecked(InWorldPartition);
+
+		if (InWorldPartition->HasStandaloneHLOD())
+		{
+			for (const TPair<FGuid, FCellData>& It : WorldPartitionHLODRuntimeData.CellsData)
+			{
+				const FCellData& CellData = It.Value;
+				FGuid CellGuid = It.Key;
+
+				StandaloneHLODCellToWorldPartitionMap.Remove(CellGuid);
+
+				for (IWorldPartitionHLODObject* HLODObject : CellData.LoadedHLODs)
+				{
+					if (HLODObject->IsStandalone())
+					{
+						StandaloneHLODObjectsReferencingUnloadedCells.FindOrAdd(CellGuid).Add(HLODObject);
+						HLODObject->SetVisibility(UWorldPartitionHLODRuntimeSubsystem::WorldPartitionHLODEnabled);
+						UE_LOG(LogHLODRuntimeSubsystem, Verbose, TEXT("Detached Standalone HLOD %s from cell %s"), *HLODObject->GetHLODNameOrLabel(), *CellGuid.ToString());
+					}
+				}
+			}
+		}
+
+#if !NO_LOGGING
 		InWorldPartition->RuntimeHash->ForEachStreamingCells([&WorldPartitionHLODRuntimeData](const UWorldPartitionRuntimeCell* Cell)
 		{
 			UE_LOG(LogHLODRuntimeSubsystem, Verbose, TEXT("Unregistering cell %s - %s"), *Cell->GetGuid().ToString(), *Cell->GetDebugName());
@@ -307,6 +355,21 @@ UWorldPartitionHLODRuntimeSubsystem::FCellData* UWorldPartitionHLODRuntimeSubsys
 	return nullptr;		
 }
 
+UWorldPartitionHLODRuntimeSubsystem::FCellData* UWorldPartitionHLODRuntimeSubsystem::GetCellData(const UWorldPartition* InWorldPartition, const FGuid& InCellGuid)
+{
+	if (InWorldPartition)
+	{
+		FWorldPartitionHLODRuntimeData* WorldPartitionHLODRuntimeData = WorldPartitionsHLODRuntimeData.Find(InWorldPartition);
+		if (WorldPartitionHLODRuntimeData)
+		{
+			check(InWorldPartition->IsStreamingEnabled());
+			return WorldPartitionHLODRuntimeData->CellsData.Find(InCellGuid);
+		}
+	}
+
+	return nullptr;
+}
+
 const TArray<IWorldPartitionHLODObject*>& UWorldPartitionHLODRuntimeSubsystem::GetHLODObjectsForCell(const UWorldPartitionRuntimeCell* InCell) const
 {
 	if (const FCellData* CellData = GetCellData(InCell))
@@ -331,6 +394,58 @@ void UWorldPartitionHLODRuntimeSubsystem::RegisterHLODObject(IWorldPartitionHLOD
 		SourceCellData->LoadedHLODs.Add(InWorldPartitionHLOD);
 		InWorldPartitionHLOD->SetVisibility(UWorldPartitionHLODRuntimeSubsystem::WorldPartitionHLODEnabled && !SourceCellData->bIsCellVisible);
 	}
+	else if (InWorldPartitionHLOD->IsStandalone())
+	{	
+		const UWorldPartition* WorldPartition = FHLODRuntimeSubsystem::GetWorldPartition(InWorldPartitionHLOD);
+
+		// Standalone HLOD actors are embedded in the main world streaming cells, so for Standalone HLOD actors that are using other Standalone HLOD actors as source actors,
+		// source cell GUID won't be correct. During streaming generation, we save the actual source cell GUID, so that we can use that override here.
+		if (const FGuid* StandaloneHLODSourceCellsGuidPtr = WorldPartition->RuntimeHash->GetStandaloneHLODActorSourceCellOverride(InWorldPartitionHLOD->GetStandaloneHLODGuid()))
+		{
+			// Source cell GUID override found
+			FGuid StandaloneHLODSourceCellGuid = *StandaloneHLODSourceCellsGuidPtr;
+			if (FCellData* StandaloneHLODSourceCellData = GetCellData(WorldPartition, StandaloneHLODSourceCellGuid))
+			{
+				// Overriden source cell found in WorldPartitionHLODRuntimeData. Add the HLOD Actor to its LoadedHLODs.
+				UE_LOG(LogHLODRuntimeSubsystem, Verbose, TEXT("Registering Standalone HLOD %s for overriden source cell '%s'"), *InWorldPartitionHLOD->GetHLODNameOrLabel(), *StandaloneHLODSourceCellGuid.ToString());
+
+				StandaloneHLODSourceCellData->LoadedHLODs.Add(InWorldPartitionHLOD);
+				InWorldPartitionHLOD->SetVisibility(UWorldPartitionHLODRuntimeSubsystem::WorldPartitionHLODEnabled && !StandaloneHLODSourceCellData->bIsCellVisible);
+			}
+			else
+			{
+				// Overriden source cell not found in WorldPartitionHLODRuntimeData. This can happen when Standalone LI is not currently loaded.
+				// Storing the HLOD object in StandaloneHLODObjectsReferencingUnloadedCells so that we can associate it later with the overridden source cell, once it's loaded.
+				UE_LOG(LogHLODRuntimeSubsystem, Verbose, TEXT("Registering Standalone HLOD %s referencing currently not loaded overriden source cell '%s'"), *InWorldPartitionHLOD->GetHLODNameOrLabel(), *StandaloneHLODSourceCellGuid.ToString());
+
+				TSet<IWorldPartitionHLODObject*>& StandaloneHLODs = StandaloneHLODObjectsReferencingUnloadedCells.FindOrAdd(StandaloneHLODSourceCellGuid);
+				StandaloneHLODs.Add(InWorldPartitionHLOD);
+				InWorldPartitionHLOD->SetVisibility(UWorldPartitionHLODRuntimeSubsystem::WorldPartitionHLODEnabled);
+			}
+		}
+		else if (TObjectPtr<UWorldPartition>* StandaloneWorldPartitionPtr = StandaloneHLODCellToWorldPartitionMap.Find(InWorldPartitionHLOD->GetSourceCellGuid()))
+		{
+			UWorldPartition* StandaloneWorldPartition = *StandaloneWorldPartitionPtr;
+			if (FCellData* StandaloneHLODSourceCellData = GetCellData(StandaloneWorldPartition, InWorldPartitionHLOD->GetSourceCellGuid()))
+			{
+				// Source cell found in WorldPartitionHLODRuntimeData
+				UE_LOG(LogHLODRuntimeSubsystem, Verbose, TEXT("Registering Standalone HLOD %s for Standalone LI source cell '%s'"), *InWorldPartitionHLOD->GetHLODNameOrLabel(), *InWorldPartitionHLOD->GetSourceCellGuid().ToString());
+
+				StandaloneHLODSourceCellData->LoadedHLODs.Add(InWorldPartitionHLOD);
+				InWorldPartitionHLOD->SetVisibility(UWorldPartitionHLODRuntimeSubsystem::WorldPartitionHLODEnabled && !StandaloneHLODSourceCellData->bIsCellVisible);
+			}
+		}
+		else
+		{
+			// Source cell not found in WorldPartitionHLODRuntimeData. This can happen when Standalone LI is not currently loaded.
+			// Storing the HLOD object in StandaloneHLODObjectsReferencingUnloadedCells so that we can associate it later with the source cell, once it's loaded.
+			UE_LOG(LogHLODRuntimeSubsystem, Verbose, TEXT("Registering Standalone HLOD %s referencing currently not loaded cell '%s'"), *InWorldPartitionHLOD->GetHLODNameOrLabel(), *InWorldPartitionHLOD->GetSourceCellGuid().ToString());
+
+			TSet<IWorldPartitionHLODObject*>& StandaloneHLODs = StandaloneHLODObjectsReferencingUnloadedCells.FindOrAdd(InWorldPartitionHLOD->GetSourceCellGuid());
+			StandaloneHLODs.Add(InWorldPartitionHLOD);
+			InWorldPartitionHLOD->SetVisibility(UWorldPartitionHLODRuntimeSubsystem::WorldPartitionHLODEnabled);
+		}
+	}
 	else
 	{
 		UE_LOG(LogHLODRuntimeSubsystem, Verbose, TEXT("Found HLOD %s referencing nonexistent cell '%s'"), *InWorldPartitionHLOD->GetHLODNameOrLabel(), *InWorldPartitionHLOD->GetSourceCellGuid().ToString());
@@ -356,6 +471,55 @@ void UWorldPartitionHLODRuntimeSubsystem::UnregisterHLODObject(IWorldPartitionHL
 
 		verify(CellData->LoadedHLODs.Remove(InWorldPartitionHLOD));
 	}
+	else if (InWorldPartitionHLOD->IsStandalone())
+	{
+		const UWorldPartition* WorldPartition = FHLODRuntimeSubsystem::GetWorldPartition(InWorldPartitionHLOD);
+
+		// Check if we have source cell override. See RegisterHLODActor for more info.
+		if (const FGuid* SourceCellGuidPtr = WorldPartition->RuntimeHash->GetStandaloneHLODActorSourceCellOverride(InWorldPartitionHLOD->GetStandaloneHLODGuid()))
+		{
+			// Source cell GUID override found
+			const FGuid SourceCellGuid = *SourceCellGuidPtr;
+			if (FCellData* SourceCellData = GetCellData(WorldPartition, SourceCellGuid))
+			{
+				// Overriden source cell found. Remove the HLOD object from its LoadedHLODs.
+				UE_LOG(LogHLODRuntimeSubsystem, Verbose, TEXT("Unregistering Standalone HLOD %s for overriden source cell %s"), *InWorldPartitionHLOD->GetHLODNameOrLabel(), *SourceCellGuid.ToString());
+
+				verify(SourceCellData->LoadedHLODs.Remove(InWorldPartitionHLOD));
+			}
+			else
+			{
+				// Overriden source cell not found. Remove the HLOD object from StandaloneHLODObjectsReferencingUnloadedCells map
+				UE_LOG(LogHLODRuntimeSubsystem, Verbose, TEXT("Unregistering Standalone HLOD %s for currently not loaded overriden source cell %s"), *InWorldPartitionHLOD->GetHLODNameOrLabel(), *SourceCellGuid.ToString());
+
+				if (TSet<IWorldPartitionHLODObject*>* StandaloneHLODActorsPtr = StandaloneHLODObjectsReferencingUnloadedCells.Find(SourceCellGuid))
+				{
+					verify(StandaloneHLODActorsPtr->Remove(InWorldPartitionHLOD));
+				}
+			}
+		}
+		else if (TObjectPtr<UWorldPartition>* StandaloneWorldPartitionPtr = StandaloneHLODCellToWorldPartitionMap.Find(InWorldPartitionHLOD->GetSourceCellGuid()))
+		{
+			UWorldPartition* StandaloneWorldPartition = *StandaloneWorldPartitionPtr;
+			if (FCellData* SourceCellData = GetCellData(StandaloneWorldPartition, InWorldPartitionHLOD->GetSourceCellGuid()))
+			{
+				// Source cell found. Remove the HLOD object from its LoadedHLODs.
+				UE_LOG(LogHLODRuntimeSubsystem, Verbose, TEXT("Unregistering Standalone HLOD %s for Standalone LI source cell %s"), *InWorldPartitionHLOD->GetHLODNameOrLabel(), *InWorldPartitionHLOD->GetSourceCellGuid().ToString());
+
+				verify(SourceCellData->LoadedHLODs.Remove(InWorldPartitionHLOD));
+			}
+		}
+		else
+		{
+			// Source cell not found. Remove the HLOD actor from StandaloneHLODObjectsReferencingUnloadedCells map
+			UE_LOG(LogHLODRuntimeSubsystem, Verbose, TEXT("Unregistering Standalone HLOD %s for currently not loaded cell %s"), *InWorldPartitionHLOD->GetHLODNameOrLabel(), *InWorldPartitionHLOD->GetSourceCellGuid().ToString());
+
+			if (TSet<IWorldPartitionHLODObject*>* StandaloneHLODActorsPtr = StandaloneHLODObjectsReferencingUnloadedCells.Find(InWorldPartitionHLOD->GetSourceCellGuid()))
+			{
+				verify(StandaloneHLODActorsPtr->Remove(InWorldPartitionHLOD));
+			}
+		}
+	}
 	else
 	{
 #if !UE_BUILD_SHIPPING
diff --git a/Engine/Source/Runtime/Engine/Private/WorldPartition/HLOD/StandaloneHLODActor.cpp b/Engine/Source/Runtime/Engine/Private/WorldPartition/HLOD/StandaloneHLODActor.cpp
new file mode 100644
--- /dev/null
+++ b/Engine/Source/Runtime/Engine/Private/WorldPartition/HLOD/StandaloneHLODActor.cpp
@@ -0,0 +1,111 @@
+// Copyright Epic Games, Inc. All Rights Reserved.
+
+#include "WorldPartition/HLOD/StandaloneHLODActor.h"
+
+#include "Engine/World.h"
+#include "LevelInstance/LevelInstanceComponent.h"
+#include "WorldPartition/LevelInstance/LevelInstanceActorDesc.h"
+
+AWorldPartitionStandaloneHLOD::AWorldPartitionStandaloneHLOD(const FObjectInitializer& ObjectInitializer)
+	: Super(ObjectInitializer)
+	, LevelInstanceActorGuid(this)
+	, LevelInstanceActorImpl(this)
+{
+	RootComponent = CreateDefaultSubobject<ULevelInstanceComponent>(TEXT("Root"));
+	RootComponent->Mobility = EComponentMobility::Static;
+
+#if WITH_EDITORONLY_DATA
+	bDefaultOutlinerExpansionState = false;
+#endif
+}
+
+const FLevelInstanceID& AWorldPartitionStandaloneHLOD::GetLevelInstanceID() const
+{
+	return LevelInstanceActorImpl.GetLevelInstanceID();
+}
+
+bool AWorldPartitionStandaloneHLOD::HasValidLevelInstanceID() const
+{
+	return LevelInstanceActorImpl.HasValidLevelInstanceID();
+}
+
+const FGuid& AWorldPartitionStandaloneHLOD::GetLevelInstanceGuid() const
+{
+	return LevelInstanceActorGuid.GetGuid();
+}
+
+const TSoftObjectPtr<UWorld>& AWorldPartitionStandaloneHLOD::GetWorldAsset() const
+{
+	return WorldAsset;
+}
+
+bool AWorldPartitionStandaloneHLOD::IsLoadingEnabled() const
+{
+	return LevelInstanceActorImpl.IsLoadingEnabled();
+}
+
+bool AWorldPartitionStandaloneHLOD::SetWorldAsset(TSoftObjectPtr<UWorld> InWorldAsset)
+{
+	WorldAsset = InWorldAsset;
+	return true;
+}
+
+#if WITH_EDITOR
+ULevelInstanceComponent* AWorldPartitionStandaloneHLOD::GetLevelInstanceComponent() const
+{
+	return Cast<ULevelInstanceComponent>(RootComponent);
+}
+
+ELevelInstanceRuntimeBehavior AWorldPartitionStandaloneHLOD::GetDesiredRuntimeBehavior() const
+{
+	return ELevelInstanceRuntimeBehavior::Partitioned;
+}
+
+ELevelInstanceRuntimeBehavior AWorldPartitionStandaloneHLOD::GetDefaultRuntimeBehavior() const
+{
+	return ELevelInstanceRuntimeBehavior::Partitioned;
+}
+
+bool AWorldPartitionStandaloneHLOD::CanEditChange(const FProperty* InProperty) const
+{
+	return false;
+}
+
+bool AWorldPartitionStandaloneHLOD::CanEditChangeComponent(const UActorComponent* Component, const FProperty* InProperty) const
+{
+	return false;
+}
+#endif // WITH_EDITOR
+
+void AWorldPartitionStandaloneHLOD::PostRegisterAllComponents()
+{
+	Super::PostRegisterAllComponents();
+
+	if (GetLocalRole() == ENetRole::ROLE_Authority && GetWorld()->IsGameWorld())
+	{
+#if !WITH_EDITOR
+		// If the level instance was spawned, not loaded
+		LevelInstanceActorGuid.AssignIfInvalid();
+#endif
+		LevelInstanceSpawnGuid = LevelInstanceActorGuid.GetGuid();
+	}
+
+	if (LevelInstanceActorGuid.IsValid())
+	{
+		LevelInstanceActorImpl.RegisterLevelInstance();
+	}
+}
+
+void AWorldPartitionStandaloneHLOD::PostUnregisterAllComponents()
+{
+	Super::PostUnregisterAllComponents();
+
+	LevelInstanceActorImpl.UnregisterLevelInstance();
+}
+
+#if WITH_EDITOR
+TUniquePtr<FWorldPartitionActorDesc> AWorldPartitionStandaloneHLOD::CreateClassActorDesc() const
+{
+	return TUniquePtr<FWorldPartitionActorDesc>(new FLevelInstanceActorDesc());
+}
+#endif
\ No newline at end of file
diff --git a/Engine/Source/Runtime/Engine/Private/WorldPartition/HLOD/StandaloneHLODSubsystem.cpp b/Engine/Source/Runtime/Engine/Private/WorldPartition/HLOD/StandaloneHLODSubsystem.cpp
new file mode 100644
--- /dev/null
+++ b/Engine/Source/Runtime/Engine/Private/WorldPartition/HLOD/StandaloneHLODSubsystem.cpp
@@ -0,0 +1,398 @@
+// Copyright Epic Games, Inc. All Rights Reserved.
+
+#include "WorldPartition/HLOD/StandaloneHLODSubsystem.h"
+
+#include UE_INLINE_GENERATED_CPP_BY_NAME(StandaloneHLODSubsystem)
+
+#if WITH_EDITOR
+
+#include "LevelInstance/LevelInstanceActor.h"
+#include "WorldPartition/ActorDescContainerSubsystem.h"
+#include "WorldPartition/HLOD/StandaloneHLODActor.h"
+#include "WorldPartition/LevelInstance/LevelInstanceActorDesc.h"
+#include "WorldPartition/WorldPartition.h"
+#include "WorldPartition/RuntimeHashSet/WorldPartitionRuntimeHashSet.h"
+#include "WorldPartition/WorldPartitionSubsystem.h"
+
+DEFINE_LOG_CATEGORY_STATIC(LogStandaloneHLODSubsystem, Log, All);
+
+void UWorldPartitionStandaloneHLODSubsystem::Initialize(FSubsystemCollectionBase& Collection)
+{
+	bRefreshCachedHLODSetups = true;
+
+	Collection.InitializeDependency<UWorldPartitionSubsystem>();
+
+	Super::Initialize(Collection);
+
+	GetWorld()->OnWorldPartitionInitialized().AddUObject(this, &UWorldPartitionStandaloneHLODSubsystem::OnWorldPartitionInitialized);
+	GetWorld()->OnWorldPartitionUninitialized().AddUObject(this, &UWorldPartitionStandaloneHLODSubsystem::OnWorldPartitionUninitialized);
+
+	GEngine->OnLevelActorAdded().AddUObject(this, &UWorldPartitionStandaloneHLODSubsystem::OnActorChanged);
+	GEngine->OnActorMoved().AddUObject(this, &UWorldPartitionStandaloneHLODSubsystem::OnActorChanged);
+	GEngine->OnLevelActorDeleted().AddUObject(this, &UWorldPartitionStandaloneHLODSubsystem::OnActorDeleted);
+	FCoreUObjectDelegates::OnObjectPropertyChanged.AddUObject(this, &UWorldPartitionStandaloneHLODSubsystem::OnObjectPropertyChanged);
+	FWorldDelegates::LevelAddedToWorld.AddUObject(this, &UWorldPartitionStandaloneHLODSubsystem::OnLevelAddedToWorld);
+}
+
+void UWorldPartitionStandaloneHLODSubsystem::Deinitialize()
+{
+	Super::Deinitialize();
+
+	GetWorld()->OnWorldPartitionInitialized().RemoveAll(this);
+	GetWorld()->OnWorldPartitionUninitialized().RemoveAll(this);
+
+	GEngine->OnLevelActorAdded().RemoveAll(this);
+	GEngine->OnActorMoved().RemoveAll(this);
+	GEngine->OnLevelActorDeleted().RemoveAll(this);
+	FCoreUObjectDelegates::OnObjectPropertyChanged.RemoveAll(this);
+}
+
+bool UWorldPartitionStandaloneHLODSubsystem::DoesSupportWorldType(const EWorldType::Type WorldType) const
+{
+	return WorldType == EWorldType::Editor || WorldType == EWorldType::Inactive;
+}
+
+void UWorldPartitionStandaloneHLODSubsystem::OnWorldPartitionInitialized(UWorldPartition* InWorldPartition)
+{
+	InWorldPartition->OnActorDescContainerInstanceRegistered.AddUObject(this, &UWorldPartitionStandaloneHLODSubsystem::OnActorDescContainerInstanceRegistered);
+	InWorldPartition->OnActorDescContainerInstanceUnregistered.AddUObject(this, &UWorldPartitionStandaloneHLODSubsystem::OnActorDescContainerInstanceUnregistered);
+
+	// Since this is created upon WP init, we missed the first broadcasts for existing container instances. Register them manually.
+	InWorldPartition->ForEachActorDescContainerInstance([this](UActorDescContainerInstance* InContainerInstance)
+	{
+		OnActorDescContainerInstanceRegistered(InContainerInstance);
+	});
+
+	if (InWorldPartition->HasStandaloneHLOD())
+	{
+		FStandaloneHLODActorParams Params;
+
+		Params.Guid = FGuid();
+		Params.Transform = FTransform::Identity;
+		Params.WorldPackageName = InWorldPartition->GetPackage()->GetPathName();
+		Params.ActorLabel = InWorldPartition->GetWorld()->GetName();
+
+		UpdateStandaloneHLODActors(Params);
+	}
+}
+
+void UWorldPartitionStandaloneHLODSubsystem::OnWorldPartitionUninitialized(UWorldPartition* InWorldPartition)
+{
+	InWorldPartition->OnActorDescContainerInstanceRegistered.RemoveAll(this);
+	InWorldPartition->OnActorDescContainerInstanceUnregistered.RemoveAll(this);
+}
+
+void UWorldPartitionStandaloneHLODSubsystem::OnActorChanged(AActor* InActor)
+{
+	if (ILevelInstanceInterface* LevelInstance = Cast<ILevelInstanceInterface>(InActor))
+	{
+		// If the changed actor is a SubWP, update its Standalone HLOD actors
+		if (LevelInstance->GetDesiredRuntimeBehavior() == ELevelInstanceRuntimeBehavior::LevelStreaming)
+		{
+			FStandaloneHLODActorParams Params;
+
+			Params.Guid = InActor->GetActorGuid();
+			Params.Transform = InActor->GetActorTransform();
+			Params.WorldPackageName = LevelInstance->GetWorldAssetPackage();
+			Params.ActorLabel = InActor->GetActorLabel();
+
+			UpdateStandaloneHLODActors(Params);
+		}
+
+		// Propagate the actor change to currently loaded level instances
+		const ULevelInstanceSubsystem* LevelInstanceSubsystem = GetWorld()->GetSubsystem<ULevelInstanceSubsystem>();
+		LevelInstanceSubsystem->ForEachLevelInstanceChild(LevelInstance, /*bRecursive=*/false, [this](ILevelInstanceInterface* ChildLevelInstance)
+		{
+			OnActorChanged(CastChecked<AActor>(ChildLevelInstance));
+			return true;
+		});
+
+		// Propagate the actor change to currently unloaded level instances
+		if (UWorldPartition* WorldPartition = GetWorld()->GetWorldPartition())
+		{
+			if (FWorldPartitionActorDescInstance* ActorDescInstance = WorldPartition->GetActorDescInstance(InActor->GetActorGuid()))
+			{
+				if (FLevelInstanceActorDesc* LevelInstanceActorDesc = (FLevelInstanceActorDesc*)ActorDescInstance->GetActorDesc())
+				{
+					UpdateStandaloneHLODActorsRecursive(*LevelInstanceActorDesc, InActor->GetActorTransform(), /*bChildrenOnly=*/true);
+				}
+			}
+		}
+	}
+}
+
+void UWorldPartitionStandaloneHLODSubsystem::OnActorDeleted(AActor* InActor)
+{
+	if (ILevelInstanceInterface* LevelInstance = Cast<ILevelInstanceInterface>(InActor))
+	{
+		// Delete LI Standalone HLOD actors
+		DeleteStandaloneHLODActors(InActor->GetActorGuid());
+
+		// Propagate actor deletion to currently loaded level instances
+		const ULevelInstanceSubsystem* LevelInstanceSubsystem = GetWorld()->GetSubsystem<ULevelInstanceSubsystem>();
+		LevelInstanceSubsystem->ForEachLevelInstanceChild(LevelInstance, /*bRecursive=*/false, [this](ILevelInstanceInterface* ChildLevelInstance)
+		{
+			OnActorDeleted(CastChecked<AActor>(ChildLevelInstance));
+			return true;
+		});
+
+		// Propagate actor deletion to currently unloaded level instances
+		if (UWorldPartition* WorldPartition = GetWorld()->GetWorldPartition())
+		{
+			if (FWorldPartitionActorDescInstance* ActorDescInstance = WorldPartition->GetActorDescInstance(InActor->GetActorGuid()))
+			{
+				if (FLevelInstanceActorDesc* LevelInstanceActorDesc = (FLevelInstanceActorDesc*)ActorDescInstance->GetActorDesc())
+				{
+					DeleteStandaloneHLODActorsRecursive(*LevelInstanceActorDesc);
+				}
+			}
+		}
+	}
+}
+
+void UWorldPartitionStandaloneHLODSubsystem::OnObjectPropertyChanged(UObject* InObject, FPropertyChangedEvent& InPropertyChangedEvent)
+{
+	if (const ILevelInstanceInterface* LevelInstance = Cast<ILevelInstanceInterface>(InObject))
+	{
+		if (AWorldPartitionStandaloneHLOD* StandaloneHLOD = Cast<AWorldPartitionStandaloneHLOD>(InObject))
+		{
+			return;
+		}
+
+		if (AActor* Actor = Cast<AActor>(InObject))
+		{
+			OnActorDeleted(Actor);
+			OnActorChanged(Actor);
+		}
+	}
+}
+
+void UWorldPartitionStandaloneHLODSubsystem::OnActorDescContainerInstanceRegistered(UActorDescContainerInstance* InContainerInstance)
+{
+	for (UActorDescContainerInstance::TConstIterator<> Iterator(InContainerInstance); Iterator; ++Iterator)
+	{
+		if (Iterator->GetActorDesc()->HasStandaloneHLOD())
+		{
+			const FLevelInstanceActorDesc& LevelInstanceActorDesc = *(FLevelInstanceActorDesc*)Iterator->GetActorDesc();
+			FTransform ContainerTransform = FTransform::Identity;
+			if (const UActorDescContainerInstance* ParentContainerInstance = InContainerInstance->GetParentContainerInstance())
+			{
+				if (FWorldPartitionActorDescInstance* ParentContainerActorDescInstance = ParentContainerInstance->GetActorDescInstance(InContainerInstance->GetContainerActorGuid()))
+				{
+					ContainerTransform = ContainerTransform * ParentContainerActorDescInstance->GetActorTransform();
+				}
+			}
+			UpdateStandaloneHLODActorsRecursive(LevelInstanceActorDesc, ContainerTransform * LevelInstanceActorDesc.GetActorTransform(), /*bChildrenOnly=*/false);
+		}
+	}
+}
+
+void UWorldPartitionStandaloneHLODSubsystem::OnActorDescContainerInstanceUnregistered(UActorDescContainerInstance* InContainerInstance)
+{
+	for (UActorDescContainerInstance::TConstIterator<> Iterator(InContainerInstance); Iterator; ++Iterator)
+	{
+		if (Iterator->GetActorDesc()->HasStandaloneHLOD())
+		{
+			if (FLevelInstanceActorDesc* LevelInstanceActorDesc = (FLevelInstanceActorDesc*)Iterator->GetActorDesc())
+			{
+				DeleteStandaloneHLODActorsRecursive(*LevelInstanceActorDesc);
+			}
+		}
+	}
+}
+
+void UWorldPartitionStandaloneHLODSubsystem::OnLevelAddedToWorld(ULevel* Level, UWorld* World)
+{
+	const ULevelInstanceSubsystem* LevelInstanceSubsystem = GetWorld()->GetSubsystem<ULevelInstanceSubsystem>();
+	if (ILevelInstanceInterface* LevelInstance = LevelInstanceSubsystem ? LevelInstanceSubsystem->GetOwningLevelInstance(Level) : nullptr)
+	{
+		LevelInstanceSubsystem->ForEachLevelInstanceChild(LevelInstance, /*bRecursive=*/false, [this](ILevelInstanceInterface* ChildLevelInstance)
+		{
+			OnActorChanged(CastChecked<AActor>(ChildLevelInstance));
+			return true;
+		});
+	}
+}
+
+void UWorldPartitionStandaloneHLODSubsystem::UpdateStandaloneHLODActorsRecursive(const FLevelInstanceActorDesc& InLevelInstanceActorDesc, const FTransform InActorTransform, bool bChildrenOnly)
+{
+	if (!bChildrenOnly && InLevelInstanceActorDesc.GetDesiredRuntimeBehavior() == ELevelInstanceRuntimeBehavior::LevelStreaming)
+	{
+		FStandaloneHLODActorParams Params;
+		Params.Guid = InLevelInstanceActorDesc.GetGuid();
+		Params.Transform = InActorTransform;
+		Params.WorldPackageName = InLevelInstanceActorDesc.GetChildContainerPackage().ToString();
+		Params.ActorLabel = InLevelInstanceActorDesc.GetActorLabel().ToString();
+
+		UpdateStandaloneHLODActors(Params);
+	}
+
+	UActorDescContainer* Container = UActorDescContainerSubsystem::GetChecked().RegisterContainer({InLevelInstanceActorDesc.GetChildContainerPackage()});
+	for (FActorDescList::TConstIterator<> ActorDescIt(Container); ActorDescIt; ++ActorDescIt)
+	{
+		if (ActorDescIt->HasStandaloneHLOD())
+		{
+			const FLevelInstanceActorDesc& LevelInstanceActorDesc = *static_cast<const FLevelInstanceActorDesc*>(*ActorDescIt);
+
+			UpdateStandaloneHLODActorsRecursive(LevelInstanceActorDesc, InActorTransform * LevelInstanceActorDesc.GetActorTransform(), /*bChildrenOnly=*/false);
+		}
+	}
+	UActorDescContainerSubsystem::GetChecked().UnregisterContainer(Container);
+}
+
+void UWorldPartitionStandaloneHLODSubsystem::UpdateStandaloneHLODActors(UWorldPartitionStandaloneHLODSubsystem::FStandaloneHLODActorParams InStandaloneHLODActorParams)
+{
+	TArray<AWorldPartitionStandaloneHLOD*>* HLODActorsPtr = StandaloneHLODActors.Find(InStandaloneHLODActorParams.Guid);
+	if (HLODActorsPtr)
+	{
+		TArray<AWorldPartitionStandaloneHLOD*>& HLODActors = *HLODActorsPtr;
+
+		for (AWorldPartitionStandaloneHLOD* HLODActor : HLODActors)
+		{
+			int32 Position = HLODActor->GetActorLabel().Find(TEXT("_HLOD"), ESearchCase::IgnoreCase, ESearchDir::FromEnd);
+			FString NameSuffix = HLODActor->GetActorLabel().RightChop(Position);
+			FString NewActorLabel = InStandaloneHLODActorParams.ActorLabel + NameSuffix;
+
+			HLODActor->SetActorTransform(InStandaloneHLODActorParams.Transform);
+			HLODActor->SetActorLabel(NewActorLabel);
+		}
+	}
+	else
+	{
+		FString FolderPath, LevelPackagePrefix;
+		UWorldPartitionStandaloneHLODSubsystem::GetStandaloneHLODFolderPathAndPackagePrefix(InStandaloneHLODActorParams.WorldPackageName, FolderPath, LevelPackagePrefix);
+
+		TArray<FString> Packages;
+		FPackageName::FindPackagesInDirectory(Packages, FolderPath);
+
+		for (const FString& Package : Packages)
+		{
+			if (!Package.Contains(LevelPackagePrefix))
+			{
+				continue;
+			}
+
+			FString PackageName = FPackageName::FilenameToLongPackageName(Package);
+			const FSoftObjectPath LODWorldPackagePath(FTopLevelAssetPath(FName(*PackageName), FName(*FPackageName::GetLongPackageAssetName(PackageName))));
+
+			int32 Position = PackageName.Find(TEXT("_HLOD"), ESearchCase::IgnoreCase, ESearchDir::FromEnd);
+			FString NameSuffix = PackageName.RightChop(Position);
+			FString ActorLabel = InStandaloneHLODActorParams.ActorLabel + NameSuffix;
+			int32 HLODIndex = FCString::Atoi(*NameSuffix.RightChop(4));
+
+			FActorSpawnParameters SpawnParams;
+			SpawnParams.ObjectFlags = RF_Transient;
+			SpawnParams.Name = FName(ActorLabel);
+			SpawnParams.NameMode = FActorSpawnParameters::ESpawnActorNameMode::Requested;
+
+			UWorldPartitionRuntimeHashSet* WorldHash = Cast<UWorldPartitionRuntimeHashSet>(GetWorld()->GetWorldPartition()->RuntimeHash);
+			if (!WorldHash)
+			{
+				UE_LOG(LogStandaloneHLODSubsystem, Log, TEXT("RuntimeHashSet not found in World %s"), *GetWorld()->GetName());
+				return;
+			}
+
+			if (bRefreshCachedHLODSetups)
+			{
+				WorldHash->ForEachHLODLayer([this](FName RuntimePartitionName, FName HLODSetupName, int32 HLODSetupIndex)
+				{
+					FString RuntimeGrid = FString::Printf(TEXT("%s:%s"), *RuntimePartitionName.ToString(), *HLODSetupName.ToString());
+					TMap<int32, FName>& HLODSetups = CachedHLODSetups.FindOrAdd(RuntimePartitionName);
+					HLODSetups.Add(HLODSetupIndex, FName(RuntimeGrid));
+					return true;
+				});
+				bRefreshCachedHLODSetups = false;
+			}
+
+			FName RuntimeGrid = NAME_None;
+			if (TMap<int32, FName>* HLODSetups = CachedHLODSetups.Find(WorldHash->GetDefaultGrid()))
+			{
+				if (FName* RuntimeGridPtr = HLODSetups->Find(HLODIndex))
+				{
+					RuntimeGrid = *RuntimeGridPtr;
+				}
+			}
+
+			if (RuntimeGrid == NAME_None)
+			{
+				UE_LOG(LogStandaloneHLODSubsystem, Log, TEXT("Couldn't resolve Runtime Grid for %s, HLODIndex %d"), *WorldHash->GetDefaultGrid().ToString(), HLODIndex);
+				return;
+			}
+
+			AWorldPartitionStandaloneHLOD* HLODActor = GetWorld()->SpawnActor<AWorldPartitionStandaloneHLOD>(AWorldPartitionStandaloneHLOD::StaticClass(), InStandaloneHLODActorParams.Transform, SpawnParams);
+			HLODActor->SetWorldAsset(TSoftObjectPtr<UWorld>(LODWorldPackagePath));
+			HLODActor->SetActorLabel(ActorLabel);
+			HLODActor->SetRuntimeGrid(RuntimeGrid);
+			HLODActor->SetFolderPath(TEXT("HLOD"));
+			HLODActor->RegisterAllComponents();
+
+			StandaloneHLODActors.FindOrAdd(InStandaloneHLODActorParams.Guid).Add(HLODActor);
+		}
+	}
+}
+
+void UWorldPartitionStandaloneHLODSubsystem::DeleteStandaloneHLODActorsRecursive(const FLevelInstanceActorDesc& InLevelInstanceActorDesc)
+{
+	DeleteStandaloneHLODActors(InLevelInstanceActorDesc.GetGuid());
+
+	UActorDescContainer* Container = UActorDescContainerSubsystem::GetChecked().RegisterContainer({ InLevelInstanceActorDesc.GetChildContainerPackage() });
+	for (FActorDescList::TConstIterator<> ActorDescIt(Container); ActorDescIt; ++ActorDescIt)
+	{
+		if (ActorDescIt->HasStandaloneHLOD())
+		{
+			const FLevelInstanceActorDesc& LevelInstanceActorDesc = *static_cast<const FLevelInstanceActorDesc*>(*ActorDescIt);
+			DeleteStandaloneHLODActorsRecursive(LevelInstanceActorDesc);
+		}
+	}
+	UActorDescContainerSubsystem::GetChecked().UnregisterContainer(Container);
+}
+
+void UWorldPartitionStandaloneHLODSubsystem::DeleteStandaloneHLODActors(FGuid InGuid)
+{
+	if (TArray<AWorldPartitionStandaloneHLOD*>* HLODActorsPtr = StandaloneHLODActors.Find(InGuid))
+	{
+		TArray<AWorldPartitionStandaloneHLOD*>& HLODActors = *HLODActorsPtr;
+		for (AWorldPartitionStandaloneHLOD* HLODActor : HLODActors)
+		{
+			GetWorld()->DestroyActor(HLODActor);
+		}
+		HLODActors.Empty();
+		StandaloneHLODActors.Remove(InGuid);
+	}
+}
+
+void UWorldPartitionStandaloneHLODSubsystem::ForEachStandaloneHLODActor(TFunctionRef<void(AWorldPartitionStandaloneHLOD*)> Func) const
+{
+	for (const TPair<FGuid, TArray<AWorldPartitionStandaloneHLOD*>>& Pair : StandaloneHLODActors)
+	{
+		const TArray<AWorldPartitionStandaloneHLOD*>& HLODActors = Pair.Value;
+		for (AWorldPartitionStandaloneHLOD* HLODActor : HLODActors)
+		{
+			Func(HLODActor);
+		}
+	}
+}
+
+void UWorldPartitionStandaloneHLODSubsystem::ForEachStandaloneHLODActorFiltered(FGuid InGuid, TFunctionRef<void(AWorldPartitionStandaloneHLOD*)> Func) const
+{
+	const TArray<AWorldPartitionStandaloneHLOD*>* HLODActorsPtr = StandaloneHLODActors.Find(InGuid);
+	if (HLODActorsPtr)
+	{
+		const TArray<AWorldPartitionStandaloneHLOD*>& HLODActors = *HLODActorsPtr;
+		for (AWorldPartitionStandaloneHLOD* HLODActor : HLODActors)
+		{
+			Func(HLODActor);
+		}
+	}
+}
+
+bool UWorldPartitionStandaloneHLODSubsystem::GetStandaloneHLODFolderPathAndPackagePrefix(const FString& InWorldPackageName, FString& OutFolderPath, FString& OutPackagePrefix)
+{
+	const FString SourceLongPackagePath = FPackageName::GetLongPackagePath(InWorldPackageName);
+	const FString SourceShortPackageName = FPackageName::GetShortName(InWorldPackageName);
+	OutFolderPath = FString::Printf(TEXT("%s/HLOD"), *SourceLongPackagePath);
+	OutPackagePrefix = FString::Printf(TEXT("%s_HLOD"), *SourceShortPackageName);
+
+	return true;
+}
+#endif
\ No newline at end of file
diff --git a/Engine/Source/Runtime/Engine/Private/WorldPartition/LevelInstance/LevelInstanceActorDesc.cpp b/Engine/Source/Runtime/Engine/Private/WorldPartition/LevelInstance/LevelInstanceActorDesc.cpp
--- a/Engine/Source/Runtime/Engine/Private/WorldPartition/LevelInstance/LevelInstanceActorDesc.cpp
+++ b/Engine/Source/Runtime/Engine/Private/WorldPartition/LevelInstance/LevelInstanceActorDesc.cpp
@@ -38,6 +38,7 @@ FLevelInstanceActorDesc::FLevelInstanceActorDesc()
 	: DesiredRuntimeBehavior(ELevelInstanceRuntimeBehavior::Partitioned)
 	, ChildContainer(nullptr)
 	, bIsChildContainerInstance(false)
+	, bHasStandaloneHLOD(false)
 {}
 
 FLevelInstanceActorDesc::~FLevelInstanceActorDesc()
@@ -60,6 +61,14 @@ void FLevelInstanceActorDesc::Init(const AActor* InActor)
 	
 	bIsChildContainerInstance = IsChildContainerInstanceInternal();
 
+	if (ULevel* LoadedLevel = LevelInstance->GetLoadedLevel())
+	{
+		if (UWorldPartition* WorldPartition = LoadedLevel->GetWorldPartition())
+		{
+			bHasStandaloneHLOD = WorldPartition->HasStandaloneHLOD();
+		}
+	}
+
 	if (LevelInstance->SupportsPropertyOverrides())
 	{
 		if (ULevelInstancePropertyOverrideAsset* Asset = LevelInstance->GetPropertyOverrideAsset())
@@ -120,7 +129,8 @@ bool FLevelInstanceActorDesc::Equals(const FWorldPartitionActorDesc* Other) cons
 		return
 			WorldAsset == LevelInstanceActorDesc->WorldAsset &&
 			PropertyOverrideAsset == LevelInstanceActorDesc->PropertyOverrideAsset &&
-			DesiredRuntimeBehavior == LevelInstanceActorDesc->DesiredRuntimeBehavior;
+			DesiredRuntimeBehavior == LevelInstanceActorDesc->DesiredRuntimeBehavior &&
+			bHasStandaloneHLOD == LevelInstanceActorDesc->bHasStandaloneHLOD;
 	}
 
 	return false;
@@ -455,5 +465,13 @@ void FLevelInstanceActorDesc::Serialize(FArchive& Ar)
 	{
 		bActorIsHLODRelevant = true;
 	}
+
+	if (!bIsDefaultActorDesc)
+	{
+		if (Ar.CustomVer(FUE5ReleaseStreamObjectVersion::GUID) >= FUE5ReleaseStreamObjectVersion::WorldPartitionAddStandaloneHLODSupport)
+		{
+			Ar << bHasStandaloneHLOD;
+		}
+	}
 }
 #endif
diff --git a/Engine/Source/Runtime/Engine/Private/WorldPartition/RuntimeHashSet/WorldPartitionRuntimeHashSet.cpp b/Engine/Source/Runtime/Engine/Private/WorldPartition/RuntimeHashSet/WorldPartitionRuntimeHashSet.cpp
--- a/Engine/Source/Runtime/Engine/Private/WorldPartition/RuntimeHashSet/WorldPartitionRuntimeHashSet.cpp
+++ b/Engine/Source/Runtime/Engine/Private/WorldPartition/RuntimeHashSet/WorldPartitionRuntimeHashSet.cpp
@@ -1176,6 +1176,23 @@ void UWorldPartitionRuntimeHashSet::FixupHLODSetup(FRuntimePartitionDesc& Runtim
 	}
 }
 
+void UWorldPartitionRuntimeHashSet::ForEachHLODLayer(TFunctionRef<bool(FName, FName, int32)> Func) const
+{
+	for (const FRuntimePartitionDesc& RuntimePartition : RuntimePartitions)
+	{
+		if (RuntimePartition.MainLayer)
+		{
+			for (int32 HLODSetupIndex = 0; HLODSetupIndex < RuntimePartition.HLODSetups.Num(); HLODSetupIndex++)
+			{
+				if (!Func(RuntimePartition.Name, RuntimePartition.HLODSetups[HLODSetupIndex].Name, HLODSetupIndex))
+				{
+					return;
+				}
+			}
+		}
+	}
+}
+
 UWorldPartitionRuntimeHashSet::FCellUniqueId UWorldPartitionRuntimeHashSet::GetCellUniqueId(const URuntimePartition::FCellDescInstance& InCellDescInstance) const
 {
 	FCellUniqueId CellUniqueId;
@@ -1297,3 +1314,8 @@ void UWorldPartitionRuntimeHashSet::UpdateRuntimeDataGridMap()
 		RuntimeNonSpatiallyLoadedDataGridList.Empty();
 	}
 }
+
+const FGuid* UWorldPartitionRuntimeHashSet::GetStandaloneHLODActorSourceCellOverride(const FGuid& InActorGuid) const
+{
+	return StandaloneHLODActorToSourceCellsMap.Find(InActorGuid);
+}
\ No newline at end of file
diff --git a/Engine/Source/Runtime/Engine/Private/WorldPartition/RuntimeHashSet/WorldPartitionRuntimeHashSetHLODGeneration.cpp b/Engine/Source/Runtime/Engine/Private/WorldPartition/RuntimeHashSet/WorldPartitionRuntimeHashSetHLODGeneration.cpp
--- a/Engine/Source/Runtime/Engine/Private/WorldPartition/RuntimeHashSet/WorldPartitionRuntimeHashSetHLODGeneration.cpp
+++ b/Engine/Source/Runtime/Engine/Private/WorldPartition/RuntimeHashSet/WorldPartitionRuntimeHashSetHLODGeneration.cpp
@@ -14,10 +14,16 @@
 #include "WorldPartition/HLOD/HLODModifier.h"
 #include "WorldPartition/HLOD/IWorldPartitionHLODUtilities.h"
 #include "WorldPartition/HLOD/IWorldPartitionHLODUtilitiesModule.h"
+#include "WorldPartition/HLOD/StandaloneHLODSubsystem.h"
 
 #include "WorldPartition/DataLayer/DataLayerManager.h"
 #include "WorldPartition/DataLayer/DataLayerUtils.h"
 
+#if WITH_EDITOR
+#include "EditorLevelUtils.h"
+#include "EditorWorldUtils.h"
+#endif
+
 #include "UObject/SavePackage.h"
 #include "HAL/PlatformFileManager.h"
 
@@ -178,8 +184,10 @@ bool UWorldPartitionRuntimeHashSet::SetupHLODActors(const IStreamingGenerationCo
 
 	// Create the HLOD creation context
 	FHLODCreationContext HLODCreationContext;
-	BaseContainerInstanceCollection->ForEachActorDescContainerInstance([&HLODCreationContext, WorldPartition](const UActorDescContainerInstance* ActorDescContainerInstance)
+	TSet<FGuid> ContentBundleGuids;
+	BaseContainerInstanceCollection->ForEachActorDescContainerInstance([&HLODCreationContext, &ContentBundleGuids, WorldPartition](const UActorDescContainerInstance* ActorDescContainerInstance)
 	{
+		ContentBundleGuids.Add(ActorDescContainerInstance->GetContentBundleGuid());
 		for (UActorDescContainerInstance::TConstIterator<AWorldPartitionHLOD> HLODIterator(ActorDescContainerInstance); HLODIterator; ++HLODIterator)
 		{
 			FWorldPartitionHandle HLODActorHandle(WorldPartition, HLODIterator->GetGuid());
@@ -192,6 +200,128 @@ bool UWorldPartitionRuntimeHashSet::SetupHLODActors(const IStreamingGenerationCo
 	int32 HLODLevel = 0;
 	while (CurrentHLODStreamingGenerationContext)
 	{
+		UWorld* StandaloneHLODWorld = nullptr;
+		if (WorldPartition->HasStandaloneHLOD())
+		{
+			UWorld* SourceWorld = WorldPartition->GetWorld();
+
+			FString FolderPath, PackagePrefix;
+			UWorldPartitionStandaloneHLODSubsystem::GetStandaloneHLODFolderPathAndPackagePrefix(SourceWorld->GetPackage()->GetName(), FolderPath, PackagePrefix);
+			FString LODLevelPackageName = FString::Printf(TEXT("%s/%s%d"), *FolderPath, *PackagePrefix, HLODLevel);
+
+			UPackage* LODPackage = LoadWorldPackageForEditor(LODLevelPackageName, EWorldType::Editor, LOAD_NoWarn);
+			if (!LODPackage)
+			{
+				LODPackage = CreatePackage(*LODLevelPackageName);
+			}
+			LODPackage->FullyLoad();
+			LODPackage->Modify();
+
+			StandaloneHLODWorld = UWorld::FindWorldInPackage(LODPackage);
+			if (!StandaloneHLODWorld)
+			{
+				// Create Standalone HLOD World if not found
+				UWorld::InitializationValues IVS;
+				IVS.RequiresHitProxies(false);
+				IVS.ShouldSimulatePhysics(false);
+				IVS.EnableTraceCollision(false);
+				IVS.CreateNavigation(false);
+				IVS.CreateAISystem(false);
+				IVS.AllowAudioPlayback(false);
+				IVS.CreatePhysicsScene(true);
+				IVS.CreateWorldPartition(true);
+
+				StandaloneHLODWorld = UWorld::CreateWorld(EWorldType::Editor, false, FPackageName::GetShortFName(LODPackage->GetFName()), LODPackage, false, /*InFeatureLevel=*/ERHIFeatureLevel::Num, &IVS, /*bInSkipInitWorld=*/ true);
+				StandaloneHLODWorld->SetFlags(RF_Public | RF_Standalone);
+
+				UWorldPartition* StandaloneHLODWorldPartition = StandaloneHLODWorld->GetWorldPartition();
+				StandaloneHLODWorldPartition->SetIsStandaloneHLODWorld(true);
+				StandaloneHLODWorldPartition->SetDefaultHLODLayer(nullptr);
+
+				// Create RuntimePartitions setup in the Standalone HLOD world based on the source world
+				if (UWorldPartitionRuntimeHashSet* SourceWorldHash = Cast<UWorldPartitionRuntimeHashSet>(WorldPartition->RuntimeHash))
+				{
+					if (UWorldPartitionRuntimeHashSet* StandaloneHLODWorldHash = Cast<UWorldPartitionRuntimeHashSet>(StandaloneHLODWorld->GetWorldPartition()->RuntimeHash))
+					{
+						TArray<FRuntimePartitionDesc>& StandaloneHLODWorldRuntimePartitions = StandaloneHLODWorldHash->RuntimePartitions;
+						StandaloneHLODWorldRuntimePartitions.Empty();
+						
+						for (FRuntimePartitionDesc& SourcePartitionDesc : SourceWorldHash->RuntimePartitions)
+						{
+							FRuntimePartitionDesc& PartitionDesc = StandaloneHLODWorldRuntimePartitions.AddDefaulted_GetRef();
+							PartitionDesc.Name = SourcePartitionDesc.Name;
+							PartitionDesc.Class = SourcePartitionDesc.Class;
+							PartitionDesc.MainLayer = DuplicateObject<URuntimePartition>(SourcePartitionDesc.MainLayer, StandaloneHLODWorldHash);
+
+							for (FRuntimePartitionHLODSetup& SourceHLODSetup : SourcePartitionDesc.HLODSetups)
+							{
+								FRuntimePartitionHLODSetup& HLODSetup = PartitionDesc.HLODSetups.AddDefaulted_GetRef();
+								HLODSetup.Name = SourceHLODSetup.Name;
+								HLODSetup.HLODLayers = SourceHLODSetup.HLODLayers;
+								HLODSetup.bIsSpatiallyLoaded = SourceHLODSetup.bIsSpatiallyLoaded;
+								HLODSetup.PartitionLayer = DuplicateObject<URuntimePartition>(SourceHLODSetup.PartitionLayer, StandaloneHLODWorldHash);
+							}
+						}
+					}
+				}
+
+				// Save World Data Layers Package
+				UPackage* WorldDataLayersPackage = StandaloneHLODWorld->PersistentLevel->GetWorldDataLayers()->GetPackage();
+				PrivateUtils::SavePackage(WorldDataLayersPackage, Params.SourceControlHelper);
+
+				PrivateUtils::SavePackage(LODPackage, Params.SourceControlHelper);
+			}
+
+			UWorldPartition* StandaloneHLODWorldPartition = StandaloneHLODWorld->GetWorldPartition();
+
+			if (!StandaloneHLODWorldPartition->IsInitialized())
+			{
+				StandaloneHLODWorldPartition->Initialize(StandaloneHLODWorld, FTransform::Identity);
+			}
+
+			// Fixup actor folders
+			if (StandaloneHLODWorld->PersistentLevel->IsUsingActorFolders())
+			{
+				if (!StandaloneHLODWorld->PersistentLevel->LoadedExternalActorFolders.IsEmpty())
+				{
+					StandaloneHLODWorld->PersistentLevel->bFixupActorFoldersAtLoad = false;
+					StandaloneHLODWorld->PersistentLevel->FixupActorFolders();
+				}
+			}
+
+			// If necessary, update Standalone HLOD world data layers based on source world
+			bool bDataLayersChanged = false;
+			SourceWorld->GetWorldDataLayers()->ForEachDataLayerInstance([StandaloneHLODWorld, &bDataLayersChanged](UDataLayerInstance* DataLayer)
+			{
+				if (DataLayer->IsA<UDataLayerInstanceWithAsset>() && StandaloneHLODWorld->GetWorldDataLayers()->GetDataLayerInstance(DataLayer->GetAsset()) == nullptr)
+				{
+					StandaloneHLODWorld->GetWorldDataLayers()->CreateDataLayer<UDataLayerInstanceWithAsset>(DataLayer->GetAsset());
+					bDataLayersChanged = true;
+				}
+				return true;
+			});
+			if (bDataLayersChanged)
+			{
+				UPackage* WorldDataLayersPackage = StandaloneHLODWorld->PersistentLevel->GetWorldDataLayers()->GetPackage();
+				PrivateUtils::SavePackage(WorldDataLayersPackage, Params.SourceControlHelper);
+			}
+
+			// Add Standalone HLOD world to be processed by World Partition builder
+			Params.OutAdditionalWorldPartitionsForStandaloneHLOD.AddUnique(StandaloneHLODWorldPartition);
+
+			StandaloneHLODWorldPartition->ForEachActorDescContainerInstance([&HLODCreationContext, &ContentBundleGuids, StandaloneHLODWorldPartition](const UActorDescContainerInstance* ActorDescContainerInstance)
+			{
+				if (ContentBundleGuids.Contains(ActorDescContainerInstance->GetContentBundleGuid()))
+				{
+					for (FActorDescContainerInstanceCollection::TIterator<AWorldPartitionHLOD> HLODIterator(StandaloneHLODWorldPartition); HLODIterator; ++HLODIterator)
+					{
+						FWorldPartitionHandle HLODActorHandle(StandaloneHLODWorldPartition, HLODIterator->GetGuid());
+						HLODCreationContext.HLODActorDescs.Emplace(HLODIterator->GetActorName(), MoveTemp(HLODActorHandle));
+					}
+				}
+			});
+		}
+
 		TMap<URuntimePartition*, TArray<URuntimePartition::FCellDescInstance>> RuntimePartitionsStreamingDescs;
 		GenerateRuntimePartitionsStreamingDescs(CurrentHLODStreamingGenerationContext.Get(), RuntimePartitionsStreamingDescs);
 
@@ -239,6 +369,7 @@ bool UWorldPartitionRuntimeHashSet::SetupHLODActors(const IStreamingGenerationCo
 
 				FHLODCreationParams HLODCreationParams;
 				HLODCreationParams.WorldPartition = WorldPartition;
+				HLODCreationParams.TargetWorld = WorldPartition->GetWorld();
 				HLODCreationParams.CellName = CellUniqueId.Name;
 				HLODCreationParams.CellGuid = CellUniqueId.Guid;
 				HLODCreationParams.CellBounds = CellBounds;
@@ -247,6 +378,33 @@ bool UWorldPartitionRuntimeHashSet::SetupHLODActors(const IStreamingGenerationCo
 				HLODCreationParams.MinVisibleDistance = RuntimePartition->LoadingRange;
 				HLODCreationParams.ContentBundleGuid = CellDescInstance.ContentBundleID;
 				HLODCreationParams.DataLayerInstances = CellDescInstance.DataLayerInstances;
+				HLODCreationParams.bIsStandalone = false;
+
+				if (WorldPartition->HasStandaloneHLOD())
+				{
+					HLODCreationParams.TargetWorld = StandaloneHLODWorld;
+					HLODCreationParams.bIsStandalone = true;
+
+					// Map data layers from the source world to data layers from Standalone HLOD world
+					HLODCreationParams.DataLayerInstances.Empty();
+					for (const UDataLayerInstance* DataLayerInstance : CellDescInstance.DataLayerInstances)
+					{
+						if (const UDataLayerInstanceWithAsset* DataLayerInstanceWithAsset = Cast<UDataLayerInstanceWithAsset>(DataLayerInstance))
+						{
+							if (const UDataLayerAsset* DataLayerAsset = DataLayerInstanceWithAsset->GetAsset())
+							{
+								if (const UDataLayerInstance* StandaloneHLODWorldDataLayerInstance = StandaloneHLODWorld->GetWorldDataLayers()->GetDataLayerInstance(DataLayerAsset))
+								{
+									HLODCreationParams.DataLayerInstances.Add(StandaloneHLODWorldDataLayerInstance);
+								}
+								else
+								{
+									UE_LOG(LogWorldPartition, Log, TEXT("Couldn't find data layer %s in Standalone HLOD world"), *DataLayerInstanceWithAsset->GetDataLayerFullName());
+								}
+							}
+						}
+					}
+				}
 
 				TArray<AWorldPartitionHLOD*> CellHLODActors = WPHLODUtilities->CreateHLODActors(HLODCreationContext, HLODCreationParams, ActorInstances);
 
@@ -295,10 +453,21 @@ bool UWorldPartitionRuntimeHashSet::SetupHLODActors(const IStreamingGenerationCo
 			HLODStreamingGenerationContext->ActorSetInstanceList.Reserve(HLODActorGuids.Num());
 			HLODStreamingGenerationContext->ActorSetContainerInstance.ActorDescViewMap = &HLODStreamingGenerationContext->ActorDescViewMap;
 
+			UWorldPartition* CurrentWorldPartition = nullptr;
+			if (WorldPartition->HasStandaloneHLOD())
+			{
+				CurrentWorldPartition = StandaloneHLODWorld->GetWorldPartition();
+				DataLayerManager = CurrentWorldPartition->GetDataLayerManager();
+			}
+			else
+			{
+				CurrentWorldPartition = WorldPartition;
+			}
+
 			UE_LOG(LogWorldPartition, Log, TEXT("%s"), *LOCTEXT("CreatingHLODContext:", "Creating HLOD context:").ToString());
 			for (const FGuid& HLODActorGuid : HLODActorGuids)
 			{
-				FWorldPartitionActorDescInstance* HLODActorDescInstance = WorldPartition->GetActorDescInstance(HLODActorGuid);
+				FWorldPartitionActorDescInstance* HLODActorDescInstance = CurrentWorldPartition->GetActorDescInstance(HLODActorGuid);
 				check(HLODActorDescInstance);
 
 				FStreamingGenerationActorDescView* HLODActorDescView = HLODStreamingGenerationContext->ActorDescViewMap.Emplace(HLODActorDescInstance);
diff --git a/Engine/Source/Runtime/Engine/Private/WorldPartition/RuntimeHashSet/WorldPartitionRuntimeHashSetStreamingGeneration.cpp b/Engine/Source/Runtime/Engine/Private/WorldPartition/RuntimeHashSet/WorldPartitionRuntimeHashSetStreamingGeneration.cpp
--- a/Engine/Source/Runtime/Engine/Private/WorldPartition/RuntimeHashSet/WorldPartitionRuntimeHashSetStreamingGeneration.cpp
+++ b/Engine/Source/Runtime/Engine/Private/WorldPartition/RuntimeHashSet/WorldPartitionRuntimeHashSetStreamingGeneration.cpp
@@ -5,9 +5,12 @@
 #include "WorldPartition/RuntimeHashSet/RuntimePartition.h"
 #include "WorldPartition/RuntimeHashSet/RuntimePartitionPersistent.h"
 #include "WorldPartition/ContentBundle/ContentBundleDescriptor.h"
+#include "WorldPartition/WorldPartitionHelpers.h"
 #include "WorldPartition/WorldPartitionStreamingPolicy.h"
 #include "WorldPartition/DataLayer/DataLayersID.h"
 #include "WorldPartition/HLOD/HLODLayer.h"
+#include "WorldPartition/HLOD/HLODActor.h"
+#include "WorldPartition/HLOD/HLODActorDesc.h"
 
 #if WITH_EDITOR
 bool UWorldPartitionRuntimeHashSet::GenerateRuntimePartitionsStreamingDescs(const IStreamingGenerationContext* StreamingGenerationContext, TMap<URuntimePartition*, TArray<URuntimePartition::FCellDescInstance>>& OutRuntimeCellDescs) const
@@ -166,6 +169,7 @@ bool UWorldPartitionRuntimeHashSet::GenerateStreaming(UWorldPartitionStreamingPo
 		return RuntimeCell;
 	};
 
+	TMap<FGuid, FGuid> StandaloneHLODActorToCell;
 	TArray<UWorldPartitionRuntimeCell*> RuntimeCells;
 	TMap<URuntimePartition*, FRuntimePartitionStreamingData> RuntimePartitionsStreamingData;
 	for (auto& [RuntimePartition, CellDescInstances] : RuntimePartitionsStreamingDescs)
@@ -183,6 +187,16 @@ bool UWorldPartitionRuntimeHashSet::GenerateStreaming(UWorldPartitionStreamingPo
 
 				PopulateRuntimeCell(RuntimeCell, CellActorInstances, OutPackagesToGenerate);
 
+				// Save HLOD actor GUID -> Cell GUID mapping to use it for Standalone HLOD source cell override
+				for (const IStreamingGenerationContext::FActorInstance& ActorInstance : CellActorInstances)
+				{
+					const FStreamingGenerationActorDescView& ActorDescView = ActorInstance.GetActorDescView();
+					if (ActorDescView.GetActorNativeClass()->IsChildOf<AWorldPartitionHLOD>())
+					{
+						StandaloneHLODActorToCell.Add(ActorDescView.GetGuid(), RuntimeCell->GetGuid());
+					}
+				}
+
 				UWorldPartitionRuntimeCellDataHashSet* RuntimeCellData = CastChecked<UWorldPartitionRuntimeCellDataHashSet>(RuntimeCell->RuntimeCellData);
 
 				if (CellDescInstance.CellBounds.IsSet())
@@ -236,6 +250,40 @@ bool UWorldPartitionRuntimeHashSet::GenerateStreaming(UWorldPartitionStreamingPo
 		}
 	}
 
+	// Standalone HLOD actors are embedded in the main world streaming cells, so for Standalone HLOD actors that are using other Standalone HLOD actors as source actors,
+	// source cell GUID won't be correct. Save the actual source cell GUID, so that we can use that override at runtime.
+	StandaloneHLODActorToSourceCellsMap.Empty();
+	StreamingGenerationContext->ForEachActorSetContainerInstance([this, &StandaloneHLODActorToCell](const IStreamingGenerationContext::FActorSetContainerInstance& ActorSetContainerInstance)
+	{
+		const FStreamingGenerationActorDescViewMap* ActorDescViewMap = ActorSetContainerInstance.ActorDescViewMap;
+		ActorDescViewMap->ForEachActorDescView([this, &StandaloneHLODActorToCell](const FStreamingGenerationActorDescView& ActorDescView)
+		{
+			if (ActorDescView.GetActorNativeClass()->IsChildOf<AWorldPartitionHLOD>())
+			{
+				const FHLODActorDesc* ActorDesc = static_cast<const FHLODActorDesc*>(ActorDescView.GetActorDesc());
+				for (const FGuid& ExternalChildHLODActorGuid : ActorDesc->GetExternalChildHLODActors())
+				{
+					if (FGuid* CellGuidPtr = StandaloneHLODActorToCell.Find(ExternalChildHLODActorGuid))
+					{
+						FGuid CellGuid = *CellGuidPtr;
+
+						if (FGuid* SourceCellGuidPtr = StandaloneHLODActorToSourceCellsMap.Find(ActorDescView.GetGuid()))
+						{
+							if (CellGuid != *SourceCellGuidPtr)
+							{
+								UE_LOG(LogWorldPartition, Log, TEXT("External child HLOD actors of %s are not in the same cell. Fix clustering settings!"), *ActorDesc->GetActorNameString())
+							}
+						}
+						else
+						{
+							StandaloneHLODActorToSourceCellsMap.Add(ActorDescView.GetGuid(), CellGuid);
+						}
+					}
+				}
+			}
+		});
+	});
+
 	//
 	// Finalize streaming object
 	//
@@ -286,4 +334,4 @@ void UWorldPartitionRuntimeHashSet::DumpStateLog(FHierarchicalLogArchive& Ar) co
 
 	Ar.Printf(TEXT(""));
 }
-#endif
+#endif
\ No newline at end of file
diff --git a/Engine/Source/Runtime/Engine/Private/WorldPartition/WorldPartition.cpp b/Engine/Source/Runtime/Engine/Private/WorldPartition/WorldPartition.cpp
--- a/Engine/Source/Runtime/Engine/Private/WorldPartition/WorldPartition.cpp
+++ b/Engine/Source/Runtime/Engine/Private/WorldPartition/WorldPartition.cpp
@@ -383,6 +383,8 @@ UWorldPartition::UWorldPartition(const FObjectInitializer& ObjectInitializer)
 
 #if WITH_EDITOR
 	bAllowShowingHLODsInEditor = true;
+	bHasStandaloneHLOD = false;
+	bIsStandaloneHLODWorld = false;
 	WorldPartitionStreamingPolicyClass = UWorldPartitionLevelStreamingPolicy::StaticClass();
 
 	if (HasAnyFlags(RF_ClassDefaultObject) && !HasAnyFlags(RF_ImmutableDefaultObject) && IsRunningCookCommandlet())
diff --git a/Engine/Source/Runtime/Engine/Private/WorldPartition/WorldPartitionStreamingGeneration.cpp b/Engine/Source/Runtime/Engine/Private/WorldPartition/WorldPartitionStreamingGeneration.cpp
--- a/Engine/Source/Runtime/Engine/Private/WorldPartition/WorldPartitionStreamingGeneration.cpp
+++ b/Engine/Source/Runtime/Engine/Private/WorldPartition/WorldPartitionStreamingGeneration.cpp
@@ -34,7 +34,10 @@
 #include "WorldPartition/ErrorHandling/WorldPartitionStreamingGenerationNullErrorHandler.h"
 #include "WorldPartition/ErrorHandling/WorldPartitionStreamingGenerationLogErrorHandler.h"
 #include "WorldPartition/HLOD/HLODActor.h"
+#include "WorldPartition/HLOD/HLODActorDesc.h"
 #include "WorldPartition/HLOD/HLODLayer.h"
+#include "WorldPartition/HLOD/StandaloneHLODActor.h"
+#include "WorldPartition/HLOD/StandaloneHLODSubsystem.h"
 #include "WorldPartition/WorldPartitionSubsystem.h"
 #include "WorldPartition/WorldPartitionRuntimeContainerResolving.h"
 #include "HAL/FileManager.h"
@@ -73,26 +76,26 @@ static FAutoConsoleCommand DumpStreamingGenerationLog(
 	})
 );
 
-class FGCUnsavedDirtyActorContainerInstances : public FGCObject
+class FGCTempActorContainerInstances : public FGCObject
 {
-	friend class FStreamingGenerationUnsavedDirtyActorDescInstance;
+	friend class FStreamingGenerationTempActorDescInstance;
 public:
 	//~ Begin FGCObject interface
 	virtual void AddReferencedObjects(FReferenceCollector& Collector)
 	{
 		Collector.AddReferencedObjects(ContainerInstances);
 	}
-	virtual FString GetReferencerName() const { return TEXT("FUnsavedDirtyActorContainer"); }
+	virtual FString GetReferencerName() const { return TEXT("FTempActorContainer"); }
 	//~ End FGCObject interface
 
 protected:
 	TSet<TObjectPtr<UActorDescContainerInstance>> ContainerInstances;
 };
 
-class FStreamingGenerationUnsavedDirtyActorDescInstance : public FWorldPartitionActorDescInstance
+class FStreamingGenerationTempActorDescInstance : public FWorldPartitionActorDescInstance
 {
 public:
-	FStreamingGenerationUnsavedDirtyActorDescInstance(UActorDescContainerInstance* InContainerInstance, AActor* InActor)
+	FStreamingGenerationTempActorDescInstance(UActorDescContainerInstance* InContainerInstance, AActor* InActor)
 	{
 		ContainerInstance = InContainerInstance;
 
@@ -106,10 +109,10 @@ public:
 		}
 	}
 
-	FStreamingGenerationUnsavedDirtyActorDescInstance(FStreamingGenerationUnsavedDirtyActorDescInstance&& InOther) = delete;
-	FStreamingGenerationUnsavedDirtyActorDescInstance& operator=(FStreamingGenerationUnsavedDirtyActorDescInstance&& InOther) = delete;
+	FStreamingGenerationTempActorDescInstance(FStreamingGenerationTempActorDescInstance&& InOther) = delete;
+	FStreamingGenerationTempActorDescInstance& operator=(FStreamingGenerationTempActorDescInstance&& InOther) = delete;
 
-	virtual ~FStreamingGenerationUnsavedDirtyActorDescInstance()
+	virtual ~FStreamingGenerationTempActorDescInstance()
 	{
 		if (IsChildContainerInstance())
 		{
@@ -117,32 +120,32 @@ public:
 		}
 	}
 
-	static TUniquePtr<FStreamingGenerationUnsavedDirtyActorDescInstance> Create(AActor* InActor, const FStreamingGenerationContainerInstanceCollection& InContainerInstanceCollection)
+	static TUniquePtr<FStreamingGenerationTempActorDescInstance> Create(AActor* InActor, const FStreamingGenerationContainerInstanceCollection& InContainerInstanceCollection)
 	{
 		UActorDescContainerInstance* HandlingContainer = const_cast<UActorDescContainerInstance*>(InContainerInstanceCollection.FindHandlingContainerInstance(InActor).Get());
 		check(HandlingContainer);
-		return MakeUnique<FStreamingGenerationUnsavedDirtyActorDescInstance>(HandlingContainer, InActor);
+		return MakeUnique<FStreamingGenerationTempActorDescInstance>(HandlingContainer, InActor);
 	}
 
 	//~ Begin FWorldPartitionActorDescInstance interface
 	virtual void RegisterChildContainerInstance() override
 	{
-		if (!UnsavedDirtyActorContainerInstances.IsValid())
+		if (!TempActorContainerInstances.IsValid())
 		{
-			UnsavedDirtyActorContainerInstances = MakeUnique<FGCUnsavedDirtyActorContainerInstances>();
+			TempActorContainerInstances = MakeUnique<FGCTempActorContainerInstances>();
 		}
 
 		ChildContainerInstance = GetActorDesc()->CreateChildContainerInstance(this);
-		UnsavedDirtyActorContainerInstances->ContainerInstances.Add(ChildContainerInstance);
+		TempActorContainerInstances->ContainerInstances.Add(ChildContainerInstance);
 	}
 	virtual void UnregisterChildContainerInstance() override
 	{
-		check(UnsavedDirtyActorContainerInstances.IsValid())
+		check(TempActorContainerInstances.IsValid())
 		{
-			UnsavedDirtyActorContainerInstances->ContainerInstances.Remove(ChildContainerInstance);
-			if (UnsavedDirtyActorContainerInstances->ContainerInstances.IsEmpty())
+			TempActorContainerInstances->ContainerInstances.Remove(ChildContainerInstance);
+			if (TempActorContainerInstances->ContainerInstances.IsEmpty())
 			{
-				UnsavedDirtyActorContainerInstances.Reset();
+				TempActorContainerInstances.Reset();
 			}
 		}
 
@@ -153,10 +156,10 @@ public:
 
 protected:
 	TUniquePtr<FWorldPartitionActorDesc> ActorDescPtr;
-	static TUniquePtr<FGCUnsavedDirtyActorContainerInstances> UnsavedDirtyActorContainerInstances;
+	static TUniquePtr<FGCTempActorContainerInstances> TempActorContainerInstances;
 };
 
-TUniquePtr<FGCUnsavedDirtyActorContainerInstances> FStreamingGenerationUnsavedDirtyActorDescInstance::UnsavedDirtyActorContainerInstances;
+TUniquePtr<FGCTempActorContainerInstances> FStreamingGenerationTempActorDescInstance::TempActorContainerInstances;
 
 template <class T>
 class TErrorHandlerSelector
@@ -656,7 +659,10 @@ class FWorldPartitionStreamingGenerator
 		TArray<FStreamingGenerationActorDescView> ContainerCollectionInstanceViews;
 
 		/** List of unsaved/dirty descriptor views, unique ptrs so that they are moveable without having to update references to them from FStreamingGenerationActorDescView::ActorDescView */
-		TArray<TUniquePtr<FStreamingGenerationUnsavedDirtyActorDescInstance>> UnsavedDirtyInstances;
+		TArray<TUniquePtr<FStreamingGenerationTempActorDescInstance>> UnsavedDirtyInstances;
+
+		/** List of Standalone HLOD actor descriptor views. */
+		TArray<TUniquePtr<FStreamingGenerationTempActorDescInstance>> StandaloneHLODActorInstances;
 
 		/** List of actor clusters for this descriptor */
 		TArray<TArray<FGuid>> Clusters;
@@ -813,7 +819,8 @@ class FWorldPartitionStreamingGenerator
 		TSet<FGuid>& OutEditorOnlyActorDescSet = InContainerCollectionInstanceDescriptor.EditorOnlyActorDescSet;
 		const FActorContainerID& InContainerID = InContainerCollectionInstanceDescriptor.ID;
 		TArray<FStreamingGenerationActorDescView>& OutContainerInstances = InContainerCollectionInstanceDescriptor.ContainerCollectionInstanceViews;
-		TArray<TUniquePtr<FStreamingGenerationUnsavedDirtyActorDescInstance>>& OutUnsavedDirtyInstances = InContainerCollectionInstanceDescriptor.UnsavedDirtyInstances;
+		TArray<TUniquePtr<FStreamingGenerationTempActorDescInstance>>& OutUnsavedDirtyInstances = InContainerCollectionInstanceDescriptor.UnsavedDirtyInstances;
+		TArray<TUniquePtr<FStreamingGenerationTempActorDescInstance>>& OutStandaloneHLODActorInstances = InContainerCollectionInstanceDescriptor.StandaloneHLODActorInstances;
 		TMap<FGuid, FTransform>& EditorOnlyParentActorTransforms = InContainerCollectionInstanceDescriptor.EditorOnlyParentActorTransforms;
 
 		// Should we handle unsaved or newly created actors?
@@ -900,7 +907,7 @@ class FWorldPartitionStreamingGenerator
 					if (bShouldHandleUnsavedActors && (bIsTempContainerPackage || (Actor->GetPackage()->IsDirty() || Actor->HasAllFlags(RF_Transient | RF_NonPIEDuplicateTransient))))
 					{
 						// Dirty, unsaved actor for PIE
-						TUniquePtr<FStreamingGenerationUnsavedDirtyActorDescInstance>& UnsavedDirtyRef = OutUnsavedDirtyInstances.Add_GetRef(FStreamingGenerationUnsavedDirtyActorDescInstance::Create(Actor, InActorDescCollection));
+						TUniquePtr<FStreamingGenerationTempActorDescInstance>& UnsavedDirtyRef = OutUnsavedDirtyInstances.Add_GetRef(FStreamingGenerationTempActorDescInstance::Create(Actor, InActorDescCollection));
 						RegisterActorDescView(FStreamingGenerationActorDescView(OutActorDescViewMap, UnsavedDirtyRef.Get(), true));
 						continue;
 					}
@@ -953,7 +960,7 @@ class FWorldPartitionStreamingGenerator
 					&& InActorDescCollection.FindHandlingContainerInstance(Actor)
 					&& !InActorDescCollection.GetActorDescInstance(Actor->GetActorGuid()))
 				{
-					TUniquePtr<FStreamingGenerationUnsavedDirtyActorDescInstance> UnsavedViewPtr = FStreamingGenerationUnsavedDirtyActorDescInstance::Create(Actor, InActorDescCollection);
+					TUniquePtr<FStreamingGenerationTempActorDescInstance> UnsavedViewPtr = FStreamingGenerationTempActorDescInstance::Create(Actor, InActorDescCollection);
 					if (ShouldRegisterActorDesc(UnsavedViewPtr.Get()))
 					{
 						FStreamingGenerationActorDescView ModifiedActorDescView(OutActorDescViewMap, UnsavedViewPtr.Get(), true);
@@ -963,6 +970,42 @@ class FWorldPartitionStreamingGenerator
 				}
 			}
 		}
+
+		if (bIncludeStandaloneHLODActors && InContainerID.IsMainContainer())
+		{
+			if (WorldPartitionStandaloneHLODSubsystem)
+			{
+				// Standalone HLOD actors
+				WorldPartitionStandaloneHLODSubsystem->ForEachStandaloneHLODActor([this, &OutStandaloneHLODActorInstances, &OutActorDescViewMap, &InActorDescCollection, &RegisterActorDescView](AWorldPartitionStandaloneHLOD* StandaloneHLODActor)
+				{
+					TUniquePtr<FStreamingGenerationTempActorDescInstance>& StandaloneHLODActorRef = OutStandaloneHLODActorInstances.Add_GetRef(FStreamingGenerationTempActorDescInstance::Create(StandaloneHLODActor, InActorDescCollection));
+					RegisterActorDescView(FStreamingGenerationActorDescView(OutActorDescViewMap, StandaloneHLODActorRef.Get(), false));
+
+					// Save information about "dependencies" between Standalone HLOD actors. We want all external child HLOD actors of a Standalone HLOD actor to end up in the same streaming cell
+					// to maintain HLODActor -> SourceCell mapping, that is used for regular HLODActors to control visibility / loading state.
+					// StandaloneHLODActorsDependencies saved here will be used later during clusters generation.
+					if (UActorDescContainer* StandaloneHLODContainer = StandaloneHLODActorRef->GetActorDesc()->GetChildContainer())
+					{
+						for (UActorDescContainer::TConstIterator<AWorldPartitionHLOD> Iterator(StandaloneHLODContainer); Iterator; ++Iterator)
+						{
+							const FHLODActorDesc* HLODActorDesc = static_cast<const FHLODActorDesc*>(*Iterator);
+							if (HLODActorDesc->GetExternalChildHLODActors().Num() > 1)
+							{
+								const TArray<FGuid>& ExternalChildHLODActors = HLODActorDesc->GetExternalChildHLODActors();
+								for (int32 ExternalChildHLODActorIndex = 0; ExternalChildHLODActorIndex < ExternalChildHLODActors.Num(); ExternalChildHLODActorIndex++)
+								{
+									for (int32 OtherExternalChildHLODActorIndex = ExternalChildHLODActorIndex + 1; OtherExternalChildHLODActorIndex < ExternalChildHLODActors.Num(); OtherExternalChildHLODActorIndex++)
+									{
+										StandaloneHLODActorsDependencies.FindOrAdd(ExternalChildHLODActors[ExternalChildHLODActorIndex]).Add(ExternalChildHLODActors[OtherExternalChildHLODActorIndex]);
+										StandaloneHLODActorsDependencies.FindOrAdd(ExternalChildHLODActors[OtherExternalChildHLODActorIndex]).Add(ExternalChildHLODActors[ExternalChildHLODActorIndex]);
+									}
+								}
+							}
+						}
+					}
+				});
+			}
+		}
 	}
 
 	void CreateActorDescriptorViewsRecursive(FContainerCollectionInstanceDescriptor&& InContainerCollectionInstanceDescriptor)
@@ -1718,7 +1761,16 @@ class FWorldPartitionStreamingGenerator
 	{
 		// Build clusters for this container - at this point, all actors references should be in the same data layers, grid, etc because of actor descriptors validation.
 		TArray<TPair<FGuid, TArray<FGuid>>> ActorsWithRefs;
-		ContainerCollectionInstanceDescriptor.ActorDescViewMap->ForEachActorDescView([&ActorsWithRefs](const FStreamingGenerationActorDescView& ActorDescView) { ActorsWithRefs.Emplace(ActorDescView.GetGuid(), ActorDescView.GetReferences()); });
+		ContainerCollectionInstanceDescriptor.ActorDescViewMap->ForEachActorDescView([&ActorsWithRefs, this](const FStreamingGenerationActorDescView& ActorDescView)
+		{ 
+			TPair<FGuid, TArray<FGuid>>& ActorWithRefs = ActorsWithRefs.Emplace_GetRef(ActorDescView.GetGuid(), ActorDescView.GetReferences());
+
+			// Include Standalone HLOD actors "dependencies". Make sure all child actors of a Standalone HLOD actor to end up in the same cluster.
+			if (TSet<FGuid>* StandaloneHLODDependencies = StandaloneHLODActorsDependencies.Find(ActorWithRefs.Key))
+			{
+				ActorWithRefs.Value.Append(StandaloneHLODDependencies->Array());
+			}
+		});
 		ContainerCollectionInstanceDescriptor.Clusters = GenerateObjectsClusters(ActorsWithRefs);
 	}
 
@@ -1758,6 +1810,7 @@ public:
 			, bCreateContainerResolver(false)
 			, bHandleUnsavedActors(false)
 			, bIsPIE(false)
+			, bIncludeStandaloneHLODActors(false)
 		{}
 
 		const UWorldPartition* WorldPartitionContext;
@@ -1766,6 +1819,7 @@ public:
 		bool bCreateContainerResolver;
 		bool bHandleUnsavedActors;
 		bool bIsPIE;
+		bool bIncludeStandaloneHLODActors;
 		TMap<FGuid, const UActorDescContainerInstance*> ActorGuidsToContainerInstanceMap;
 		TArray<TSubclassOf<AActor>> FilteredClasses;
 		FName DefaultGrid;
@@ -1778,6 +1832,7 @@ public:
 		FWorldPartitionStreamingGeneratorParams& SetCreateContainerResolver(bool bInCreateContainerResolver) { bCreateContainerResolver = bInCreateContainerResolver; return *this; }
 		FWorldPartitionStreamingGeneratorParams& SetHandleUnsavedActors(bool bInHandleUnsavedActors) { bHandleUnsavedActors = bInHandleUnsavedActors; return *this; }
 		FWorldPartitionStreamingGeneratorParams& SetIsPIE(bool bInIsPIE) { bIsPIE = bInIsPIE; return *this; }
+		FWorldPartitionStreamingGeneratorParams& SetIncludeStandaloneHLODActors(bool bInIncludeStandaloneHLODActors) { bIncludeStandaloneHLODActors = bInIncludeStandaloneHLODActors; return *this; }
 		FWorldPartitionStreamingGeneratorParams& SetActorGuidsToContainerInstanceMap(const TMap<FGuid, const UActorDescContainerInstance*>& InActorGuidsToContainerInstanceMap) { ActorGuidsToContainerInstanceMap = InActorGuidsToContainerInstanceMap; return *this; }
 		FWorldPartitionStreamingGeneratorParams& SetFilteredClasses(const TArray<TSubclassOf<AActor>>& InFilteredClasses) { FilteredClasses = InFilteredClasses; return *this; }
 		FWorldPartitionStreamingGeneratorParams& SetDefaultGrid(FName InDefaultGrid) { DefaultGrid = InDefaultGrid; return *this; }
@@ -1793,6 +1848,7 @@ public:
 		, bCreateContainerResolver(Params.bCreateContainerResolver)
 		, bHandleUnsavedActors(Params.bHandleUnsavedActors)
 		, bIsPIE(Params.bIsPIE)
+		, bIncludeStandaloneHLODActors(Params.bIncludeStandaloneHLODActors)
 		, FilteredClasses(Params.FilteredClasses)
 		, DefaultGrid(Params.DefaultGrid)
 		, IsValidGrid(Params.IsValidGrid)
@@ -1803,6 +1859,7 @@ public:
 		UWorld* OwningWorld = WorldPartitionContext ? WorldPartitionContext->GetWorld() : nullptr;
 		WorldPartitionSubsystem = OwningWorld ? UWorld::GetSubsystem<UWorldPartitionSubsystem>(OwningWorld) : nullptr;
 		DataLayerManager = OwningWorld ? WorldPartitionContext->GetDataLayerManager() : nullptr;
+		WorldPartitionStandaloneHLODSubsystem = OwningWorld ? UWorld::GetSubsystem<UWorldPartitionStandaloneHLODSubsystem>(OwningWorld) : nullptr;
 	}
 
 	void PreparationPhase(const FStreamingGenerationContainerInstanceCollection& ContainerInstanceCollection)
@@ -1959,11 +2016,13 @@ public:
 private:
 	const UWorldPartition* WorldPartitionContext;
 	UWorldPartitionSubsystem* WorldPartitionSubsystem;
+	UWorldPartitionStandaloneHLODSubsystem* WorldPartitionStandaloneHLODSubsystem;
 	const UDataLayerManager* DataLayerManager;
 	bool bEnableStreaming;
 	bool bCreateContainerResolver;
 	bool bHandleUnsavedActors;
 	bool bIsPIE;
+	bool bIncludeStandaloneHLODActors;
 	TArray<TSubclassOf<AActor>> FilteredClasses;
 	FName DefaultGrid;
 	TFunction<bool(FName, const UClass*)> IsValidGrid;
@@ -1982,6 +2041,9 @@ private:
 
 	/** Maps containers IDs to their filtered actors use while creating FContainerCollectionInstanceDescriptor */
 	TMap<FActorContainerID, TSet<FGuid>> ContainerFilteredActors;
+
+	/** Maps Standalone HLOD actor guid to their "dependencies". Used to make sure all child actors of a Standalone HLOD actor end up in the same streaming cell */
+	TMap<FGuid, TSet<FGuid>> StandaloneHLODActorsDependencies;
 };
 
 bool UWorldPartition::GenerateStreaming(const FGenerateStreamingParams& InParams, FGenerateStreamingContext& InContext)
@@ -2039,6 +2101,7 @@ bool UWorldPartition::GenerateContainerStreaming(const FGenerateStreamingParams&
 		.SetWorldPartitionContext(this)
 		.SetHandleUnsavedActors(bIsPIE && !GetTypedOuter<UWorld>()->IsGameWorld())
 		.SetIsPIE(bIsPIE)
+		.SetIncludeStandaloneHLODActors(IsMainWorldPartition())
 		.SetDefaultGrid(RuntimeHash ? RuntimeHash->GetDefaultGrid() : NAME_None)
 		.SetIsValidGrid([this](FName GridName, const UClass* ActorClass) { return RuntimeHash->IsValidGrid(GridName, ActorClass); })
 		.SetIsValidHLODLayer([this](FName GridName, const FSoftObjectPath& HLODLayerPath) { return RuntimeHash->IsValidHLODLayer(GridName, HLODLayerPath); })
@@ -2445,6 +2508,7 @@ void UWorldPartition::CheckForErrors(const FCheckForErrorsParams& InParams)
 			.SetWorldPartitionContext(WorldPartition)
 			.SetHandleUnsavedActors(!!WorldPartition)
 			.SetIsPIE(WorldPartition ? WorldPartition->bIsPIE : false)
+			.SetIncludeStandaloneHLODActors(WorldPartition ? WorldPartition->IsMainWorldPartition() : false)
 			.SetErrorHandler(ErrorHandlerSelector.Get())
 			.SetDefaultGrid(WorldPartitionRuntimeHash ? WorldPartitionRuntimeHash->GetDefaultGrid() : NAME_None)
 			.SetIsValidGrid([WorldPartitionRuntimeHash](FName GridName, const UClass* ActorClass) { return WorldPartitionRuntimeHash ? WorldPartitionRuntimeHash->IsValidGrid(GridName, ActorClass) : true; })
diff --git a/Engine/Source/Runtime/Engine/Public/WorldPartition/HLOD/HLODActor.h b/Engine/Source/Runtime/Engine/Public/WorldPartition/HLOD/HLODActor.h
--- a/Engine/Source/Runtime/Engine/Public/WorldPartition/HLOD/HLODActor.h
+++ b/Engine/Source/Runtime/Engine/Public/WorldPartition/HLOD/HLODActor.h
@@ -45,6 +45,8 @@ public:
 	ENGINE_API virtual TSet<UObject*> GetAssetsToWarmup() const override;
 	ENGINE_API virtual void SetVisibility(bool bIsVisible) override;
 	ENGINE_API virtual const FGuid& GetSourceCellGuid() const override;
+	ENGINE_API virtual bool IsStandalone() const override;
+	ENGINE_API virtual const FGuid& GetStandaloneHLODGuid() const override;
 	// ~End IWorldPartitionHLODObject interface
 
 #if WITH_EDITOR
@@ -58,6 +60,7 @@ public:
 	ENGINE_API const FHLODBuildInputStats& GetInputStats() const;
 
 	void SetRequireWarmup(bool InRequireWarmup) { bRequireWarmup = InRequireWarmup; }
+	ENGINE_API void SetIsStandalone(bool bInIsStandalone);
 
 	ENGINE_API void SetSourceCellGuid(const FGuid& InSourceCellGuid);
 	inline void SetLODLevel(uint32 InLODLevel) { LODLevel = InLODLevel; }
@@ -155,6 +158,9 @@ private:
 	UPROPERTY()
 	FGuid SourceCellGuid;
 
+	UPROPERTY()
+	FGuid StandaloneHLODGuid;
+
 #if WITH_EDITORONLY_DATA
 	UPROPERTY()
 	TSoftObjectPtr<UWorldPartitionRuntimeCell> SourceCell_DEPRECATED;
diff --git a/Engine/Source/Runtime/Engine/Public/WorldPartition/HLOD/HLODActorDesc.h b/Engine/Source/Runtime/Engine/Public/WorldPartition/HLOD/HLODActorDesc.h
--- a/Engine/Source/Runtime/Engine/Public/WorldPartition/HLOD/HLODActorDesc.h
+++ b/Engine/Source/Runtime/Engine/Public/WorldPartition/HLOD/HLODActorDesc.h
@@ -26,6 +26,7 @@ public:
 	typedef TMap<FName, int64>	FStats;
 
 	inline const TArray<FGuid>& GetChildHLODActors() const { return ChildHLODActors; }
+	inline const TArray<FGuid>& GetExternalChildHLODActors() const { return ExternalChildHLODActors; }
 	inline const FTopLevelAssetPath& GetSourceHLODLayer() const { return SourceHLODLayer; }
 	inline const FStats& GetStats() const { return HLODStats; }
 	ENGINE_API int64 GetStat(FName InStatName) const;
@@ -43,6 +44,8 @@ protected:
 
 	TArray<FGuid> ChildHLODActors;
 
+	TArray<FGuid> ExternalChildHLODActors;
+
 	FTopLevelAssetPath SourceHLODLayer;
 
 	FStats HLODStats;
diff --git a/Engine/Source/Runtime/Engine/Public/WorldPartition/HLOD/HLODRuntimeSubsystem.h b/Engine/Source/Runtime/Engine/Public/WorldPartition/HLOD/HLODRuntimeSubsystem.h
--- a/Engine/Source/Runtime/Engine/Public/WorldPartition/HLOD/HLODRuntimeSubsystem.h
+++ b/Engine/Source/Runtime/Engine/Public/WorldPartition/HLOD/HLODRuntimeSubsystem.h
@@ -124,6 +124,7 @@ private:
 	ENGINE_API const FCellData* GetCellData(const UWorldPartitionRuntimeCell* InCell) const;
 	ENGINE_API FCellData* GetCellData(const UWorldPartitionRuntimeCell* InCell);
 	ENGINE_API FCellData* GetCellData(IWorldPartitionHLODObject* InWorldPartitionHLOD);
+	ENGINE_API FCellData* GetCellData(const UWorldPartition* InWorldPartition, const FGuid& InCellGuid);
 
 	// Keep track of the state of warmup for an HLOD object
 	struct FWorldPartitionHLODWarmupState
@@ -174,6 +175,12 @@ private:
 	FWorldPartitionHLODActorUnregisteredEvent	HLODActorUnregisteredEvent;
 	PRAGMA_ENABLE_DEPRECATION_WARNINGS
 
+	// Referenced Cell GUID -> HLOD Actor
+	TMap<FGuid, TSet<IWorldPartitionHLODObject*>> StandaloneHLODObjectsReferencingUnloadedCells;
+	
+	// Cell GUID -> WorldPartition
+	TMap<FGuid, TObjectPtr<UWorldPartition>> StandaloneHLODCellToWorldPartitionMap;
+
 #if !UE_BUILD_SHIPPING
 	TSet<IWorldPartitionHLODObject*> OutdatedHLODObjects;
 #endif
diff --git a/Engine/Source/Runtime/Engine/Public/WorldPartition/HLOD/IWorldPartitionHLODObject.h b/Engine/Source/Runtime/Engine/Public/WorldPartition/HLOD/IWorldPartitionHLODObject.h
--- a/Engine/Source/Runtime/Engine/Public/WorldPartition/HLOD/IWorldPartitionHLODObject.h
+++ b/Engine/Source/Runtime/Engine/Public/WorldPartition/HLOD/IWorldPartitionHLODObject.h
@@ -32,4 +32,10 @@ public:
 
 	/** Returns the associated source cell guid of this HLOD object. */
 	virtual const FGuid& GetSourceCellGuid() const = 0;
+
+	/** Returns whether the HLOD object is part of Standalone HLOD. */
+	virtual bool IsStandalone() const = 0;
+
+	/** Returns Standalone HLOD Guid */
+	virtual const FGuid& GetStandaloneHLODGuid() const = 0;
 };
\ No newline at end of file
diff --git a/Engine/Source/Runtime/Engine/Public/WorldPartition/HLOD/IWorldPartitionHLODUtilities.h b/Engine/Source/Runtime/Engine/Public/WorldPartition/HLOD/IWorldPartitionHLODUtilities.h
--- a/Engine/Source/Runtime/Engine/Public/WorldPartition/HLOD/IWorldPartitionHLODUtilities.h
+++ b/Engine/Source/Runtime/Engine/Public/WorldPartition/HLOD/IWorldPartitionHLODUtilities.h
@@ -29,6 +29,7 @@ struct FHLODCreationContext
 struct FHLODCreationParams
 {
 	UWorldPartition* WorldPartition;
+	UWorld* TargetWorld;
 
 	FGuid CellGuid;
 	FString CellName;
@@ -37,6 +38,7 @@ struct FHLODCreationParams
 	uint32 HLODLevel;
 	FGuid ContentBundleGuid;
 	TArray<const UDataLayerInstance*> DataLayerInstances;
+	bool bIsStandalone;
 
 	const UExternalDataLayerAsset* GetExternalDataLayerAsset() const
 	{
diff --git a/Engine/Source/Runtime/Engine/Public/WorldPartition/HLOD/StandaloneHLODActor.h b/Engine/Source/Runtime/Engine/Public/WorldPartition/HLOD/StandaloneHLODActor.h
new file mode 100644
--- /dev/null
+++ b/Engine/Source/Runtime/Engine/Public/WorldPartition/HLOD/StandaloneHLODActor.h
@@ -0,0 +1,59 @@
+// Copyright Epic Games, Inc. All Rights Reserved.
+
+#pragma once
+
+#include "LevelInstance/LevelInstanceActorGuid.h"
+#include "LevelInstance/LevelInstanceActorImpl.h"
+#include "LevelInstance/LevelInstanceInterface.h"
+#include "StandaloneHLODActor.generated.h"
+
+UCLASS(NotPlaceable, MinimalAPI, HideCategories=(Rendering, Replication, Collision, Physics, Navigation, Networking, Input, Actor, LevelInstance, Cooking))
+class AWorldPartitionStandaloneHLOD : public AActor, public ILevelInstanceInterface
+{
+	GENERATED_BODY()
+
+public:
+	ENGINE_API AWorldPartitionStandaloneHLOD(const FObjectInitializer& ObjectInitializer);
+
+	//~ Begin ILevelInstanceInterface
+	ENGINE_API virtual const FLevelInstanceID& GetLevelInstanceID() const override;
+	ENGINE_API virtual bool HasValidLevelInstanceID() const override;
+	ENGINE_API virtual const FGuid& GetLevelInstanceGuid() const override;
+	ENGINE_API virtual const TSoftObjectPtr<UWorld>& GetWorldAsset() const override;
+	ENGINE_API virtual bool IsLoadingEnabled() const override;
+	ENGINE_API virtual bool SetWorldAsset(TSoftObjectPtr<UWorld> InWorldAsset) override;
+
+#if WITH_EDITOR
+	ENGINE_API virtual ULevelInstanceComponent* GetLevelInstanceComponent() const override;
+	ENGINE_API virtual ELevelInstanceRuntimeBehavior GetDesiredRuntimeBehavior() const override;
+	ENGINE_API virtual ELevelInstanceRuntimeBehavior GetDefaultRuntimeBehavior() const override;
+#endif // WITH_EDITOR
+	//~ End ILevelInstanceInterface
+
+	//~ Begin UObject Interface
+#if WITH_EDITOR
+	ENGINE_API virtual bool CanEditChange(const FProperty* InProperty) const override;
+	ENGINE_API virtual bool CanEditChangeComponent(const UActorComponent* Component, const FProperty* InProperty) const override;
+#endif
+	//~ End UObject Interface.
+
+	//~ Begin AActor Interface.
+	ENGINE_API void PostRegisterAllComponents() override;
+	ENGINE_API void PostUnregisterAllComponents() override;
+	
+#if WITH_EDITOR
+	ENGINE_API virtual TUniquePtr<class FWorldPartitionActorDesc> CreateClassActorDesc() const override;
+#endif
+	//~ End AActor Interface.
+
+protected:
+	UPROPERTY()
+	TSoftObjectPtr<UWorld> WorldAsset;
+	
+	UPROPERTY(Transient)
+	FGuid LevelInstanceSpawnGuid;
+
+private:
+	FLevelInstanceActorGuid LevelInstanceActorGuid;
+	FLevelInstanceActorImpl LevelInstanceActorImpl;
+};
\ No newline at end of file
diff --git a/Engine/Source/Runtime/Engine/Public/WorldPartition/HLOD/StandaloneHLODSubsystem.h b/Engine/Source/Runtime/Engine/Public/WorldPartition/HLOD/StandaloneHLODSubsystem.h
new file mode 100644
--- /dev/null
+++ b/Engine/Source/Runtime/Engine/Public/WorldPartition/HLOD/StandaloneHLODSubsystem.h
@@ -0,0 +1,69 @@
+// Copyright Epic Games, Inc. All Rights Reserved.
+
+#pragma once
+
+#include "Subsystems/WorldSubsystem.h"
+#include "WorldPartition/WorldPartitionActorDesc.h"
+
+#include "StandaloneHLODSubsystem.generated.h"
+
+class AWorldPartitionStandaloneHLOD;
+class FLevelInstanceActorDesc;
+class UActorDescContainerInstance;
+
+UCLASS(MinimalAPI)
+class UWorldPartitionStandaloneHLODSubsystem : public UWorldSubsystem
+{
+	GENERATED_BODY()
+
+#if WITH_EDITOR
+private:
+	struct FStandaloneHLODActorParams
+	{
+		FGuid Guid;
+		FTransform Transform;
+		FString WorldPackageName;
+		FString ActorLabel;
+	};
+
+public:
+
+	//~ Begin USubsystem Interface.
+	ENGINE_API virtual void Initialize(FSubsystemCollectionBase& Collection) override;
+	ENGINE_API virtual void Deinitialize() override;
+	//~ End USubsystem Interface.
+
+	//~ Begin UWorldSubsystem Interface.
+	ENGINE_API virtual bool DoesSupportWorldType(const EWorldType::Type WorldType) const override;
+	//~ End UWorldSubsystem Interface.
+
+	ENGINE_API void UpdateStandaloneHLODActors(FStandaloneHLODActorParams InStandaloneHLODActorParams);
+	ENGINE_API void DeleteStandaloneHLODActors(FGuid InGuid);
+
+	ENGINE_API void UpdateStandaloneHLODActorsRecursive(const FLevelInstanceActorDesc& InLevelInstanceActorDesc, const FTransform InActorTransform, bool bChildrenOnly);
+	ENGINE_API void DeleteStandaloneHLODActorsRecursive(const FLevelInstanceActorDesc& InLevelInstanceActorDesc);
+
+	ENGINE_API void ForEachStandaloneHLODActor(TFunctionRef<void(AWorldPartitionStandaloneHLOD*)> Func) const;
+	ENGINE_API void ForEachStandaloneHLODActorFiltered(FGuid InGuid, TFunctionRef<void(AWorldPartitionStandaloneHLOD*)> Func) const;
+
+	static ENGINE_API bool GetStandaloneHLODFolderPathAndPackagePrefix(const FString& InWorldPackageName, FString& OutFolderPath, FString& OutPackagePrefix);
+
+private:
+	void OnWorldPartitionInitialized(UWorldPartition* InWorldPartition);
+	void OnWorldPartitionUninitialized(UWorldPartition* InWorldPartition);
+
+	void OnActorChanged(AActor* InActor);
+	void OnActorDeleted(AActor* InActor);
+	void OnObjectPropertyChanged(UObject* InObject, FPropertyChangedEvent& InPropertyChangedEvent);
+
+	void OnActorDescContainerInstanceRegistered(UActorDescContainerInstance* InContainerInstance);
+	void OnActorDescContainerInstanceUnregistered(UActorDescContainerInstance* InContainerInstance);
+
+	void OnLevelAddedToWorld(ULevel* Level, UWorld* World);
+
+	TMap<FGuid, TArray<AWorldPartitionStandaloneHLOD*>> StandaloneHLODActors;
+
+	TMap<FName, TMap<int32, FName>> CachedHLODSetups;
+	bool bRefreshCachedHLODSetups;
+#endif
+};
\ No newline at end of file
diff --git a/Engine/Source/Runtime/Engine/Public/WorldPartition/LevelInstance/LevelInstanceActorDesc.h b/Engine/Source/Runtime/Engine/Public/WorldPartition/LevelInstance/LevelInstanceActorDesc.h
--- a/Engine/Source/Runtime/Engine/Public/WorldPartition/LevelInstance/LevelInstanceActorDesc.h
+++ b/Engine/Source/Runtime/Engine/Public/WorldPartition/LevelInstance/LevelInstanceActorDesc.h
@@ -9,6 +9,7 @@
 #include "UObject/ObjectPtr.h"
 #include "Containers/Map.h"
 #include "WorldPartition/WorldPartitionActorDesc.h"
+#include "WorldPartition/HLOD/StandaloneHLODActor.h"
 #include "WorldPartition/LevelInstance/LevelInstancePropertyOverrideDesc.h"
 #include "WorldPartition/Filter/WorldPartitionActorFilter.h"
 
@@ -38,6 +39,10 @@ public:
 	ENGINE_API virtual void CheckForErrors(const IWorldPartitionActorDescInstanceView* InActorDescView, IStreamingGenerationErrorHandler* ErrorHandler) const override;
 
 	virtual UActorDescContainer* GetChildContainer() const override { check(IsChildContainerInstance()); return ChildContainer.Get(); }
+
+	virtual bool HasStandaloneHLOD() const override { return bHasStandaloneHLOD; }
+	virtual ELevelInstanceRuntimeBehavior GetDesiredRuntimeBehavior() const { return DesiredRuntimeBehavior; }
+
 protected:
 	ENGINE_API virtual void Init(const AActor* InActor) override;
 	ENGINE_API virtual void Init(const FWorldPartitionActorDescInitData& DescData) override;
@@ -60,6 +65,7 @@ protected:
 
 	FWorldPartitionActorFilter Filter;
 	bool bIsChildContainerInstance;
+	bool bHasStandaloneHLOD;
 
 	FSoftObjectPath PropertyOverrideAsset;
 	TSharedPtr<FLevelInstancePropertyOverrideDesc> PropertyOverrideDesc;
diff --git a/Engine/Source/Runtime/Engine/Public/WorldPartition/RuntimeHashSet/WorldPartitionRuntimeHashSet.h b/Engine/Source/Runtime/Engine/Public/WorldPartition/RuntimeHashSet/WorldPartitionRuntimeHashSet.h
--- a/Engine/Source/Runtime/Engine/Public/WorldPartition/RuntimeHashSet/WorldPartitionRuntimeHashSet.h
+++ b/Engine/Source/Runtime/Engine/Public/WorldPartition/RuntimeHashSet/WorldPartitionRuntimeHashSet.h
@@ -161,6 +161,8 @@ class UWorldPartitionRuntimeHashSet final : public UWorldPartitionRuntimeHash
 public:
 	ENGINE_API virtual bool Draw2D(FWorldPartitionDraw2DContext& DrawContext) const override;
 	ENGINE_API virtual void Draw3D(const TArray<FWorldPartitionStreamingSource>& Sources) const override;
+	ENGINE_API virtual const FGuid* GetStandaloneHLODActorSourceCellOverride(const FGuid& InActorGuid) const override;
+	ENGINE_API void ForEachHLODLayer(TFunctionRef<bool(FName, FName, int32)> Func) const; // RuntimePartitionName, HLODSetupName, HLODSetupIndex
 
 #if WITH_EDITOR
 	// Streaming generation interface
@@ -243,6 +245,9 @@ private:
 	UPROPERTY(Transient)
 	TMap<FGuid, TObjectPtr<URuntimeHashSetExternalStreamingObject>> WorldAssetStreamingObjects;
 
+	UPROPERTY()
+	TMap<FGuid, FGuid> StandaloneHLODActorToSourceCellsMap;
+
 	// Optimized data
 	TMap<FName, TArray<const FRuntimePartitionStreamingData*>> RuntimeSpatiallyLoadedDataGridMap;
 	TArray<const FRuntimePartitionStreamingData*> RuntimeNonSpatiallyLoadedDataGridList;
diff --git a/Engine/Source/Runtime/Engine/Public/WorldPartition/WorldPartition.h b/Engine/Source/Runtime/Engine/Public/WorldPartition/WorldPartition.h
--- a/Engine/Source/Runtime/Engine/Public/WorldPartition/WorldPartition.h
+++ b/Engine/Source/Runtime/Engine/Public/WorldPartition/WorldPartition.h
@@ -354,6 +354,7 @@ public:
 
 		ISourceControlHelper* SourceControlHelper;
 		bool bReportOnly;
+		mutable TArray<TObjectPtr<UWorldPartition>> OutAdditionalWorldPartitionsForStandaloneHLOD;
 
 		FSetupHLODActorsParams& SetSourceControlHelper(ISourceControlHelper* InSourceControlHelper) { SourceControlHelper = InSourceControlHelper; return *this; }
 		FSetupHLODActorsParams& SetReportOnly(bool bInReportOnly) { bReportOnly = bInReportOnly; return *this; }
@@ -470,8 +471,13 @@ public:
 	bool IsEnablingStreamingJustified() const { return bEnablingStreamingJustified; }
 
 	bool IsHLODsInEditorAllowed() const { return bAllowShowingHLODsInEditor; }
+
+	void SetIsStandaloneHLODWorld(bool bInIsStandaloneHLODWorld) { bIsStandaloneHLODWorld = bInIsStandaloneHLODWorld; }
 #endif
 
+	bool HasStandaloneHLOD() const { return bHasStandaloneHLOD; }
+	bool IsStandaloneHLODWorld() const { return bIsStandaloneHLODWorld; }
+
 public:
 	static ENGINE_API bool IsSimulating(bool bIncludeTestEnableSimulationStreamingSource = true);
 	ENGINE_API int32 GetStreamingStateEpoch() const;
@@ -608,6 +614,12 @@ private:
 	uint8 bAllowShowingHLODsInEditor : 1;
 #endif
 
+	UPROPERTY(EditAnywhere, Category = WorldPartitionSetup, meta = (DisplayName = "Build Standalone HLOD", EditConditionHides, EditCondition = "bEnableStreaming", HideEditConditionToggle))
+	uint8 bHasStandaloneHLOD : 1;
+
+	UPROPERTY()
+	uint8 bIsStandaloneHLODWorld : 1;
+
 	/** if set to true, this removes any content bundles from this world and also removes content bundle editing */
 	UPROPERTY(EditAnywhere, Category = WorldPartitionSetup, AdvancedDisplay)
 	uint8 bDisableContentBundles : 1;
diff --git a/Engine/Source/Runtime/Engine/Public/WorldPartition/WorldPartitionActorDesc.h b/Engine/Source/Runtime/Engine/Public/WorldPartition/WorldPartitionActorDesc.h
--- a/Engine/Source/Runtime/Engine/Public/WorldPartition/WorldPartitionActorDesc.h
+++ b/Engine/Source/Runtime/Engine/Public/WorldPartition/WorldPartitionActorDesc.h
@@ -245,7 +245,7 @@ class FWorldPartitionActorDesc
 	friend class FAssetHeaderPatcherInner;
 	friend class FActorDescArchive;
 	friend class FWorldPartitionActorDescInstance;
-	friend class FStreamingGenerationUnsavedDirtyActorDescInstance;
+	friend class FStreamingGenerationTempActorDescInstance;
 	friend class FPropertyOverrideUtils;
 	template<class U> friend class TActorDescContainerCollection;
 
@@ -486,6 +486,8 @@ public:
 
 	ENGINE_API virtual bool Equals(const FWorldPartitionActorDesc* Other) const;	
 
+	virtual bool HasStandaloneHLOD() const { return false; }
+
 	/**
 	 * Returns true if resaving this actor will have an impact on streaming generation. Before class descriptors, properties changed on a Blueprint 
 	 * that would affect streaming generation weren't taken into account until the actors affected by the change were resaved.
diff --git a/Engine/Source/Runtime/Engine/Public/WorldPartition/WorldPartitionRuntimeHash.h b/Engine/Source/Runtime/Engine/Public/WorldPartition/WorldPartitionRuntimeHash.h
--- a/Engine/Source/Runtime/Engine/Public/WorldPartition/WorldPartitionRuntimeHash.h
+++ b/Engine/Source/Runtime/Engine/Public/WorldPartition/WorldPartitionRuntimeHash.h
@@ -279,6 +279,7 @@ public:
 	virtual bool ContainsRuntimeHash(const FString& Name) const { return false; }
 	virtual bool IsStreaming3D() const { return true; }
 	virtual bool GetShouldMergeStreamingSourceInfo() const { return false; }
+	virtual const FGuid* GetStandaloneHLODActorSourceCellOverride(const FGuid& InActorGuid) const { return nullptr; }
 
 	static ENGINE_API URuntimeHashExternalStreamingObjectBase* CreateExternalStreamingObject(TSubclassOf<URuntimeHashExternalStreamingObjectBase> InClass, UObject* InOuter, UWorld* InOuterWorld);
 	ENGINE_API UWorldPartitionRuntimeCell* CreateRuntimeCell(UClass* CellClass, UClass* CellDataClass, const FString& CellName, const FString& CellInstanceSuffix, UObject* InOuter = nullptr);
