At the end of the cook also add package dependencies.
This fixes an issue where packages were no longer being passed to the UpdateChunkMap now that the UpdateChunkMap happens at the end of the cook.

diff --git a/Engine/Source/Editor/UnrealEd/Private/Commandlets/AssetRegistryGenerator.cpp b/Engine/Source/Editor/UnrealEd/Private/Commandlets/AssetRegistryGenerator.cpp
--- a/Engine/Source/Editor/UnrealEd/Private/Commandlets/AssetRegistryGenerator.cpp
+++ b/Engine/Source/Editor/UnrealEd/Private/Commandlets/AssetRegistryGenerator.cpp
@@ -1813,8 +1813,31 @@ void FAssetRegistryGenerator::FinalizeChunkIDs(const TSet<FName>& InCookedPackag
 
 	if (bGenerateChunks)
 	{
-		// Build the chunk map with only the packages that have been cooked.
-		UAssetManager::Get().UpdateCachedChunkMapAfterCook(CookedPackages, StartupPackages);
+		TSet<FName> AdditionalUncookedDependenciesForChunkSearch;
+		TSet<FName> ProcessedPackages;
+		ProcessedPackages.Reserve(CookedPackages.Num());
+		TArray<FName> PackagesToProcess = CookedPackages.Array();
+		while (PackagesToProcess.Num() > 0)
+		{
+			const FName PackageName = PackagesToProcess[0];
+			PackagesToProcess.RemoveAtSwap(0);
+			if (ProcessedPackages.Contains(PackageName))
+			{
+				continue;
+			}
+
+			TArray<FName> Dependencies;
+			AssetRegistry.GetDependencies(PackageName, Dependencies, UE::AssetRegistry::EDependencyCategory::Package, UE::AssetRegistry::EDependencyQuery::Game);
+			AssetRegistry.GetDependencies(PackageName, Dependencies, UE::AssetRegistry::EDependencyCategory::Package, UE::AssetRegistry::EDependencyQuery::Build);
+
+			AdditionalUncookedDependenciesForChunkSearch.Append(Dependencies);
+			PackagesToProcess.Append(Dependencies);
+			ProcessedPackages.Add(PackageName);
+		}
+		TSet<FName> CookedPackagesAndMapDependencies = CookedPackages;
+		CookedPackagesAndMapDependencies.Append(AdditionalUncookedDependenciesForChunkSearch);
+		// Update the chunk map based upon what we have cooked and their game dependencies to account for external actors.
+		UAssetManager::Get().UpdateCachedChunkMapAfterCook(CookedPackagesAndMapDependencies, StartupPackages);
 	}
 	// add all the packages to the unassigned package list
 	for (FName CookedPackage : CookedPackages)
