Ensure that redirectors have been applied to collections prior to querying them in the asset localization gather

#rnx

diff --git a/Engine/Source/Developer/Localization/Private/UserGeneratedContentLocalization.cpp b/Engine/Source/Developer/Localization/Private/UserGeneratedContentLocalization.cpp
--- a/Engine/Source/Developer/Localization/Private/UserGeneratedContentLocalization.cpp
+++ b/Engine/Source/Developer/Localization/Private/UserGeneratedContentLocalization.cpp
@@ -425,6 +425,7 @@ bool ExportLocalization(TArrayView<const TSharedRef<IPlugin>> Plugins, const FEx
 				}
 				
 				ConfigSection.Add(TEXT("SearchAllAssets"), TEXT("false"));
+				ConfigSection.Add(TEXT("ApplyRedirectorsToCollections"), TEXT("false"));
 
 				GatherConfig.AddGatherTextStep(GatherStepIndex++, MoveTemp(ConfigSection));
 			}
diff --git a/Engine/Source/Editor/UnrealEd/Classes/Commandlets/GatherTextFromAssetsCommandlet.h b/Engine/Source/Editor/UnrealEd/Classes/Commandlets/GatherTextFromAssetsCommandlet.h
--- a/Engine/Source/Editor/UnrealEd/Classes/Commandlets/GatherTextFromAssetsCommandlet.h
+++ b/Engine/Source/Editor/UnrealEd/Classes/Commandlets/GatherTextFromAssetsCommandlet.h
@@ -137,6 +137,7 @@ private:
 	bool bFixStaleGatherCache;
 	bool bFixMissingGatherCache;
 	bool bSearchAllAssets;
+	bool bApplyRedirectorsToCollections;
 	bool bShouldGatherFromEditorOnlyData;
 	bool bShouldExcludeDerivedClasses;
 	bool bFixPackageLocalizationIdConflict;
diff --git a/Engine/Source/Editor/UnrealEd/Private/Commandlets/GatherTextFromAssetsCommandlet.cpp b/Engine/Source/Editor/UnrealEd/Private/Commandlets/GatherTextFromAssetsCommandlet.cpp
--- a/Engine/Source/Editor/UnrealEd/Private/Commandlets/GatherTextFromAssetsCommandlet.cpp
+++ b/Engine/Source/Editor/UnrealEd/Private/Commandlets/GatherTextFromAssetsCommandlet.cpp
@@ -290,6 +290,7 @@ UGatherTextFromAssetsCommandlet::UGatherTextFromAssetsCommandlet(const FObjectIn
 	, NumPackagesDupLocId(0)
 	, bSkipGatherCache(false)
 	, bSearchAllAssets(true)
+	, bApplyRedirectorsToCollections(true)
 	, bShouldGatherFromEditorOnlyData(false)
 	, bShouldExcludeDerivedClasses(false)
 	, bFixPackageLocalizationIdConflict(false)
@@ -365,8 +366,8 @@ void CalculateDependenciesImpl(IAssetRegistry& InAssetRegistry, const FName& InP
 void UGatherTextFromAssetsCommandlet::CalculateDependenciesForPackagesPendingGather()
 {
 	UE_SCOPED_TIMER(TEXT("UGatherTextFromAssetsCommandlet::CalculateDependenciesForPackagesPendingGather"), LogGatherTextFromAssetsCommandlet, Display);
-	FAssetRegistryModule& AssetRegistryModule = FModuleManager::LoadModuleChecked<FAssetRegistryModule>(TEXT("AssetRegistry"));
-	IAssetRegistry& AssetRegistry = AssetRegistryModule.Get();
+	
+	IAssetRegistry& AssetRegistry = IAssetRegistry::GetChecked();
 	TMap<FName, TSet<FName>> PackageNameToDependencies;
 
 	for (FPackagePendingGather& PackagePendingGather : PackagesPendingGather)
@@ -570,6 +571,16 @@ bool UGatherTextFromAssetsCommandlet::BuildFirstPassFilter(FARFilter& InOutFilte
 /** Builds a filter based on the specified collections to be used for gathering.*/
 bool UGatherTextFromAssetsCommandlet::BuildCollectionFilter(FARFilter& InOutFilter, const TArray<FString>& Collections) const
 {
+	if (bApplyRedirectorsToCollections && Collections.Num() > 0)
+	{
+		UE_SCOPED_TIMER(TEXT("UGatherTextFromAssetsCommandlet::DiscoverAllAssetsForCollectionQuery"), LogGatherTextFromAssetsCommandlet, Display);
+
+		// The asset registry scan must finish before redirectors have been applied to collections, 
+		// so wait on that now, unless we've been asked to skip it
+		UE_LOG(LogGatherTextFromAssetsCommandlet, Display, TEXT("Waiting for asset discovery to finish before querying collections..."));
+		IAssetRegistry::GetChecked().SearchAllAssets(true);
+	}
+
 	bool bHasFailedToGetACollection = false;
 	FCollectionManagerModule& CollectionManagerModule = FCollectionManagerModule::GetModule();
 	ICollectionManager& CollectionManager = CollectionManagerModule.Get();
@@ -1663,9 +1674,9 @@ int32 UGatherTextFromAssetsCommandlet::Main(const FString& Params)
 	UE_LOG(LogGatherTextFromAssetsCommandlet, Display, TEXT("Discovering assets to gather..."));
 	const double DiscoveringAssetsStartTime = FPlatformTime::Seconds();
 	{
-		UE_SCOPED_TIMER(TEXT("UGatherTextFromAssetsCommandlet::SearchAssetRegistryForAllAssets"), LogGatherTextFromAssetsCommandlet, Display);
-		FAssetRegistryModule& AssetRegistryModule = FModuleManager::LoadModuleChecked<FAssetRegistryModule>(TEXT("AssetRegistry"));
-		IAssetRegistry& AssetRegistry = AssetRegistryModule.Get();
+		UE_SCOPED_TIMER(TEXT("UGatherTextFromAssetsCommandlet::DiscoverAllAssetsForGatherFilter"), LogGatherTextFromAssetsCommandlet, Display);
+
+		IAssetRegistry& AssetRegistry = IAssetRegistry::GetChecked();
 
 		if (bSearchAllAssets || FirstPassFilter.PackagePaths.IsEmpty())
 		{
@@ -1934,6 +1945,11 @@ bool UGatherTextFromAssetsCommandlet::ConfigureFromScript(const FString& GatherT
 		bSearchAllAssets = true;
 	}
 
+	if (!GetBoolFromConfig(*SectionName, TEXT("ApplyRedirectorsToCollections"), bApplyRedirectorsToCollections, GatherTextConfigPath))
+	{
+		bApplyRedirectorsToCollections = true;
+	}
+
 	// Get whether we should gather editor-only data. Typically only useful for the localization of UE itself.
 	if (!GetBoolFromConfig(*SectionName, TEXT("ShouldGatherFromEditorOnlyData"), bShouldGatherFromEditorOnlyData, GatherTextConfigPath))
 	{
