Suppress unsolicited package messages during cook validation

diff --git a/Engine/Source/Editor/UnrealEd/Private/CookOnTheFlyServer.cpp b/Engine/Source/Editor/UnrealEd/Private/CookOnTheFlyServer.cpp
--- a/Engine/Source/Editor/UnrealEd/Private/CookOnTheFlyServer.cpp
+++ b/Engine/Source/Editor/UnrealEd/Private/CookOnTheFlyServer.cpp
@@ -3353,6 +3353,8 @@ EDataValidationResult UCookOnTheFlyServer::ValidateSourcePackage(UE::Cook::FPack
 	// Run asset validation if requested
 	if (EnumHasAnyFlags(CookByTheBookOptions->StartupOptions, ECookByTheBookOptions::RunAssetValidation) && UE::Cook::FDelegates::ValidateSourcePackage.IsBound())
 	{
+		FCookLoadScope CookLoadScope(ECookLoadType::EditorOnly);
+		
 		TArray<FAssetData> ExternalObjects;
 		if (World)
 		{
