Fix nullptr crash in HLOD builder.

diff --git a/Engine/Source/Editor/UnrealEd/Private/WorldPartition/WorldPartitionHLODsBuilder.cpp b/Engine/Source/Editor/UnrealEd/Private/WorldPartition/WorldPartitionHLODsBuilder.cpp
--- a/Engine/Source/Editor/UnrealEd/Private/WorldPartition/WorldPartitionHLODsBuilder.cpp
+++ b/Engine/Source/Editor/UnrealEd/Private/WorldPartition/WorldPartitionHLODsBuilder.cpp
@@ -1048,7 +1048,8 @@ bool UWorldPartitionHLODsBuilder::GenerateBuildManifest(TMap<FString, TPair<int3
 
 	// If we're generating manifest for Standalone HLOD levels, each of them needs a separate config file
 	TArray<FConfigFile> ConfigFiles;
-	ConfigFiles.SetNum(WorldPartition->HasStandaloneHLOD() ? AdditionalWorldPartitionsForStandaloneHLOD.Num() : 1);
+	bool bHasStandaloneHLOD = WorldPartition ? WorldPartition->HasStandaloneHLOD() : false;
+	ConfigFiles.SetNum(bHasStandaloneHLOD ? AdditionalWorldPartitionsForStandaloneHLOD.Num() : 1);
 	for (FConfigFile& ConfigFile : ConfigFiles)
 	{
 		ConfigFile.SetInt64(TEXT("General"), TEXT("BuilderCount"), BuilderCount);
@@ -1073,7 +1074,7 @@ bool UWorldPartitionHLODsBuilder::GenerateBuildManifest(TMap<FString, TPair<int3
 
 		for (int32 WorldIndex = 0; WorldIndex < BuildersWorkload[WorkloadIndex].PerWorldHLODWorkloads.Num(); WorldIndex++)
 		{
-			UWorldPartition* CurrentWorldPartition = WorldPartition->HasStandaloneHLOD() ? AdditionalWorldPartitionsForStandaloneHLOD[WorldIndex].Get() : WorldPartition;
+			UWorldPartition* CurrentWorldPartition = bHasStandaloneHLOD ? AdditionalWorldPartitionsForStandaloneHLOD[WorldIndex].Get() : WorldPartition;
 			for(const FGuid& ActorGuid : BuildersWorkload[WorkloadIndex].PerWorldHLODWorkloads[WorldIndex])
 			{
 				ConfigFiles[WorldIndex].AddToSection(*SectionName, TEXT("+HLODActorGuid"), ActorGuid.ToString(EGuidFormats::Digits));
