Fully decouple the UGC localization target name from the DLC plugin name

The few remaining places that assumed the plugin name always matched the localization target name have been updated to read the localization target info from the .uplugin instead

#rnx

diff --git a/Engine/Source/Developer/Localization/Private/UserGeneratedContentLocalization.cpp b/Engine/Source/Developer/Localization/Private/UserGeneratedContentLocalization.cpp
--- a/Engine/Source/Developer/Localization/Private/UserGeneratedContentLocalization.cpp
+++ b/Engine/Source/Developer/Localization/Private/UserGeneratedContentLocalization.cpp
@@ -144,9 +144,7 @@ FString GetLocalizationScratchDirectory(const TSharedRef<IPlugin>& Plugin)
 
 FString GetLocalizationTargetName(const TSharedRef<IPlugin>& Plugin)
 {
-	// Note: If you change this naming scheme you'll need to handle backwards computability with existing data,
-	//       and other code (such as the cooker) that assumes it can use the plugin name as the target name
-	//       (@see GetLocalizationTargetDirectory overload that takes the PluginName and PluginContentDirectory)
+	// Note: If you change this naming scheme you'll need to handle backwards computability with existing data
 	return Plugin->GetName();
 }
 
@@ -162,9 +160,9 @@ FString GetLocalizationTargetDirectory(const TSharedRef<IPlugin>& Plugin, const
 		: FPaths::Combine(LocalizationTargetRootDirectoryOverride, Plugin->GetName());
 }
 
-FString GetLocalizationTargetDirectory(const FString& PluginName, const FString& PluginContentDirectory)
+FString GetLocalizationTargetDirectory(const FString& LocalizationTargetName, const FString& PluginContentDirectory)
 {
-	return FPaths::Combine(PluginContentDirectory, TEXT("Localization"), PluginName);
+	return FPaths::Combine(PluginContentDirectory, TEXT("Localization"), LocalizationTargetName);
 }
 
 FString GetLocalizationTargetUGCLocFile(const FString& LocalizationTargetName, const FString& LocalizationTargetDirectory)
diff --git a/Engine/Source/Developer/Localization/Public/UserGeneratedContentLocalization.h b/Engine/Source/Developer/Localization/Public/UserGeneratedContentLocalization.h
--- a/Engine/Source/Developer/Localization/Public/UserGeneratedContentLocalization.h
+++ b/Engine/Source/Developer/Localization/Public/UserGeneratedContentLocalization.h
@@ -201,7 +201,7 @@ LOCALIZATION_API FString GetLocalizationTargetDirectory(const TSharedRef<IPlugin
 /**
  * Utility to get the path of the UGC localization target that would be used by the given plugin name and directory.
  */
-LOCALIZATION_API FString GetLocalizationTargetDirectory(const FString& PluginName, const FString& PluginContentDirectory);
+LOCALIZATION_API FString GetLocalizationTargetDirectory(const FString& LocalizationTargetName, const FString& PluginContentDirectory);
 
 /**
  * Utility to get the UGCLoc file for for the given UGC localization target name and directory.
diff --git a/Engine/Source/Editor/UnrealEd/Private/CookOnTheFlyServer.cpp b/Engine/Source/Editor/UnrealEd/Private/CookOnTheFlyServer.cpp
--- a/Engine/Source/Editor/UnrealEd/Private/CookOnTheFlyServer.cpp
+++ b/Engine/Source/Editor/UnrealEd/Private/CookOnTheFlyServer.cpp
@@ -12476,22 +12476,35 @@ void UCookOnTheFlyServer::CompileDLCLocalization(FBeginCookContext& BeginContext
 		return;
 	}
 
-	// Used to validate that we're not loading/compiling invalid cultures during the compile step below
-	FUserGeneratedContentLocalizationDescriptor DefaultUGCLocDescriptor;
-	DefaultUGCLocDescriptor.InitializeFromProject();
-
-	// Also filter the validation list by the cultures we're cooking against
-	DefaultUGCLocDescriptor.CulturesToGenerate.RemoveAll([this](const FString& Culture)
+	if (TSharedPtr<IPlugin> Plugin = IPluginManager::Get().FindPlugin(CookByTheBookOptions->DlcName))
 	{
-		return !CookByTheBookOptions->AllCulturesToCook.Contains(Culture);
-	});
+		if (Plugin->GetDescriptor().LocalizationTargets.Num() > 0)
+		{
+			// Used to validate that we're not loading/compiling invalid cultures during the compile step below
+			FUserGeneratedContentLocalizationDescriptor DefaultUGCLocDescriptor;
+			DefaultUGCLocDescriptor.InitializeFromProject();
 
-	// Compile UGC localization (if available) for this DLC plugin
-	const FString InputLocalizationTargetDirectory = UserGeneratedContentLocalization::GetLocalizationTargetDirectory(CookByTheBookOptions->DlcName, GetContentDirectoryForDLC());
-	for (const FBeginCookContextPlatform& PlatformContext : BeginContext.PlatformContexts)
+			// Also filter the validation list by the cultures we're cooking against
+			DefaultUGCLocDescriptor.CulturesToGenerate.RemoveAll([this](const FString& Culture)
+			{
+				return !CookByTheBookOptions->AllCulturesToCook.Contains(Culture);
+			});
+
+			// Compile UGC localization (if available) for this DLC plugin
+			for (const FLocalizationTargetDescriptor& LocalizationTarget : Plugin->GetDescriptor().LocalizationTargets)
+			{
+				const FString InputLocalizationTargetDirectory = UserGeneratedContentLocalization::GetLocalizationTargetDirectory(LocalizationTarget.Name, GetContentDirectoryForDLC());
+				for (const FBeginCookContextPlatform& PlatformContext : BeginContext.PlatformContexts)
+				{
+					const FString OutputLocalizationTargetDirectory = ConvertToFullSandboxPath(InputLocalizationTargetDirectory, /*bForWrite*/true, PlatformContext.TargetPlatform->PlatformName());
+					UserGeneratedContentLocalization::CompileLocalization(LocalizationTarget.Name, InputLocalizationTargetDirectory, OutputLocalizationTargetDirectory, GetDefault<UUserGeneratedContentLocalizationSettings>()->bValidateDLCLocalizationDuringCook ? &DefaultUGCLocDescriptor : nullptr);
+				}
+			}
+		}
+	}
+	else
 	{
-		const FString OutputLocalizationTargetDirectory = ConvertToFullSandboxPath(InputLocalizationTargetDirectory, /*bForWrite*/true, PlatformContext.TargetPlatform->PlatformName());
-		UserGeneratedContentLocalization::CompileLocalization(CookByTheBookOptions->DlcName, InputLocalizationTargetDirectory, OutputLocalizationTargetDirectory, GetDefault<UUserGeneratedContentLocalizationSettings>()->bValidateDLCLocalizationDuringCook ? &DefaultUGCLocDescriptor : nullptr);
+		UE_LOG(LogCook, Warning, TEXT("DLC plugin '%s' was not known to IPluginManager. UGC localization has not been compiled and will be missing at runtime."), *CookByTheBookOptions->DlcName);
 	}
 }
 
