IncrementalCook: Add log diagnostics at end of cook for how many packages of each AssetClass were found incrementally modified.
#rnx
#incrementalcook

diff --git a/Engine/Source/Editor/UnrealEd/Private/CookOnTheFlyServer.cpp b/Engine/Source/Editor/UnrealEd/Private/CookOnTheFlyServer.cpp
--- a/Engine/Source/Editor/UnrealEd/Private/CookOnTheFlyServer.cpp
+++ b/Engine/Source/Editor/UnrealEd/Private/CookOnTheFlyServer.cpp
@@ -7687,6 +7687,7 @@ void UCookOnTheFlyServer::PopulateCookedPackages(TArrayView<const ITargetPlatfor
 	// have to check whether each one is already cooked.
 	bool bFirstPlatform = true;
 	COOK_STAT(DetailedCookStats::NumPackagesIncrementallySkipped = 0);
+	COOK_STAT(DetailedCookStats::NumPackagesIncrementallyModifiedByClass.Empty());
 	for (const ITargetPlatform* TargetPlatform : TargetPlatforms)
 	{
 		FAssetRegistryGenerator& PlatformAssetRegistry = *(PlatformManager->GetPlatformData(TargetPlatform)->RegistryGenerator);
@@ -9917,6 +9918,11 @@ void UCookOnTheFlyServer::ShutdownCookSession()
 
 		PrintFinishStats();
 		OutputHierarchyTimers();
+		if (PlatformManager->GetSessionPlatforms().Num() > 0 &&
+			PlatformManager->GetPlatformData(PlatformManager->GetSessionPlatforms()[0])->bAllowIncrementalResults)
+		{
+			OutputIncrementalCookStats();
+		}
 		PrintDetailedCookStats();
 
 		// BroadcastCookFinished needs to be called before clearing the session data, so that subscribers
@@ -10654,6 +10660,8 @@ void UCookOnTheFlyServer::BeginCookSandbox(FBeginCookContext& BeginContext)
 		{
 			// Set the NumPackagesIncrementallySkipped field to include all of the already CookedPackages
 			COOK_STAT(DetailedCookStats::NumPackagesIncrementallySkipped = 0);
+			// TODO: Initializing NumPackagesIncrementallyModifiedByClass is not yet implemented in this path (cooking
+			// multiple times in editor process) because we plan to deprecate this path soon.
 			const ITargetPlatform* TargetPlatform = AlreadyCookedPlatforms[0];
 			PackageDatas->LockAndEnumeratePackageDatas([TargetPlatform](UE::Cook::FPackageData* PackageData)
 			{
diff --git a/Engine/Source/Editor/UnrealEd/Private/Cooker/CookPackageData.cpp b/Engine/Source/Editor/UnrealEd/Private/Cooker/CookPackageData.cpp
--- a/Engine/Source/Editor/UnrealEd/Private/Cooker/CookPackageData.cpp
+++ b/Engine/Source/Editor/UnrealEd/Private/Cooker/CookPackageData.cpp
@@ -59,7 +59,7 @@ static FAutoConsoleVariableRef CVarPollAsyncPeriod(
 // FPackageData
 FPackagePlatformData::FPackagePlatformData()
 	: Reachability((uint8)EReachability::None), ReachabilityVisitedByCluster((uint8)EReachability::None)
-	, bSaveTimedOut(0), bCookable(1), bExplorable(1), bExplorableOverride(0) , bIncrementallyUnmodified(0)
+	, bSaveTimedOut(0), bCookable(1), bExplorable(1), bExplorableOverride(0) , IncrementallyUnmodified(0)
 	, bIncrementallySkipped(0), bRegisteredForCachedObjectsInOuter(0), bCommitted(0)
 	, CookResults((uint8)ECookResult::NotAttempted)
 {
diff --git a/Engine/Source/Editor/UnrealEd/Private/Cooker/CookPackageData.h b/Engine/Source/Editor/UnrealEd/Private/Cooker/CookPackageData.h
--- a/Engine/Source/Editor/UnrealEd/Private/Cooker/CookPackageData.h
+++ b/Engine/Source/Editor/UnrealEd/Private/Cooker/CookPackageData.h
@@ -256,9 +256,19 @@ struct FPackagePlatformData
 	/** Called on CookWorkers to indicate reachable,cookable,etc for packages sent from Director. */
 	void MarkCommittableForWorker(EReachability InReachability, FCookWorkerClient& CookWorkerClient);
 
-	/** The package was found to be unmodified in the current incremental cook. */
+	/**
+	 * Returns whether SetIncrementallyUnmodified has been called in the current cook session.
+	 * IncrementallyUnmodified like a TOptional tracks an unset state in addition to its value if set.
+	 */
+	bool IsIncrementallyUnmodifiedSet() const;
+	/**
+	 * The package was found to be unmodified in the current incremental cook.
+	 * Returns false if !IsIncrementallyUnmodifiedSet().
+	 */
 	bool IsIncrementallyUnmodified() const;
 	void SetIncrementallyUnmodified(bool bValue);
+	/** Restore the value to its unset state (e.g. beginning a new cook session). */
+	void ClearIncrementallyUnmodified();
 
 	/**
 	 * The package was skipped in the current incremental cook. This might not be equal to IncrementallyUnmodified,
@@ -296,7 +306,7 @@ private:
 	uint32 bCookable : 1;
 	uint32 bExplorable : 1;
 	uint32 bExplorableOverride : 1;
-	uint32 bIncrementallyUnmodified : 1;
+	uint32 IncrementallyUnmodified : 2;
 	uint32 bIncrementallySkipped : 1;
 	uint32 bRegisteredForCachedObjectsInOuter : 1;
 	uint32 bReportedToDirector : 1;
@@ -1570,14 +1580,24 @@ inline void FPackagePlatformData::SetExplorableOverride(bool bValue)
 	bExplorableOverride = (uint32)bValue;
 }
 
+inline bool FPackagePlatformData::IsIncrementallyUnmodifiedSet() const
+{
+	return IncrementallyUnmodified != 0;
+}
+
 inline bool FPackagePlatformData::IsIncrementallyUnmodified() const
 {
-	return bIncrementallyUnmodified != 0;
+	return IncrementallyUnmodified == 2;
 }
 
 inline void FPackagePlatformData::SetIncrementallyUnmodified(bool bValue)
 {
-	bIncrementallyUnmodified = (uint32)bValue;
+	IncrementallyUnmodified = bValue ? 2 : 1;
+}
+
+inline void FPackagePlatformData::ClearIncrementallyUnmodified()
+{
+	IncrementallyUnmodified = 0;
 }
 
 inline bool FPackagePlatformData::IsIncrementallySkipped() const
diff --git a/Engine/Source/Editor/UnrealEd/Private/Cooker/CookProfiling.cpp b/Engine/Source/Editor/UnrealEd/Private/Cooker/CookProfiling.cpp
--- a/Engine/Source/Editor/UnrealEd/Private/Cooker/CookProfiling.cpp
+++ b/Engine/Source/Editor/UnrealEd/Private/Cooker/CookProfiling.cpp
@@ -250,6 +250,7 @@ namespace DetailedCookStats
 	std::atomic<int32> NumDetectedLoads{ 0 };
 	int32 NumRequestedLoads = 0;
 	uint32 NumPackagesIncrementallySkipped = 0;
+	TMap<FTopLevelAssetPath, uint32> NumPackagesIncrementallyModifiedByClass;
 	FCookStatsManager::FAutoRegisterCallback RegisterCookOnTheFlyServerStats([](FCookStatsManager::AddStatFuncRef AddStat)
 		{
 			AddStat(TEXT("Package.Load"), FCookStatsManager::CreateKeyValueArray(
@@ -1129,5 +1130,39 @@ void SendLogCookStats(ECookMode::Type CookMode)
 	}
 }
 
+} // namespace DetailedCookStats
+
+void OutputIncrementalCookStats()
+{
+	using namespace DetailedCookStats;
+
+	uint32 NumPackagesIncrementallyModifiedByClassSum = 0;
+	TArray<TPair<FTopLevelAssetPath, uint32>> ClassAndNums = NumPackagesIncrementallyModifiedByClass.Array();
+	ClassAndNums.RemoveAllSwap(
+		[&NumPackagesIncrementallyModifiedByClassSum](const TPair<FTopLevelAssetPath, uint32>& Pair)
+		{
+			NumPackagesIncrementallyModifiedByClassSum += Pair.Value;
+			return Pair.Value == 0;
+		});
+	if (NumPackagesIncrementallyModifiedByClassSum > 0)
+	{
+		Algo::Sort(ClassAndNums,
+			[](const TPair<FTopLevelAssetPath, uint32>& A, const TPair<FTopLevelAssetPath, uint32>& B)
+			{
+				if (A.Value != B.Value)
+				{
+					return A.Value > B.Value;
+				}
+				return A.Key.Compare(B.Key) < 0;
+			});
+
+		UE_LOG(LogCookStats, Display, TEXT("Incrementally Modified Counts for Each Class:"));
+		UE_LOG(LogCookStats, Display, TEXT("\tTotal, %u"), NumPackagesIncrementallyModifiedByClassSum);
+		for (const TPair<FTopLevelAssetPath, uint32>& Pair : ClassAndNums)
+		{
+			UE_LOG(LogCookStats, Display, TEXT("\t%s, %u"), *Pair.Key.ToString(), Pair.Value);
+		}
+	}
 }
-#endif
\ No newline at end of file
+
+#endif // ENABLE_COOK_STATS
\ No newline at end of file
diff --git a/Engine/Source/Editor/UnrealEd/Private/Cooker/CookProfiling.h b/Engine/Source/Editor/UnrealEd/Private/Cooker/CookProfiling.h
--- a/Engine/Source/Editor/UnrealEd/Private/Cooker/CookProfiling.h
+++ b/Engine/Source/Editor/UnrealEd/Private/Cooker/CookProfiling.h
@@ -13,6 +13,11 @@
 #define OUTPUT_COOKTIMING ENABLE_COOK_STATS
 #define PROFILE_NETWORK 0
 
+#if ENABLE_COOK_STATS
+#include "Containers/Map.h"
+#include "UObject/TopLevelAssetPath.h"
+#endif
+
 #if OUTPUT_COOKTIMING
 #include "Trace/Trace.inl"
 #include "ProfilingDebugging/FormatArgsTrace.h"
@@ -24,6 +29,7 @@ struct FWeakObjectPtr;
 #if OUTPUT_COOKTIMING
 
 void OutputHierarchyTimers();
+void OutputIncrementalCookStats();
 void ClearHierarchyTimers();
 
 struct FHierarchicalTimerInfo;
@@ -160,6 +166,7 @@ extern bool IsUnversioned;
 extern std::atomic<int32> NumDetectedLoads;
 extern int32 NumRequestedLoads;
 extern uint32 NumPackagesIncrementallySkipped;
+extern TMap<FTopLevelAssetPath, uint32> NumPackagesIncrementallyModifiedByClass;
 extern int32 PeakRequestQueueSize;
 extern int32 PeakLoadQueueSize;
 extern int32 PeakSaveQueueSize;
diff --git a/Engine/Source/Editor/UnrealEd/Private/Cooker/CookRequestCluster.cpp b/Engine/Source/Editor/UnrealEd/Private/Cooker/CookRequestCluster.cpp
--- a/Engine/Source/Editor/UnrealEd/Private/Cooker/CookRequestCluster.cpp
+++ b/Engine/Source/Editor/UnrealEd/Private/Cooker/CookRequestCluster.cpp
@@ -2229,7 +2229,7 @@ void FRequestCluster::FGraphSearch::FExploreEdgesContext::ProcessPlatformAttachm
 					ParentGenerationHelper->MarkPackageIncrementallySkipped(*PackageData, TargetPlatform,
 						true /* bIncrementallySkipped */);
 				}
-				if (CookResult == ECookResult::Succeeded && PlatformIndex == FirstSessionPlatformIndex)
+				if (PlatformIndex == FirstSessionPlatformIndex && CookResult == ECookResult::Succeeded)
 				{
 					COOK_STAT(++DetailedCookStats::NumPackagesIncrementallySkipped);
 				}
@@ -2273,6 +2273,28 @@ void FRequestCluster::FGraphSearch::FExploreEdgesContext::ProcessPlatformAttachm
 					ParentGenerationHelper->MarkPackageIncrementallySkipped(*PackageData, TargetPlatform,
 						false /* bIncrementallySkipped */);
 				}
+
+#if ENABLE_COOK_STATS
+				if (PlatformIndex == FirstSessionPlatformIndex &&
+					QueryPlatformData.bIncrementallyModifiedInstigated &&
+					(QueryPlatformData.CookAttachments.CommitStatus == IPackageWriter::ECommitStatus::Success
+						|| QueryPlatformData.CookAttachments.CommitStatus == IPackageWriter::ECommitStatus::Error))
+				{
+					FTopLevelAssetPath AssetClassPath;
+					FARFilter PackageNameFilter;
+					PackageNameFilter.PackageNames.Add(PackageData->GetPackageName());
+					Cluster.AssetRegistry.EnumerateAssets(PackageNameFilter, [&AssetClassPath](const FAssetData& AssetData)
+						{
+							AssetClassPath = AssetData.AssetClassPath;
+							return false;
+						});
+
+					if (AssetClassPath.IsValid())
+					{
+						++DetailedCookStats::NumPackagesIncrementallyModifiedByClass.FindOrAdd(AssetClassPath, 0);
+					}
+				}
+#endif
 			}
 		}
 	}
@@ -2307,11 +2329,17 @@ void FRequestCluster::FGraphSearch::FExploreEdgesContext::ProcessPlatformAttachm
 void FRequestCluster::FGraphSearch::FExploreEdgesContext::SetIncrementallyUnmodified(int32 PlatformIndex,
 	bool bIncrementallyUnmodified, FPackagePlatformData& PackagePlatformData)
 {
-	Vertex->GetPlatformData()[PlatformIndex].bIncrementallyUnmodified.Emplace(bIncrementallyUnmodified);
+	FQueryPlatformData& ClusterPlatformData = Vertex->GetPlatformData()[PlatformIndex];
+	ClusterPlatformData.bIncrementallyUnmodified.Emplace(bIncrementallyUnmodified);
 	if (bIncrementallyUnmodified)
 	{
 		PackagePlatformData.SetIncrementallyUnmodified(true);
 	}
+	else if (!PackagePlatformData.IsIncrementallyUnmodifiedSet())
+	{
+		ClusterPlatformData.bIncrementallyModifiedInstigated = true;
+		PackagePlatformData.SetIncrementallyUnmodified(false);
+	}
 }
 
 FRequestCluster::FVertexData* FRequestCluster::AllocateVertex(FName PackageName, FPackageData* PackageData)
diff --git a/Engine/Source/Editor/UnrealEd/Private/Cooker/CookRequestCluster.h b/Engine/Source/Editor/UnrealEd/Private/Cooker/CookRequestCluster.h
--- a/Engine/Source/Editor/UnrealEd/Private/Cooker/CookRequestCluster.h
+++ b/Engine/Source/Editor/UnrealEd/Private/Cooker/CookRequestCluster.h
@@ -140,6 +140,7 @@ private:
 		bool bExploreCompleted = false;
 		bool bIncrementallyUnmodifiedRequested = false;
 		bool bTransitiveBuildDependenciesResolvedAsNotModified = false;
+		bool bIncrementallyModifiedInstigated = false;
 		TOptional<bool> bIncrementallyUnmodified;
 	private:
 		std::atomic<EAsyncQueryStatus> AsyncQueryStatus;
