Changed a localization warning to use structured logging

diff --git a/Engine/Source/Editor/UnrealEd/Private/Commandlets/GatherTextFromSourceCommandlet.cpp b/Engine/Source/Editor/UnrealEd/Private/Commandlets/GatherTextFromSourceCommandlet.cpp
--- a/Engine/Source/Editor/UnrealEd/Private/Commandlets/GatherTextFromSourceCommandlet.cpp
+++ b/Engine/Source/Editor/UnrealEd/Private/Commandlets/GatherTextFromSourceCommandlet.cpp
@@ -1400,7 +1400,12 @@ bool UGatherTextFromSourceCommandlet::ParseSourceText(const FString& Text, const
 				ParseCtxt.WithinStringLiteral = *(Cursor - 1) == TEXT('\\');
 			}
 
-			UE_CLOG(!ParseCtxt.WithinStringLiteral, LogGatherTextFromSourceCommandlet, Warning, TEXT("A string literal was not correctly terminated. File %s at line %d, starting line: %s"), *ParseCtxt.Filename, ParseCtxt.LineNumber, ParseCtxt.WithinStartingLine);
+			UE_CLOGFMT(!ParseCtxt.WithinStringLiteral, LogGatherTextFromSourceCommandlet, Warning, "A string literal was not correctly terminated. File {file} at line {line}, starting line: {startingLine}",
+				("file", *ParseCtxt.Filename),
+				("line", ParseCtxt.LineNumber),
+				("startingLine", ParseCtxt.WithinStartingLine),
+				("id", GatherTextFromSourceCommandlet::LocalizationLogIdentifier)
+			);
 		}
 	}
 
