Cooker: Extend the AddExternalImportValidation of missing UClass imports to imports of all native types.

#rnx

diff --git a/Engine/Plugins/MovieScene/MovieRenderPipeline/Source/MovieRenderPipelineCore/Private/MovieRenderPipelineCoreModule.cpp b/Engine/Plugins/MovieScene/MovieRenderPipeline/Source/MovieRenderPipelineCore/Private/MovieRenderPipelineCoreModule.cpp
--- a/Engine/Plugins/MovieScene/MovieRenderPipeline/Source/MovieRenderPipelineCore/Private/MovieRenderPipelineCoreModule.cpp
+++ b/Engine/Plugins/MovieScene/MovieRenderPipeline/Source/MovieRenderPipelineCore/Private/MovieRenderPipelineCoreModule.cpp
@@ -6,18 +6,27 @@
 #include "Modules/ModuleManager.h"
 #include "MovieRenderPipelineDataTypes.h"
 #include "Engine/World.h"
+#include "HAL/IConsoleManager.h"
 #include "TimerManager.h"
 #include "MoviePipeline.h"
 #include "Graph/MovieGraphPipeline.h"
 #include "Graph/Nodes/MovieGraphBurnInNode.h"
 #include "UObject/ICookInfo.h"
 
+#if WITH_EDITOR
+static TAutoConsoleVariable<bool> CVarAddCodeReferencedAssetsToCook(
+	TEXT("MoviePipeline.AddCodeReferencedAssetsToCook"), true,
+	TEXT("Add code-referenced-only assets to the cook.")
+	TEXT("Can be disabled for a project that restricts this plugin's usage to the Editor target."),
+	ECVF_Default);
+#endif
+
 FName IMoviePipelineBurnInExtension::ModularFeatureName = "ModularFeature_MoviePipelineBurnInExt";
 
 void FMovieRenderPipelineCoreModule::StartupModule()
 {
 #if WITH_EDITOR
-	if (IsRunningCookCommandlet())
+	if (IsRunningCookCommandlet() && CVarAddCodeReferencedAssetsToCook.GetValueOnGameThread())
 	{
 		UE::Cook::FDelegates::ModifyCook.AddLambda(
 			[](UE::Cook::ICookInfo& CookInfo, TArray<UE::Cook::FPackageCookRule>& InOutPackageCookRules)
diff --git a/Engine/Source/Editor/UnrealEd/Private/CookOnTheFlyServer.cpp b/Engine/Source/Editor/UnrealEd/Private/CookOnTheFlyServer.cpp
--- a/Engine/Source/Editor/UnrealEd/Private/CookOnTheFlyServer.cpp
+++ b/Engine/Source/Editor/UnrealEd/Private/CookOnTheFlyServer.cpp
@@ -10775,8 +10775,7 @@ UE::Cook::FCookSavePackageContext* UCookOnTheFlyServer::CreateSaveContext(const
 			{
 				for (const TObjectPtr<UObject>& Object : ValidationContext.Imports)
 				{
-					UClass* Class = Cast<UClass>(Object);
-					if (Class && Class->IsNative() && SuppressedNativeScriptPackages.Contains(Class->GetPackage()->GetFName()))
+					if (Object->IsNative() && SuppressedNativeScriptPackages.Contains(Object->GetPackage()->GetFName()))
 					{
 						FInstigator Instigator = GetInstigator(ValidationContext.Package->GetFName(), EReachability::Runtime);
 						bool bIsError = true;
@@ -10787,19 +10786,40 @@ UE::Cook::FCookSavePackageContext* UCookOnTheFlyServer::CreateSaveContext(const
 							bIsError = false;
 						}
 
-						// If you receive this message in a package that you do want to cook, you can remove the object of the
-						// unavailable class by overriding UObject::NeedsLoadForTargetPlatform on that class to return false.
-						if (bIsError)
+						if (Object->IsA<UClass>())
 						{
-							UE_ASSET_LOG(LogCook, Error, ValidationContext.Package, TEXT("Failed to cook %s for platform %s. It imports class %s, which is in a module that is not available on the platform."),
-								*ValidationContext.Package->GetName(), *TargetPlatform->PlatformName(), *Class->GetPathName());
-							return ESavePackageResult::ValidatorError;
+							if (bIsError)
+							{
+								// If you receive this message in a package that you do want to cook, you can remove the object of the
+								// unavailable class by overriding UObject::NeedsLoadForTargetPlatform on that class to return false.
+								UE_ASSET_LOG(LogCook, Error, ValidationContext.Package, TEXT("Failed to cook %s for platform %s. It imports class %s, which is in a module that is not available on the platform."),
+									*ValidationContext.Package->GetName(), *TargetPlatform->PlatformName(), *Object->GetPathName());
+								return ESavePackageResult::ValidatorError;
+							}
+							else
+							{
+								UE_ASSET_LOG(LogCook, Display, ValidationContext.Package, TEXT("Skipping package %s for platform %s. It imports class %s, which is in a module that is not available on the platform."),
+									*ValidationContext.Package->GetName(), *TargetPlatform->PlatformName(), *Object->GetPathName());
+								return ESavePackageResult::ValidatorSuppress;
+							}
 						}
 						else
 						{
-							UE_ASSET_LOG(LogCook, Display, ValidationContext.Package, TEXT("Skipping package %s for platform %s. It imports class %s, which is in a module that is not available on the platform."),
-								*ValidationContext.Package->GetName(), *TargetPlatform->PlatformName(), *Class->GetPathName());
-							return ESavePackageResult::ValidatorSuppress;
+							if (bIsError)
+							{
+								// If you receive this message in a package that you do want to cook, you can remove the reference to the missing object
+								// marking the property holding reference as editor-only, or otherwise failing to serialize it, or nulling it out during PreSave
+								// and restoring it during PostSave.
+								UE_ASSET_LOG(LogCook, Error, ValidationContext.Package, TEXT("Failed to cook %s for platform %s. It imports type '%s', which is in a module that is not available on the platform."),
+									*ValidationContext.Package->GetName(), *TargetPlatform->PlatformName(), *Object->GetFullName());
+								return ESavePackageResult::ValidatorError;
+							}
+							else
+							{
+								UE_ASSET_LOG(LogCook, Display, ValidationContext.Package, TEXT("Skipping package %s for platform %s. It imports type '%s', which is in a module that is not available on the platform."),
+									*ValidationContext.Package->GetName(), *TargetPlatform->PlatformName(), *Object->GetFullName());
+								return ESavePackageResult::ValidatorSuppress;
+							}
 						}
 					}
 				}
