[AssetThumbnailWidget] Adjusting AssetThumbnail default behavior to match old one

diff --git a/Engine/Source/Editor/ContentBrowser/Private/AssetViewWidgets.cpp b/Engine/Source/Editor/ContentBrowser/Private/AssetViewWidgets.cpp
--- a/Engine/Source/Editor/ContentBrowser/Private/AssetViewWidgets.cpp
+++ b/Engine/Source/Editor/ContentBrowser/Private/AssetViewWidgets.cpp
@@ -1987,6 +1987,7 @@ void SAssetListItem::Construct( const FArguments& InArgs )
 			ThumbnailConfig.AssetSystemInfoProvider = MakeShared<FAssetSystemContentBrowserInfoProvider>(AssetItem);
 			ThumbnailConfig.AllowAssetStatusThumbnailOverlay = MakeAttributeSP(this, &SAssetListItem::ShouldAllowAssetStatusOverlay, ThumbnailConfig.bForceGenericThumbnail);
 			ThumbnailConfig.ShowAssetColor = true;
+			ThumbnailConfig.ShowAssetBorder = true;
 			ThumbnailConfig.AlwaysExpandTooltip = MakeAttributeSP(this, &SAssetListItem::ShouldExpandTooltip);
 		}
 
diff --git a/Engine/Source/Editor/PropertyEditor/Private/UserInterface/PropertyEditor/SPropertyEditorAsset.cpp b/Engine/Source/Editor/PropertyEditor/Private/UserInterface/PropertyEditor/SPropertyEditorAsset.cpp
--- a/Engine/Source/Editor/PropertyEditor/Private/UserInterface/PropertyEditor/SPropertyEditorAsset.cpp
+++ b/Engine/Source/Editor/PropertyEditor/Private/UserInterface/PropertyEditor/SPropertyEditorAsset.cpp
@@ -630,6 +630,7 @@ void SPropertyEditorAsset::Construct(const FArguments& InArgs, const TSharedPtr<
 		if (UE::Editor::ContentBrowser::IsNewStyleEnabled())
 		{
 			ThumbnailPadding = 0.f;
+			AssetThumbnailConfig.BorderPadding = FMargin(1.f);
 			AssetThumbnailConfig.AssetBorderImageOverride = TAttribute<const FSlateBrush*>::CreateSP(this, &SPropertyEditorAsset::GetThumbnailBorder);
 		}
 		else
@@ -655,6 +656,12 @@ void SPropertyEditorAsset::Construct(const FArguments& InArgs, const TSharedPtr<
 			}
 		}
 
+		TSharedRef<SWidget> ThumbnailWidget = AssetThumbnail->MakeThumbnailWidget(AssetThumbnailConfig);
+		if (UE::Editor::ContentBrowser::IsNewStyleEnabled())
+		{
+			ThumbnailWidget->SetToolTip(TSharedPtr<IToolTip>());
+		}
+
 		ValueContentBox->AddSlot()
 		.Padding(0.0f,3.0f,5.0f,0.0f)
 		.AutoWidth()
@@ -679,7 +686,7 @@ void SPropertyEditorAsset::Construct(const FArguments& InArgs, const TSharedPtr<
 						.WidthOverride(static_cast<float>(InArgs._ThumbnailSize.X))
 						.HeightOverride(static_cast<float>(InArgs._ThumbnailSize.Y))
 						[
-							AssetThumbnail->MakeThumbnailWidget(AssetThumbnailConfig)
+							ThumbnailWidget
 						]
 					]
 				]
diff --git a/Engine/Source/Editor/UMGEditor/Private/Components/AssetThumbnailWidget.cpp b/Engine/Source/Editor/UMGEditor/Private/Components/AssetThumbnailWidget.cpp
--- a/Engine/Source/Editor/UMGEditor/Private/Components/AssetThumbnailWidget.cpp
+++ b/Engine/Source/Editor/UMGEditor/Private/Components/AssetThumbnailWidget.cpp
@@ -25,7 +25,9 @@ FAssetThumbnailConfig FAssetThumbnailWidgetSettings::ToThumbnailConfig() const
 	});
 	Result.HintColorAndOpacity					= HintColorAndOpacity;
 	Result.AssetTypeColorOverride				= bOverrideAssetTypeColor ? AssetTypeColorOverride : TOptional<FLinearColor>();
-	Result.BorderPadding						= Padding;
+	PRAGMA_DISABLE_DEPRECATION_WARNINGS
+	Result.Padding								= Padding;
+	PRAGMA_ENABLE_DEPRECATION_WARNINGS
 	Result.GenericThumbnailSize					= GenericThumbnailSize;
 	Result.ColorStripOrientation				= static_cast<EThumbnailColorStripOrientation>(ColorStripOrientation);
 	return Result;
diff --git a/Engine/Source/Editor/UMGEditor/Public/Components/AssetThumbnailWidget.h b/Engine/Source/Editor/UMGEditor/Public/Components/AssetThumbnailWidget.h
--- a/Engine/Source/Editor/UMGEditor/Public/Components/AssetThumbnailWidget.h
+++ b/Engine/Source/Editor/UMGEditor/Public/Components/AssetThumbnailWidget.h
@@ -63,10 +63,11 @@ struct UMGEDITOR_API FAssetThumbnailWidgetSettings
 	
 	UPROPERTY(EditAnywhere, BlueprintReadWrite, Category = "Appearance", meta = (EditCondition = "bOverrideAssetTypeColor", EditConditionHides))
 	FLinearColor AssetTypeColorOverride = FLinearColor(0.0f, 0.0f, 0.0f, 0.0f);
-	
-	UPROPERTY(EditAnywhere, BlueprintReadWrite, Category = "Appearance")
+
+	UE_DEPRECATED(5.6, "Border padding is handled by the thumbnail internally, this will no longer have any effect in future release.")
+	UPROPERTY(EditAnywhere, BlueprintReadWrite, Category = "Appearance", meta = (ToolTip = "Warning: Border padding is handled by the thumbnail internally, this will no longer have any effect in future release"))
 	FMargin Padding;
-	
+
 	UPROPERTY(EditAnywhere, BlueprintReadWrite, Category = "Appearance")
 	int32 GenericThumbnailSize = 64;
 	
@@ -74,6 +75,14 @@ struct UMGEDITOR_API FAssetThumbnailWidgetSettings
 	EThumbnailColorStripOrientation_BlueprintType ColorStripOrientation = EThumbnailColorStripOrientation_BlueprintType::HorizontalBottomEdge;
 
 	FAssetThumbnailConfig ToThumbnailConfig() const;
+
+	PRAGMA_DISABLE_DEPRECATION_WARNINGS
+	FAssetThumbnailWidgetSettings() = default;
+	FAssetThumbnailWidgetSettings(const FAssetThumbnailWidgetSettings&) = default;
+	FAssetThumbnailWidgetSettings& operator=(const FAssetThumbnailWidgetSettings&) = default;
+	FAssetThumbnailWidgetSettings(FAssetThumbnailWidgetSettings&&) = default;
+	FAssetThumbnailWidgetSettings& operator=(FAssetThumbnailWidgetSettings&&) = default;
+	PRAGMA_ENABLE_DEPRECATION_WARNINGS
 };
 
 /** This widget can be given an asset and it will render its thumbnail. Editor-only. */
diff --git a/Engine/Source/Editor/UnrealEd/Private/AssetThumbnail.cpp b/Engine/Source/Editor/UnrealEd/Private/AssetThumbnail.cpp
--- a/Engine/Source/Editor/UnrealEd/Private/AssetThumbnail.cpp
+++ b/Engine/Source/Editor/UnrealEd/Private/AssetThumbnail.cpp
@@ -54,6 +54,7 @@
 #include "Widgets/Layout/SSeparator.h"
 #include "Widgets/Text/SRichTextBlock.h"
 #include "ContentBrowserUtils.h"
+#include "Widgets/Colors/SColorBlock.h"
 
 namespace AssetThumbnailPool
 {
@@ -98,6 +99,7 @@ public:
 		, _AdditionalTooltipInSmallView(SNullWidget::NullWidget)
 		, _ShowAssetColor(false)
 		, _AssetBorderImageOverride()
+		, _ShowAssetBorder(false)
 		, _AlwaysExpandTooltip(false)
 		, _ColorStripOrientation(EThumbnailColorStripOrientation::HorizontalBottomEdge)
 		{}
@@ -124,6 +126,7 @@ public:
 		SLATE_ATTRIBUTE(TSharedPtr<SWidget>, AdditionalTooltipInSmallView)
 		SLATE_ATTRIBUTE(bool, ShowAssetColor)
 		SLATE_ATTRIBUTE(const FSlateBrush*, AssetBorderImageOverride)
+		SLATE_ARGUMENT(bool, ShowAssetBorder)
 		SLATE_ARGUMENT(bool, CanDisplayEditModePrimitiveTools)
 		SLATE_ATTRIBUTE(EVisibility, IsEditModeVisible)
 		SLATE_ATTRIBUTE(bool, AlwaysExpandTooltip)
@@ -157,8 +160,7 @@ public:
 		EditModeVisibility = InArgs._IsEditModeVisible;
 		AlwaysExpandTooltip = InArgs._AlwaysExpandTooltip;
 		AdditionalTooltipInSmallView = InArgs._AdditionalTooltipInSmallView;
-		// If the old padding was used, use that instead of the new one to avoid positioning issue
-		BorderPadding = InArgs._Padding == FMargin(0.f) ? InArgs._BorderPadding : InArgs._Padding;
+		BorderPadding = InArgs._BorderPadding;
 		AssetThumbnail->OnAssetDataChanged().AddSP(this, &SAssetThumbnail::OnAssetDataChanged);
 		const FAssetData& AssetData = AssetThumbnail->GetAssetData();
 
@@ -329,17 +331,14 @@ public:
 
 		auto AddAssetColor = [&]()
 		{
-			TAttribute<FMargin> ColorStripMargin = MakeAttributeSP(this, &SAssetThumbnail::GetAssetColorPadding);
-
 			// The asset color strip
 			OverlayWidget->AddSlot()
 			.HAlign(ColorStripOrientation == EThumbnailColorStripOrientation::HorizontalBottomEdge ? HAlign_Fill : HAlign_Right)
 			.VAlign(ColorStripOrientation == EThumbnailColorStripOrientation::HorizontalBottomEdge ? VAlign_Bottom : VAlign_Fill)
 			[
-				SAssignNew(AssetColorStripWidget, SBorder)
-				.BorderImage(FAppStyle::GetBrush("WhiteBrush"))
-				.BorderBackgroundColor(AssetColor)
-				.Padding(ColorStripMargin)
+				SAssignNew(AssetColorStripBlockWidget, SColorBlock)
+				.Color_Lambda([this] () { return AssetColor; })
+				.Size(2.f)
 				.Visibility(TAttribute<EVisibility>::CreateSP(this, &SAssetThumbnail::GetAssetColorVisibility))
 			];
 		};
@@ -349,37 +348,44 @@ public:
 			if (ShowAssetColor.IsSet())
 			{
 				AddAssetColor();
-				
+			}
+		}
+		else
+		{
+			AddAssetColor();
+		}
+
+		if (UE::Editor::ContentBrowser::IsNewStyleEnabled())
+		{
+			// Default border creation
+			if (InArgs._ShowAssetBorder || AssetBorderImageOverride.IsSet())
+			{
 				// Image overlay to match the Hovered/Selected status, by design only part of the asset color should be highlighted
 				const TSharedRef<SImage> OverlayImage = SNew(SImage)
 					.Image(this, &SAssetThumbnail::GetThumbnailBorderBrush)
 					.Visibility(EVisibility::HitTestInvisible);
 
 				OverlayWidget->AddSlot()
-				.Padding(TAttribute<FMargin>::CreateSP(this, &SAssetThumbnail::GetAssetColorMargin))
+				.Padding(TAttribute<FMargin>::CreateSP(this, &SAssetThumbnail::GetOverlayThumbnailBorderPadding))
 				[
 					OverlayImage
 				];
 
 				// Border used to contain the Thumbnail in the ratio
-				const TSharedRef<SBorder> OverlayBorder = SNew(SBorder)
-					.Padding(TAttribute<FMargin>::CreateSP(this, &SAssetThumbnail::GetOverlayBorderMargin))
+				const TSharedRef<SBorder> ThumbnailBorder = SNew(SBorder)
+					.Padding(TAttribute<FMargin>::CreateSP(this, &SAssetThumbnail::GetThumbnailBorderPadding))
 					.BorderImage(this, &SAssetThumbnail::GetThumbnailBorderBrush);
 
 				if (AssetBorderImageOverride.IsSet())
 				{
 					OverlayImage->SetImage(AssetBorderImageOverride);
-					OverlayBorder->SetBorderImage(AssetBorderImageOverride);
+					ThumbnailBorder->SetBorderImage(AssetBorderImageOverride);
 				}
 
-				OverlayBorder->SetContent(OverlayWidget);
-				ContentWidget = OverlayBorder;
+				ThumbnailBorder->SetContent(OverlayWidget);
+				ContentWidget = ThumbnailBorder;
 			}
 		}
-		else
-		{
-			AddAssetColor();
-		}
 
 		if (UE::Editor::ContentBrowser::IsNewStyleEnabled())
 		{
@@ -1037,10 +1043,11 @@ private:
 
 	EVisibility GetAssetColorVisibility() const
 	{
-		return ShowAssetColor.Get() ? EVisibility::Visible : EVisibility::Collapsed;
+		// if the new style is not enabled always show it
+		return !UE::Editor::ContentBrowser::IsNewStyleEnabled() || ShowAssetColor.Get() ? EVisibility::Visible : EVisibility::Collapsed;
 	}
 
-	FMargin GetAssetColorMargin() const
+	FMargin GetOverlayThumbnailBorderPadding() const
 	{
 		constexpr float DefaultPadding = -2.f;
 		// Need this to be negative for the Overlay so that it overlaps with the actual border which will keep the ratio correct and won't let the thumbnail over the limits.
@@ -1052,7 +1059,7 @@ private:
 		return FMargin(DefaultPadding);
 	}
 
-	FMargin GetOverlayBorderMargin() const
+	FMargin GetThumbnailBorderPadding() const
 	{
 		constexpr float DefaultPadding = 2.f;
 		// Asset strip margin
@@ -1163,7 +1170,7 @@ private:
 	float GetAssetThumbnailBorderPadding() const
 	{
 		// Already take into account Left and Right
-		constexpr float DefaultPadding = 2.f;
+		constexpr float DefaultPadding = 0.f;
 
 		// Asset strip margin
 		if (BorderPadding.IsSet())
@@ -1283,11 +1290,6 @@ private:
 			AssetColor = AssetTypeActions->GetTypeColor();
 		}
 
-		if (AssetColorStripWidget.IsValid())
-		{
-			AssetColorStripWidget->SetBorderBackgroundColor(AssetColor);	
-		}
-
 		UpdateThumbnailVisibilities();
 	}
 
@@ -1342,20 +1344,6 @@ private:
 		return 2.0f;
 	}
 
-	FMargin GetAssetColorPadding() const
-	{
-		if (ColorStripOrientation == EThumbnailColorStripOrientation::HorizontalBottomEdge)
-		{
-			const float Height = GetAssetColorThickness();
-			return FMargin(0, Height, 0, 0);
-		}
-		else
-		{
-			const float Width = GetAssetColorThickness();
-			return FMargin(Width, 0, 0, 0);
-		}
-	}
-
 	const FSlateBrush* GetClassThumbnailBrush() const
 	{
 		if (ClassThumbnailBrushOverride.IsNone())
@@ -1644,7 +1632,7 @@ private:
 	TSharedPtr<SBorder> ClassIconWidget;
 	TSharedPtr<SBorder> RenderedThumbnailWidget;
 	TSharedPtr<SBorder> AssetBackgroundWidget;
-	TSharedPtr<SBorder> AssetColorStripWidget;
+	TSharedPtr<SColorBlock> AssetColorStripBlockWidget;
 	TSharedPtr<FAssetThumbnail> AssetThumbnail;
 	FCurveSequence ViewportFadeAnimation;
 	FCurveHandle ViewportFadeCurve;
@@ -1901,6 +1889,7 @@ TSharedRef<SWidget> FAssetThumbnail::MakeThumbnailWidget( const FAssetThumbnailC
 			.CanDisplayEditModePrimitiveTools(InConfig.bCanDisplayEditModePrimitiveTools)
 			.IsEditModeVisible(InConfig.IsEditModeVisible)
 			.AssetBorderImageOverride(InConfig.AssetBorderImageOverride)
+			.ShowAssetBorder(InConfig.ShowAssetBorder)
 			.AlwaysExpandTooltip(InConfig.AlwaysExpandTooltip)
 			.ColorStripOrientation(InConfig.ColorStripOrientation);
 	}
diff --git a/Engine/Source/Editor/UnrealEd/Private/DragAndDrop/AssetDragDropOp.cpp b/Engine/Source/Editor/UnrealEd/Private/DragAndDrop/AssetDragDropOp.cpp
--- a/Engine/Source/Editor/UnrealEd/Private/DragAndDrop/AssetDragDropOp.cpp
+++ b/Engine/Source/Editor/UnrealEd/Private/DragAndDrop/AssetDragDropOp.cpp
@@ -90,6 +90,7 @@ TSharedPtr<SWidget> FAssetDragDropOp::GetDefaultDecorator() const
 		{
 			FAssetThumbnailConfig AssetThumbnailConfig;
 			AssetThumbnailConfig.ShowAssetColor = true;
+			AssetThumbnailConfig.BorderPadding = FMargin(1.f);
 			AssetThumbnailConfig.AssetBorderImageOverride = FAppStyle::GetBrush("ContentBrowser.ThumbnailDragDropBackground");
 			ThumbnailWidget = AssetThumbnail->MakeThumbnailWidget(AssetThumbnailConfig);
 		}
diff --git a/Engine/Source/Editor/UnrealEd/Public/AssetThumbnail.h b/Engine/Source/Editor/UnrealEd/Public/AssetThumbnail.h
--- a/Engine/Source/Editor/UnrealEd/Public/AssetThumbnail.h
+++ b/Engine/Source/Editor/UnrealEd/Public/AssetThumbnail.h
@@ -79,6 +79,7 @@ struct FAssetThumbnailConfig
 		, AllowAssetStatusThumbnailOverlay( false )
 		, ShowAssetColor(true)
 		, AssetBorderImageOverride()
+		, ShowAssetBorder(false)
 		, bCanDisplayEditModePrimitiveTools(false)
 		, IsEditModeVisible()
 		, AlwaysExpandTooltip(false)
@@ -120,8 +121,10 @@ struct FAssetThumbnailConfig
 	TAttribute< bool > AllowAssetStatusThumbnailOverlay;
 	/** Whether to show the asset color or not */
 	TAttribute< bool > ShowAssetColor;
-	/** Border override brush needs to have an OutlineWidth of 2 if normal size or 1 if small size  */
+	/** Border override brush, will be used if set, otherwise if ShowAssetBorder is true the default border will be used, do not assign both if no border is wanted */
 	TAttribute<const FSlateBrush*> AssetBorderImageOverride;
+	/** If true will create a default border for the Thumbnail, overridden by AssetBorderImageOverride, if this is false and no override is given no border will be created */
+	bool ShowAssetBorder;
 	bool bCanDisplayEditModePrimitiveTools;
 	TAttribute<EVisibility> IsEditModeVisible;
 	TAttribute<bool> AlwaysExpandTooltip;
