Cooker: Downgrade spurious warning in cook that occurs in some strange garbage collection edge cases. We were tracking these edge cases for an investigation, but we now think the cooker handles them correctly and they do not need to cause a warning.
#rnx

diff --git a/Engine/Source/Editor/UnrealEd/Private/Cooker/CookPackageData.cpp b/Engine/Source/Editor/UnrealEd/Private/Cooker/CookPackageData.cpp
--- a/Engine/Source/Editor/UnrealEd/Private/Cooker/CookPackageData.cpp
+++ b/Engine/Source/Editor/UnrealEd/Private/Cooker/CookPackageData.cpp
@@ -1962,7 +1962,7 @@ bool FPendingCookedPlatformData::PollIsComplete()
 	FName CurrentPackageName = LocalObject->GetPackage()->GetFName();
 	if (CurrentPackageName != PackageData.GetPackageName())
 	{
-		UE_LOG(LogCook, Warning,
+		UE_LOG(LogCook, Display,
 			TEXT("We were waiting for IsCachedCookedPlatformData to return true for %s in package %s, but that object has been moved out of the package. We will stop waiting on it."),
 			*LocalObject->GetFullName(), *PackageData.GetPackageName().ToString());
 		Release();
@@ -1970,7 +1970,7 @@ bool FPendingCookedPlatformData::PollIsComplete()
 	}
 	if (LocalObject->HasAnyFlags(EObjectFlags::RF_MirroredGarbage))
 	{
-		UE_LOG(LogCook, Warning,
+		UE_LOG(LogCook, Display,
 			TEXT("We were waiting for IsCachedCookedPlatformData to return true for %s, but that object is now marked for garbage. We will stop waiting on it."),
 			*LocalObject->GetFullName());
 		Release();
