Vanilla engine mode value
- Replace how we decide if a plugin is part of the vanilla build by checking if its part of the TargetReceipt instead of if it has a marketplace URL, as 3rd party plugins that are hosted outside of the UE marketplace don't have a marketplace URL.

diff --git a/Engine/Source/Editor/UnrealEd/Private/EditorEngine.cpp b/Engine/Source/Editor/UnrealEd/Private/EditorEngine.cpp
--- a/Engine/Source/Editor/UnrealEd/Private/EditorEngine.cpp
+++ b/Engine/Source/Editor/UnrealEd/Private/EditorEngine.cpp
@@ -277,6 +277,7 @@
 
 #include "Engine/EngineCustomTimeStep.h"
 #include "Engine/TimecodeProvider.h"
+#include "TargetReceipt.h"
 
 DEFINE_LOG_CATEGORY_STATIC(LogEditor, Log, All);
 
@@ -5887,19 +5888,47 @@ void UEditorEngine::UpdateIsVanillaProduct()
 			DesktopPlatform->EnumerateLauncherEngineInstallations(Installations);
 			if (Installations.Contains(Identifier))
 			{
-				// Check if we have any marketplace plugins enabled
-				bool bHasMarketplacePlugin = false;
-				for (const TSharedRef<IPlugin>& Plugin : IPluginManager::Get().GetEnabledPlugins())
+				bool bHasThirdPartyPlugin = false;
+
+				FString TargetReceiptFilePath = FTargetReceipt::GetDefaultPath(*FPaths::EngineDir(), TEXT("UnrealEditor"),
+													FPlatformProcess::GetBinariesSubdirectory(),
+													FApp::GetBuildConfiguration(),
+													FPlatformProcess::GetArchitectureSuffix());
+				
+				FTargetReceipt Receipt;
+				if (Receipt.Read(*TargetReceiptFilePath))
 				{
-					if (Plugin->GetDescriptor().MarketplaceURL.Len() > 0)
-					{
-						bHasMarketplacePlugin = true;
-						break;
+					// Check if we have any plugins that weren't part of the build
+					for (const TSharedRef<IPlugin>& Plugin : IPluginManager::Get().GetEnabledPlugins())
+					{	
+						if (Plugin->GetDescriptor().SupportsTargetPlatform(FPlatformMisc::GetUBTPlatform()) && Plugin->GetDescriptor().Modules.Num() > 0)
+						{
+							if (!Receipt.BuildPlugins.Contains(Plugin->GetName()))
+							{
+								// Check if we have build products for that plugin. Mainly for plugins that aren't enabled by default but would have had their binaries built.
+								if (Algo::FindByPredicate(Receipt.BuildProducts, 
+									[&Plugin](const FBuildProduct& BuildProduct)
+									{
+										// Look for the plugin .modules file
+										return BuildProduct.Type == TEXT("RequiredResource") && FPaths::IsUnderDirectory(BuildProduct.Path, Plugin->GetBaseDir());
+									}) == nullptr)
+								{
+								
+									bHasThirdPartyPlugin = true;
+									break;
+								}
+							}
+						}
 					}
 				}
+				else
+				{
+					// Can't validate the set of plugins, assume that this isn't a vanilla build
+					bHasThirdPartyPlugin = true;
+				}
 
 				// If not, we're running Epic-only code.
-				if (!bHasMarketplacePlugin)
+				if (!bHasThirdPartyPlugin)
 				{
 					bResult = true;
 				}
