Preview Platform Blueprint Functions
TogglePreviewPlatform - Toggle the Preview Platform On/Off
SetPreviewPlatform -  Sets the Preview Platform. The Preview Platform is set froma  list of valid Preview Platforms, the same present in the Preview Platform Menu.

diff --git a/Engine/Source/Editor/UnrealEd/Private/Subsystems/EditorSubsystemBlueprintLibrary.cpp b/Engine/Source/Editor/UnrealEd/Private/Subsystems/EditorSubsystemBlueprintLibrary.cpp
--- a/Engine/Source/Editor/UnrealEd/Private/Subsystems/EditorSubsystemBlueprintLibrary.cpp
+++ b/Engine/Source/Editor/UnrealEd/Private/Subsystems/EditorSubsystemBlueprintLibrary.cpp
@@ -1,10 +1,94 @@
 // Copyright Epic Games, Inc. All Rights Reserved.
 
 #include "Subsystems/EditorSubsystemBlueprintLibrary.h"
-
+#include "DataDrivenShaderPlatformInfo.h"
 #include "Editor.h"
 
 /*static*/ UEditorSubsystem* UEditorSubsystemBlueprintLibrary::GetEditorSubsystem(TSubclassOf<UEditorSubsystem> Class)
 {
 	return GEditor->GetEditorSubsystemBase(Class);
 }
+
+FName GetPreviewPlatformDropdownName(const FPreviewPlatformMenuItem& Item)
+{
+	// We cannot use the localized name, because it will get cached and not work for other Languages.
+	// Instead we take the PreviewShaderPlatformName and remove the postfix and suffix if they are present
+	// and concatenate it with the DeviceProfileName if present.
+
+	FString DropdownName = Item.PreviewShaderPlatformName.ToString();
+	DropdownName.RemoveFromEnd(TEXT("_Preview"));
+	DropdownName.RemoveFromStart(TEXT("SP_"));
+
+	if (Item.DeviceProfileName != NAME_None)
+	{
+		DropdownName += TEXT(" ") + Item.DeviceProfileName.ToString();
+	}
+
+	return FName(DropdownName);
+}
+
+TArray<FName> UEditorSubsystemBlueprintLibrary::GetPreviewPlatformOptions() const
+{
+	// We gather all the available PreviewPlatform options from the PreviewPlatformMenu.
+
+	TArray<FName> Outputs;
+	for (const FPreviewPlatformMenuItem& Item : FDataDrivenPlatformInfoRegistry::GetAllPreviewPlatformMenuItems())
+	{
+		FName PreviewNameAndDeviceProfileName = GetPreviewPlatformDropdownName(Item);
+		Outputs.Add(PreviewNameAndDeviceProfileName);
+	}
+
+	return Outputs;
+}
+
+void UEditorSubsystemBlueprintLibrary::TogglePreviewPlatform()
+{
+	GEditor->ToggleFeatureLevelPreview();
+}
+
+void UEditorSubsystemBlueprintLibrary::SetPreviewPlatform(FName PreviewShaderPlatformName)
+{
+	if (PreviewShaderPlatformName == NAME_None)
+	{
+		return;
+	}
+
+	bool bFoundMatch = false;
+	for (const FPreviewPlatformMenuItem& Item : FDataDrivenPlatformInfoRegistry::GetAllPreviewPlatformMenuItems())
+	{
+		// We check to see if we can find the matching PreviewShaderPlatform DeviceProfile combination in the Preview Platform Menu.
+
+		FName NameToCheck = GetPreviewPlatformDropdownName(Item);
+		if (NameToCheck == PreviewShaderPlatformName)
+		{
+			bFoundMatch = true;
+			EShaderPlatform ShaderPlatform = FDataDrivenShaderPlatformInfo::GetShaderPlatformFromName(Item.PreviewShaderPlatformName);
+
+			if (ShaderPlatform < SP_NumPlatforms)
+			{
+				const bool bIsDefaultShaderPlatform = FDataDrivenShaderPlatformInfo::GetPreviewShaderPlatformParent(ShaderPlatform) == GMaxRHIShaderPlatform;
+
+				auto GetPreviewFeatureLevelInfo = [&]()
+					{
+						if (bIsDefaultShaderPlatform)
+						{
+							return FPreviewPlatformInfo(GMaxRHIFeatureLevel, GMaxRHIShaderPlatform, NAME_None, NAME_None, NAME_None, false, NAME_None);
+						}
+
+						const ERHIFeatureLevel::Type FeatureLevel = GetMaxSupportedFeatureLevel(ShaderPlatform);
+						return FPreviewPlatformInfo(FeatureLevel, ShaderPlatform, Item.PlatformName, Item.ShaderFormat, Item.DeviceProfileName, true, Item.PreviewShaderPlatformName, Item.OptionalFriendlyNameOverride);
+					};
+
+				FPreviewPlatformInfo PreviewFeatureLevelInfo = GetPreviewFeatureLevelInfo();
+				GEditor->SetPreviewPlatform(PreviewFeatureLevelInfo, false);
+				break;
+			}
+		}
+	}
+
+	// If we cannot find a match we emit a warning.
+	if (!bFoundMatch)
+	{
+		UE_LOG(LogBlueprint, Warning, TEXT("SetPreviewPlatform Blueprint Function could not find the PreviewShaderPlatform DeviceProfile combination: %s"), *PreviewShaderPlatformName.ToString());
+	}
+}
\ No newline at end of file
diff --git a/Engine/Source/Editor/UnrealEd/Public/Subsystems/EditorSubsystemBlueprintLibrary.h b/Engine/Source/Editor/UnrealEd/Public/Subsystems/EditorSubsystemBlueprintLibrary.h
--- a/Engine/Source/Editor/UnrealEd/Public/Subsystems/EditorSubsystemBlueprintLibrary.h
+++ b/Engine/Source/Editor/UnrealEd/Public/Subsystems/EditorSubsystemBlueprintLibrary.h
@@ -19,4 +19,19 @@ public:
 	UFUNCTION(BlueprintPure, Category = "Editor Subsystems", meta = (BlueprintInternalUseOnly = "true"))
 	static UNREALED_API UEditorSubsystem* GetEditorSubsystem(TSubclassOf<UEditorSubsystem> Class);
 
-};
+	UFUNCTION()
+	TArray<FName> GetPreviewPlatformOptions() const;
+	
+	/**
+	 * Toggle the Preview Platform.
+	 *
+	 */
+	UFUNCTION(BlueprintCallable, Category = "Editor")
+	static UNREALED_API void TogglePreviewPlatform();
+
+	/**
+	 * Sets the Preview Platform
+	 */
+	UFUNCTION(BlueprintCallable, Category = "Editor")
+	static UNREALED_API void SetPreviewPlatform(UPARAM(meta = (GetOptions = "GetPreviewPlatformOptions")) FName PreviewShaderPlatformName);
+};
\ No newline at end of file
