Cooker: Provide a single CookerInterface back pointer to the CookedPackageWriters, rather than a separate function pointer for each operation. This supports an upcoming change to add another callback.
#rnx

diff --git a/Engine/Source/Developer/IoStoreUtilities/Private/ZenStoreWriter.cpp b/Engine/Source/Developer/IoStoreUtilities/Private/ZenStoreWriter.cpp
--- a/Engine/Source/Developer/IoStoreUtilities/Private/ZenStoreWriter.cpp
+++ b/Engine/Source/Developer/IoStoreUtilities/Private/ZenStoreWriter.cpp
@@ -455,9 +455,9 @@ void FZenStoreWriter::WritePackageTrailer(const FPackageTrailerInfo& Info, const
 void FZenStoreWriter::RegisterDeterminismHelper(UObject* SourceObject,
 	const TRefCountPtr<UE::Cook::IDeterminismHelper>& DeterminismHelper)
 {
-	if (RegisterDeterminismHelperCallback)
+	if (Cooker)
 	{
-		RegisterDeterminismHelperCallback(SourceObject, DeterminismHelper);
+		Cooker->RegisterDeterminismHelper(this, SourceObject, DeterminismHelper);
 	}
 }
 
@@ -1516,7 +1516,8 @@ void FZenStoreWriter::CompleteExportsArchiveForDiff(FPackageInfo& Info, FLargeMe
 EPackageWriterResult FZenStoreWriter::BeginCacheForCookedPlatformData(
 	FBeginCacheForCookedPlatformDataInfo& Info)
 {
-	return BeginCacheCallback(Info);
+	check(Cooker); // Fallback for non-cooker case not yet implemented
+	return Cooker->CookerBeginCacheForCookedPlatformData(Info);
 }
 
 TFuture<FCbObject> FZenStoreWriter::WriteMPCookMessageForPackage(FName PackageName)
diff --git a/Engine/Source/Developer/IoStoreUtilities/Public/ZenStoreWriter.h b/Engine/Source/Developer/IoStoreUtilities/Public/ZenStoreWriter.h
--- a/Engine/Source/Developer/IoStoreUtilities/Public/ZenStoreWriter.h
+++ b/Engine/Source/Developer/IoStoreUtilities/Public/ZenStoreWriter.h
@@ -71,14 +71,23 @@ public:
 		return this;
 	};
 
-	// Delegates to forward PackageWriter events onto UCookOnTheFlyServer when cooking
-	void SetBeginCacheCallback(FBeginCacheCallback&& InBeginCacheCallback)
+	virtual void SetCooker(UE::PackageWriter::Private::ICookerInterface* CookerInterface) override
 	{
-		BeginCacheCallback = MoveTemp(InBeginCacheCallback);
+		Cooker = CookerInterface;
 	}
-	void SetRegisterDeterminismHelperCallback(FRegisterDeterminismHelperCallback&& InRegisterDeterminismHelperCallback)
+
+	typedef TUniqueFunction<EPackageWriterResult(FBeginCacheForCookedPlatformDataInfo& Info)> 
+		FDeprecatedBeginCacheCallback;
+	typedef TUniqueFunction<void(UObject* SourceObject,
+		const TRefCountPtr<UE::Cook::IDeterminismHelper>& DeterminismHelper)>
+		FDeprecatedRegisterDeterminismHelperCallback;
+	UE_DEPRECATED(5.7, "Use SetCooker instead")
+	void SetBeginCacheCallback(FDeprecatedBeginCacheCallback&&)
+	{
+	}
+	UE_DEPRECATED(5.7, "Use SetCooker instead")
+	void SetRegisterDeterminismHelperCallback(FDeprecatedRegisterDeterminismHelperCallback&&)
 	{
-		RegisterDeterminismHelperCallback = MoveTemp(InRegisterDeterminismHelperCallback);
 	}
 
 	struct ZenHostInfo
@@ -92,7 +101,7 @@ public:
 	virtual FCapabilities GetCapabilities() const override
 	{
 		FCapabilities Result;
-		Result.bDeterminismDebug = (bool)RegisterDeterminismHelperCallback;
+		Result.bDeterminismDebug = Cooker ? Cooker->IsDeterminismDebug() : false;
 		return Result;
 	}
 
@@ -266,8 +275,7 @@ private:
 	TUniquePtr<FZenFileSystemManifest>	ZenFileSystemManifest;
 	TMap<FName, TArray<FString>>		PackageAdditionalFiles;
 
-	FBeginCacheCallback					BeginCacheCallback;
-	FRegisterDeterminismHelperCallback RegisterDeterminismHelperCallback;
+	UE::PackageWriter::Private::ICookerInterface* Cooker = nullptr;
 
 	FEntryCreatedEvent					EntryCreatedEvent;
 	FCriticalSection					CommitEventCriticalSection;
diff --git a/Engine/Source/Editor/UnrealEd/Classes/CookOnTheSide/CookOnTheFlyServer.h b/Engine/Source/Editor/UnrealEd/Classes/CookOnTheSide/CookOnTheFlyServer.h
--- a/Engine/Source/Editor/UnrealEd/Classes/CookOnTheSide/CookOnTheFlyServer.h
+++ b/Engine/Source/Editor/UnrealEd/Classes/CookOnTheSide/CookOnTheFlyServer.h
@@ -12,6 +12,7 @@
 #include "Misc/EnumClassFlags.h"
 #include "Misc/Optional.h"
 #include "Misc/PackageAccessTracking.h"
+#include "Serialization/PackageWriter.h"
 #include "Templates/Function.h"
 #include "Templates/RefCounting.h"
 #include "Templates/SharedPointer.h"
@@ -251,7 +252,11 @@ typedef TMap<FName, TMap<FName, TMap<FName, TArray<FString>>>> FIniSettingContai
 
 }
 UCLASS(MinimalAPI)
-class UCookOnTheFlyServer : public UObject, public FTickableEditorObject, public FExec, public UE::Cook::ICookInfo
+class UCookOnTheFlyServer : public UObject,
+	public FTickableEditorObject,
+	public FExec,
+	public UE::Cook::ICookInfo,
+	public UE::PackageWriter::Private::ICookerInterface
 {
 	GENERATED_BODY()
 
@@ -694,7 +699,15 @@ public:
 		UE::Cook::EProcessType ProcessType = UE::Cook::EProcessType::AllMPCook) override;
 	UNREALED_API virtual void UnregisterCollector(UE::Cook::IMPCollector* Collector) override;
 	UNREALED_API virtual void GetCulturesToCook(TArray<FString>& OutCulturesToCook) const override;
-
+	// End ICookInfo interface
+
+	// UE::PackageWriter::Private::ICookerInterface
+	UNREALED_API virtual EPackageWriterResult CookerBeginCacheForCookedPlatformData(
+		UE::PackageWriter::Private::FBeginCacheForCookedPlatformDataInfo& Info) override;
+	UNREALED_API virtual void RegisterDeterminismHelper(ICookedPackageWriter* PackageWriter, UObject* SourceObject,
+		const TRefCountPtr<UE::Cook::IDeterminismHelper>& DeterminismHelper) override;
+	UNREALED_API virtual bool IsDeterminismDebug() const override;
+	// End UE::PackageWriter::Private::ICookerInterface
 
 	/** Dumps cooking stats to the log. Run from the exec command "Cook stats". */
 	UNREALED_API void DumpStats();
diff --git a/Engine/Source/Editor/UnrealEd/Private/CookOnTheFlyServer.cpp b/Engine/Source/Editor/UnrealEd/Private/CookOnTheFlyServer.cpp
--- a/Engine/Source/Editor/UnrealEd/Private/CookOnTheFlyServer.cpp
+++ b/Engine/Source/Editor/UnrealEd/Private/CookOnTheFlyServer.cpp
@@ -6767,6 +6767,41 @@ void UCookOnTheFlyServer::GetCulturesToCook(TArray<FString>& OutCulturesToCook)
 	}
 }
 
+EPackageWriterResult UCookOnTheFlyServer::CookerBeginCacheForCookedPlatformData(
+	UE::PackageWriter::Private::FBeginCacheForCookedPlatformDataInfo& Info)
+{
+	return SavePackageBeginCacheForCookedPlatformData(Info.PackageName,
+		Info.TargetPlatform, Info.SaveableObjects, Info.SaveFlags);
+}
+
+void UCookOnTheFlyServer::RegisterDeterminismHelper(ICookedPackageWriter* PackageWriter,
+	UObject* SourceObject, const TRefCountPtr<UE::Cook::IDeterminismHelper>& DeterminismHelper)
+{
+	if (!IsDeterminismDebug())
+	{
+		return;
+	}
+	UE::Cook::FCookSavePackageContext* PackageWriterContext = nullptr;
+	for (UE::Cook::FCookSavePackageContext* Context : SavePackageContexts)
+	{
+		if (Context->PackageWriter == PackageWriter)
+		{
+			PackageWriterContext = Context;
+			break;
+		}
+	}
+	if (!PackageWriterContext || !PackageWriterContext->DeterminismManager)
+	{
+		return;
+	}
+	PackageWriterContext->DeterminismManager->RegisterDeterminismHelper(SourceObject, DeterminismHelper);
+}
+
+bool UCookOnTheFlyServer::IsDeterminismDebug() const
+{
+	return bDeterminismDebug || DiffModeHelper->IsDeterminismDebug();
+}
+
 void UCookOnTheFlyServer::DumpStats()
 {
 	UE_LOG(LogCook, Display, TEXT("IntStats:"));
@@ -10699,23 +10734,10 @@ UE::Cook::FCookSavePackageContext* UCookOnTheFlyServer::CreateSaveContext(const
 	TSharedPtr<ICookArtifactReader> CookArtifactReader = nullptr;
 	ICookedPackageWriter* PackageWriter = nullptr;
 	FString WriterDebugName;
-	ICookedPackageWriter::FBeginCacheCallback BeginCacheCallback(
-		[this](ICookedPackageWriter::FBeginCacheForCookedPlatformDataInfo& Info)
-		{
-			return this->SavePackageBeginCacheForCookedPlatformData(Info.PackageName,
-				Info.TargetPlatform, Info.SaveableObjects, Info.SaveFlags);
-		});
 	DiffModeHelper->Initialize();
-	ICookedPackageWriter::FRegisterDeterminismHelperCallback RegisterDeterminismHelperCallback;
 	if (bDeterminismDebug || DiffModeHelper->IsDeterminismDebug())
 	{
 		DeterminismManager = MakeUnique<FDeterminismManager>();
-		RegisterDeterminismHelperCallback =
-			[this, DeterminismManagerPtr = DeterminismManager.Get()]
-			(UObject* SourceObject, const TRefCountPtr<UE::Cook::IDeterminismHelper>& DeterminismHelper)
-			{
-				DeterminismManagerPtr->RegisterDeterminismHelper(SourceObject, DeterminismHelper);
-			};
 	}
 	if (IsUsingZenStore())
 	{
@@ -10726,8 +10748,7 @@ UE::Cook::FCookSavePackageContext* UCookOnTheFlyServer::CreateSaveContext(const
 		AllContextArtifactReader->AddLayer(ZenReader);
 		CookArtifactReader = LayeredReader;
 		FZenStoreWriter* ZenWriter = new FZenStoreWriter(ResolvedRootPath, ResolvedMetadataPath, TargetPlatform, LayeredReader);
-		ZenWriter->SetBeginCacheCallback(MoveTemp(BeginCacheCallback));
-		ZenWriter->SetRegisterDeterminismHelperCallback(MoveTemp(RegisterDeterminismHelperCallback));
+		ZenWriter->SetCooker(this);
 		PackageWriter = ZenWriter;
 		WriterDebugName = TEXT("ZenStore");
 	}
@@ -10735,9 +10756,8 @@ UE::Cook::FCookSavePackageContext* UCookOnTheFlyServer::CreateSaveContext(const
 	{
 		CookArtifactReader = SharedLooseFilesCookArtifactReader;
 		PackageWriter = new FLooseCookedPackageWriter(ResolvedRootPath, ResolvedMetadataPath, TargetPlatform,
-			GetAsyncIODelete(), *SandboxFile, MoveTemp(BeginCacheCallback),
-			MoveTemp(RegisterDeterminismHelperCallback),
-			SharedLooseFilesCookArtifactReader.ToSharedRef());
+			GetAsyncIODelete(), *SandboxFile, SharedLooseFilesCookArtifactReader.ToSharedRef());
+		PackageWriter->SetCooker(this);
 		WriterDebugName = TEXT("LooseCookedPackageWriter");
 	}
 
diff --git a/Engine/Source/Editor/UnrealEd/Private/Cooker/CookFunctionLibrary.cpp b/Engine/Source/Editor/UnrealEd/Private/Cooker/CookFunctionLibrary.cpp
--- a/Engine/Source/Editor/UnrealEd/Private/Cooker/CookFunctionLibrary.cpp
+++ b/Engine/Source/Editor/UnrealEd/Private/Cooker/CookFunctionLibrary.cpp
@@ -143,27 +143,40 @@ void UCookFunctionLibrary::CookAsset(UObject* Object, const FString& ForPlatform
 		TargetPlatform, CookInfo);
 	FArchiveCookData CookData(*TargetPlatform, ArchiveCookContext);
 
+	class FCookerInterface : public UE::PackageWriter::Private::ICookerInterface
+	{
+	public:
+		virtual EPackageWriterResult CookerBeginCacheForCookedPlatformData(
+			UE::PackageWriter::Private::FBeginCacheForCookedPlatformDataInfo& Info) override
+		{
+			return EPackageWriterResult::Success; // saving will fail if we don't say good things happened
+		}
+		virtual void RegisterDeterminismHelper(ICookedPackageWriter* PackageWriter, UObject* SourceObject,
+			const TRefCountPtr<UE::Cook::IDeterminismHelper>& DeterminismHelper) override
+		{
+		}
+		virtual bool IsDeterminismDebug() const override
+		{
+			return false;
+		}
+	};
+	FCookerInterface CookerInterface;
+
 	// The AsyncIODelete object optimizes the cleaning of the saved/cook directory. Happily not relevant
 	// for callers of this function, but an important optimization for the cook commandlet.
 	TUniquePtr<FAsyncIODelete> AsyncIODelete = MakeUnique<FAsyncIODelete>(); 
-	ICookedPackageWriter::FBeginCacheCallback BeginCacheCallback(
-		[](ICookedPackageWriter::FBeginCacheForCookedPlatformDataInfo& Info) -> EPackageWriterResult
-		{
-			return EPackageWriterResult::Success; // saving will fail if we don't say good things happened
-		});
 
 	TArray<TSharedRef<IPlugin> > PluginsToRemap; // none?
 	TUniquePtr<UE::Cook::FCookSandbox> SandboxFileObj = MakeUnique<UE::Cook::FCookSandbox>(OutputName, PluginsToRemap);
-	ICookedPackageWriter::FRegisterDeterminismHelperCallback RegisterDeterminismHelperCallback;
 	FLooseCookedPackageWriter* LooseWriter = new FLooseCookedPackageWriter(
 		OutputName, 
 		OutputName, 
 		TargetPlatform, 
 		*AsyncIODelete, 
 		*SandboxFileObj,
-		MoveTemp(BeginCacheCallback),
-		MoveTemp(RegisterDeterminismHelperCallback),
 		MakeShared<FLooseFilesCookArtifactReader>());
+	LooseWriter->SetCooker(&CookerInterface);
+
 	// ~FSavePackageContext will delete the loose writer, hence naked new above:
 	FSavePackageContext SaveContext(TargetPlatform, LooseWriter); 
 	FSavePackageArgs SaveArgs;
diff --git a/Engine/Source/Editor/UnrealEd/Private/Cooker/DiffPackageWriter.h b/Engine/Source/Editor/UnrealEd/Private/Cooker/DiffPackageWriter.h
--- a/Engine/Source/Editor/UnrealEd/Private/Cooker/DiffPackageWriter.h
+++ b/Engine/Source/Editor/UnrealEd/Private/Cooker/DiffPackageWriter.h
@@ -15,6 +15,11 @@ class FDiffPackageWriter : public ICookedPackageWriter
 public:
 	FDiffPackageWriter(TUniquePtr<ICookedPackageWriter>&& InInner, UE::Cook::FDeterminismManager* InDeterminismManager);
 
+	virtual void SetCooker(UE::PackageWriter::Private::ICookerInterface* CookerInterface) override
+	{
+		// Nothing to do, our Inner has the Cooker interface, we don't use it directly
+	}
+
 	// IPackageWriter
 	virtual FCapabilities GetCapabilities() const override
 	{
diff --git a/Engine/Source/Editor/UnrealEd/Private/Cooker/LooseCookedPackageWriter.cpp b/Engine/Source/Editor/UnrealEd/Private/Cooker/LooseCookedPackageWriter.cpp
--- a/Engine/Source/Editor/UnrealEd/Private/Cooker/LooseCookedPackageWriter.cpp
+++ b/Engine/Source/Editor/UnrealEd/Private/Cooker/LooseCookedPackageWriter.cpp
@@ -27,17 +27,13 @@ LLM_DEFINE_TAG(Cooker_PackageStoreManifest);
 
 FLooseCookedPackageWriter::FLooseCookedPackageWriter(const FString& InOutputPath,
 	const FString& InMetadataDirectoryPath, const ITargetPlatform* InTargetPlatform, FAsyncIODelete& InAsyncIODelete,
-	UE::Cook::FCookSandbox& InSandboxFile, FBeginCacheCallback&& InBeginCacheCallback,
-	FRegisterDeterminismHelperCallback&& InRegisterDeterminismHelperCallback,
-	TSharedRef<FLooseFilesCookArtifactReader> InCookArtifactReader)
+	UE::Cook::FCookSandbox& InSandboxFile, TSharedRef<FLooseFilesCookArtifactReader> InCookArtifactReader)
 	: CookArtifactReader(InCookArtifactReader)
 	, OutputPath(InOutputPath)
 	, MetadataDirectoryPath(InMetadataDirectoryPath)
 	, TargetPlatform(*InTargetPlatform)
 	, SandboxFile(InSandboxFile)
 	, AsyncIODelete(InAsyncIODelete)
-	, BeginCacheCallback(MoveTemp(InBeginCacheCallback))
-	, RegisterDeterminismHelperCallback(MoveTemp(InRegisterDeterminismHelperCallback))
 {
 }
 
@@ -167,9 +163,9 @@ FDateTime FLooseCookedPackageWriter::GetPreviousCookTime() const
 void FLooseCookedPackageWriter::RegisterDeterminismHelper(UObject* SourceObject,
 	const TRefCountPtr<UE::Cook::IDeterminismHelper>& DeterminismHelper)
 {
-	if (RegisterDeterminismHelperCallback)
+	if (Cooker)
 	{
-		RegisterDeterminismHelperCallback(SourceObject, DeterminismHelper);
+		Cooker->RegisterDeterminismHelper(this, SourceObject, DeterminismHelper);
 	}
 }
 
@@ -195,7 +191,8 @@ void FLooseCookedPackageWriter::Initialize(const FCookInfo& Info)
 EPackageWriterResult FLooseCookedPackageWriter::BeginCacheForCookedPlatformData(
 	FBeginCacheForCookedPlatformDataInfo& Info)
 {
-	return BeginCacheCallback(Info);
+	check(Cooker); // Fallback for non-cooker case not yet implemented
+	return Cooker->CookerBeginCacheForCookedPlatformData(Info);
 }
 
 void FLooseCookedPackageWriter::WriteOplogEntry(FCbWriter& Writer, const FOplogPackageInfo& PackageInfo)
diff --git a/Engine/Source/Editor/UnrealEd/Private/Cooker/LooseCookedPackageWriter.h b/Engine/Source/Editor/UnrealEd/Private/Cooker/LooseCookedPackageWriter.h
--- a/Engine/Source/Editor/UnrealEd/Private/Cooker/LooseCookedPackageWriter.h
+++ b/Engine/Source/Editor/UnrealEd/Private/Cooker/LooseCookedPackageWriter.h
@@ -25,14 +25,18 @@ public:
 
 	FLooseCookedPackageWriter(const FString& OutputPath, const FString& MetadataDirectoryPath,
 		const ITargetPlatform* TargetPlatform, FAsyncIODelete& InAsyncIODelete,
-		UE::Cook::FCookSandbox& InSandboxFile, FBeginCacheCallback&& InBeginCacheCallback,
-		FRegisterDeterminismHelperCallback&& InRegisterDeterminismHelperCallback, TSharedRef<FLooseFilesCookArtifactReader> CookArtifactReader);
+		UE::Cook::FCookSandbox& InSandboxFile, TSharedRef<FLooseFilesCookArtifactReader> CookArtifactReader);
 	~FLooseCookedPackageWriter();
 
+	virtual void SetCooker(UE::PackageWriter::Private::ICookerInterface* CookerInterface) override
+	{
+		Cooker = CookerInterface;
+	}
+
 	virtual FCapabilities GetCapabilities() const override
 	{
 		FCapabilities Result = Super::GetCapabilities();
-		Result.bDeterminismDebug = (bool)RegisterDeterminismHelperCallback;
+		Result.bDeterminismDebug = Cooker ? Cooker->IsDeterminismDebug() : false;
 		return Result;
 	}
 
@@ -117,8 +121,7 @@ private:
 	TMap<FName, FOplogPackageInfo> Oplog;
 	UE::Cook::FCookSandbox& SandboxFile;
 	FAsyncIODelete& AsyncIODelete;
-	FBeginCacheCallback BeginCacheCallback;
-	FRegisterDeterminismHelperCallback RegisterDeterminismHelperCallback;
+	UE::PackageWriter::Private::ICookerInterface* Cooker = nullptr;
 	bool bLegacyIterativeSharedBuild = false;
 	bool bProvidePerPackageResults = false;
 };
diff --git a/Engine/Source/Runtime/Core/Public/Serialization/PackageWriter.h b/Engine/Source/Runtime/Core/Public/Serialization/PackageWriter.h
--- a/Engine/Source/Runtime/Core/Public/Serialization/PackageWriter.h
+++ b/Engine/Source/Runtime/Core/Public/Serialization/PackageWriter.h
@@ -250,12 +250,53 @@ struct FPackageHashes : FThreadSafeRefCountedObject
 	TFuture<int> CompletionFuture;
 };
 
+namespace UE::PackageWriter::Private
+{
+
+/**
+ * Information passed from SavePackage into a PackageWriter during cooking for the calls to
+ * BeginCacheForCookedPlatformData on the UObjects in the package.
+ * This type is a cooker implementation detail and might be changed without deprecation in a future version.
+ */
+struct FBeginCacheForCookedPlatformDataInfo
+{
+	FName PackageName;
+	const ITargetPlatform* TargetPlatform;
+	TConstArrayView<UObject*> SaveableObjects;
+	uint32 SaveFlags;
+};
+
+/**
+ * Backpointer interface to the Cooker for functions that are implemented in Cooker code for PackageWriters used
+ * by the cooker rather than CookedPackageWriter code, but that are triggered from CookedPackageWriter code.
+ * This type is a cooker implementation detail and might be changed without deprecation in a future version.
+ */
+class ICookerInterface
+{
+public:
+	virtual EPackageWriterResult CookerBeginCacheForCookedPlatformData(FBeginCacheForCookedPlatformDataInfo& Info) = 0;
+	virtual void RegisterDeterminismHelper(ICookedPackageWriter* PackageWriter, UObject* SourceObject,
+		const TRefCountPtr<UE::Cook::IDeterminismHelper>& DeterminismHelper) = 0;
+	virtual bool IsDeterminismDebug() const = 0;
+
+protected:
+	~ICookerInterface() = default;
+};
+
+} // namespace UE::PackageWriter::Private
+
 /** Interface for cooking that writes cooked packages to storage usable by the runtime game. */
 class ICookedPackageWriter : public IPackageWriter
 {
 public:
 	virtual ~ICookedPackageWriter() = default;
 
+	/**
+	 * Called by the Cooker for PackageWriters it is using. If provided, the interface must be used to handle the
+	 * virtual functions of the same name on the ICookedPackageWriter.
+	 */
+	virtual void SetCooker(UE::PackageWriter::Private::ICookerInterface* CookerInterface) = 0;
+
 	enum class EPackageHeaderFormat
 	{
 		PackageFileSummary,
@@ -411,17 +452,13 @@ public:
 		unimplemented();
 	}
 
-	struct FBeginCacheForCookedPlatformDataInfo
-	{
-		FName PackageName;
-		const ITargetPlatform* TargetPlatform;
-		TConstArrayView<UObject*> SaveableObjects;
-		uint32 SaveFlags;
-	};
+	using FBeginCacheForCookedPlatformDataInfo = UE::PackageWriter::Private::FBeginCacheForCookedPlatformDataInfo;
 	// Helper callback type for Writers that need to send the message on to UCookOnTheFlyServer
-	using FBeginCacheCallback = TUniqueFunction<EPackageWriterResult(FBeginCacheForCookedPlatformDataInfo& Info)>;
-	using FRegisterDeterminismHelperCallback = TUniqueFunction<void(UObject* SourceObject,
-		const TRefCountPtr<UE::Cook::IDeterminismHelper>& DeterminismHelper)>;
+	UE_DEPRECATED(5.7, "Use ICookedPackageWriterCookerInterface.BeginCacheForCookedPlatformData instead.")
+	typedef TUniqueFunction<EPackageWriterResult(FBeginCacheForCookedPlatformDataInfo& Info)> FBeginCacheCallback;
+	UE_DEPRECATED(5.7, "Use ICookedPackageWriterCookerInterface.RegisterDeterminismHelper instead.")
+	typedef TUniqueFunction<void(UObject* SourceObject,
+		const TRefCountPtr<UE::Cook::IDeterminismHelper>& DeterminismHelper)> FRegisterDeterminismHelperCallback;
 
 	virtual EPackageWriterResult BeginCacheForCookedPlatformData(FBeginCacheForCookedPlatformDataInfo& Info)
 	{
