Cooker: PackageWriter: Add FBasePackageWriter and FBaseCookePackageWriter classes in CoreUObject to provide shared implementation for ZenStoreWriter and FLooseCookedPackageWriter, and to remove the accidental link dependency that was added from Core to CoreUObject.
Move some FilePackageWriterUtil structures and functions into BasePackageWriter to allow sharing with FZenStoreWriter.
#rnx

diff --git a/Engine/Source/Developer/IoStoreUtilities/Public/PackageStoreWriter.h b/Engine/Source/Developer/IoStoreUtilities/Public/PackageStoreWriter.h
--- a/Engine/Source/Developer/IoStoreUtilities/Public/PackageStoreWriter.h
+++ b/Engine/Source/Developer/IoStoreUtilities/Public/PackageStoreWriter.h
@@ -2,12 +2,12 @@
 
 #pragma once
 
+#include "Serialization/BasePackageWriter.h"
 #include "Serialization/PackageStore.h"
-#include "Serialization/PackageWriter.h"
 
 class FZenStoreWriter;
 
-class IPackageStoreWriter : public ICookedPackageWriter
+class IPackageStoreWriter : public FBaseCookedPackageWriter
 {
 public:
 	/** Identify as a member of this interface from the ICookedPackageWriter api. */
diff --git a/Engine/Source/Editor/UnrealEd/Private/Cooker/LooseCookedPackageWriter.h b/Engine/Source/Editor/UnrealEd/Private/Cooker/LooseCookedPackageWriter.h
--- a/Engine/Source/Editor/UnrealEd/Private/Cooker/LooseCookedPackageWriter.h
+++ b/Engine/Source/Editor/UnrealEd/Private/Cooker/LooseCookedPackageWriter.h
@@ -18,10 +18,10 @@ namespace UE::Cook { class FCookSandbox; }
 namespace UE::Cook { struct FCookSandboxConvertCookedPathToPackageNameContext; }
 
 /** A CookedPackageWriter that saves cooked packages in separate .uasset,.uexp,.ubulk files in the Saved\Cooked\[Platform] directory. */
-class FLooseCookedPackageWriter : public TPackageWriterToSharedBuffer<ICookedPackageWriter>
+class FLooseCookedPackageWriter : public TPackageWriterToSharedBuffer<FBaseCookedPackageWriter>
 {
 public:
-	using Super = TPackageWriterToSharedBuffer<ICookedPackageWriter>;
+	using Super = TPackageWriterToSharedBuffer<FBaseCookedPackageWriter>;
 
 	FLooseCookedPackageWriter(const FString& OutputPath, const FString& MetadataDirectoryPath,
 		const ITargetPlatform* TargetPlatform, FAsyncIODelete& InAsyncIODelete,
diff --git a/Engine/Source/Editor/UnrealEd/Private/EditorDomain/EditorDomainUtils.cpp b/Engine/Source/Editor/UnrealEd/Private/EditorDomain/EditorDomainUtils.cpp
--- a/Engine/Source/Editor/UnrealEd/Private/EditorDomain/EditorDomainUtils.cpp
+++ b/Engine/Source/Editor/UnrealEd/Private/EditorDomain/EditorDomainUtils.cpp
@@ -1483,7 +1483,7 @@ void RequestEditorDomainPackage(const FPackagePath& PackagePath,
 }
 
 /** Stores data from SavePackage in accessible fields */
-class FEditorDomainPackageWriter final : public TPackageWriterToSharedBuffer<IPackageWriter>
+class FEditorDomainPackageWriter final : public TPackageWriterToSharedBuffer<FBasePackageWriter>
 {
 public:
 	// IPackageWriter
diff --git a/Engine/Source/Editor/UnrealEd/Private/Serialization/DefaultCookedFilePackageWriter.cpp b/Engine/Source/Editor/UnrealEd/Private/Serialization/DefaultCookedFilePackageWriter.cpp
--- a/Engine/Source/Editor/UnrealEd/Private/Serialization/DefaultCookedFilePackageWriter.cpp
+++ b/Engine/Source/Editor/UnrealEd/Private/Serialization/DefaultCookedFilePackageWriter.cpp
@@ -6,9 +6,9 @@
 #include "Misc/Paths.h"
 #include "Serialization/FilePackageWriterUtil.h"
 
-FDefaultCookedFilePackageWriter::FDefaultCookedFilePackageWriter(const FString& BaseFilename)
-	: TPackageWriterToSharedBuffer<IPackageWriter>()
-	, BaseFilename(BaseFilename)
+FDefaultCookedFilePackageWriter::FDefaultCookedFilePackageWriter(const FString& InBaseFilename)
+	: TPackageWriterToSharedBuffer<FBasePackageWriter>()
+	, BaseFilename(InBaseFilename)
 {}
 
 void FDefaultCookedFilePackageWriter::WritePackageData(const IPackageWriter::FPackageInfo& Info, FLargeMemoryWriter& ExportsArchive, const TArray<FFileRegion>& FileRegions)
diff --git a/Engine/Source/Editor/UnrealEd/Private/Serialization/FilePackageWriterUtil.cpp b/Engine/Source/Editor/UnrealEd/Private/Serialization/FilePackageWriterUtil.cpp
--- a/Engine/Source/Editor/UnrealEd/Private/Serialization/FilePackageWriterUtil.cpp
+++ b/Engine/Source/Editor/UnrealEd/Private/Serialization/FilePackageWriterUtil.cpp
@@ -8,78 +8,6 @@
 #include "Serialization/MemoryWriter.h"
 #include "UObject/SavePackage.h"
 
-namespace FilePackageWriterUtil
-{
-	static void WriteToFile(const FString& Filename, const FCompositeBuffer& Buffer)
-	{
-		IFileManager& FileManager = IFileManager::Get();
-
-		struct FFailureReason
-		{
-			uint32 LastErrorCode = 0;
-			bool bSizeMatchFailed = false;
-			int64 ExpectedSize = 0;
-			int64 ActualSize = 0;
-			bool bArchiveError = false;
-		};
-		TOptional<FFailureReason> FailureReason;
-
-		for (int32 Tries = 0; Tries < 3; ++Tries)
-		{
-			FArchive* Ar = FileManager.CreateFileWriter(*Filename);
-			if (!Ar)
-			{
-				if (!FailureReason)
-				{
-					FailureReason = FFailureReason{ FPlatformMisc::GetLastError(), false };
-				}
-				continue;
-			}
-
-			int64 DataSize = 0;
-			for (const FSharedBuffer& Segment : Buffer.GetSegments())
-			{
-				int64 SegmentSize = static_cast<int64>(Segment.GetSize());
-				Ar->Serialize(const_cast<void*>(Segment.GetData()), SegmentSize);
-				DataSize += SegmentSize;
-			}
-			bool bArchiveError = Ar->IsError();
-			delete Ar;
-
-			int64 ActualSize = FileManager.FileSize(*Filename);
-			if (ActualSize != DataSize)
-			{
-				if (!FailureReason)
-				{
-					FailureReason = FFailureReason{ 0, true, DataSize, ActualSize, bArchiveError };
-				}
-				FileManager.Delete(*Filename);
-				continue;
-			}
-			return;
-		}
-
-		FString ReasonText;
-		if (FailureReason && FailureReason->bSizeMatchFailed)
-		{
-			ReasonText = FString::Printf(TEXT("Unexpected file size. Tried to write %" INT64_FMT " but resultant size was %" INT64_FMT ".%s")
-				TEXT(" Another operation is modifying the file, or the write operation failed to write completely."),
-				FailureReason->ExpectedSize, FailureReason->ActualSize, FailureReason->bArchiveError ? TEXT(" Ar->Serialize failed.") : TEXT(""));
-		}
-		else if (FailureReason && FailureReason->LastErrorCode != 0)
-		{
-			TCHAR LastErrorText[1024];
-			FPlatformMisc::GetSystemErrorMessage(LastErrorText, UE_ARRAY_COUNT(LastErrorText), FailureReason->LastErrorCode);
-			ReasonText = LastErrorText;
-		}
-		else
-		{
-			ReasonText = TEXT("Unknown failure reason.");
-		}
-		UE_LOG(LogSavePackage, Fatal, TEXT("SavePackage Async write %s failed: %s"), *Filename, *ReasonText);
-	}
-}
-
 FFilePackageWriterUtil::FWritePackageParameters::FWritePackageParameters(FRecord& InRecord,
 	const IPackageWriter::FCommitPackageInfo& InInfo,
 	TMap<FName, TRefCountPtr<FPackageHashes>>* InAllPackageHashes,
@@ -92,53 +20,6 @@ FFilePackageWriterUtil::FWritePackageParameters::FWritePackageParameters(FRecord
 	, bProvidePerPackageResult(bInProvidePerPackageResult)
 {}
 
-void FFilePackageWriterUtil::FWriteFileData::HashAndWrite(FMD5& AccumulatedHash, const TRefCountPtr<FPackageHashes>& PackageHashes, IPackageWriter::EWriteOptions WriteOptions) const
-{
-	//@todo: FH: Should we calculate the hash of both output, currently only the main package output hash is calculated
-	if (EnumHasAnyFlags(WriteOptions, IPackageWriter::EWriteOptions::ComputeHash) && bContributeToHash)
-	{
-		for (const FSharedBuffer& Segment : Buffer.GetSegments())
-		{
-			AccumulatedHash.Update(static_cast<const uint8*>(Segment.GetData()), Segment.GetSize());
-		}
-
-		if (ChunkId.IsValid())
-		{
-			FBlake3 ChunkHash;
-			for (const FSharedBuffer& Segment : Buffer.GetSegments())
-			{
-				ChunkHash.Update(static_cast<const uint8*>(Segment.GetData()), Segment.GetSize());
-			}
-			FIoHash FinalHash(ChunkHash.Finalize());
-			PackageHashes->ChunkHashes.Add(ChunkId, FinalHash);
-		}
-	}
-
-	if ((bIsSidecar && EnumHasAnyFlags(WriteOptions, IPackageWriter::EWriteOptions::WriteSidecars)) ||
-		(!bIsSidecar && EnumHasAnyFlags(WriteOptions, IPackageWriter::EWriteOptions::WritePackage)))
-	{
-		const FString* WriteFilename = &Filename;
-		FString FilenameBuffer;
-		if (EnumHasAnyFlags(WriteOptions, IPackageWriter::EWriteOptions::SaveForDiff))
-		{
-			FilenameBuffer = FPaths::Combine(FPaths::GetPath(Filename),
-				FPaths::GetBaseFilename(Filename) + TEXT("_ForDiff") + FPaths::GetExtension(Filename, true));
-			WriteFilename = &FilenameBuffer;
-		}
-		FilePackageWriterUtil::WriteToFile(*WriteFilename, Buffer);
-
-		if (Regions.Num() > 0)
-		{
-			TArray<uint8> Memory;
-			FMemoryWriter Ar(Memory);
-			FFileRegion::SerializeFileRegions(Ar, const_cast<TArray<FFileRegion>&>(Regions));
-
-			FilePackageWriterUtil::WriteToFile(*WriteFilename + FFileRegion::RegionsFileExtension,
-				FCompositeBuffer(FSharedBuffer::MakeView(Memory.GetData(), Memory.Num())));
-		}
-	}
-}
-
 void FFilePackageWriterUtil::WritePackage(FWritePackageParameters& Parameters)
 {
 	if (Parameters.Info.Status == IPackageWriter::ECommitStatus::Success)
@@ -189,12 +70,13 @@ void FFilePackageWriterUtil::CollectForSaveBulkData(FRecord& Record, FCommitCont
 		}
 		else
 		{
-			FWriteFileData& OutputFile = Context.OutputFiles.Emplace_GetRef();
+			UE::PackageWriter::Private::FWriteFileData& OutputFile = Context.OutputFiles.Emplace_GetRef();
 			OutputFile.Filename = BulkRecord.Info.LooseFilePath;
 			OutputFile.Buffer = FCompositeBuffer(BulkRecord.Buffer);
 			OutputFile.Regions = MoveTemp(BulkRecord.Regions);
 			OutputFile.bIsSidecar = true;
-			OutputFile.bContributeToHash = BulkRecord.Info.MultiOutputIndex == 0; // Only caculate the main package output hash
+			// Only calculate the CookHash for files in the main package output. Ignore it for the Optional package output.
+			OutputFile.bContributeToHash = BulkRecord.Info.MultiOutputIndex == 0;
 			OutputFile.ChunkId = BulkRecord.Info.ChunkId;
 		}
 	}
@@ -218,11 +100,12 @@ void FFilePackageWriterUtil::CollectForSaveAdditionalFileRecords(FRecord& Record
 {
 	for (FPackageWriterRecords::FAdditionalFile& AdditionalRecord : Record.AdditionalFiles)
 	{
-		FWriteFileData& OutputFile = Context.OutputFiles.Emplace_GetRef();
+		UE::PackageWriter::Private::FWriteFileData& OutputFile = Context.OutputFiles.Emplace_GetRef();
 		OutputFile.Filename = AdditionalRecord.Info.Filename;
 		OutputFile.Buffer = FCompositeBuffer(AdditionalRecord.Buffer);
 		OutputFile.bIsSidecar = true;
-		OutputFile.bContributeToHash = AdditionalRecord.Info.MultiOutputIndex == 0; // Only calculate the main package output hash
+		// Only calculate the CookHash for files in the main package output. Ignore it for the Optional package output.
+		OutputFile.bContributeToHash = AdditionalRecord.Info.MultiOutputIndex == 0;
 		OutputFile.ChunkId = AdditionalRecord.Info.ChunkId;
 	}
 }
@@ -274,20 +157,22 @@ void FFilePackageWriterUtil::CollectForSaveExportsBuffers(FRecord& Record, FComm
 	
 		// Header (.uasset/.umap)
 		{
-			FWriteFileData& OutputFile = Context.OutputFiles.Emplace_GetRef();
+			UE::PackageWriter::Private::FWriteFileData& OutputFile = Context.OutputFiles.Emplace_GetRef();
 			OutputFile.Filename = Package.Info.LooseFilePath;
 			OutputFile.Buffer = FCompositeBuffer(
 				FSharedBuffer::MakeView(HeaderAndExportsData.GetData(), HeaderSize, HeaderAndExportsData));
 			OutputFile.bIsSidecar = false;
-			OutputFile.bContributeToHash = Package.Info.MultiOutputIndex == 0; // Only calculate the main package output hash
+			// Only calculate the CookHash for files in the main package output. Ignore it for the Optional package output.
+			OutputFile.bContributeToHash = Package.Info.MultiOutputIndex == 0;
 		}
 
 		// Exports + AllAppendedData (.uexp)
 		{
-			FWriteFileData& OutputFile = Context.OutputFiles.Emplace_GetRef();
+			UE::PackageWriter::Private::FWriteFileData& OutputFile = Context.OutputFiles.Emplace_GetRef();
 			OutputFile.Filename = FPaths::ChangeExtension(Package.Info.LooseFilePath, LexToString(EPackageExtension::Exports));
 			OutputFile.bIsSidecar = false;
-			OutputFile.bContributeToHash = Package.Info.MultiOutputIndex == 0; // Only calculate the main package output hash
+			// Only calculate the CookHash for files in the main package output. Ignore it for the Optional package output.
+			OutputFile.bContributeToHash = Package.Info.MultiOutputIndex == 0;
 
 			int32 NumBuffers = ExportsBuffers.Num();
 			TArray<FSharedBuffer> BuffersForComposition;
@@ -314,7 +199,9 @@ void FFilePackageWriterUtil::CollectForSaveExportsBuffers(FRecord& Record, FComm
 	}
 }
 
-void FFilePackageWriterUtil::AsyncSaveOutputFiles(FCommitContext& Context, TMap<FName, TRefCountPtr<FPackageHashes>>* AllPackageHashes, FCriticalSection* PackageHashesLock, bool bProvidePerPackageResult)
+void FFilePackageWriterUtil::AsyncSaveOutputFiles(FCommitContext& Context,
+	TMap<FName, TRefCountPtr<FPackageHashes>>* AllPackageHashes, FCriticalSection* PackageHashesLock,
+	bool bProvidePerPackageResult)
 {
 	if (bProvidePerPackageResult && !AllPackageHashes)
 	{
@@ -348,31 +235,35 @@ void FFilePackageWriterUtil::AsyncSaveOutputFiles(FCommitContext& Context, TMap<
 			ThisPackageHashes->CompletionFuture = PackageHashesCompletionPromise->GetFuture();
 		}
 
-		bool bAlreadyExisted = false;
-		if(AllPackageHashes)
+		if (AllPackageHashes)
 		{
 			FScopeLock PackageHashesScopeLock(PackageHashesLock);
+			// Store a RefCount to the FPackageHashes in AllPackageHashes so e.g. replication of the package's save
+			// results from CookWorker to CookDirector can be delayed until the Future is ready. Give an error if we
+			// try to assign the future twice, it is only supposed to happen once, when the package is saved, and we
+			// don't want the burden of supporting multiple Futures
 			TRefCountPtr<FPackageHashes>& ExistingPackageHashes = AllPackageHashes->FindOrAdd(Context.Info.PackageName);
-			// This calculation of bAlreadyExisted looks weird but we're finding the _refcount_, not the hashes. So if it gets
-			// constructed, it's not actually assigned a pointer.
-			bAlreadyExisted = ExistingPackageHashes.IsValid();
+			if (ExistingPackageHashes.IsValid())
+			{
+				UE_LOG(LogSavePackage, Error,
+					TEXT("FCookedFilePackageWriter encountered the same package twice in a cook! (%s)"),
+					*Context.Info.PackageName.ToString());
+			}
 			ExistingPackageHashes = ThisPackageHashes;
 		}
-		if (bAlreadyExisted)
-		{
-			UE_LOG(LogSavePackage, Error, TEXT("FCookedFilePackageWriter encountered the same package twice in a cook! (%s)"), *Context.Info.PackageName.ToString());
-		}
 	}
 
 	UE::Tasks::Launch(TEXT("HashAndWriteCookedFile"),
 		[OutputFiles = MoveTemp(Context.OutputFiles), WriteOptions = Context.Info.WriteOptions,
-		ThisPackageHashes = MoveTemp(ThisPackageHashes), PackageHashesCompletionPromise = MoveTemp(PackageHashesCompletionPromise)]
+		ThisPackageHashes = MoveTemp(ThisPackageHashes),
+		PackageHashesCompletionPromise = MoveTemp(PackageHashesCompletionPromise)]
 		() mutable
 		{
 			FMD5 AccumulatedHash;
-			for (const FWriteFileData& OutputFile : OutputFiles)
+			for (const UE::PackageWriter::Private::FWriteFileData& OutputFile : OutputFiles)
 			{
-				OutputFile.HashAndWrite(AccumulatedHash, ThisPackageHashes, WriteOptions);
+				UE::PackageWriter::Private::HashAndWrite(OutputFile, AccumulatedHash, ThisPackageHashes,
+					WriteOptions);
 			}
 
 			if (EnumHasAnyFlags(WriteOptions, IPackageWriter::EWriteOptions::ComputeHash))
diff --git a/Engine/Source/Editor/UnrealEd/Public/Serialization/DefaultCookedFilePackageWriter.h b/Engine/Source/Editor/UnrealEd/Public/Serialization/DefaultCookedFilePackageWriter.h
--- a/Engine/Source/Editor/UnrealEd/Public/Serialization/DefaultCookedFilePackageWriter.h
+++ b/Engine/Source/Editor/UnrealEd/Public/Serialization/DefaultCookedFilePackageWriter.h
@@ -9,12 +9,12 @@
  * BaseFilename is the filename of the package to write. The extension is ignored.
  * Then extension of each file is specified to match Unreal's loading expectations.
  */
-class FDefaultCookedFilePackageWriter : public TPackageWriterToSharedBuffer<IPackageWriter>
+class FDefaultCookedFilePackageWriter : public TPackageWriterToSharedBuffer<FBasePackageWriter>
 {
 public:
-	using Super = TPackageWriterToSharedBuffer<IPackageWriter>;
+	using Super = TPackageWriterToSharedBuffer<FBasePackageWriter>;
 
-	FDefaultCookedFilePackageWriter(const FString& BaseFilename);
+	FDefaultCookedFilePackageWriter(const FString& InBaseFilename);
 
 	void WritePackageData(const IPackageWriter::FPackageInfo& Info, FLargeMemoryWriter& ExportsArchive, const TArray<FFileRegion>& FileRegions) override;
 	void WriteBulkData(const IPackageWriter::FBulkDataInfo& Info, const FIoBuffer& BulkData, const TArray<FFileRegion>& FileRegions) override;
diff --git a/Engine/Source/Editor/UnrealEd/Public/Serialization/FilePackageWriterUtil.h b/Engine/Source/Editor/UnrealEd/Public/Serialization/FilePackageWriterUtil.h
--- a/Engine/Source/Editor/UnrealEd/Public/Serialization/FilePackageWriterUtil.h
+++ b/Engine/Source/Editor/UnrealEd/Public/Serialization/FilePackageWriterUtil.h
@@ -46,28 +46,12 @@ private:
 		TArray<FFileRegion> Regions;
 	};
 
-	/**
-	 * The data needed to asynchronously write one of the files (.uasset, .uexp, .ubulk, any optional and any additional),
-	 * without reference back to other data on this writer.
-	 */
-	struct FWriteFileData
-	{
-		FString Filename;
-		FCompositeBuffer Buffer;
-		TArray<FFileRegion> Regions;
-		bool bIsSidecar;
-		bool bContributeToHash = true;
-		FIoChunkId ChunkId = FIoChunkId::InvalidChunkId;
-
-		void HashAndWrite(FMD5& AccumulatedHash, const TRefCountPtr<FPackageHashes>& PackageHashes, IPackageWriter::EWriteOptions WriteOptions) const;
-	};
-
 	/** Stack data for the helper functions of CommitPackageInternal. */
 	struct FCommitContext
 	{
 		const IPackageWriter::FCommitPackageInfo& Info;
 		TArray<TArray<FExportBuffer>> ExportsBuffers;
-		TArray<FWriteFileData> OutputFiles;
+		TArray<UE::PackageWriter::Private::FWriteFileData> OutputFiles;
 	};
 
 	static void AsyncSave(FWritePackageParameters& Parameters);
@@ -78,5 +62,7 @@ private:
 	static void CollectForSaveExportsFooter(FRecord& Record, FCommitContext& Context);
 	static void CollectForSaveExportsPackageTrailer(FRecord& Record, FCommitContext& Context);
 	static void CollectForSaveExportsBuffers(FRecord& Record, FCommitContext& Context);
-	static void AsyncSaveOutputFiles(FCommitContext& Context, TMap<FName, TRefCountPtr<FPackageHashes>>* AllPackageHashes, FCriticalSection* PackageHashesLock, bool bProvidePerPackageResult);
+	static void AsyncSaveOutputFiles(FCommitContext& Context,
+		TMap<FName,TRefCountPtr<FPackageHashes>>* AllPackageHashes, FCriticalSection* PackageHashesLock,
+		bool bProvidePerPackageResult);
 };
diff --git a/Engine/Source/Runtime/Core/Public/Serialization/PackageWriter.h b/Engine/Source/Runtime/Core/Public/Serialization/PackageWriter.h
--- a/Engine/Source/Runtime/Core/Public/Serialization/PackageWriter.h
+++ b/Engine/Source/Runtime/Core/Public/Serialization/PackageWriter.h
@@ -2,12 +2,11 @@
 
 #pragma once
 
-// HEADER_UNIT_UNSUPPORTED - Contains code with dll export not allowed COREUOBJECT_API
-
 #include "Async/Future.h"
 #include "Containers/StringView.h"
 #include "IO/IoDispatcher.h"
 #include "IO/IoHash.h"
+#include "Memory/CompositeBuffer.h"
 #include "Misc/AssertionMacros.h"
 #include "Misc/DateTime.h"
 #include "Misc/EnumClassFlags.h"
@@ -204,10 +203,10 @@ public:
 	virtual void WritePackageTrailer(const FPackageTrailerInfo& Info, const FIoBuffer& Data) = 0;
 
 	/** Create the FLargeMemoryWriter to which the Header and Exports are written during the save. */
-	COREUOBJECT_API virtual TUniquePtr<FLargeMemoryWriter> CreateLinkerArchive(FName PackageName, UObject* Asset, uint16 MultiOutputIndex);
+	virtual TUniquePtr<FLargeMemoryWriter> CreateLinkerArchive(FName PackageName, UObject* Asset, uint16 MultiOutputIndex) = 0;
 
 	/** Returns an archive to be used when serializing exports. */
-	COREUOBJECT_API virtual TUniquePtr<FLargeMemoryWriter> CreateLinkerExportsArchive(FName PackageName, UObject* Asset, uint16 MultiOutputIndex);
+	virtual TUniquePtr<FLargeMemoryWriter> CreateLinkerExportsArchive(FName PackageName, UObject* Asset, uint16 MultiOutputIndex) = 0;
 
 	/**
 	 * Overridden by PackageWriters that handle bDeterminismDebug=true. A system will call this function to register
@@ -266,6 +265,25 @@ struct FBeginCacheForCookedPlatformDataInfo
 	uint32 SaveFlags;
 };
 
+/**
+ * The data needed to asynchronously write that was passed to a PackageWriter (.uasset, .uexp, .ubulk, any optional
+ * and any additional), without reference back to other data on the Writer.
+ * This type is a cooker implementation detail and might be changed without deprecation in a future version.
+ */
+struct FWriteFileData
+{
+	FString Filename;
+	FCompositeBuffer Buffer;
+	TArray<FFileRegion> Regions;
+	bool bIsSidecar = false;
+	/**
+	 * True if file should be hashed and its hash accumulated into the total cookhash for the package writing the
+	 * file.
+	 */
+	bool bContributeToHash = true;
+	FIoChunkId ChunkId = FIoChunkId::InvalidChunkId;
+};
+
 /**
  * Backpointer interface to the Cooker for functions that are implemented in Cooker code for PackageWriters used
  * by the cooker rather than CookedPackageWriter code, but that are triggered from CookedPackageWriter code.
diff --git a/Engine/Source/Runtime/CoreUObject/Internal/Serialization/PackageWriterToSharedBuffer.h b/Engine/Source/Runtime/CoreUObject/Internal/Serialization/PackageWriterToSharedBuffer.h
--- a/Engine/Source/Runtime/CoreUObject/Internal/Serialization/PackageWriterToSharedBuffer.h
+++ b/Engine/Source/Runtime/CoreUObject/Internal/Serialization/PackageWriterToSharedBuffer.h
@@ -6,7 +6,7 @@
 #include "HAL/CriticalSection.h"
 #include "Memory/SharedBuffer.h"
 #include "Misc/AssertionMacros.h"
-#include "Serialization/PackageWriter.h"
+#include "Serialization/BasePackageWriter.h"
 #include "Templates/UniquePtr.h"
 
 class FPackageWriterRecords
@@ -84,13 +84,14 @@ private:
 
 /**
  * A base class for IPackageWriter subclasses that writes to records that are read in CommitPackage.
- * To avoid diamond inheritance, this class specifies the interface it is implementing (either IPackageWriter
- * or ICookedPackageWriter) by template. Subclasses should derive from one of
- *     TPackageWriterToSharedBuffer<IPackageWriter>
- *     TPackageWriterToSharedBuffer<ICookedPackageWriter>.
+ * To avoid diamond inheritance, this class specifies the base class for the interface it is implementing (either
+ * FBasePackageWriter for IPackageWriter or FBaseCookedPackageWriter for ICookedPackageWriter) by template. Subclasses
+ * should derive from one of
+ *     TPackageWriterToSharedBuffer<FBasePackageWriter>
+ *     TPackageWriterToSharedBuffer<FBaseCookedPackageWriter>.
  */
-template <typename BaseInterface>
-class TPackageWriterToSharedBuffer : public BaseInterface
+template <typename BaseClass>
+class TPackageWriterToSharedBuffer : public BaseClass
 {
 public:
 	using FWritePackageRecord = FPackageWriterRecords::FWritePackage;
@@ -161,7 +162,7 @@ protected:
 };
 
 /** A PackageWriter that saves the portions of the SavePackage output to sharedbuffers on a record structure. */
-class FPackageWriterToRecord : public TPackageWriterToSharedBuffer<IPackageWriter>
+class FPackageWriterToRecord : public TPackageWriterToSharedBuffer<FBasePackageWriter>
 {
 public:
 	FPackageWriterRecords::FPackage SavedRecord;
diff --git a/Engine/Source/Runtime/CoreUObject/Private/Serialization/BasePackageWriter.cpp b/Engine/Source/Runtime/CoreUObject/Private/Serialization/BasePackageWriter.cpp
new file mode 100644
--- /dev/null
+++ b/Engine/Source/Runtime/CoreUObject/Private/Serialization/BasePackageWriter.cpp
@@ -0,0 +1,180 @@
+// Copyright Epic Games, Inc. All Rights Reserved.
+
+#include "Serialization/BasePackageWriter.h"
+
+#include "Serialization/LargeMemoryWriter.h"
+
+#if WITH_EDITOR
+#include "HAL/FileManager.h"
+#include "HAL/PlatformMisc.h"
+#include "Misc/Optional.h"
+#include "Misc/Paths.h"
+#include "Serialization/MemoryWriter.h"
+#include "UObject/SavePackage.h"
+#endif
+
+namespace UE::PackageWriter::Private
+{
+
+TUniquePtr<FLargeMemoryWriter> BaseCreateLinkerArchive(FName PackageName,
+	UObject* Asset, uint16 /*MultiOutputIndex*/)
+{
+	// The LargeMemoryWriter does not need to be persistent; the LinkerSave wraps it and reports Persistent=true
+	bool bIsPersistent = false;
+	return TUniquePtr<FLargeMemoryWriter>(new FLargeMemoryWriter(0, bIsPersistent, *PackageName.ToString()));
+}
+
+TUniquePtr<FLargeMemoryWriter> BaseCreateLinkerExportsArchive(FName PackageName,
+	UObject* Asset, uint16 /*MultiOutputIndex*/)
+{
+	const bool bPersistent = true;
+	return MakeUnique<FLargeMemoryWriter>(0, bPersistent, *PackageName.ToString());
+}
+
+#if WITH_EDITOR
+
+static void WriteToFile(const FString& Filename, const FCompositeBuffer& Buffer)
+{
+	IFileManager& FileManager = IFileManager::Get();
+
+	struct FFailureReason
+	{
+		uint32 LastErrorCode = 0;
+		bool bSizeMatchFailed = false;
+		int64 ExpectedSize = 0;
+		int64 ActualSize = 0;
+		bool bArchiveError = false;
+	};
+	TOptional<FFailureReason> FailureReason;
+
+	for (int32 Tries = 0; Tries < 3; ++Tries)
+	{
+		FArchive* Ar = FileManager.CreateFileWriter(*Filename);
+		if (!Ar)
+		{
+			if (!FailureReason)
+			{
+				FailureReason = FFailureReason{ FPlatformMisc::GetLastError(), false };
+			}
+			continue;
+		}
+
+		int64 DataSize = 0;
+		for (const FSharedBuffer& Segment : Buffer.GetSegments())
+		{
+			int64 SegmentSize = static_cast<int64>(Segment.GetSize());
+			Ar->Serialize(const_cast<void*>(Segment.GetData()), SegmentSize);
+			DataSize += SegmentSize;
+		}
+		bool bArchiveError = Ar->IsError();
+		delete Ar;
+
+		int64 ActualSize = FileManager.FileSize(*Filename);
+		if (ActualSize != DataSize)
+		{
+			if (!FailureReason)
+			{
+				FailureReason = FFailureReason{ 0, true, DataSize, ActualSize, bArchiveError };
+			}
+			FileManager.Delete(*Filename);
+			continue;
+		}
+		return;
+	}
+
+	FString ReasonText;
+	if (FailureReason && FailureReason->bSizeMatchFailed)
+	{
+		ReasonText = FString::Printf(TEXT("Unexpected file size. Tried to write %" INT64_FMT " but resultant size was %" INT64_FMT ".%s")
+			TEXT(" Another operation is modifying the file, or the write operation failed to write completely."),
+			FailureReason->ExpectedSize, FailureReason->ActualSize, FailureReason->bArchiveError ? TEXT(" Ar->Serialize failed.") : TEXT(""));
+	}
+	else if (FailureReason && FailureReason->LastErrorCode != 0)
+	{
+		TCHAR LastErrorText[1024];
+		FPlatformMisc::GetSystemErrorMessage(LastErrorText, UE_ARRAY_COUNT(LastErrorText), FailureReason->LastErrorCode);
+		ReasonText = LastErrorText;
+	}
+	else
+	{
+		ReasonText = TEXT("Unknown failure reason.");
+	}
+	UE_LOG(LogSavePackage, Fatal, TEXT("SavePackage Async write %s failed: %s"), *Filename, *ReasonText);
+}
+
+void HashAndWrite(const UE::PackageWriter::Private::FWriteFileData& FileData, FMD5& AccumulatedHash,
+	const TRefCountPtr<FPackageHashes>& PackageHashes, IPackageWriter::EWriteOptions WriteOptions)
+{
+	//@todo: FH: Should we calculate the hash of both output, currently only the main package output hash is calculated
+	if (EnumHasAnyFlags(WriteOptions, IPackageWriter::EWriteOptions::ComputeHash) && FileData.bContributeToHash)
+	{
+		for (const FSharedBuffer& Segment : FileData.Buffer.GetSegments())
+		{
+			AccumulatedHash.Update(static_cast<const uint8*>(Segment.GetData()), Segment.GetSize());
+		}
+
+		if (FileData.ChunkId.IsValid())
+		{
+			FBlake3 ChunkHash;
+			for (const FSharedBuffer& Segment : FileData.Buffer.GetSegments())
+			{
+				ChunkHash.Update(static_cast<const uint8*>(Segment.GetData()), Segment.GetSize());
+			}
+			FIoHash FinalHash(ChunkHash.Finalize());
+			PackageHashes->ChunkHashes.Add(FileData.ChunkId, FinalHash);
+		}
+	}
+
+	if ((FileData.bIsSidecar && EnumHasAnyFlags(WriteOptions, IPackageWriter::EWriteOptions::WriteSidecars)) ||
+		(!FileData.bIsSidecar && EnumHasAnyFlags(WriteOptions, IPackageWriter::EWriteOptions::WritePackage)))
+	{
+		const FString* WriteFilename = &FileData.Filename;
+		FString FilenameBuffer;
+		if (EnumHasAnyFlags(WriteOptions, IPackageWriter::EWriteOptions::SaveForDiff))
+		{
+			FilenameBuffer = FPaths::Combine(FPaths::GetPath(FileData.Filename),
+				FPaths::GetBaseFilename(FileData.Filename) + TEXT("_ForDiff")
+				+ FPaths::GetExtension(FileData.Filename, true));
+			WriteFilename = &FilenameBuffer;
+		}
+		WriteToFile(*WriteFilename, FileData.Buffer);
+
+		if (FileData.Regions.Num() > 0)
+		{
+			TArray<uint8> Memory;
+			FMemoryWriter Ar(Memory);
+			FFileRegion::SerializeFileRegions(Ar, const_cast<TArray<FFileRegion>&>(FileData.Regions));
+
+			WriteToFile(*WriteFilename + FFileRegion::RegionsFileExtension,
+				FCompositeBuffer(FSharedBuffer::MakeView(Memory.GetData(), Memory.Num())));
+		}
+	}
+}
+
+#endif
+
+} // UE::PackageWriter::Private
+
+TUniquePtr<FLargeMemoryWriter> FBasePackageWriter::CreateLinkerArchive(FName PackageName,
+	UObject* Asset, uint16 MultiOutputIndex)
+{
+	return UE::PackageWriter::Private::BaseCreateLinkerArchive(PackageName, Asset, MultiOutputIndex);
+}
+
+TUniquePtr<FLargeMemoryWriter> FBasePackageWriter::CreateLinkerExportsArchive(FName PackageName,
+	UObject* Asset, uint16 MultiOutputIndex)
+{
+	return UE::PackageWriter::Private::BaseCreateLinkerExportsArchive(PackageName, Asset, MultiOutputIndex);
+}
+
+TUniquePtr<FLargeMemoryWriter> FBaseCookedPackageWriter::CreateLinkerArchive(FName PackageName,
+	UObject* Asset, uint16 MultiOutputIndex)
+{
+	return UE::PackageWriter::Private::BaseCreateLinkerArchive(PackageName, Asset, MultiOutputIndex);
+}
+
+TUniquePtr<FLargeMemoryWriter> FBaseCookedPackageWriter::CreateLinkerExportsArchive(FName PackageName,
+	UObject* Asset, uint16 MultiOutputIndex)
+{
+	return UE::PackageWriter::Private::BaseCreateLinkerExportsArchive(PackageName, Asset, MultiOutputIndex);
+}
diff --git a/Engine/Source/Runtime/CoreUObject/Private/Serialization/PackageWriterToSharedBuffer.cpp b/Engine/Source/Runtime/CoreUObject/Private/Serialization/PackageWriterToSharedBuffer.cpp
--- a/Engine/Source/Runtime/CoreUObject/Private/Serialization/PackageWriterToSharedBuffer.cpp
+++ b/Engine/Source/Runtime/CoreUObject/Private/Serialization/PackageWriterToSharedBuffer.cpp
@@ -7,19 +7,6 @@
 
 DEFINE_LOG_CATEGORY_STATIC(LogPackageWriter, Display, All);
 
-TUniquePtr<FLargeMemoryWriter> IPackageWriter::CreateLinkerArchive(FName PackageName, UObject* Asset, uint16 /*MultiOutputIndex*/)
-{
-	// The LargeMemoryWriter does not need to be persistent; the LinkerSave wraps it and reports Persistent=true
-	bool bIsPersistent = false; 
-	return TUniquePtr<FLargeMemoryWriter>(new FLargeMemoryWriter(0, bIsPersistent, *PackageName.ToString()));
-}
-
-TUniquePtr<FLargeMemoryWriter> IPackageWriter::CreateLinkerExportsArchive(FName PackageName, UObject* Asset, uint16 /*MultiOutputIndex*/)
-{
-	const bool bPersistent = true;
-	return MakeUnique<FLargeMemoryWriter>(0, bPersistent, *PackageName.ToString());
-}
-
 static FSharedBuffer IoBufferToSharedBuffer(const FIoBuffer& InBuffer)
 {
 	InBuffer.EnsureOwned();
diff --git a/Engine/Source/Runtime/CoreUObject/Public/Serialization/BasePackageWriter.h b/Engine/Source/Runtime/CoreUObject/Public/Serialization/BasePackageWriter.h
new file mode 100644
--- /dev/null
+++ b/Engine/Source/Runtime/CoreUObject/Public/Serialization/BasePackageWriter.h
@@ -0,0 +1,48 @@
+// Copyright Epic Games, Inc. All Rights Reserved.
+
+#pragma once
+
+#include "Serialization/PackageWriter.h"
+
+/** Shared Implementation for IPackageWriters that are not ICookedPackageWriters. */
+class FBasePackageWriter : public IPackageWriter
+{
+public:
+	/** Create the FLargeMemoryWriter to which the Header and Exports are written during the save. */
+	COREUOBJECT_API virtual TUniquePtr<FLargeMemoryWriter> CreateLinkerArchive(FName PackageName,
+		UObject* Asset, uint16 MultiOutputIndex) override;
+
+	/** Returns an archive to be used when serializing exports. */
+	COREUOBJECT_API virtual TUniquePtr<FLargeMemoryWriter> CreateLinkerExportsArchive(FName PackageName,
+		UObject* Asset, uint16 MultiOutputIndex) override;
+};
+
+/** Shared Implementation for ICookedPackageWriters. */
+class FBaseCookedPackageWriter : public ICookedPackageWriter
+{
+public:
+	/** Create the FLargeMemoryWriter to which the Header and Exports are written during the save. */
+	COREUOBJECT_API virtual TUniquePtr<FLargeMemoryWriter> CreateLinkerArchive(FName PackageName,
+		UObject* Asset, uint16 MultiOutputIndex) override;
+
+	/** Returns an archive to be used when serializing exports. */
+	COREUOBJECT_API virtual TUniquePtr<FLargeMemoryWriter> CreateLinkerExportsArchive(FName PackageName,
+		UObject* Asset, uint16 MultiOutputIndex) override;
+};
+
+// Shared helper functions for PackageWriters and CookedPackageWriters
+namespace UE::PackageWriter::Private
+{
+
+#if WITH_EDITOR
+/**
+ * Synchronously write if requested and hash if requested the data in WriteFileData to the filename in WriteFileData,
+ * based on the given WriteOptions and WriteFileData's filetype flags.
+ * This function is a cooker implementation detail and might be changed without deprecation in a future version.
+ */
+COREUOBJECT_API void HashAndWrite(const UE::PackageWriter::Private::FWriteFileData& WriteFileData,
+	FMD5& AccumulatedHash, const TRefCountPtr<FPackageHashes>& PackageHashes,
+	IPackageWriter::EWriteOptions WriteOptions);
+#endif
+
+}
