[ComponentTypeRegistry]

- Defers rebuilding the type registry until after AssetRegistry gathering is complete. Otherwise, the TypeRegistry will aggressively take contend with the AssetRegistry as it tries to register data, slowing both down considerably. Worse still, since multiple mechanisms can force refreshing, we may force rebuild more than once while initializing.

Boot the editor to an interactive empty map:
4m 47s -> 2m 48s

#rnx

diff --git a/Engine/Source/Editor/UnrealEd/Private/ComponentTypeRegistry.cpp b/Engine/Source/Editor/UnrealEd/Private/ComponentTypeRegistry.cpp
--- a/Engine/Source/Editor/UnrealEd/Private/ComponentTypeRegistry.cpp
+++ b/Engine/Source/Editor/UnrealEd/Private/ComponentTypeRegistry.cpp
@@ -269,7 +269,13 @@ struct FComponentTypeRegistryData
 	/** Implementation of FTickableEditorObject */
 	virtual void Tick(float) override;
 	virtual ETickableTickType GetTickableTickType() const override { return ETickableTickType::Conditional; }
-	virtual bool IsTickable() const override { return GUnrealEd != nullptr; }
+	virtual bool IsTickable() const override 
+	{ 
+		// Only necessary to tick when in editor. However,
+		// avoid ticking until the AssetRegistry is done gathering otherwise 
+		// we can introduce excessive amounts of contention and cache updates when ticking.
+		return GUnrealEd != nullptr && !AssetRegistryModule.Get().IsGathering();
+	}
 	virtual TStatId GetStatId() const override { RETURN_QUICK_DECLARE_CYCLE_STAT(FTypeDatabaseUpdater, STATGROUP_Tickables); }
 	
 	/** Implementation of FGCObject */
@@ -291,6 +297,7 @@ public:
 	TArray<FAssetData> PendingAssetData;
 	TMap<FTopLevelAssetPath, int32> ClassPathToClassListIndexMap;
 	FOnComponentTypeListChanged ComponentListChanged;
+	FAssetRegistryModule& AssetRegistryModule;
 	bool bNeedsRefreshNextTick;
 };
 
@@ -423,7 +430,8 @@ void FComponentTypeRegistryData::AddBasicShapeComponents(TArray<FComponentClassC
 }
 
 FComponentTypeRegistryData::FComponentTypeRegistryData()
-	: bNeedsRefreshNextTick(false)
+	: AssetRegistryModule(FModuleManager::LoadModuleChecked<FAssetRegistryModule>(TEXT("AssetRegistry")))
+	, bNeedsRefreshNextTick(false)
 {
 	const auto HandleAdded = [](const FAssetData& Data, FComponentTypeRegistryData* Parent)
 	{
@@ -434,9 +442,10 @@ FComponentTypeRegistryData::FComponentTypeRegistryData()
 		}
 	};
 
-	IAssetRegistry& AssetRegistry = FModuleManager::LoadModuleChecked<FAssetRegistryModule>(AssetRegistryConstants::ModuleName).Get();
+	IAssetRegistry& AssetRegistry = AssetRegistryModule.Get();
 	AssetRegistry.OnAssetAdded().AddStatic(HandleAdded, this);
 	AssetRegistry.OnAssetRemoved().AddStatic(HandleAdded, this);
+	AssetRegistry.OnKnownGathersComplete().AddLambda([](FComponentTypeRegistryData* Data) { Data->Invalidate(); }, this);
 
 	const auto HandleRenamed = [](const FAssetData& Data, const FString&, FComponentTypeRegistryData* Parent)
 	{
@@ -450,6 +459,13 @@ FComponentTypeRegistryData::FComponentTypeRegistryData()
 
 void FComponentTypeRegistryData::ForceRefreshComponentList()
 {
+	if (!IsTickable())
+	{
+		// If we aren't tickable, don't refresh for the same reasons why we don't tick. This early out is 
+		// necessary should ForceRefreshComponentList be called explicitly rather than Invalidate().
+		return;
+	}
+
 	bNeedsRefreshNextTick = false;
 	ComponentClassList.Empty();
 	ComponentTypeList.Empty();
@@ -598,7 +614,7 @@ void FComponentTypeRegistryData::ForceRefreshComponentList()
 			ComponentTypeList.Add(MoveTemp(Entry));
 		};
 
-		IAssetRegistry& AssetRegistry = FModuleManager::LoadModuleChecked<FAssetRegistryModule>(AssetRegistryConstants::ModuleName).Get();
+		IAssetRegistry& AssetRegistry = AssetRegistryModule.Get();
 
 		TSet<FTopLevelAssetPath> ActorComponentDerivedClassNames;
 		{
@@ -688,18 +704,8 @@ void FComponentTypeRegistryData::Tick(float)
 	TRACE_CPUPROFILER_EVENT_SCOPE(FComponentTypeRegistryData::Tick);
 
 	bool bRequiresRefresh = bNeedsRefreshNextTick;
-
-	if (PendingAssetData.Num() != 0)
+	if (!bRequiresRefresh && PendingAssetData.Num() != 0)
 	{
-		FAssetRegistryModule& AssetRegistryModule = FModuleManager::LoadModuleChecked<FAssetRegistryModule>(TEXT("AssetRegistry"));
-
-		// Avoid querying the asset registry until it has finished discovery, 
-		// as doing so may force it to update temporary caches many times:
-		if(AssetRegistryModule.Get().IsLoadingAssets())
-		{
-			return;
-		}
-
 		TArray<FTopLevelAssetPath> ClassNames;
 		ClassNames.Add(UActorComponent::StaticClass()->GetClassPathName());
 		TSet<FTopLevelAssetPath> DerivedClassNames;
@@ -720,6 +726,7 @@ void FComponentTypeRegistryData::Tick(float)
 
 	if (bRequiresRefresh)
 	{
+		PendingAssetData.Empty();
 		ForceRefreshComponentList();
 	}
 }
diff --git a/Engine/Source/Editor/UnrealEd/Public/ComponentTypeRegistry.h b/Engine/Source/Editor/UnrealEd/Public/ComponentTypeRegistry.h
--- a/Engine/Source/Editor/UnrealEd/Public/ComponentTypeRegistry.h
+++ b/Engine/Source/Editor/UnrealEd/Public/ComponentTypeRegistry.h
@@ -31,9 +31,9 @@ struct FComponentTypeRegistry
 
 	/**
 	 * Called when the user changes the text in the search box.
-	 * @OutComponentList Pointer that will be set to the (globally shared) component type list
-	 * @return Deleate that can be used to handle change notifications. change notifications are raised when entries are 
-	 *	added or removed from the component type list
+	 * @OutComponentList Pointer that will be set to the (globally shared) component type list.
+	 * @return Delegate that can be used to handle change notifications. Change notifications are raised when entries are 
+	 * added or removed from the component type list.
 	 */
 	UNREALED_API FOnComponentTypeListChanged& SubscribeToComponentList(TArray<FComponentClassComboEntryPtr>*& OutComponentList);
 	UNREALED_API FOnComponentTypeListChanged& SubscribeToComponentList(const TArray<FComponentTypeEntry>*& OutComponentList);
