Downgrade Warning to Display

#rnx

diff --git a/Engine/Source/Editor/UnrealEd/Private/CookOnTheFlyServer.cpp b/Engine/Source/Editor/UnrealEd/Private/CookOnTheFlyServer.cpp
--- a/Engine/Source/Editor/UnrealEd/Private/CookOnTheFlyServer.cpp
+++ b/Engine/Source/Editor/UnrealEd/Private/CookOnTheFlyServer.cpp
@@ -12552,7 +12552,7 @@ void UCookOnTheFlyServer::CompileDLCLocalization(FBeginCookContext& BeginContext
 	}
 	else
 	{
-		UE_LOG(LogCook, Warning, TEXT("DLC plugin '%s' was not known to IPluginManager. UGC localization has not been compiled and will be missing at runtime."), *CookByTheBookOptions->DlcName);
+		UE_LOG(LogCook, Display, TEXT("DLC plugin '%s' was not known to IPluginManager. UGC localization has not been compiled and will be missing at runtime."), *CookByTheBookOptions->DlcName);
 	}
 }
 
