Better handling of objects that will also be re-instantiated in the PrePopulateInstancingGraphForArchetype

#rnx

diff --git a/Engine/Source/Editor/Kismet/Private/BlueprintCompilationManager.cpp b/Engine/Source/Editor/Kismet/Private/BlueprintCompilationManager.cpp
--- a/Engine/Source/Editor/Kismet/Private/BlueprintCompilationManager.cpp
+++ b/Engine/Source/Editor/Kismet/Private/BlueprintCompilationManager.cpp
@@ -3147,8 +3147,13 @@ void FBlueprintCompilationManagerImpl::ReinstanceBatch(TArray<FReinstancingJob>&
 					Template = *NewArchetypeTemplate;
 				}
 
-				FObjectInstancingGraph InstancingGraph;
-				const bool bUseInstancingGraph = FBlueprintCompileReinstancer::PrePopulateInstancingGraphForArchetype(InstancingGraph, OldInstance, Template, &OldArchetypeToNewArchetype);
+				FObjectInstancingGraph PrePopulatedInstancingGraph;
+				const bool bUsePrePopulatedInstancingGraph = FBlueprintCompileReinstancer::PrePopulateInstancingGraphForArchetype(PrePopulatedInstancingGraph, OldInstance, Template, &OldArchetypeToNewArchetype);
+				if (bUsePrePopulatedInstancingGraph)
+				{
+					// Make sure the template object as all it references path up before using to instance another archetype
+					FArchiveReplaceObjectRef<UObject> ReplaceInCDOAr(Template, OldArchetypeToNewArchetype);
+				}
 
 				UObject* Destination = GetTransientOuterForRename(OldInstance->GetClass());
 				OldInstance->Rename(
@@ -3163,7 +3168,7 @@ void FBlueprintCompilationManagerImpl::ReinstanceBatch(TArray<FReinstancingJob>&
 				FMakeClassSpawnableOnScope TemporarilySpawnable(NewClass);
 				const EObjectFlags FlagMask = RF_Public | RF_ArchetypeObject | RF_Transactional | RF_Transient | RF_TextExportTransient | RF_InheritableComponentTemplate | RF_Standalone; //TODO: what about RF_RootSet?
 
-				UObject* NewArchetype = NewObject<UObject>(OriginalOuter, NewClass, OriginalName, OriginalFlags & FlagMask, Template, false, bUseInstancingGraph ? &InstancingGraph : nullptr);
+				UObject* NewArchetype = NewObject<UObject>(OriginalOuter, NewClass, OriginalName, OriginalFlags & FlagMask, Template, false, bUsePrePopulatedInstancingGraph ? &PrePopulatedInstancingGraph : nullptr);
 
 				OldArchetypeToNewArchetype.Add(OldInstance, NewArchetype);
 				if (OldToNewTemplatesForClass)
@@ -4236,12 +4241,13 @@ void FBlueprintCompilationManager::FlushCompilationQueue(FUObjectSerializeContex
 	}
 }
 
-void FBlueprintCompilationManager::FlushCompilationQueueAndReinstance()
+void FBlueprintCompilationManager::FlushCompilationQueueAndReinstance(const bool bFindAndReplaceCDOReferences /* = false*/)
 {
 	if(BPCMImpl)
 	{
-		BPCMImpl->FlushCompilationQueueImpl(false, nullptr, nullptr, nullptr);
-		BPCMImpl->FlushReinstancingQueueImpl();
+		TMap<UClass*, TMap<UObject*, UObject*>> OldToNewTemplates;
+		BPCMImpl->FlushCompilationQueueImpl(false, nullptr, nullptr, nullptr, bFindAndReplaceCDOReferences ? &OldToNewTemplates : nullptr);
+		BPCMImpl->FlushReinstancingQueueImpl(bFindAndReplaceCDOReferences, bFindAndReplaceCDOReferences ? &OldToNewTemplates : nullptr);
 
 		BPCMImpl->OldCDOs.Empty();
 	}
diff --git a/Engine/Source/Editor/Kismet/Public/BlueprintCompilationManager.h b/Engine/Source/Editor/Kismet/Public/BlueprintCompilationManager.h
--- a/Engine/Source/Editor/Kismet/Public/BlueprintCompilationManager.h
+++ b/Engine/Source/Editor/Kismet/Public/BlueprintCompilationManager.h
@@ -50,7 +50,7 @@ struct KISMET_API FBlueprintCompilationManager
 	/**
 	 * Flushes the compilation queue and finishes reinstancing
 	 */
-	static void FlushCompilationQueueAndReinstance();
+	static void FlushCompilationQueueAndReinstance(const bool bFindAndReplaceCDOReferences = false);
 
 	/**
 	 * Immediately compiles the blueprint, no expectation that related blueprints be subsequently compiled.
diff --git a/Engine/Source/Editor/UnrealEd/Private/Kismet2/KismetReinstanceUtilities.cpp b/Engine/Source/Editor/UnrealEd/Private/Kismet2/KismetReinstanceUtilities.cpp
--- a/Engine/Source/Editor/UnrealEd/Private/Kismet2/KismetReinstanceUtilities.cpp
+++ b/Engine/Source/Editor/UnrealEd/Private/Kismet2/KismetReinstanceUtilities.cpp
@@ -2280,11 +2280,11 @@ void FBlueprintCompileReinstancer::ReplaceObjectHelper(UObject*& OldObject, UCla
 			}
 		}
 
-		FObjectInstancingGraph InstancingGraph;
-		const bool bUseInstancingGraph = PrePopulateInstancingGraphForArchetype(InstancingGraph, OldObject, NewArchetype, &OldToNewInstanceMap);
+		FObjectInstancingGraph PrePopulatedInstancingGraph;
+		const bool bUsePrePopulatedInstancingGraph = PrePopulateInstancingGraphForArchetype(PrePopulatedInstancingGraph, OldObject, NewArchetype, &OldToNewInstanceMap);
 		
 		FMakeClassSpawnableOnScope TemporarilySpawnable(NewClass);
-		NewUObject = NewObject<UObject>(DestinationOuter, NewClass, NewName, RF_NoFlags, NewArchetype, false, bUseInstancingGraph ? &InstancingGraph : nullptr, OldExternalPackage);
+		NewUObject = NewObject<UObject>(DestinationOuter, NewClass, NewName, RF_NoFlags, NewArchetype, false, bUsePrePopulatedInstancingGraph ? &PrePopulatedInstancingGraph : nullptr, OldExternalPackage);
 	}
 
 	check(NewUObject != nullptr);
@@ -3314,11 +3314,11 @@ bool FBlueprintCompileReinstancer::PrePopulateInstancingGraphForArchetype(FObjec
 		return false;
 	}
 
-	auto GatherMatchingSubObjects = [&InstancingGraph, OldInstance](UObject* Instance, UObject* Template, auto GatherMatchingSubObjectsRecurse) -> void
+	auto GatherMatchingSubObjects = [&InstancingGraph, OldInstance, OldToNewInstanceMap](UObject* Instance, UObject* Template, auto GatherMatchingSubObjectsRecurse) -> void
 	{
 		InstancingGraph.AddNewObject(Instance, Template);
 
-		ForEachObjectWithOuter(Instance, [OldInstance, Template, GatherMatchingSubObjectsRecurse](UObject* SubObject)
+		ForEachObjectWithOuter(Instance, [OldInstance, Template, GatherMatchingSubObjectsRecurse, OldToNewInstanceMap](UObject* SubObject)
 		{
 			// Skip the object that is about to be re-instantiated
 			if (SubObject == OldInstance)
@@ -3336,9 +3336,17 @@ bool FBlueprintCompileReinstancer::PrePopulateInstancingGraphForArchetype(FObjec
 			// Sanitization of the archetype
 			if (SubObjectArchetype->GetOuter() != Template)
 			{
-				UE_LOG(LogClass, Warning, TEXT("Missmatched while trying to rebuild instancing graph for '%s', expecting the outer of archetype %s to be %s but was %s"), 
+				// Handle the fact that subobject will be re-instantiated later by pointing its archetype to the new one.
+				if (UObject** NewSubObjectArchetype = OldToNewInstanceMap ? OldToNewInstanceMap->Find(SubObjectArchetype) : nullptr)
+				{
+					SubObjectArchetype = *NewSubObjectArchetype;
+				}
+				else
+				{
+					UE_LOG(LogClass, VeryVerbose, TEXT("Missmatched while trying to rebuild instancing graph for '%s', expecting the outer of archetype %s to be %s but was %s"), 
 						*GetPathNameSafe(OldInstance), *GetPathNameSafe(SubObjectArchetype), *GetPathNameSafe(Template), *GetPathNameSafe(SubObjectArchetype->GetOuter()));
-				return;
+					return;
+				}
 			}
 
 			GatherMatchingSubObjectsRecurse(SubObject, SubObjectArchetype, GatherMatchingSubObjectsRecurse);
