Cooker: PackageWriter: Add UE::PackageWriter::ICookerInterface::WriteFileOnCookDirector for the cooker to handle replicating UObject::CookAdditionalFilesOverride files to a synchronized write queue on cookdirector when in MPCook.
#rnx

diff --git a/Engine/Source/Developer/IoStoreUtilities/Private/ZenStoreWriter.cpp b/Engine/Source/Developer/IoStoreUtilities/Private/ZenStoreWriter.cpp
--- a/Engine/Source/Developer/IoStoreUtilities/Private/ZenStoreWriter.cpp
+++ b/Engine/Source/Developer/IoStoreUtilities/Private/ZenStoreWriter.cpp
@@ -386,48 +386,52 @@ void FZenStoreWriter::WriteAdditionalFile(const FAdditionalFileInfo& Info, const
 	FPendingPackageState& ExistingState = GetPendingPackage(Info.PackageName);
 
 	FFileDataEntry& FileEntry = ExistingState.FileData.AddDefaulted_GetRef();
-	
-	FileData.EnsureOwned();
 
-	auto WriteToFile = [](const FString& Filename, const FIoBuffer& FileData)
+	// TODO: Change FileData into FIoBuffer&&.
+	// SavePackage currently allows us to destroy not only the FileData structure but even the data it points to,
+	// but it would be nice to indicate that contract in the function call.
+	FSharedBuffer FileDataSharedBuffer = IoBufferToSharedBuffer(const_cast<FIoBuffer&>(FileData));
+
+#if !WITH_EDITOR
+	// ZenStoreWriter is only used on editor, and we've added some dependencies on editor-only functionality of
+	// SavePackage.
+	checkNoEntry();
+#else
+	UE::SavePackageUtilities::IncrementOutstandingAsyncWrites();
+	FileEntry.CompressedPayload = Async(EAsyncExecution::TaskGraph,
+		[this, FileDataSharedBuffer, Filename=Info.Filename]()
 	{
 		ON_SCOPE_EXIT
 		{
 			UE::SavePackageUtilities::DecrementOutstandingAsyncWrites();
 		};
 
-		IFileManager& FileManager = IFileManager::Get();
-		int64 DataSize = IntCastChecked<int64>(FileData.DataSize());
-
-		for (int Tries = 0; Tries < 3; ++Tries)
+		UE::PackageWriter::Private::FWriteFileData WriteFileData;
+		WriteFileData.Filename = Filename;
+		WriteFileData.Buffer = FCompositeBuffer(FileDataSharedBuffer);
+		WriteFileData.bIsSidecar = true;
+		WriteFileData.bContributeToHash = false;
+		// For robustness, we handle all AdditionalFiles reported by UObject::CookAdditionalFiles as if they might be
+		// written by multiple packages, without regard for their filename.
+		WriteFileData.bPackageSpecificFilename = false;
+
+		// WriteFileOnCookDirector or HashAndWrite will calculate the hash for us if we pass in
+		// EWriteOptions::WriteHash, but we need a different hash - the hash of the compressed file that we store
+		// in ZenServer. So omit WriteHash and just pass in unused placeholders for the hash output.
+		FMD5 UnusedHash;
+		TRefCountPtr<FPackageHashes> UnusedHashes;
+		if (Cooker)
 		{
-			if (FArchive* Ar = FileManager.CreateFileWriter(*Filename))
-			{
-				Ar->Serialize(const_cast<uint8*>(FileData.GetData()), DataSize);
-				bool bArchiveError = Ar->IsError();
-				delete Ar;
-
-				int64 ActualSize = FileManager.FileSize(*Filename);
-				if (ActualSize != DataSize)
-				{
-					FileManager.Delete(*Filename);
-
-					UE_LOG(LogZenStoreWriter, Fatal, TEXT("Could not save to %s! Tried to write %" INT64_FMT " bytes but resultant size was %" INT64_FMT ".%s"),
-						*Filename, DataSize, ActualSize, bArchiveError ? TEXT(" Ar->Serialize failed.") : TEXT(""));
-				}
-				return;
-			}
+			Cooker->WriteFileOnCookDirector(WriteFileData, UnusedHash, UnusedHashes, EWriteOptions::Write);
+		}
+		else
+		{
+			UE::PackageWriter::Private::HashAndWrite(WriteFileData, UnusedHash, UnusedHashes, EWriteOptions::Write);
 		}
 
-		UE_LOG(LogZenStoreWriter, Fatal, TEXT("Could not write to %s!"), *Filename);
-	};
-
-	UE::SavePackageUtilities::IncrementOutstandingAsyncWrites();
-	FileEntry.CompressedPayload = Async(EAsyncExecution::TaskGraph, [this, Info, FileData, WriteToFile]()
-	{
-		WriteToFile(Info.Filename, FileData);
-		return FCompressedBuffer::Compress(FSharedBuffer::MakeView(FileData.GetView()), Compressor, CompressionLevel);
+		return FCompressedBuffer::Compress(FileDataSharedBuffer, Compressor, CompressionLevel);
 	});
+#endif
 
 	const FZenFileSystemManifestEntry& ManifestEntry = ZenFileSystemManifest->CreateManifestEntry(Info.Filename);
 	FileEntry.Info					= Info;
diff --git a/Engine/Source/Editor/UnrealEd/Classes/CookOnTheSide/CookOnTheFlyServer.h b/Engine/Source/Editor/UnrealEd/Classes/CookOnTheSide/CookOnTheFlyServer.h
--- a/Engine/Source/Editor/UnrealEd/Classes/CookOnTheSide/CookOnTheFlyServer.h
+++ b/Engine/Source/Editor/UnrealEd/Classes/CookOnTheSide/CookOnTheFlyServer.h
@@ -707,6 +707,9 @@ public:
 	UNREALED_API virtual void RegisterDeterminismHelper(ICookedPackageWriter* PackageWriter, UObject* SourceObject,
 		const TRefCountPtr<UE::Cook::IDeterminismHelper>& DeterminismHelper) override;
 	UNREALED_API virtual bool IsDeterminismDebug() const override;
+	UNREALED_API virtual void WriteFileOnCookDirector(const UE::PackageWriter::Private::FWriteFileData& FileData,
+		FMD5& AccumulatedHash, const TRefCountPtr<FPackageHashes>& PackageHashes,
+		IPackageWriter::EWriteOptions WriteOptions) override;
 	// End UE::PackageWriter::Private::ICookerInterface
 
 	/** Dumps cooking stats to the log. Run from the exec command "Cook stats". */
diff --git a/Engine/Source/Editor/UnrealEd/Private/CookOnTheFlyServer.cpp b/Engine/Source/Editor/UnrealEd/Private/CookOnTheFlyServer.cpp
--- a/Engine/Source/Editor/UnrealEd/Private/CookOnTheFlyServer.cpp
+++ b/Engine/Source/Editor/UnrealEd/Private/CookOnTheFlyServer.cpp
@@ -6715,7 +6715,7 @@ FString UCookOnTheFlyServer::GetCookOutputFolder(const ITargetPlatform* TargetPl
 	{
 		return FString();
 	}
-	if (!PlatformManager->GetSessionPlatforms().Contains(TargetPlatform))
+	if (!PlatformManager->HasSessionPlatform(TargetPlatform))
 	{
 		return FString();
 	}
@@ -6802,6 +6802,67 @@ bool UCookOnTheFlyServer::IsDeterminismDebug() const
 	return bDeterminismDebug || DiffModeHelper->IsDeterminismDebug();
 }
 
+void UCookOnTheFlyServer::WriteFileOnCookDirector(const UE::PackageWriter::Private::FWriteFileData& FileData,
+	FMD5& AccumulatedHash, const TRefCountPtr<FPackageHashes>& PackageHashes,
+	IPackageWriter::EWriteOptions WriteOptions)
+{
+	// This function can be called from Async threads, so we can only call cooker functions that are threadsafe;
+	// many cooker functions are read/write only from scheduler thread.
+	if (CookWorkerClient || CookDirector)
+	{
+		FString CookRoot;
+		FStringView RelPath;
+		const ITargetPlatform* TargetPlatform = nullptr;
+		FString TargetFileNameAbsPath = FPaths::ConvertRelativePathToFull(FileData.Filename);
+
+		{
+			UE::Cook::FPlatformManager::FReadScopeLock PlatformScopeLock(PlatformManager->ReadLockPlatforms());
+			PlatformManager->EnumerateSessionPlatforms(
+				[&CookRoot, &RelPath, &TargetPlatform, &TargetFileNameAbsPath, this]
+				(const ITargetPlatform* IterTargetPlatform)
+				{
+					if (!TargetPlatform)
+					{
+						CookRoot = FPaths::ConvertRelativePathToFull(GetCookOutputFolder(IterTargetPlatform));
+						if (FPathViews::TryMakeChildPathRelativeTo(TargetFileNameAbsPath, CookRoot, RelPath))
+						{
+							TargetPlatform = IterTargetPlatform;
+						}
+					}
+				});
+		}
+
+		if (!TargetPlatform)
+		{
+			UE_LOG(LogCook, Error,
+				TEXT("Output file requested during cooking from CookAdditionalFilesOverride was invalid. ")
+				TEXT("Output files must be under the platform's cooked root output, such as %s. ")
+				TEXT("File %s is not under the cooked root for any platform being cooked. The file will be missing from the cooked output."),
+				!CookRoot.IsEmpty() ? *CookRoot : TEXT("<UnknownBecauseNoSessionPlatforms>"),
+				*FileData.Filename);
+			return;
+		}
+		FString FileTransferRoot = FPaths::Combine(CookRoot, UE::Cook::MPCookTransferRootRelPath);
+		FString TempFileName = FPaths::CreateTempFilename(*FileTransferRoot);
+
+		UE::PackageWriter::Private::FWriteFileData TempFileData = FileData;
+		TempFileData.Filename = TempFileName;
+		UE::PackageWriter::Private::HashAndWrite(TempFileData, AccumulatedHash, PackageHashes, WriteOptions);
+		if (CookWorkerClient)
+		{
+			CookWorkerClient->ReportFileTransfer(TempFileName, FileData.Filename);
+		}
+		else
+		{
+			CookDirector->ReportFileTransferFromLocalWorker(TempFileName, FileData.Filename);
+		}
+	}
+	else
+	{
+		UE::PackageWriter::Private::HashAndWrite(FileData, AccumulatedHash, PackageHashes, WriteOptions);
+	}
+}
+
 void UCookOnTheFlyServer::DumpStats()
 {
 	UE_LOG(LogCook, Display, TEXT("IntStats:"));
diff --git a/Engine/Source/Editor/UnrealEd/Private/Cooker/CookDirector.cpp b/Engine/Source/Editor/UnrealEd/Private/Cooker/CookDirector.cpp
--- a/Engine/Source/Editor/UnrealEd/Private/Cooker/CookDirector.cpp
+++ b/Engine/Source/Editor/UnrealEd/Private/Cooker/CookDirector.cpp
@@ -1077,6 +1077,17 @@ void FCookDirector::ShutdownCookSession()
 	// Kill any connections that had just been made and not yet assigned to a Server
 	PendingConnections.Reset();
 
+	// Delete temp directories
+	{
+		FScopeLock FileTransferScopeLock(&FileTransferLock);
+		for (const ITargetPlatform* TargetPlatform : COTFS.PlatformManager->GetSessionPlatforms())
+		{
+			FString FileTransferRoot = FPaths::Combine(COTFS.GetCookOutputFolder(TargetPlatform),
+				UE::Cook::MPCookTransferRootRelPath);
+			IFileManager::Get().DeleteDirectory(*FileTransferRoot, false /* bRequireExists */, true /* bTree */);
+		}
+	}
+
 	// Restore the FCookDirector to its original state so that it is ready for a new session
 	bWorkersInitialized = false;
 	bIsFirstAssignment = true;
@@ -1189,6 +1200,42 @@ void FCookDirector::HandleHeartbeatMessage(FMPCollectorServerMessageContext& Con
 	Context.GetCookWorkerServer()->SetLastReceivedHeartbeatNumberInLock(Message.HeartbeatNumber);
 }
 
+void FCookDirector::HandleFileTransferMessage(FCookWorkerServer& FromWorker, FFileTransferMessage&& Message)
+{
+	MoveFileSynchronized(*Message.TempFileName, *Message.TargetFileName, FromWorker.GetWorkerId(),
+		FromWorker.GetProfileId());
+
+}
+
+void FCookDirector::ReportFileTransferFromLocalWorker(FStringView TempFileName, FStringView TargetFileName)
+{
+	FScopeLock FileTransferScopeLock(&FileTransferLock);
+	TStringBuilder<256> TempFileNameStr(InPlace, TempFileName);
+	TStringBuilder<256> TargetFileNameStr(InPlace, TargetFileName);
+	MoveFileSynchronized(*TempFileNameStr, *TargetFileNameStr, FWorkerId::Local(), (uint32)-1);
+}
+
+void FCookDirector::MoveFileSynchronized(const TCHAR* TempFileName, const TCHAR* TargetFileName,
+	FWorkerId WorkerId, uint32 WorkerProfileId)
+{
+	// This function can called within the lock of a CookWorkerServer, and we're blocking on a critical section.
+	// Currently the work done is a little bit of IO to rename files within the Saved\Cooked directory.
+	// If it becomes more contended or does more work, we could create a new thread or use AsyncTasks to do the work;
+	// we would need to make sure that Director finalization includes waiting on the thread or Tasks,
+	// and if using multiple tasks, they would need to synchronize to avoid trying to write to the same TargetFile at
+	// once.
+	FScopeLock FileTransferScopeLock(&FileTransferLock);
+	bool bMoveSucceeded = IFileManager::Get().Move(TargetFileName, TempFileName, true /* bReplace */,
+		false /* EvenIfReadOnly */);
+	if (!bMoveSucceeded)
+	{
+		UE_LOG(LogCook, Error,
+			TEXT("Could not move UObject::CookAdditionalFilesOverride file transferred from CookWorker %s from %s -> %s. ")
+			TEXT("The file will be missing from cook output."),
+			WorkerId.IsLocal() ? TEXT("Local") : *FString::Printf(TEXT("%u"), WorkerProfileId),
+			TempFileName, TargetFileName);
+	}
+}
 
 FCookDirector::FPendingConnection::FPendingConnection(FPendingConnection&& Other)
 {
diff --git a/Engine/Source/Editor/UnrealEd/Private/Cooker/CookDirector.h b/Engine/Source/Editor/UnrealEd/Private/Cooker/CookDirector.h
--- a/Engine/Source/Editor/UnrealEd/Private/Cooker/CookDirector.h
+++ b/Engine/Source/Editor/UnrealEd/Private/Cooker/CookDirector.h
@@ -31,6 +31,7 @@ namespace UE::Cook { class FCookWorkerServer; }
 namespace UE::Cook { struct FAssignPackageExtraData; }
 namespace UE::Cook { struct FCookWorkerProfileData; }
 namespace UE::Cook { struct FDirectorEventMessage; }
+namespace UE::Cook { struct FFileTransferMessage; }
 namespace UE::Cook { struct FGeneratorEventMessage; }
 namespace UE::Cook { struct FGenerationHelper; }
 namespace UE::Cook { struct FHeartbeatMessage; }
@@ -150,6 +151,13 @@ public:
 	/** The message CookWorkerServer sends to the remote process once it is ready to connect. */
 	const FInitialConfigMessage& GetInitialConfigMessage();
 
+	/**
+	 * This function is called from the Director's local Worker. It is for FileTransfers that have to be synchronized
+	 * to avoid file write collisions, across all CookWorkers including the Director's local Worker.
+	 * Messages from Remote workers instead go through HandleFileTransferMessage.
+	 */
+	void ReportFileTransferFromLocalWorker(FStringView TempFileName, FStringView TargetFileName);
+
 private:
 	enum class ELoadBalanceAlgorithm
 	{
@@ -231,6 +239,9 @@ private:
 	/** Log the occurrence of a heartbeat message from a CookWorker. */
 	void HandleHeartbeatMessage(FMPCollectorServerMessageContext& Context, bool bReadSuccessful,
 		FHeartbeatMessage&& Message);
+	void HandleFileTransferMessage(FCookWorkerServer& FromWorker, FFileTransferMessage&& Message);
+	void MoveFileSynchronized(const TCHAR* TempFileName, const TCHAR* TargetFileName,
+		FWorkerId WorkerId, uint32 WorkerProfileId);
 
 	/** Move the given worker from active workers to the list of workers shutting down. */
 	void AbortWorker(FWorkerId WorkerId, ECookDirectorThread TickThread);
@@ -312,6 +323,9 @@ private:
 	TUniquePtr<FRetractionHandler> RetractionHandler;
 	bool bWorkersActive = false;
 
+	// FileTransfers (they have their own critical section to synchronize IO from multiple CookWorkers)
+	mutable FCriticalSection FileTransferLock;
+
 	friend class UE::Cook::FCookWorkerServer;
 };
 
diff --git a/Engine/Source/Editor/UnrealEd/Private/Cooker/CookFunctionLibrary.cpp b/Engine/Source/Editor/UnrealEd/Private/Cooker/CookFunctionLibrary.cpp
--- a/Engine/Source/Editor/UnrealEd/Private/Cooker/CookFunctionLibrary.cpp
+++ b/Engine/Source/Editor/UnrealEd/Private/Cooker/CookFunctionLibrary.cpp
@@ -14,6 +14,7 @@
 #include "Interfaces/ITargetPlatform.h"
 #include "Interfaces/ITargetPlatformManagerModule.h"
 #include "LooseFilesCookArtifactReader.h"
+#include "Serialization/BasePackageWriter.h"
 #include "Serialization/CompactBinaryPackage.h"
 #include "Serialization/CompactBinaryWriter.h"
 #include "Subsystems/UnrealEditorSubsystem.h"
@@ -159,6 +160,13 @@ void UCookFunctionLibrary::CookAsset(UObject* Object, const FString& ForPlatform
 		{
 			return false;
 		}
+		virtual void WriteFileOnCookDirector(const UE::PackageWriter::Private::FWriteFileData& FileData,
+			FMD5& AccumulatedHash, const TRefCountPtr<FPackageHashes>& PackageHashes,
+			IPackageWriter::EWriteOptions WriteOptions) override
+		{
+			UE::PackageWriter::Private::HashAndWrite(FileData, AccumulatedHash, PackageHashes, WriteOptions);
+		}
+
 	};
 	FCookerInterface CookerInterface;
 
diff --git a/Engine/Source/Editor/UnrealEd/Private/Cooker/CookPlatformManager.h b/Engine/Source/Editor/UnrealEd/Private/Cooker/CookPlatformManager.h
--- a/Engine/Source/Editor/UnrealEd/Private/Cooker/CookPlatformManager.h
+++ b/Engine/Source/Editor/UnrealEd/Private/Cooker/CookPlatformManager.h
@@ -18,9 +18,7 @@ class ITargetPlatform;
 class UCookOnTheFlyServer;
 namespace UE::Cook { class IAssetRegistryReporter; }
 
-namespace UE
-{
-namespace Cook
+namespace UE::Cook
 {
 
 	/*
@@ -101,6 +99,12 @@ namespace Cook
 
 		/** Returns the TargetPlatforms that are active for the CurrentCookByTheBook session or CookOnTheFly request. */
 		const TArray<const ITargetPlatform*>& GetSessionPlatforms() const;
+		/**
+		 * Calls the Callback with each platform in GetSessionPlatforms. This is the accessor that must be used when
+		 * trying to read the SessionPlatforms from an async thread in an FPlatformManager::FReadScopeLock.
+		 */
+		template <typename CallbackType>
+		void EnumerateSessionPlatforms(CallbackType&& Callback) const;
 		int32 GetNumSessionPlatforms() const;
 
 		/** Return whether the given platform is already in the list of platforms for the current session. */
@@ -233,8 +237,13 @@ namespace Cook
 		bool bHasSelectedSessionPlatforms = false;
 	};
 
-}
-}
+} // namespace UE::Cook
+
+
+///////////////////////////////////////////////////////
+// Inline implementations
+///////////////////////////////////////////////////////
+
 
 template <typename Value>
 void RemapMapKeys(TMap<const ITargetPlatform*, Value>& Map, const TMap<ITargetPlatform*,
@@ -257,3 +266,20 @@ inline void RemapArrayElements(TArray<const ITargetPlatform*>& Array, const TMap
 		Old = Remap[Old];
 	}
 }
+
+
+namespace UE::Cook
+{
+
+template <typename CallbackType>
+void FPlatformManager::EnumerateSessionPlatforms(CallbackType&& Callback) const
+{
+	checkf(IsSchedulerThread() || IsInPlatformsLock(), TEXT("Access to SessionPlatforms is only legal on non-scheduler threads when inside a ReadLockPlatforms scope."));
+	checkf(bHasSelectedSessionPlatforms, TEXT("Calling GetSessionPlatforms or (any of the top level cook functions that call it) without first calling SelectSessionPlatforms is invalid"));
+	for (const ITargetPlatform* TargetPlatform : SessionPlatforms)
+	{
+		Callback(TargetPlatform);
+	}
+}
+
+} // namespace UE::Cook
diff --git a/Engine/Source/Editor/UnrealEd/Private/Cooker/CookTypes.h b/Engine/Source/Editor/UnrealEd/Private/Cooker/CookTypes.h
--- a/Engine/Source/Editor/UnrealEd/Private/Cooker/CookTypes.h
+++ b/Engine/Source/Editor/UnrealEd/Private/Cooker/CookTypes.h
@@ -480,6 +480,7 @@ namespace UE::Cook
 
 	extern FGuid CookIncrementalVersion;
 
+	constexpr FStringView MPCookTransferRootRelPath(TEXTVIEW("MPCookTransfer"));
 } // namespace UE::Cook
 
 bool LexTryParseString(FPlatformMemoryStats::EMemoryPressureStatus& OutValue, FStringView Text);
diff --git a/Engine/Source/Editor/UnrealEd/Private/Cooker/CookWorkerClient.cpp b/Engine/Source/Editor/UnrealEd/Private/Cooker/CookWorkerClient.cpp
--- a/Engine/Source/Editor/UnrealEd/Private/Cooker/CookWorkerClient.cpp
+++ b/Engine/Source/Editor/UnrealEd/Private/Cooker/CookWorkerClient.cpp
@@ -339,6 +339,14 @@ void FCookWorkerClient::ReportLogMessage(const FReplicatedLogData& LogData)
 	LogMessageHandler->ClientReportLogMessage(LogData);
 }
 
+void FCookWorkerClient::ReportFileTransfer(FStringView TempFileName, FStringView TargetFileName)
+{
+	FScopeLock FileTransferScopeLock(&FileTransferLock);
+	FFileTransferMessage& Message = PendingFileTransfers.Emplace_GetRef();
+	Message.TempFileName = TempFileName;
+	Message.TargetFileName = TargetFileName;
+}
+
 void FCookWorkerClient::ReportGeneratorQueuedGeneratedPackages(FGenerationHelper& GenerationHelper)
 {
 	PendingGeneratorEvents.Add(FGeneratorEventMessage(EGeneratorEvent::QueuedGeneratedPackages,
@@ -679,6 +687,16 @@ void FCookWorkerClient::SendPendingResults()
 		}
 		PendingGeneratorEvents.Reset();
 	}
+
+	TArray<FFileTransferMessage> LocalPendingFileTransfers;
+	{
+		FScopeLock FileTransferScopeLock(&FileTransferLock);
+		Swap(LocalPendingFileTransfers, PendingFileTransfers);
+	}
+	for (FFileTransferMessage& FileTransferMessage : LocalPendingFileTransfers)
+	{
+		SendMessage(FileTransferMessage);
+	}
 }
 
 void FCookWorkerClient::PumpReceiveMessages()
diff --git a/Engine/Source/Editor/UnrealEd/Private/Cooker/CookWorkerClient.h b/Engine/Source/Editor/UnrealEd/Private/Cooker/CookWorkerClient.h
--- a/Engine/Source/Editor/UnrealEd/Private/Cooker/CookWorkerClient.h
+++ b/Engine/Source/Editor/UnrealEd/Private/Cooker/CookWorkerClient.h
@@ -20,6 +20,7 @@ namespace UE::Cook { struct FAssignPackagesMessage; }
 namespace UE::Cook { struct FDirectorConnectionInfo; }
 namespace UE::Cook { struct FDirectorEventMessage; }
 namespace UE::Cook { struct FDiscoveredPackageReplication; }
+namespace UE::Cook { struct FFileTransferMessage; }
 namespace UE::Cook { struct FHeartbeatMessage; }
 namespace UE::Cook { struct FInitialConfigMessage; }
 namespace UE::Cook { struct FGenerationHelper; }
@@ -72,6 +73,7 @@ public:
 	void ReportDiscoveredPackage(const FPackageData& PackageData, const FInstigator& Instigator,
 		FDiscoveredPlatformSet&& ReachablePlatforms, FGenerationHelper* ParentGenerationHelper, EUrgency Urgency);
 	void ReportLogMessage(const FReplicatedLogData& LogData);
+	void ReportFileTransfer(FStringView TempFileName, FStringView TargetFileName);
 
 	/**
 	 * Queue a message to the server that a PackageGenerator queued its generated packages for cooking, and will
@@ -184,6 +186,10 @@ private:
 	FCriticalSection PendingResultsLock;
 	TArray<TUniquePtr<FPackageRemoteResult>> PendingResults;
 	TMap<FPackageRemoteResult*, FPendingResultNeedingAsyncWork> PendingResultsNeedingAsyncWork;
+
+	// Variables Read/Write only within FileTransferLock
+	FCriticalSection FileTransferLock;
+	TArray<FFileTransferMessage> PendingFileTransfers;
 };
 
 }
\ No newline at end of file
diff --git a/Engine/Source/Editor/UnrealEd/Private/Cooker/CookWorkerServer.cpp b/Engine/Source/Editor/UnrealEd/Private/Cooker/CookWorkerServer.cpp
--- a/Engine/Source/Editor/UnrealEd/Private/Cooker/CookWorkerServer.cpp
+++ b/Engine/Source/Editor/UnrealEd/Private/Cooker/CookWorkerServer.cpp
@@ -874,6 +874,18 @@ void FCookWorkerServer::HandleReceiveMessagesInternal()
 				HandleGeneratorMessage(GeneratorMessage);
 			}
 		}
+		else if (Message.MessageType == FFileTransferMessage::MessageType)
+		{
+			FFileTransferMessage FileTransferMessage;
+			if (!FileTransferMessage.TryRead(Message.Object))
+			{
+				LogInvalidMessage(TEXT("FFileTransferMessage"));
+			}
+			else
+			{
+				Director.HandleFileTransferMessage(*this, MoveTemp(FileTransferMessage));
+			}
+		}
 		else
 		{
 			TRefCountPtr<IMPCollector>* Collector = Director.Collectors.Find(Message.MessageType);
@@ -1996,4 +2008,20 @@ void FPackageWriterMPCollector::ServerReceiveMessage(FMPCollectorServerMessageCo
 
 FGuid FPackageWriterMPCollector::MessageType(TEXT("D2B1CE3FD26644AF9EC28FBADB1BD331"));
 
-}
\ No newline at end of file
+void FFileTransferMessage::Write(FCbWriter& Writer) const
+{
+	Writer << "S" << TempFileName;
+	Writer << "T" << TargetFileName;
+}
+
+bool FFileTransferMessage::TryRead(FCbObjectView Object)
+{
+	bool bOk = true;
+	bOk = LoadFromCompactBinary(Object["S"], TempFileName) & bOk;
+	bOk = LoadFromCompactBinary(Object["T"], TargetFileName) & bOk;
+	return bOk;
+}
+
+FGuid FFileTransferMessage::MessageType(0x49205105, 0x9D62427F, 0xAA3F2F7C, 0x118152E0);
+
+} // namespace UE::Cook
diff --git a/Engine/Source/Editor/UnrealEd/Private/Cooker/CookWorkerServer.h b/Engine/Source/Editor/UnrealEd/Private/Cooker/CookWorkerServer.h
--- a/Engine/Source/Editor/UnrealEd/Private/Cooker/CookWorkerServer.h
+++ b/Engine/Source/Editor/UnrealEd/Private/Cooker/CookWorkerServer.h
@@ -552,4 +552,24 @@ private:
 	static FGuid MessageType;
 };
 
+/**
+ * Message from CookWorker requesting that a file written from the CookWorker be moved into its target path
+ * in the cook output directory on the CookDirector. Used for extra files specified by
+ * UObject::CookAdditionalFilesOverride. The same file might be specified by multiple packages, so even though
+ * CookWorkers and CookDirector are on the same machine, we need to have them moved to final location by the
+ * CookDirector to avoid the possibility of multiple CookWorkers trying to write the same file at once.
+ */
+struct FFileTransferMessage : public IMPCollectorMessage
+{
+	virtual void Write(FCbWriter& Writer) const override;
+	virtual bool TryRead(FCbObjectView Object) override;
+	virtual FGuid GetMessageType() const override { return MessageType; }
+	virtual const TCHAR* GetDebugName() const override { return TEXT("FileTransferMessage"); }
+
+public:
+	FString TempFileName;
+	FString TargetFileName;
+	static FGuid MessageType;
+};
+
 }
\ No newline at end of file
diff --git a/Engine/Source/Editor/UnrealEd/Private/Cooker/LooseCookedPackageWriter.cpp b/Engine/Source/Editor/UnrealEd/Private/Cooker/LooseCookedPackageWriter.cpp
--- a/Engine/Source/Editor/UnrealEd/Private/Cooker/LooseCookedPackageWriter.cpp
+++ b/Engine/Source/Editor/UnrealEd/Private/Cooker/LooseCookedPackageWriter.cpp
@@ -56,7 +56,11 @@ void FLooseCookedPackageWriter::CommitPackageInternal(FPackageWriterRecords::FPa
 {
 	{
 		FFilePackageWriterUtil::FRecord& Record = static_cast<FFilePackageWriterUtil::FRecord&>(BaseRecord);
-		FFilePackageWriterUtil::FWritePackageParameters Parameters(Record, Info, &AllPackageHashes, &PackageHashesLock, bProvidePerPackageResults);
+		FFilePackageWriterUtil::FWritePackageParameters Parameters(Record, Info);
+		Parameters.Cooker = Cooker;
+		Parameters.AllPackageHashes = &AllPackageHashes;
+		Parameters.PackageHashesLock = &PackageHashesLock;
+		Parameters.bProvidePerPackageResult = bProvidePerPackageResults;
 		FFilePackageWriterUtil::WritePackage(Parameters);
 	}
 
diff --git a/Engine/Source/Editor/UnrealEd/Private/Serialization/DefaultCookedFilePackageWriter.cpp b/Engine/Source/Editor/UnrealEd/Private/Serialization/DefaultCookedFilePackageWriter.cpp
--- a/Engine/Source/Editor/UnrealEd/Private/Serialization/DefaultCookedFilePackageWriter.cpp
+++ b/Engine/Source/Editor/UnrealEd/Private/Serialization/DefaultCookedFilePackageWriter.cpp
@@ -36,7 +36,7 @@ void FDefaultCookedFilePackageWriter::WriteBulkData(const IPackageWriter::FBulkD
 void FDefaultCookedFilePackageWriter::CommitPackageInternal(FPackageWriterRecords::FPackage&& BaseRecord, const FCommitPackageInfo& Info)
 {
 	FFilePackageWriterUtil::FRecord& Record = static_cast<FFilePackageWriterUtil::FRecord&>(BaseRecord);
-	FFilePackageWriterUtil::FWritePackageParameters Parameters(Record, Info, nullptr, nullptr, false);
+	FFilePackageWriterUtil::FWritePackageParameters Parameters(Record, Info);
 	FFilePackageWriterUtil::WritePackage(Parameters);
 }
 
diff --git a/Engine/Source/Editor/UnrealEd/Private/Serialization/FilePackageWriterUtil.cpp b/Engine/Source/Editor/UnrealEd/Private/Serialization/FilePackageWriterUtil.cpp
--- a/Engine/Source/Editor/UnrealEd/Private/Serialization/FilePackageWriterUtil.cpp
+++ b/Engine/Source/Editor/UnrealEd/Private/Serialization/FilePackageWriterUtil.cpp
@@ -13,12 +13,19 @@ FFilePackageWriterUtil::FWritePackageParameters::FWritePackageParameters(FRecord
 	TMap<FName, TRefCountPtr<FPackageHashes>>* InAllPackageHashes,
 	FCriticalSection* InPackageHashesLock,
 	bool bInProvidePerPackageResult)
+	: FWritePackageParameters(InRecord, InInfo)
+{
+	AllPackageHashes = InAllPackageHashes;
+	PackageHashesLock = InPackageHashesLock;
+	bProvidePerPackageResult = bInProvidePerPackageResult;
+}
+
+FFilePackageWriterUtil::FWritePackageParameters::FWritePackageParameters(FRecord& InRecord,
+	const IPackageWriter::FCommitPackageInfo& InInfo)
 	: Record(InRecord)
 	, Info(InInfo)
-	, AllPackageHashes(InAllPackageHashes)
-	, PackageHashesLock(InPackageHashesLock)
-	, bProvidePerPackageResult(bInProvidePerPackageResult)
-{}
+{
+}
 
 void FFilePackageWriterUtil::WritePackage(FWritePackageParameters& Parameters)
 {
@@ -31,6 +38,7 @@ void FFilePackageWriterUtil::WritePackage(FWritePackageParameters& Parameters)
 void FFilePackageWriterUtil::AsyncSave(FWritePackageParameters& Parameters)
 {
 	FCommitContext Context{ Parameters.Info };
+	Context.Cooker = Parameters.Cooker;
 
 	// The order of these collection calls is important, both for ExportsBuffers (affects the meaning of offsets
 	// to those buffers) and for OutputFiles (affects the calculation of the Hash for the set of PackageData)
@@ -77,6 +85,7 @@ void FFilePackageWriterUtil::CollectForSaveBulkData(FRecord& Record, FCommitCont
 			OutputFile.bIsSidecar = true;
 			// Only calculate the CookHash for files in the main package output. Ignore it for the Optional package output.
 			OutputFile.bContributeToHash = BulkRecord.Info.MultiOutputIndex == 0;
+			OutputFile.bPackageSpecificFilename = true;
 			OutputFile.ChunkId = BulkRecord.Info.ChunkId;
 		}
 	}
@@ -106,6 +115,9 @@ void FFilePackageWriterUtil::CollectForSaveAdditionalFileRecords(FRecord& Record
 		OutputFile.bIsSidecar = true;
 		// Only calculate the CookHash for files in the main package output. Ignore it for the Optional package output.
 		OutputFile.bContributeToHash = AdditionalRecord.Info.MultiOutputIndex == 0;
+		// For robustness, we handle all AdditionalFiles reported by UObject::CookAdditionalFiles as if they might be
+		// written by multiple packages, without regard for their filename.
+		OutputFile.bPackageSpecificFilename = false;
 		OutputFile.ChunkId = AdditionalRecord.Info.ChunkId;
 	}
 }
@@ -164,6 +176,7 @@ void FFilePackageWriterUtil::CollectForSaveExportsBuffers(FRecord& Record, FComm
 			OutputFile.bIsSidecar = false;
 			// Only calculate the CookHash for files in the main package output. Ignore it for the Optional package output.
 			OutputFile.bContributeToHash = Package.Info.MultiOutputIndex == 0;
+			OutputFile.bPackageSpecificFilename = true;
 		}
 
 		// Exports + AllAppendedData (.uexp)
@@ -173,6 +186,7 @@ void FFilePackageWriterUtil::CollectForSaveExportsBuffers(FRecord& Record, FComm
 			OutputFile.bIsSidecar = false;
 			// Only calculate the CookHash for files in the main package output. Ignore it for the Optional package output.
 			OutputFile.bContributeToHash = Package.Info.MultiOutputIndex == 0;
+			OutputFile.bPackageSpecificFilename = true;
 
 			int32 NumBuffers = ExportsBuffers.Num();
 			TArray<FSharedBuffer> BuffersForComposition;
@@ -256,14 +270,22 @@ void FFilePackageWriterUtil::AsyncSaveOutputFiles(FCommitContext& Context,
 	UE::Tasks::Launch(TEXT("HashAndWriteCookedFile"),
 		[OutputFiles = MoveTemp(Context.OutputFiles), WriteOptions = Context.Info.WriteOptions,
 		ThisPackageHashes = MoveTemp(ThisPackageHashes),
-		PackageHashesCompletionPromise = MoveTemp(PackageHashesCompletionPromise)]
+		PackageHashesCompletionPromise = MoveTemp(PackageHashesCompletionPromise),
+		Cooker = Context.Cooker]
 		() mutable
 		{
 			FMD5 AccumulatedHash;
 			for (const UE::PackageWriter::Private::FWriteFileData& OutputFile : OutputFiles)
 			{
-				UE::PackageWriter::Private::HashAndWrite(OutputFile, AccumulatedHash, ThisPackageHashes,
-					WriteOptions);
+				if (!OutputFile.bPackageSpecificFilename && Cooker)
+				{
+					Cooker->WriteFileOnCookDirector(OutputFile, AccumulatedHash, ThisPackageHashes, WriteOptions);
+				}
+				else
+				{
+					UE::PackageWriter::Private::HashAndWrite(
+						OutputFile, AccumulatedHash, ThisPackageHashes, WriteOptions);
+				}
 			}
 
 			if (EnumHasAnyFlags(WriteOptions, IPackageWriter::EWriteOptions::ComputeHash))
diff --git a/Engine/Source/Editor/UnrealEd/Public/Serialization/FilePackageWriterUtil.h b/Engine/Source/Editor/UnrealEd/Public/Serialization/FilePackageWriterUtil.h
--- a/Engine/Source/Editor/UnrealEd/Public/Serialization/FilePackageWriterUtil.h
+++ b/Engine/Source/Editor/UnrealEd/Public/Serialization/FilePackageWriterUtil.h
@@ -24,15 +24,20 @@ public:
 	{
 		FRecord& Record;
 		const IPackageWriter::FCommitPackageInfo& Info;
-		TMap<FName, TRefCountPtr<FPackageHashes>>* AllPackageHashes;
-		FCriticalSection* PackageHashesLock;
-		bool bProvidePerPackageResult;
+		UE::PackageWriter::Private::ICookerInterface* Cooker = nullptr;
+		TMap<FName, TRefCountPtr<FPackageHashes>>* AllPackageHashes = nullptr;
+		FCriticalSection* PackageHashesLock = nullptr;
+		bool bProvidePerPackageResult = false;
 
+		UE_DEPRECATED(5.7, "Use constructor that takes only the non-null parameters, and set other parameters individually.")
 		FWritePackageParameters(FRecord& InRecord,
 			const IPackageWriter::FCommitPackageInfo& InInfo,
 			TMap<FName, TRefCountPtr<FPackageHashes>>* InAllPackageHashes,
 			FCriticalSection* InPackageHashesLock,
 			bool bInProvidePerPackageResult);
+
+		FWritePackageParameters(FRecord& InRecord,
+			const IPackageWriter::FCommitPackageInfo& InInfo);
 	};
 
 	static void WritePackage(FWritePackageParameters& Parameters);
@@ -52,6 +57,7 @@ private:
 		const IPackageWriter::FCommitPackageInfo& Info;
 		TArray<TArray<FExportBuffer>> ExportsBuffers;
 		TArray<UE::PackageWriter::Private::FWriteFileData> OutputFiles;
+		UE::PackageWriter::Private::ICookerInterface* Cooker = nullptr;
 	};
 
 	static void AsyncSave(FWritePackageParameters& Parameters);
diff --git a/Engine/Source/Runtime/Core/Public/Serialization/PackageWriter.h b/Engine/Source/Runtime/Core/Public/Serialization/PackageWriter.h
--- a/Engine/Source/Runtime/Core/Public/Serialization/PackageWriter.h
+++ b/Engine/Source/Runtime/Core/Public/Serialization/PackageWriter.h
@@ -281,6 +281,12 @@ struct FWriteFileData
 	 * file.
 	 */
 	bool bContributeToHash = true;
+	/**
+	 * True only if the file is known to have a filename specific to the package saving the file, so that two
+	 * packages cannot both try to write the same file. If it is true we can write the file directly on CookWorkers
+	 * during MPCook without the possibility of two CookWorkers writing the file at the same time.
+	 */
+	bool bPackageSpecificFilename = false;
 	FIoChunkId ChunkId = FIoChunkId::InvalidChunkId;
 };
 
@@ -296,6 +302,8 @@ public:
 	virtual void RegisterDeterminismHelper(ICookedPackageWriter* PackageWriter, UObject* SourceObject,
 		const TRefCountPtr<UE::Cook::IDeterminismHelper>& DeterminismHelper) = 0;
 	virtual bool IsDeterminismDebug() const = 0;
+	virtual void WriteFileOnCookDirector(const FWriteFileData& FileData, FMD5& AccumulatedHash,
+		const TRefCountPtr<FPackageHashes>& PackageHashes, IPackageWriter::EWriteOptions WriteOptions) = 0;
 
 protected:
 	~ICookerInterface() = default;
