IncrementalValidate: Fix incorrect value for the number of Modified packages reported in the log. It was counting uncooked external actor packages, when all other values for package counts skip uncooked packages.
#rnx
#cookincremental

diff --git a/Engine/Source/Developer/IoStoreUtilities/Private/ZenStoreWriter.cpp b/Engine/Source/Developer/IoStoreUtilities/Private/ZenStoreWriter.cpp
--- a/Engine/Source/Developer/IoStoreUtilities/Private/ZenStoreWriter.cpp
+++ b/Engine/Source/Developer/IoStoreUtilities/Private/ZenStoreWriter.cpp
@@ -1371,12 +1371,11 @@ void FZenStoreWriter::RemoveCookedPackages()
 	PackageNameToIndex.Empty();
 }
 
-void FZenStoreWriter::UpdatePackageModificationStatus(FName PackageName, bool bIncrementallyUnmodified,
-	bool& bInOutShouldIncrementallySkip)
+void FZenStoreWriter::UpdatePackageModifiedStatus(FUpdatePackageModifiedStatusContext& Context)
 {
-	TRACE_CPUPROFILER_EVENT_SCOPE(FZenStoreWriter::UpdatePackageModificationStatus);
+	TRACE_CPUPROFILER_EVENT_SCOPE(FZenStoreWriter::UpdatePackageModifiedStatus);
 
-	if (!bIncrementallyUnmodified)
+	if (!Context.bIncrementallyUnmodified)
 	{
 		return;
 	}
@@ -1385,13 +1384,13 @@ void FZenStoreWriter::UpdatePackageModificationStatus(FName PackageName, bool bI
 
 	{
 		FWriteScopeLock _(EntriesLock);
-		int32* Index = PackageNameToIndex.Find(PackageName);
+		int32* Index = PackageNameToIndex.Find(Context.PackageName);
 		if (!Index)
 		{
-			if (!FPackageName::IsScriptPackage(WriteToString<128>(PackageName)))
+			if (!FPackageName::IsScriptPackage(WriteToString<128>(Context.PackageName)))
 			{
-				UE_LOG(LogZenStoreWriter, Verbose, TEXT("UpdatePackageModificationStatus called with package %s that is not in the oplog."),
-					*PackageName.ToString());
+				UE_LOG(LogZenStoreWriter, Verbose, TEXT("UpdatePackageModifiedStatus called with package %s that is not in the oplog."),
+					*Context.PackageName.ToString());
 			}
 			return;
 		}
diff --git a/Engine/Source/Developer/IoStoreUtilities/Public/ZenStoreWriter.h b/Engine/Source/Developer/IoStoreUtilities/Public/ZenStoreWriter.h
--- a/Engine/Source/Developer/IoStoreUtilities/Public/ZenStoreWriter.h
+++ b/Engine/Source/Developer/IoStoreUtilities/Public/ZenStoreWriter.h
@@ -163,8 +163,7 @@ public:
 
 	IOSTOREUTILITIES_API virtual void RemoveCookedPackages(TArrayView<const FName> PackageNamesToRemove) override;
 	IOSTOREUTILITIES_API virtual void RemoveCookedPackages() override;
-	IOSTOREUTILITIES_API virtual void UpdatePackageModificationStatus(FName PackageName, bool bIncrementallyUnmodified,
-		bool& bInOutShouldIncrementallySkip) override;
+	IOSTOREUTILITIES_API virtual void UpdatePackageModifiedStatus(FUpdatePackageModifiedStatusContext& Context) override;
 	IOSTOREUTILITIES_API bool GetPreviousCookedBytes(const FPackageInfo& Info, FPreviousCookedBytesData& OutData) override;
 	IOSTOREUTILITIES_API void CompleteExportsArchiveForDiff(FPackageInfo& Info, FLargeMemoryWriter& ExportsArchive) override;
 	IOSTOREUTILITIES_API virtual EPackageWriterResult BeginCacheForCookedPlatformData(FBeginCacheForCookedPlatformDataInfo& Info) override;
diff --git a/Engine/Source/Editor/UnrealEd/Private/CookOnTheFlyServer.cpp b/Engine/Source/Editor/UnrealEd/Private/CookOnTheFlyServer.cpp
--- a/Engine/Source/Editor/UnrealEd/Private/CookOnTheFlyServer.cpp
+++ b/Engine/Source/Editor/UnrealEd/Private/CookOnTheFlyServer.cpp
@@ -7873,8 +7873,13 @@ void UCookOnTheFlyServer::PopulateCookedPackages(TArrayView<const ITargetPlatfor
 						PackageData->FindOrAddPlatformData(TargetPlatform).SetIncrementallyUnmodified(true);
 					}
 					bool bShouldLegacyIterativeSkip = bLegacyIterativelyUnmodified;
-					PackageWriter.UpdatePackageModificationStatus(PackageName, bLegacyIterativelyUnmodified,
-						bShouldLegacyIterativeSkip);
+					ICookedPackageWriter::FUpdatePackageModifiedStatusContext Context;
+					Context.PackageName = PackageName;
+					Context.bIncrementallyUnmodified = bLegacyIterativelyUnmodified;
+					Context.bPreviouslyCooked = true;
+					Context.bInOutShouldIncrementallySkip = bShouldLegacyIterativeSkip;
+					PackageWriter.UpdatePackageModifiedStatus(Context);
+					bShouldLegacyIterativeSkip = Context.bInOutShouldIncrementallySkip;
 					if (bShouldLegacyIterativeSkip && !bLegacyIterativelyUnmodified)
 					{
 						// Override the PackageWriter's request to skip the modified generator package, because we
diff --git a/Engine/Source/Editor/UnrealEd/Private/Cooker/CookGenerationHelper.cpp b/Engine/Source/Editor/UnrealEd/Private/Cooker/CookGenerationHelper.cpp
--- a/Engine/Source/Editor/UnrealEd/Private/Cooker/CookGenerationHelper.cpp
+++ b/Engine/Source/Editor/UnrealEd/Private/Cooker/CookGenerationHelper.cpp
@@ -2773,8 +2773,13 @@ void FCookGenerationInfo::LegacyIterativeCookValidateOrClear(FGenerationHelper&
 	}
 	bool bShouldLegacyIterativelySkip = bOutLegacyIterativeUnmodified;
 	ICookedPackageWriter& PackageWriter = COTFS.FindOrCreatePackageWriter(TargetPlatform);
-	PackageWriter.UpdatePackageModificationStatus(PackageData->GetPackageName(), bOutLegacyIterativeUnmodified,
-		bShouldLegacyIterativelySkip);
+	ICookedPackageWriter::FUpdatePackageModifiedStatusContext Context;
+	Context.PackageName = PackageData->GetPackageName();
+	Context.bIncrementallyUnmodified = bOutLegacyIterativeUnmodified;
+	Context.bPreviouslyCooked = true;
+	Context.bInOutShouldIncrementallySkip = bShouldLegacyIterativelySkip;
+	PackageWriter.UpdatePackageModifiedStatus(Context);
+	bShouldLegacyIterativelySkip = Context.bInOutShouldIncrementallySkip;
 	if (bShouldLegacyIterativelySkip)
 	{
 		PackageData->SetPlatformCooked(TargetPlatform, ECookResult::Succeeded);
diff --git a/Engine/Source/Editor/UnrealEd/Private/Cooker/CookProfiling.cpp b/Engine/Source/Editor/UnrealEd/Private/Cooker/CookProfiling.cpp
--- a/Engine/Source/Editor/UnrealEd/Private/Cooker/CookProfiling.cpp
+++ b/Engine/Source/Editor/UnrealEd/Private/Cooker/CookProfiling.cpp
@@ -1160,7 +1160,8 @@ void OutputIncrementalCookStats()
 		UE_LOG(LogCookStats, Display, TEXT("\tTotal, %u"), NumPackagesIncrementallyModifiedByClassSum);
 		for (const TPair<FTopLevelAssetPath, uint32>& Pair : ClassAndNums)
 		{
-			UE_LOG(LogCookStats, Display, TEXT("\t%s, %u"), *Pair.Key.ToString(), Pair.Value);
+			UE_LOG(LogCookStats, Display, TEXT("\t%s, %u"),
+				Pair.Key.IsValid() ? *Pair.Key.ToString() : TEXT("<UnknownClass>"), Pair.Value);
 		}
 	}
 }
diff --git a/Engine/Source/Editor/UnrealEd/Private/Cooker/CookRequestCluster.cpp b/Engine/Source/Editor/UnrealEd/Private/Cooker/CookRequestCluster.cpp
--- a/Engine/Source/Editor/UnrealEd/Private/Cooker/CookRequestCluster.cpp
+++ b/Engine/Source/Editor/UnrealEd/Private/Cooker/CookRequestCluster.cpp
@@ -2170,6 +2170,9 @@ void FRequestCluster::FGraphSearch::FExploreEdgesContext::ProcessPlatformAttachm
 		ECookResult CookResult =
 			QueryPlatformData.CookAttachments.CommitStatus == IPackageWriter::ECommitStatus::Success
 			? ECookResult::Succeeded : ECookResult::Failed;
+		bool bPreviouslyCooked =
+			(QueryPlatformData.CookAttachments.CommitStatus == IPackageWriter::ECommitStatus::Success)
+			| (QueryPlatformData.CookAttachments.CommitStatus == IPackageWriter::ECommitStatus::Error);
 		if (QueryPlatformData.CookAttachments.CommitStatus == IPackageWriter::ECommitStatus::Error)
 		{
 			// Recook packages with errors even if they have not changed, so that the error is not hidden from
@@ -2192,8 +2195,38 @@ void FRequestCluster::FGraphSearch::FExploreEdgesContext::ProcessPlatformAttachm
 
 		if (Cluster.TraversalMarkCookable())
 		{
-			PackageWriter->UpdatePackageModificationStatus(PackageName, bIncrementallyUnmodified,
-				bShouldIncrementallySkip);
+			ICookedPackageWriter::FUpdatePackageModifiedStatusContext Context;
+			Context.PackageName = PackageName;
+			check(!bIncrementallyUnmodified || !QueryPlatformData.bIncrementallyModifiedInstigated);
+			Context.bIncrementallyUnmodified = bIncrementallyUnmodified;
+			Context.bPreviouslyCooked = bPreviouslyCooked;
+			Context.bInOutShouldIncrementallySkip = bShouldIncrementallySkip;
+			PackageWriter->UpdatePackageModifiedStatus(Context);
+			bShouldIncrementallySkip = Context.bInOutShouldIncrementallySkip;
+		
+#if ENABLE_COOK_STATS
+			if (PlatformIndex == FirstSessionPlatformIndex && QueryPlatformData.bIncrementallyModifiedInstigated &&
+				bPreviouslyCooked)
+			{
+				FTopLevelAssetPath AssetClassPath;
+				FARFilter PackageNameFilter;
+				PackageNameFilter.PackageNames.Add(PackageData->GetPackageName());
+				Cluster.AssetRegistry.EnumerateAssets(PackageNameFilter, [&AssetClassPath](const FAssetData& AssetData)
+					{
+						AssetClassPath = AssetData.AssetClassPath;
+						return false;
+					});
+
+				if (AssetClassPath.IsValid())
+				{
+					++DetailedCookStats::NumPackagesIncrementallyModifiedByClass.FindOrAdd(AssetClassPath, 0);
+				}
+				else
+				{
+					++DetailedCookStats::NumPackagesIncrementallyModifiedByClass.FindOrAdd(AssetClassPath, 0);
+				}
+			}
+#endif
 		}
 
 		TRefCountPtr<FGenerationHelper> ParentGenerationHelper;
@@ -2273,28 +2306,6 @@ void FRequestCluster::FGraphSearch::FExploreEdgesContext::ProcessPlatformAttachm
 					ParentGenerationHelper->MarkPackageIncrementallySkipped(*PackageData, TargetPlatform,
 						false /* bIncrementallySkipped */);
 				}
-
-#if ENABLE_COOK_STATS
-				if (PlatformIndex == FirstSessionPlatformIndex &&
-					QueryPlatformData.bIncrementallyModifiedInstigated &&
-					(QueryPlatformData.CookAttachments.CommitStatus == IPackageWriter::ECommitStatus::Success
-						|| QueryPlatformData.CookAttachments.CommitStatus == IPackageWriter::ECommitStatus::Error))
-				{
-					FTopLevelAssetPath AssetClassPath;
-					FARFilter PackageNameFilter;
-					PackageNameFilter.PackageNames.Add(PackageData->GetPackageName());
-					Cluster.AssetRegistry.EnumerateAssets(PackageNameFilter, [&AssetClassPath](const FAssetData& AssetData)
-						{
-							AssetClassPath = AssetData.AssetClassPath;
-							return false;
-						});
-
-					if (AssetClassPath.IsValid())
-					{
-						++DetailedCookStats::NumPackagesIncrementallyModifiedByClass.FindOrAdd(AssetClassPath, 0);
-					}
-				}
-#endif
 			}
 		}
 	}
diff --git a/Engine/Source/Editor/UnrealEd/Private/Cooker/DiffPackageWriter.h b/Engine/Source/Editor/UnrealEd/Private/Cooker/DiffPackageWriter.h
--- a/Engine/Source/Editor/UnrealEd/Private/Cooker/DiffPackageWriter.h
+++ b/Engine/Source/Editor/UnrealEd/Private/Cooker/DiffPackageWriter.h
@@ -113,10 +113,9 @@ public:
 	{
 		Inner->RemoveCookedPackages();
 	}
-	virtual void UpdatePackageModificationStatus(FName PackageName, bool bIncrementallyUnmodified,
-		bool& bInOutShouldIncrementallySkip) override
+	virtual void UpdatePackageModifiedStatus(FUpdatePackageModifiedStatusContext& Context) override
 	{
-		Inner->UpdatePackageModificationStatus(PackageName, bIncrementallyUnmodified, bInOutShouldIncrementallySkip);
+		Inner->UpdatePackageModifiedStatus(Context);
 	}
 	virtual EPackageWriterResult BeginCacheForCookedPlatformData(FBeginCacheForCookedPlatformDataInfo& Info) override
 	{
diff --git a/Engine/Source/Editor/UnrealEd/Private/Cooker/IncrementalValidatePackageWriter.cpp b/Engine/Source/Editor/UnrealEd/Private/Cooker/IncrementalValidatePackageWriter.cpp
--- a/Engine/Source/Editor/UnrealEd/Private/Cooker/IncrementalValidatePackageWriter.cpp
+++ b/Engine/Source/Editor/UnrealEd/Private/Cooker/IncrementalValidatePackageWriter.cpp
@@ -620,17 +620,16 @@ void FIncrementalValidatePackageWriter::Initialize(const FCookInfo& CookInfo)
 	Super::Initialize(CookInfo);
 }
 
-void FIncrementalValidatePackageWriter::UpdatePackageModificationStatus(FName PackageName, bool bIncrementallyUnmodified,
-	bool& bInOutShouldIncrementallySkip)
+void FIncrementalValidatePackageWriter::UpdatePackageModifiedStatus(FUpdatePackageModifiedStatusContext& Context)
 {
 	// We need to not skip previously cooked generator packages, if they were modified and we're read only we still
 	// need to cook them so we can investigate their generated packages. Look up whether they were a generator in
 	// the previous cook results.
 	bool bKnownGenerator = false;
-	if (bIncrementallyUnmodified)
+	if (Context.bIncrementallyUnmodified)
 	{
 		// GenerationHelpers were created for previously cooked generators by UCOTFS::PopulateCookedPackages.
-		UE::Cook::FPackageData* PackageData = COTFS.PackageDatas->FindPackageDataByPackageName(PackageName);
+		UE::Cook::FPackageData* PackageData = COTFS.PackageDatas->FindPackageDataByPackageName(Context.PackageName);
 		if (PackageData)
 		{
 			TRefCountPtr<UE::Cook::FGenerationHelper> GenerationHelper = PackageData->GetGenerationHelper();
@@ -645,31 +644,45 @@ void FIncrementalValidatePackageWriter::UpdatePackageModificationStatus(FName Pa
 	{
 	case EPhase::AllInOnePhase:		
 		// Save the input value for bIncrementallyUnmodified, and report skippable for the modified if possible
-		if (bIncrementallyUnmodified)
+		if (Context.bIncrementallyUnmodified)
 		{
-			SetPackageStatus(PackageName, EPackageStatus::DeclaredUnmodified_NotYetProcessed);
-			bInOutShouldIncrementallySkip = false;
+			SetPackageStatus(Context.PackageName, EPackageStatus::DeclaredUnmodified_NotYetProcessed);
+			Context.bInOutShouldIncrementallySkip = false;
 		}
 		else
 		{
-			SetPackageStatus(PackageName, EPackageStatus::DeclaredModified_WillNotVerify);
+			if (Context.bPreviouslyCooked)
+			{
+				SetPackageStatus(Context.PackageName, EPackageStatus::DeclaredModified_WillNotVerify_Cooked);
+			}
+			else
+			{
+				SetPackageStatus(Context.PackageName, EPackageStatus::DeclaredModified_WillNotVerify_NotCooked);
+			}
 			if (!bKnownGenerator && bReadOnly)
 			{
-				bInOutShouldIncrementallySkip = true;
+				Context.bInOutShouldIncrementallySkip = true;
 			}
 		}
 		break;
 	case EPhase::Phase1:
 		// Save the incrementally unmodified packages to verify their diffs. Skip the non-generator incrementally
 		// modified packages. Do not skip generator packages; we need to save them to test their generated packages.
-		bInOutShouldIncrementallySkip = !bIncrementallyUnmodified && !bKnownGenerator;
-		if (!bIncrementallyUnmodified)
+		Context.bInOutShouldIncrementallySkip = !Context.bIncrementallyUnmodified && !bKnownGenerator;
+		if (!Context.bIncrementallyUnmodified)
 		{
-			SetPackageStatus(PackageName, EPackageStatus::DeclaredModified_WillNotVerify);
+			if (Context.bPreviouslyCooked)
+			{
+				SetPackageStatus(Context.PackageName, EPackageStatus::DeclaredModified_WillNotVerify_Cooked);
+			}
+			else
+			{
+				SetPackageStatus(Context.PackageName, EPackageStatus::DeclaredModified_WillNotVerify_NotCooked);
+			}
 		}
 		else
 		{
-			SetPackageStatus(PackageName, EPackageStatus::DeclaredUnmodified_NotYetProcessed);
+			SetPackageStatus(Context.PackageName, EPackageStatus::DeclaredUnmodified_NotYetProcessed);
 		}
 		break;
 	case EPhase::Phase2:
@@ -679,8 +692,8 @@ void FIncrementalValidatePackageWriter::UpdatePackageModificationStatus(FName Pa
 		// is responsible for getting those resaved. Reexecute save for the packages that were IncrementalFailed
 		// from phase1, so we can test whether they are indeterministic. Always save generators so we can
 		// test their generated packages.
-		EPackageStatus PackageStatus = GetPackageStatus(PackageName);
-		bInOutShouldIncrementallySkip =
+		EPackageStatus PackageStatus = GetPackageStatus(Context.PackageName);
+		Context.bInOutShouldIncrementallySkip =
 			(PackageStatus == EPackageStatus::DeclaredUnmodified_ConfirmedUnmodified
 				|| PackageStatus == EPackageStatus::DeclaredUnmodified_FoundModified_OnIgnoreList)
 			&& !bKnownGenerator;
@@ -691,10 +704,11 @@ void FIncrementalValidatePackageWriter::UpdatePackageModificationStatus(FName Pa
 		break;
 	}
 
-	bool bInnerIncrementallyUnmodified = bInOutShouldIncrementallySkip;
-	bool bInnerInOutShouldIncrementallySkip = bInOutShouldIncrementallySkip;
-	Inner->UpdatePackageModificationStatus(PackageName, bInnerIncrementallyUnmodified, bInnerInOutShouldIncrementallySkip);
-	checkf(bInnerInOutShouldIncrementallySkip == bInnerIncrementallyUnmodified,
+	FUpdatePackageModifiedStatusContext InnerContext = Context;
+	InnerContext.bIncrementallyUnmodified = Context.bInOutShouldIncrementallySkip;
+	InnerContext.bInOutShouldIncrementallySkip = Context.bInOutShouldIncrementallySkip;
+	Inner->UpdatePackageModifiedStatus(InnerContext);
+	checkf(InnerContext.bInOutShouldIncrementallySkip == InnerContext.bIncrementallyUnmodified,
 		TEXT("FIncrementalValidatePackageWriter is not supported with an Inner that modifies bInOutShouldIncrementallySkip."));
 }
 
@@ -745,7 +759,7 @@ void FIncrementalValidatePackageWriter::BeginCook(const FCookInfo& Info)
 			UE_CLOG(COTFS.GetCookMode() != ECookMode::CookWorker,
 				LogIncrementalValidate, Display,
 				TEXT("%d packages were found during Phase1 to be modified or new and will be resaved."),
-				StatusCounts[EPackageStatus::DeclaredModified_WillNotVerify]);
+				StatusCounts[EPackageStatus::DeclaredModified_WillNotVerify_Cooked]);
 			break;
 		}
 	default:
@@ -770,7 +784,7 @@ void FIncrementalValidatePackageWriter::EndCook(const FCookInfo& Info)
 		UE_CLOG(COTFS.GetCookMode() != ECookMode::CookWorker,
 			LogIncrementalValidate, Display,
 			TEXT("Modified: %d. DetectedUnmodified: %d. ValidatedUnmodified: %d. IncrementalSkipFalsePositive: %d."),
-			StatusCounts[EPackageStatus::DeclaredModified_WillNotVerify], 
+			StatusCounts[EPackageStatus::DeclaredModified_WillNotVerify_Cooked],
 			DetectedUnmodified,
 			StatusCounts[EPackageStatus::DeclaredUnmodified_ConfirmedUnmodified]
 			+ StatusCounts[EPackageStatus::DeclaredUnmodified_FoundModified_OnIgnoreList],
@@ -840,7 +854,7 @@ void FIncrementalValidatePackageWriter::EndCook(const FCookInfo& Info)
 		UE_CLOG(COTFS.GetCookMode() != ECookMode::CookWorker,
 			LogIncrementalValidate, Display,
 			TEXT("Modified: %d. DetectedUnmodified: %d. ValidatedUnmodified: %d. IncrementalSkipFalsePositiveOrIndeterminism: %d."),
-			StatusCounts[EPackageStatus::DeclaredModified_WillNotVerify],
+			StatusCounts[EPackageStatus::DeclaredModified_WillNotVerify_Cooked],
 			DetectedUnmodified,
 			StatusCounts[EPackageStatus::DeclaredUnmodified_ConfirmedUnmodified]
 			+ StatusCounts[EPackageStatus::DeclaredUnmodified_FoundModified_OnIgnoreList],
@@ -858,7 +872,7 @@ void FIncrementalValidatePackageWriter::EndCook(const FCookInfo& Info)
 		UE_CLOG(COTFS.GetCookMode() != ECookMode::CookWorker,
 			LogIncrementalValidate, Display,
 			TEXT("Modified: %d. DetectedUnmodified: %d. ValidatedUnmodified: %d. Indeterminism: %d. IncrementalSkipFalsePositive: %d."),
-			StatusCounts[EPackageStatus::DeclaredModified_WillNotVerify],
+			StatusCounts[EPackageStatus::DeclaredModified_WillNotVerify_Cooked],
 			DetectedUnmodified,
 			StatusCounts[EPackageStatus::DeclaredUnmodified_ConfirmedUnmodified]
 			+ StatusCounts[EPackageStatus::DeclaredUnmodified_FoundModified_OnIgnoreList],
@@ -978,12 +992,11 @@ bool FIncrementalValidatePackageWriter::IsAnotherSaveNeededInternal(FSavePackage
 			}
 			else if (!bNewPackage)
 			{
-
 				SetPackageStatus(BeginInfo.PackageName, EPackageStatus::DeclaredUnmodified_ConfirmedUnmodified);
 			}
 			else
 			{
-				SetPackageStatus(BeginInfo.PackageName, EPackageStatus::DeclaredModified_WillNotVerify);
+				SetPackageStatus(BeginInfo.PackageName, EPackageStatus::DeclaredModified_WillNotVerify_Cooked);
 			}
 
 			if (bReadOnly)
@@ -1049,7 +1062,7 @@ bool FIncrementalValidatePackageWriter::IsAnotherSaveNeededInternal(FSavePackage
 			// New packages need to be resaved in Phase2; for our purposes they are equivalent
 			// to a package that the IncrementalCook detected as modified.
 			// Do not add an entry for it in our results for incremental packages, and do not resave it in this pass
-			SetPackageStatus(BeginInfo.PackageName, EPackageStatus::DeclaredModified_WillNotVerify);
+			SetPackageStatus(BeginInfo.PackageName, EPackageStatus::DeclaredModified_WillNotVerify_Cooked);
 			return false;
 		}
 	case EPhase::Phase2:
@@ -1218,7 +1231,8 @@ void FIncrementalValidatePackageWriter::SetPackageStatus(FName PackageName, EPac
 	switch (NewStatus)
 	{
 	case EPackageStatus::DeclaredUnmodified_ConfirmedUnmodified: [[fallthrough]];
-	case EPackageStatus::DeclaredModified_WillNotVerify:
+	case EPackageStatus::DeclaredModified_WillNotVerify_Cooked: [[fallthrough]];
+	case EPackageStatus::DeclaredModified_WillNotVerify_NotCooked:
 		// For the non-error cases, clear the AssetClass so that we don't waste bandwidth or diskspace to store it.
 		Info.AssetClass = FTopLevelAssetPath();
 		break;
diff --git a/Engine/Source/Editor/UnrealEd/Private/Cooker/IncrementalValidatePackageWriter.h b/Engine/Source/Editor/UnrealEd/Private/Cooker/IncrementalValidatePackageWriter.h
--- a/Engine/Source/Editor/UnrealEd/Private/Cooker/IncrementalValidatePackageWriter.h
+++ b/Engine/Source/Editor/UnrealEd/Private/Cooker/IncrementalValidatePackageWriter.h
@@ -45,16 +45,13 @@ public:
 	// ICookedPackageWriter
 	virtual FCookCapabilities GetCookCapabilities() const override;
 	virtual void Initialize(const FCookInfo& CookInfo) override;
-	virtual void UpdatePackageModificationStatus(FName PackageName, bool bIncrementallyUnmodified,
-		bool& bInOutShouldIncrementallySkip) override;
+	virtual void UpdatePackageModifiedStatus(FUpdatePackageModifiedStatusContext& Context) override;
 	virtual void BeginCook(const FCookInfo& Info) override;
 	virtual void EndCook(const FCookInfo& Info) override;
 	virtual void UpdateSaveArguments(FSavePackageArgs& SaveArgs) override;
 	virtual bool IsAnotherSaveNeeded(FSavePackageResultStruct& PreviousResult, FSavePackageArgs& SaveArgs) override;
 
 protected:
-	enum class EPackageStatus : uint8;
-
 	enum class ESaveAction : uint8
 	{
 		CheckForDiffs,
@@ -71,7 +68,8 @@ protected:
 		DeclaredUnmodified_FoundModified_FalsePositive,
 		DeclaredUnmodified_FoundModified_OnIgnoreList,
 		DeclaredUnmodified_NotYetProcessed,
-		DeclaredModified_WillNotVerify,
+		DeclaredModified_WillNotVerify_NotCooked,
+		DeclaredModified_WillNotVerify_Cooked,
 		Count
 	};
 
diff --git a/Engine/Source/Editor/UnrealEd/Private/Cooker/LooseCookedPackageWriter.cpp b/Engine/Source/Editor/UnrealEd/Private/Cooker/LooseCookedPackageWriter.cpp
--- a/Engine/Source/Editor/UnrealEd/Private/Cooker/LooseCookedPackageWriter.cpp
+++ b/Engine/Source/Editor/UnrealEd/Private/Cooker/LooseCookedPackageWriter.cpp
@@ -506,11 +506,6 @@ void FLooseCookedPackageWriter::RemoveCookedPackagesByPackageName(TArrayView<con
 	}
 }
 
-void FLooseCookedPackageWriter::UpdatePackageModificationStatus(FName PackageName,
-	bool bIncrementallyUnmodified, bool& bInOutShouldIncrementallySkip)
-{
-}
-
 TFuture<FCbObject> FLooseCookedPackageWriter::WriteMPCookMessageForPackage(FName PackageName)
 {
 	FCbFieldIterator OplogEntryField;
diff --git a/Engine/Source/Editor/UnrealEd/Private/Cooker/LooseCookedPackageWriter.h b/Engine/Source/Editor/UnrealEd/Private/Cooker/LooseCookedPackageWriter.h
--- a/Engine/Source/Editor/UnrealEd/Private/Cooker/LooseCookedPackageWriter.h
+++ b/Engine/Source/Editor/UnrealEd/Private/Cooker/LooseCookedPackageWriter.h
@@ -64,8 +64,6 @@ public:
 	virtual ECommitStatus GetCommitStatus(FName PackageName) override;
 	virtual void RemoveCookedPackages(TArrayView<const FName> PackageNamesToRemove) override;
 	virtual void RemoveCookedPackages() override;
-	virtual void UpdatePackageModificationStatus(FName PackageName, bool bIncrementallyUnmodified,
-		bool& bInOutShouldIncrementallySkip) override;
 	virtual TFuture<FCbObject> WriteMPCookMessageForPackage(FName PackageName) override;
 	virtual bool TryReadMPCookMessageForPackage(FName PackageName, FCbObjectView Message) override;
 	virtual bool GetPreviousCookedBytes(const FPackageInfo& Info, FPreviousCookedBytesData& OutData) override;
diff --git a/Engine/Source/Runtime/Core/Public/Serialization/PackageWriter.h b/Engine/Source/Runtime/Core/Public/Serialization/PackageWriter.h
--- a/Engine/Source/Runtime/Core/Public/Serialization/PackageWriter.h
+++ b/Engine/Source/Runtime/Core/Public/Serialization/PackageWriter.h
@@ -440,14 +440,25 @@ public:
 	 */
 	virtual void RemoveCookedPackages() = 0;
 
-	UE_DEPRECATED(5.4, "No longer called; override UpdatePackageModificationStatus instead")
-	virtual void MarkPackagesUpToDate(TArrayView<const FName> UpToDatePackages) {}
+	struct FUpdatePackageModifiedStatusContext
+	{
+		FName PackageName;
+		bool bIncrementallyUnmodified = false;
+		bool bPreviouslyCooked = false;
+		bool bInOutShouldIncrementallySkip = false;
+	};
 	/**
 	 * During incremental cooking, signal the cooked package has been checked for changes,
 	 * and decide whether it should be incrementally skipped.
 	 */
+	virtual void UpdatePackageModifiedStatus(FUpdatePackageModifiedStatusContext& Context)
+	{
+	}
+	UE_DEPRECATED(5.4, "No longer called; override UpdatePackageModifiedStatus instead")
 	virtual void UpdatePackageModificationStatus(FName PackageName, bool bIncrementallyUnmodified,
-		bool& bInOutShouldIncrementallySkip) = 0;
+		bool& bInOutShouldIncrementallySkip)
+	{
+	}
 
 	struct FDeleteByFree
 	{
