Fix runtime dependencies of generated package:
The world splitter returns a list of runtime dependencies for generated packages containing external actors and levels. During the next cook, those runtime dependencies will be cooked. External actors will be skipped because they never cook. But levels will be cooked even if we don't need to cook them since they are already saved inside the generated packages. What we need as runtime dependencies for the generated packages are the runtime dependencies of the levels.

To do so use the flags in FAssetDependency to control how the generation helper handles dependencies:
- FWorldPartitionCookPackageSplitter::BuildPackagesToGenerateList : Make all the dependencies EditorOnly and Build. EditorOnly means it will not be added as a dependency. Build means its runtime dependencies will be added to the runtime dependencies of the generated package.
- FGenerationHelper::FinishGeneratedPlatformSave :
- Handle dependencies with flags Build.
- Add special case to ignore maps.
- Add never cooked package to runtime dependencies even if they are editor only. That's because I want to keep external actors in the runtime dependencies.

#rnx

diff --git a/Engine/Source/Editor/UnrealEd/Private/Cooker/CookGenerationHelper.cpp b/Engine/Source/Editor/UnrealEd/Private/Cooker/CookGenerationHelper.cpp
--- a/Engine/Source/Editor/UnrealEd/Private/Cooker/CookGenerationHelper.cpp
+++ b/Engine/Source/Editor/UnrealEd/Private/Cooker/CookGenerationHelper.cpp
@@ -1653,36 +1653,54 @@ void FGenerationHelper::FinishGeneratedPlatformSave(FPackageData& PackageData, c
 
 	FCookGenerationInfo::FPlatformData& InfoPlatformData = Info->FindCheckedPlatformData(TargetPlatform);
 
-	// There should be no package dependencies present for the package from the global assetregistry
-	// because it is newly created. Add on the package dependencies declared for it from the CookPackageSplitter.
-	// These dependencies are used in the AssetRegistry entry and runtime dependencies for the generated package
-	OutDependencies = Info->PackageDependencies;
-
-	// For ExternalPackage (that'll get stripped from the cook), append their dependencies gathered from 
-	// the AssetRegistry to PackageDependencies so that the GeneratedPackage dependencies are complete
-	TArray<FAssetDependency> PackageDependencies;
-	TSet<FAssetDependency> ExternalPackageDependencies;
-
 	// External packages are detected by their presence in the NeverCookPackageList 
 	FThreadSafeSet<FName>& NeverCookPackageList = GetOwner().GetPackageDatas().GetCookOnTheFlyServer().PackageTracker->NeverCookPackageList; 
 
 	for (const FAssetDependency& Dependency : Info->PackageDependencies)
 	{
-		// Verify it's an external package before adding its dependencies
-		if (NeverCookPackageList.Contains(Dependency.AssetId.PackageName))
+		// Only dependencies flagged as Build add their own dependencies as runtime dependencies
+		bool bIsBuildDependency = EnumHasAnyFlags(Dependency.Properties, UE::AssetRegistry::EDependencyProperty::Build);
+		if (bIsBuildDependency)
 		{
-			PackageDependencies.Reset();
+			TArray<FAssetDependency> PackageDependencies;
 			IAssetRegistry::GetChecked().GetDependencies(FAssetIdentifier(Dependency.AssetId.PackageName), PackageDependencies,
-				UE::AssetRegistry::EDependencyCategory::Package);
+				UE::AssetRegistry::EDependencyCategory::Package, UE::AssetRegistry::EDependencyQuery::Game);
+
+			for (const FAssetDependency& PackageDependency : PackageDependencies)
+			{
+				TStringBuilder<256> PackakgeDependencyName(InPlace, PackageDependency.AssetId.PackageName);
+
+				// TODO : This is a special case for world splitter. If the package is a level, ignore it. A level can contain another level. 
+				// But they are already included in the generated package so there is no need to add them as runtime dependencies. It would trigger 
+				// a recook of those levels which is unnecessary. We should find a better way to do this. One solution would be to have the world splitter
+				// find the correct dependencies directly so I don't need any special case here.
+				
+				// If it's a script package, add it directly. Calling GetLocalFullPath to check if it's a map will trigger a warning.
+				if (FPackageName::IsScriptPackage(PackakgeDependencyName))
+				{
+					OutDependencies.Add(PackageDependency);
+					continue;
+				}
+				
+				FPackagePath DependencyPackagePath = FPackagePath::FromPackageNameChecked(PackakgeDependencyName);
+				EPackageExtension Extension = FPackagePath::ParseExtension(DependencyPackagePath.GetLocalFullPath());
+				if (Extension == EPackageExtension::Map)
+				{
+					continue;
+				}
 
-			for(const FAssetDependency& PackageDependency : PackageDependencies)
-			{				
-				ExternalPackageDependencies.Add(FAssetDependency::PackageDependency(PackageDependency.AssetId.PackageName, PackageDependency.Properties));
+				OutDependencies.Add(PackageDependency);
 			}
-		}		
-	}
+		}
 
-	OutDependencies.Append(ExternalPackageDependencies.Array());
+		// The world splitter flags every dependency as EditorOnly but I still want external actors added as runtime dependency.
+		bool bNeverCookedPackage = NeverCookPackageList.Contains(Dependency.AssetId.PackageName);
+		bool bIsGameDependency = EnumHasAnyFlags(Dependency.Properties, UE::AssetRegistry::EDependencyProperty::Game);
+		if (bIsGameDependency || bNeverCookedPackage)
+		{
+			OutDependencies.Add(Dependency);
+		}
+	}
 
 	// Add on the custom cook BuildResultDependencies declared for the generated package by the CookPackageSplitter.
 	OutBuildResultDependencies = Info->BuildResultDependencies;
diff --git a/Engine/Source/Runtime/Engine/Private/WorldPartition/WorldPartitionCookPackageSplitter.cpp b/Engine/Source/Runtime/Engine/Private/WorldPartition/WorldPartitionCookPackageSplitter.cpp
--- a/Engine/Source/Runtime/Engine/Private/WorldPartition/WorldPartitionCookPackageSplitter.cpp
+++ b/Engine/Source/Runtime/Engine/Private/WorldPartition/WorldPartitionCookPackageSplitter.cpp
@@ -138,8 +138,11 @@ void FWorldPartitionCookPackageSplitter::BuildPackagesToGenerateList(TArray<ICoo
 					GeneratedPackage.PackageDependencies.Reset(ActorPackageNames.Num());
 					for (FName ActorPackageName : ActorPackageNames)
 					{
+						//The logic here for the dependency property is:
+						// EditorOnly (no Game flag) : so the dependency is not added as a runtime dependency.
+						// Build : so the runtime dependencies of ActorPackageName are added as runtime dependencies.
 						GeneratedPackage.PackageDependencies.Add(FAssetDependency::PackageDependency(ActorPackageName,
-							UE::AssetRegistry::EDependencyProperty::Hard | UE::AssetRegistry::EDependencyProperty::Game));
+							UE::AssetRegistry::EDependencyProperty::Hard | UE::AssetRegistry::EDependencyProperty::Build));
 					}
 				}
 			}
