Cooker: Remove unnecessary if statement that is triggering a static analysis compile warning.
#incrementalcook
#rnx

diff --git a/Engine/Source/Editor/UnrealEd/Private/Cooker/CookRequestCluster.cpp b/Engine/Source/Editor/UnrealEd/Private/Cooker/CookRequestCluster.cpp
--- a/Engine/Source/Editor/UnrealEd/Private/Cooker/CookRequestCluster.cpp
+++ b/Engine/Source/Editor/UnrealEd/Private/Cooker/CookRequestCluster.cpp
@@ -2217,14 +2217,7 @@ void FRequestCluster::FGraphSearch::FExploreEdgesContext::ProcessPlatformAttachm
 						return false;
 					});
 
-				if (AssetClassPath.IsValid())
-				{
-					++DetailedCookStats::NumPackagesIncrementallyModifiedByClass.FindOrAdd(AssetClassPath, 0);
-				}
-				else
-				{
-					++DetailedCookStats::NumPackagesIncrementallyModifiedByClass.FindOrAdd(AssetClassPath, 0);
-				}
+				++DetailedCookStats::NumPackagesIncrementallyModifiedByClass.FindOrAdd(AssetClassPath, 0);
 			}
 #endif
 		}
