Add Disable Preview Platform Blueprint Node

diff --git a/Engine/Source/Editor/UnrealEd/Private/Subsystems/EditorSubsystemBlueprintLibrary.cpp b/Engine/Source/Editor/UnrealEd/Private/Subsystems/EditorSubsystemBlueprintLibrary.cpp
--- a/Engine/Source/Editor/UnrealEd/Private/Subsystems/EditorSubsystemBlueprintLibrary.cpp
+++ b/Engine/Source/Editor/UnrealEd/Private/Subsystems/EditorSubsystemBlueprintLibrary.cpp
@@ -91,4 +91,10 @@ void UEditorSubsystemBlueprintLibrary::SetPreviewPlatform(FName PreviewShaderPla
 	{
 		UE_LOG(LogBlueprint, Warning, TEXT("SetPreviewPlatform Blueprint Function could not find the PreviewShaderPlatform DeviceProfile combination: %s"), *PreviewShaderPlatformName.ToString());
 	}
+}
+
+void UEditorSubsystemBlueprintLibrary::DisablePreviewPlatform()
+{
+	FPreviewPlatformInfo DisablePreviewPlatformInfo(GMaxRHIFeatureLevel, GMaxRHIShaderPlatform, NAME_None, NAME_None, NAME_None, false, NAME_None);
+	GEditor->SetPreviewPlatform(DisablePreviewPlatformInfo, false);
 }
\ No newline at end of file
diff --git a/Engine/Source/Editor/UnrealEd/Public/Subsystems/EditorSubsystemBlueprintLibrary.h b/Engine/Source/Editor/UnrealEd/Public/Subsystems/EditorSubsystemBlueprintLibrary.h
--- a/Engine/Source/Editor/UnrealEd/Public/Subsystems/EditorSubsystemBlueprintLibrary.h
+++ b/Engine/Source/Editor/UnrealEd/Public/Subsystems/EditorSubsystemBlueprintLibrary.h
@@ -24,7 +24,6 @@ public:
 	
 	/**
 	 * Toggle the Preview Platform.
-	 *
 	 */
 	UFUNCTION(BlueprintCallable, Category = "Editor")
 	static UNREALED_API void TogglePreviewPlatform();
@@ -34,4 +33,10 @@ public:
 	 */
 	UFUNCTION(BlueprintCallable, Category = "Editor")
 	static UNREALED_API void SetPreviewPlatform(UPARAM(meta = (GetOptions = "GetPreviewPlatformOptions")) FName PreviewShaderPlatformName);
+
+	/**
+	 * Disable the Preview Platform.
+	 */
+	UFUNCTION(BlueprintCallable, Category = "Editor")
+	static UNREALED_API void DisablePreviewPlatform();
 };
\ No newline at end of file
