Added support for showing Verse paths to log message tokens.

diff --git a/Engine/Source/Editor/UnrealEd/Private/UnrealEdMisc.cpp b/Engine/Source/Editor/UnrealEd/Private/UnrealEdMisc.cpp
--- a/Engine/Source/Editor/UnrealEd/Private/UnrealEdMisc.cpp
+++ b/Engine/Source/Editor/UnrealEd/Private/UnrealEdMisc.cpp
@@ -7,6 +7,7 @@
 #include "Model.h"
 #include "TickableEditorObject.h"
 #include "Components/PrimitiveComponent.h"
+#include "Containers/VersePath.h"
 #include "Misc/MessageDialog.h"
 #include "HAL/FileManager.h"
 #include "Misc/CommandLine.h"
@@ -21,6 +22,7 @@
 #include "Framework/Docking/TabManager.h"
 #include "Toolkits/FConsoleCommandExecutor.h"
 #include "TexAlignTools.h"
+#include "IAssetTools.h"
 #include "ISourceControlModule.h"
 #include "Editor/UnrealEdEngine.h"
 #include "Settings/EditorExperimentalSettings.h"
@@ -1540,7 +1542,19 @@ FText FUnrealEdMisc::OnGetAssetDataDisplayName(const FAssetData& InObject, const
 		Buffer << TEXT(")");
 		return FText::FromStringView(Buffer.ToView());
 	}
-	else if (bFullPath)
+	
+	// Verse paths only exist for top level assets, so show the same value regardless of bFullPath;
+	// we will fall through if we have a subobject.
+	if (IAssetTools::Get().ShowingContentVersePath())
+	{
+		UE::Core::FVersePath VersePath = InObject.GetVersePath();
+		if (VersePath.IsValid())
+		{
+			return FText::FromString(MoveTemp(VersePath).ToString());
+		}
+	}
+	
+	if (bFullPath)
 	{
 		InObject.AppendObjectPath(Buffer);
 		return FText::FromStringView(Buffer.ToView());
@@ -1714,9 +1728,30 @@ FText FUnrealEdMisc::OnGetDisplayName(const UObject* InObject, const bool bFullP
 		{
 			Name = FText::FromString( bFullPath ? Actor->GetPathName() : Actor->GetName() );
 		}
-		else if (InObject != NULL)
+		else if (InObject != nullptr)
 		{
-			Name = FText::FromString( bFullPath ? InObject->GetPathName() : InObject->GetName() );
+			if (bFullPath)
+			{
+				UE::Core::FVersePath VersePath;
+
+				if (IAssetTools::Get().ShowingContentVersePath())
+				{
+					VersePath = InObject->GetVersePath();
+				}
+
+				if (VersePath.IsValid())
+				{
+					Name = FText::FromString(MoveTemp(VersePath).ToString());
+				}
+				else
+				{
+					Name = FText::FromString(InObject->GetPathName());
+				}
+			}
+			else
+			{
+				Name = FText::FromString(InObject->GetName());
+			}
 		}
 	}
 
