Unshelved from pending changelist '42992443':

Fix for crash when trying to get a typed element against a null ActorComponent: Early out from the component visualizer.
FLevelEditorViewportClient::ProcessClick aleady accounts for an unset typed element handle.

diff --git a/Engine/Source/Editor/UnrealEd/Public/ComponentVisualizer.h b/Engine/Source/Editor/UnrealEd/Public/ComponentVisualizer.h
--- a/Engine/Source/Editor/UnrealEd/Public/ComponentVisualizer.h
+++ b/Engine/Source/Editor/UnrealEd/Public/ComponentVisualizer.h
@@ -53,7 +53,9 @@ struct HComponentVisProxy : public HHitProxy
 
 	virtual FTypedElementHandle GetElementHandle() const override
 	{
-		return UEngineElementsLibrary::AcquireEditorComponentElementHandle(Component.Get());
+		return Component.IsValid() ? 
+			UEngineElementsLibrary::AcquireEditorComponentElementHandle(Component.Get()) :
+			FTypedElementHandle();
 	}
 
 	TWeakObjectPtr<const UActorComponent> Component;
