IncrementalCook: Turn on incrementalcook dependency storage by default, but keep forcerecook=true by default.
Change the meaning of -cookincremental to mean !forcerecook, add a separate config-only variable LegacyIterative to compare which system is used, and add a config-only variable CookIncrementalDefaultIncremental to control whether forcerecook is on by default.
Add parsing of forcerecook=false to allow turning CookIncrementalDefaultIncremental=true in commandline.
#rnx

diff --git a/Engine/Config/BaseEditor.ini b/Engine/Config/BaseEditor.ini
--- a/Engine/Config/BaseEditor.ini
+++ b/Engine/Config/BaseEditor.ini
@@ -348,18 +348,22 @@ CookContentMissingSeverity=Error
 ; But this warning can be spurious on farm machines so the project may want to downgrade it to display.
 CookerIdleWarningSeverity=Warning
 
-; When configured on, CookIncremental causes every cooked package to test its dependencies (its package, the schemas of
+; When Incremental Cook is enabled, either by CookIncrementalDefault=true or by -forcerecook=false on commandline,
+; Incremental Cook causes every cooked package to test its dependencies (its package, the schemas of
 ; its classes, its referenced config settings, its referenced other packages, and others) for whether it has changed.
 ; If no dependencies have changed since the last time it was cooked, the previously-cooked version will be used in the
 ; current cook instead of cooking it again; this is used to improve cook performance. The ability to be incrementally
-; skipped can be turned;; off for packages that contain instances of classes that are known to have hidden dependencies
+; skipped can be turned off for packages that contain instances of classes that are known to have hidden dependencies
 ; and therefore might be incorrectly skipped even when they have changed, see IncrementalClassAllowList.
 ;
-; When configured off, cooks are full recooks by default - every package is recooked even if previously cooked - and
-; packages will skip their cook only if the commandline argument -legacyiterative is passed. LegacyIterative cooks are
-; less robust than IncrementalCooks, because the only dependencies they know of are package dependencies from the
-; AssetRegistry. Changes to class schemas do not cause packages with instances of those classes to recook.
-CookIncremental=false
+; When Incremental Cook is disabled or -forcerecook or -forcerecook=true is on commandline, cooks are full recooks by\
+; default - every package is recooked even if previously cooked.
+CookIncrementalDefaultIncremental=false
+; Use the legacy iterative mode, which cooked all packages by default, and when running with -iterate it skipped
+; packages that were previously cooked, were not changed, and did not have any transitive AssetRegistry dependencies
+; that had changed. This mode was not as robust as the new Incremental Cook, because it did not handle changes to
+; c++ serialization, changes to non-package files, or hidden dependencies undeclared to the AssetRegistry
+LegacyIterative=false
 PackagesPerGC=0
 IdleTimeToGC=20
 ; 2 GB of temporary memory is need to cook large textures, so the default for all projects is to trigger a
diff --git a/Engine/Source/Editor/UnrealEd/Private/CookOnTheFlyServer.cpp b/Engine/Source/Editor/UnrealEd/Private/CookOnTheFlyServer.cpp
--- a/Engine/Source/Editor/UnrealEd/Private/CookOnTheFlyServer.cpp
+++ b/Engine/Source/Editor/UnrealEd/Private/CookOnTheFlyServer.cpp
@@ -10477,21 +10477,43 @@ namespace UE::Cook
 void FBeginCookConfigSettings::LoadLocal(FBeginCookContext& BeginContext)
 {
 	const TCHAR* CommandLine = FCommandLine::Get();
-	GConfig->GetBool(TEXT("CookSettings"), TEXT("CookIncremental"), bCookIncremental, GEditorIni);
-	if (FParse::Param(CommandLine, TEXT("CookIncremental")))
+	bool bLegacyIterative = false;
+	GConfig->GetBool(TEXT("CookSettings"), TEXT("LegacyIterative"), bLegacyIterative, GEditorIni);
+	bCookIncremental = !bLegacyIterative;
+
+	if (bCookIncremental && FParse::Param(CommandLine, TEXT("ODSC")))
 	{
-		bCookIncremental = true;
+		// INCREMENTALCOOK_TODO: implement incremental cook for DLC
+		UE_LOG(LogCook, Display,
+			TEXT("INCREMENTAL COOK DEPENDENCIES: Disabled. The cooker is running as ODSC which forces use of legacyiterative."));
+		bCookIncremental = false;
 	}
-	FString Text;
-	if (FParse::Value(CommandLine, TEXT("-CookIncremental="), Text))
+	else if (bCookIncremental && BeginContext.COTFS.IsCookingDLC())
 	{
-		LexFromString(bCookIncremental, Text);
+		// INCREMENTALCOOK_TODO: implement incremental cook for DLC
+		UE_LOG(LogCook, Display, 
+			TEXT("INCREMENTAL COOK DEPENDENCIES: Disabled. Incremental Cook Dependencies are disabled because the cooker is cooking DLC, ")
+			TEXT("and we do not yet support Incremental Cook dependency evaluation in DLC. ")
+			TEXT("Falling back to legacy iterative cook mode."));
+		bCookIncremental = false;
 	}
-
-	// TODO: incremental cook is not yet implemented for DLC
-	bCookIncremental &= !BeginContext.COTFS.IsCookingDLC();
-	// Incremental cook uses TargetDomain storage of dependencies which is only implemented in ZenStore
-	bCookIncremental &= BeginContext.COTFS.IsUsingZenStore();
+	else if (bCookIncremental && !BeginContext.COTFS.IsUsingZenStore())
+	{
+		// Incremental cook uses TargetDomain storage of dependencies which is only implemented in ZenStore
+		UE_LOG(LogCook, Display,
+			TEXT("INCREMENTAL COOK DEPENDENCIES: Disabled. Incremental Cook Dependencies are disabled because the cooker is not using ZenStore, ")
+			TEXT("and we do not yet support storage of Incremental Cook dependency data with loose cooked packages. ")
+			TEXT("Falling back to legacy iterative cook mode."));
+		bCookIncremental = false;
+	}
+	else
+	{
+		UE_CLOG(bCookIncremental, LogCook, Display,
+			TEXT("INCREMENTAL COOK DEPENDENCIES: Enabled. Incremental Cook Dependencies are enabled in Editor.ini:[CookSettings]:LegacyIterative=false."));
+		UE_CLOG(!bCookIncremental, LogCook, Display,
+			TEXT("INCREMENTAL COOK DEPENDENCIES: Disabled. Incremental Cook Dependencies are disabled in Editor.ini:[CookSettings]:LegacyIterative=true."));
+	}
+		 
 	bCookIncrementalAllowAllClasses = FParse::Param(CommandLine, TEXT("CookIncrementalAllowAllClasses"));
 
 	FParse::Value(CommandLine, TEXT("-CookShowInstigator="), CookShowInstigator);
@@ -10577,8 +10599,49 @@ void UCookOnTheFlyServer::LoadBeginCookIncrementalFlags(FBeginCookContext& Begin
 
 void UCookOnTheFlyServer::LoadBeginCookIncrementalFlagsLocal(FBeginCookContext& BeginContext)
 {
-	const bool bIsDiffOnly = FParse::Param(FCommandLine::Get(), TEXT("DIFFONLY"));
-	bool bForceRecook = FParse::Param(FCommandLine::Get(), TEXT("fullcook")) || FParse::Param(FCommandLine::Get(), TEXT("forcerecook"));
+	const TCHAR* CommandLine = FCommandLine::Get();
+	const bool bIsDiffOnly = FParse::Param(CommandLine, TEXT("DIFFONLY"));
+	bool bDefaultIncremental = true;
+	if (bCookIncremental)
+	{
+		GConfig->GetBool(TEXT("CookSettings"), TEXT("CookIncrementalDefaultIncremental"), bDefaultIncremental, GEditorIni);
+	}
+	else
+	{
+		bDefaultIncremental = IsCookFlagSet(ECookInitializationFlags::LegacyIterative);
+	}
+
+	bool bForceRecook = !bDefaultIncremental;
+	bool bForceRecookCommandLineArgPresent = false;
+	if (FParse::Param(CommandLine, TEXT("CookIncremental")))
+	{
+		bForceRecook = false;
+		bForceRecookCommandLineArgPresent = true;
+	}
+	if (IsCookFlagSet(ECookInitializationFlags::LegacyIterative))
+	{
+		bForceRecook = false;
+		bForceRecookCommandLineArgPresent = true;
+	}
+	FString Text;
+	if (FParse::Value(CommandLine, TEXT("-CookIncremental="), Text))
+	{
+		bool bValue = true;
+		LexFromString(bValue, Text);
+		bForceRecook = !bValue;
+		bForceRecookCommandLineArgPresent = true;
+	}
+	if (FParse::Param(CommandLine, TEXT("fullcook")) || FParse::Param(CommandLine, TEXT("forcerecook")))
+	{
+		bForceRecook = true;
+		bForceRecookCommandLineArgPresent = true;
+	}
+	if (FParse::Value(CommandLine, TEXT("-forcerecook="), Text))
+	{
+		LexFromString(bForceRecook, Text);
+		bForceRecookCommandLineArgPresent = true;
+	}
+
 	const bool bIncrementalOrLegacyIterative = !bForceRecook && (bCookIncremental || IsCookFlagSet(ECookInitializationFlags::LegacyIterative));
 	const bool bIsSharedLegacyIterativeCook = !bCookIncremental && IsCookFlagSet(ECookInitializationFlags::LegacyIterativeSharedBuild);
 
@@ -10637,20 +10700,30 @@ void UCookOnTheFlyServer::LoadBeginCookIncrementalFlagsLocal(FBeginCookContext&
 			bool bIncrementalOrLegacyIterativeAllowed = true;
 			if (!bIncrementalOrLegacyIterative && !PlatformData->bIsSandboxInitialized)
 			{
-				if (bForceRecook)
+				if (bCookIncremental)
 				{
-					UE_LOG(LogCook, Display,
-						TEXT("FULL COOK: -forcerecook was specified. Deleting previously cooked packages for platform %s and recooking all packages discovered in the current cook."),
-						*TargetPlatform->PlatformName());
+					if (!bDefaultIncremental && !bForceRecookCommandLineArgPresent)
+					{
+						UE_LOG(LogCook, Display,
+							TEXT("FULL COOK: Incremental Cooks are disabled by default in Editor.ini:[CookSettings]:CookIncrementalDefaultIncremental=false, and commandline did not specify -forcerecook=false. ")
+							TEXT("Deleting previously cooked packages for platform %s and recooking all packages discovered in the current cook."),
+							*TargetPlatform->PlatformName());
+					}
+					else
+					{
+						UE_LOG(LogCook, Display,
+							TEXT("FULL COOK: -forcerecook=true (or an equivalent) was specified. Deleting previously cooked packages for platform %s and recooking all packages discovered in the current cook."),
+							*TargetPlatform->PlatformName());
+					}
 				}
 				else
 				{
 					UE_LOG(LogCook, Display,
-						TEXT("FULL COOK: Neither -legacyiterative nor -cookincremental were specified. Deleting previously cooked packages for platform %s and recooking all packages discovered in the current cook."),
+						TEXT("FULL COOK: Incremental Cook Dependencies are disabled, and -legacyiterative was not specified. Deleting previously cooked packages for platform %s and recooking all packages discovered in the current cook."),
 						*TargetPlatform->PlatformName());
 				}
 				bIncrementalOrLegacyIterativeAllowed = false;
-				if (!bForceRecook && bRunningAsShaderServer)
+				if (!bForceRecookCommandLineArgPresent && bRunningAsShaderServer)
 				{
 					UE_LOG(LogCook, Display,
 						TEXT("'-odsc' was passed on commandline, but '-legacyiterative' was not, so the cooker as a side effect is clearing cook results. The build will need to be recooked before it can be staged. Add the commandline argument '-legacyiterative' to avoid this unnecessary clear.")
@@ -10675,8 +10748,8 @@ void UCookOnTheFlyServer::LoadBeginCookIncrementalFlagsLocal(FBeginCookContext&
 			if (bIncrementalOrLegacyIterativeAllowed)
 			{
 				UE_LOG(LogCook, Display,
-					TEXT("INCREMENTAL COOK: %s was specified and global settings are still valid. Keeping previously cooked packages for platform %s and cooking only packages that have been modified."),
-					bCookIncremental ? TEXT("-cookincremental") : TEXT("-legacyiterative"),
+					TEXT("INCREMENTAL COOK: %s settings are still valid. Keeping previously cooked packages for platform %s and cooking only packages that have been modified."),
+					bCookIncremental ? TEXT("Global") : TEXT("-legacyiterative was specified and global"),
 					*TargetPlatform->PlatformName());
 				PlatformContext.bFullBuild = false;
 				PlatformContext.bAllowIncrementalResults = true;
diff --git a/Engine/Source/Programs/AutomationTool/AutomationUtils/ProjectParams.cs b/Engine/Source/Programs/AutomationTool/AutomationUtils/ProjectParams.cs
--- a/Engine/Source/Programs/AutomationTool/AutomationUtils/ProjectParams.cs
+++ b/Engine/Source/Programs/AutomationTool/AutomationUtils/ProjectParams.cs
@@ -30,7 +30,7 @@ namespace AutomationTool
 	[ParamHelp("run", "Run the game after it is built (including server, if -server)", ParamType = typeof(bool))]
 	[ParamHelp("zenstore", "Save cooked output data to the Zen storage server", ParamType = typeof(bool))]
 	[ParamHelp("iterate", "Uses the iterative cooking/deploy", ParamType = typeof(bool))]
-	[ParamHelp("snapshot", "Imports the most recently published project snapshot (if cooking, this forces the use of -cookincremental)", ParamType = typeof(bool))]
+	[ParamHelp("snapshot", "Imports the most recently published project snapshot (if cooking, this forces the use of -forcerecook=false)", ParamType = typeof(bool))]
 	[ParamHelp("cookincremental", "Cooks incrementally using previously-cooked data or a published snapshot as a base", ParamType = typeof(bool))]
 	public interface IProjectParamsHelpers
 	{
diff --git a/Engine/Source/Programs/AutomationTool/Scripts/CookCommand.Automation.cs b/Engine/Source/Programs/AutomationTool/Scripts/CookCommand.Automation.cs
--- a/Engine/Source/Programs/AutomationTool/Scripts/CookCommand.Automation.cs
+++ b/Engine/Source/Programs/AutomationTool/Scripts/CookCommand.Automation.cs
@@ -57,7 +57,7 @@ namespace AutomationScripts
 			}
 			if (Params.IncrementalCook)
 			{
-				CommandletParams += " -cookincremental";
+				CommandletParams += " -forcerecook=false";
 			}
 			if (Params.IterativeCooking)
 			{
