Fix double logging of package diffs: This is due to FIncrementalValidatePackageWriter::BeginPackage not calling Super::BeginPackage in every case. The consequence is the variable BeginInfo is not up to date and contains the data of a previously cooked package. Then when MarkPackageCompletedOnDirector is called in IsAnotherSaveNeeded, we pass it BeginInfo.PackageName, which is the name of a previously cooked package and not the currently cooked package.

- Factorize the code setting internal variables from FDiffPackageWriter::BeginPackage to a new function BeginPackageDiff.
- In FIncrementalValidatePackageWriter::BeginPackage, call BeginPackageDiff when necessary (when Super::BeginPackage is not called).

#rnx

diff --git a/Engine/Source/Editor/UnrealEd/Private/Cooker/DiffPackageWriter.cpp b/Engine/Source/Editor/UnrealEd/Private/Cooker/DiffPackageWriter.cpp
--- a/Engine/Source/Editor/UnrealEd/Private/Cooker/DiffPackageWriter.cpp
+++ b/Engine/Source/Editor/UnrealEd/Private/Cooker/DiffPackageWriter.cpp
@@ -148,14 +148,8 @@ void FDiffPackageWriter::RemoveParam(FString& InOutParams, const TCHAR* InParamT
 
 void FDiffPackageWriter::BeginPackage(const FBeginPackageInfo& Info)
 {
-	bIsDifferent = false;
-	bNewPackage = false;
-	bHasStartedSecondSave = false;
-	Accumulators[0].SafeRelease();
-	Accumulators[1].SafeRelease();
+	BeginPackageDiff(Info);
 
-	BeginInfo = Info;
-	Package = FindObjectFast<UPackage>(nullptr, BeginInfo.PackageName);
 	ConditionallyDumpObjList();
 	ConditionallyDumpObjects();
 	Inner->BeginPackage(Info);
@@ -251,6 +245,19 @@ bool FDiffPackageWriter::IsPackageDiffAllowed() const
 	}
 	return true;
 }
+
+void FDiffPackageWriter::BeginPackageDiff(const FBeginPackageInfo& Info)
+{
+	bIsDifferent = false;
+	bNewPackage = false;
+	bHasStartedSecondSave = false;
+	Accumulators[0].SafeRelease();
+	Accumulators[1].SafeRelease();
+
+	BeginInfo = Info;
+	Package = FindObjectFast<UPackage>(nullptr, BeginInfo.PackageName);
+}
+
 UE::DiffWriter::FMessageCallback FDiffPackageWriter::GetDiffWriterMessageCallback()
 {
 	return UE::DiffWriter::FMessageCallback([this](ELogVerbosity::Type Verbosity, FStringView Message)
diff --git a/Engine/Source/Editor/UnrealEd/Private/Cooker/DiffPackageWriter.h b/Engine/Source/Editor/UnrealEd/Private/Cooker/DiffPackageWriter.h
--- a/Engine/Source/Editor/UnrealEd/Private/Cooker/DiffPackageWriter.h
+++ b/Engine/Source/Editor/UnrealEd/Private/Cooker/DiffPackageWriter.h
@@ -150,6 +150,14 @@ protected:
 	UE::DiffWriter::FAccumulator& ConstructAccumulator(FName PackageName, UObject* Asset, uint16 MultiOutputIndex);
 	bool IsPackageDiffAllowed() const;
 
+	/** 
+	 * When BeginPackage is called, CommitPackage must be called as well. But sometimes we don't want or don't need to do all this and we just want to
+	 * initialize the internal variables as if BeginPackage was called. That's what BeginPackageDiff does.
+	 * For example FIncrementalValidatePackageWriter can use an internal writer so FDiffPackageWriter::BeginPackage/CommitPackage will not be called, but to
+	 * make sure the internal variables are correct and the object behaves as expected, BeginPackageDiff is called instead.
+	*/
+	void BeginPackageDiff(const FBeginPackageInfo& Info);
+
 	TRefCountPtr<UE::DiffWriter::FAccumulator> Accumulators[2];
 	FBeginPackageInfo BeginInfo;
 	TSet<FTopLevelAssetPath> CompareDenyListClasses;
diff --git a/Engine/Source/Editor/UnrealEd/Private/Cooker/IncrementalValidatePackageWriter.cpp b/Engine/Source/Editor/UnrealEd/Private/Cooker/IncrementalValidatePackageWriter.cpp
--- a/Engine/Source/Editor/UnrealEd/Private/Cooker/IncrementalValidatePackageWriter.cpp
+++ b/Engine/Source/Editor/UnrealEd/Private/Cooker/IncrementalValidatePackageWriter.cpp
@@ -357,6 +357,8 @@ void FIncrementalValidatePackageWriter::BeginPackage(const FBeginPackageInfo& In
 		}
 		else
 		{
+			BeginPackageDiff(Info);
+
 			// Not incrementally skippable, so we expect it to change. No need to look for diffs, just save to disk.
 			if (bReadOnly)
 			{
@@ -381,6 +383,7 @@ void FIncrementalValidatePackageWriter::BeginPackage(const FBeginPackageInfo& In
 			{
 				// Already saved in Phase 1; no need to diff it or save it now
 				SaveAction = ESaveAction::IgnoreResults;
+				BeginPackageDiff(Info);
 			}
 			else if (PackageStatus == EPackageStatus::DeclaredUnmodified_FoundModified_IndeterminismOrFalsePositive)
 			{
@@ -392,6 +395,7 @@ void FIncrementalValidatePackageWriter::BeginPackage(const FBeginPackageInfo& In
 				// This is an IncrementallyModified package. It was found during Phase1 to be modified and would in a
 				// normal incremental cook be resaved rather than incrementally skipped. Resave it as normal.
 				SaveAction = ESaveAction::SaveToInner;
+				BeginPackageDiff(Info);
 				Inner->BeginPackage(Info);
 			}
 			break;
