Updated headers using UnrealCodeFixup to make sure dllstorage is on methods/staticvars instead of types. Part 5/n

diff --git a/Engine/Source/Editor/UnrealEd/Public/EditorViewportSelectability.h b/Engine/Source/Editor/UnrealEd/Public/EditorViewportSelectability.h
--- a/Engine/Source/Editor/UnrealEd/Public/EditorViewportSelectability.h
+++ b/Engine/Source/Editor/UnrealEd/Public/EditorViewportSelectability.h
@@ -9,6 +9,8 @@
 #include "Math/Color.h"
 #include "Misc/Optional.h"
 
+#define UE_API UNREALED_API
+
 class AActor;
 class FCanvas;
 class FEditorViewportClient;
@@ -28,36 +30,36 @@ DECLARE_DELEGATE_RetVal_OneParam(bool, FOnIsObjectSelectableInViewport, UObject*
  * Manages level actor viewport selectability and hovered visual states.
  * Contains static methods to enable outside modules to implement their own management.
  */
-class UNREALED_API FEditorViewportSelectability
+class FEditorViewportSelectability
 {
 public:
 	/** Default text to display in the viewport when selection is limited as a helpful reminder to the user. */
-	static const FText DefaultLimitedSelectionText;
+	static UE_API const FText DefaultLimitedSelectionText;
 
 	/** Updates a single primitive component's hovered state and visuals. */
-	static void UpdatePrimitiveVisuals(const bool bInSelectedLimited, UPrimitiveComponent* const InPrimitive, const TOptional<FColor>& InColor = TOptional<FColor>());
+	static UE_API void UpdatePrimitiveVisuals(const bool bInSelectedLimited, UPrimitiveComponent* const InPrimitive, const TOptional<FColor>& InColor = TOptional<FColor>());
 
 	/** Updates a list of hovered primitive component's hovered state and visuals */
-	static bool UpdateHoveredPrimitive(const bool bInSelectedLimited
+	static UE_API bool UpdateHoveredPrimitive(const bool bInSelectedLimited
 		, UPrimitiveComponent* const InPrimitiveComponent
 		, TMap<TWeakObjectPtr<UPrimitiveComponent>, TOptional<FColor>>& InOutHoveredPrimitiveComponents
 		, const TFunctionRef<bool(UObject*)>& InSelectablePredicate);
 
 	/** Updates an actors hovered state and visuals. */
-	static bool UpdateHoveredActorPrimitives(const bool bInSelectedLimited
+	static UE_API bool UpdateHoveredActorPrimitives(const bool bInSelectedLimited
 		, AActor* const InActor
 		, TMap<TWeakObjectPtr<UPrimitiveComponent>, TOptional<FColor>>& InOutHoveredPrimitiveComponents
 		, const TFunctionRef<bool(UObject*)>& InSelectablePredicate);
 
-	static FText GetLimitedSelectionText(const TSharedPtr<FUICommandInfo>& InToggleAction, const FText& InDefaultText = DefaultLimitedSelectionText);
+	static UE_API FText GetLimitedSelectionText(const TSharedPtr<FUICommandInfo>& InToggleAction, const FText& InDefaultText = DefaultLimitedSelectionText);
 
-	static void DrawEnabledTextNotice(FCanvas* const InCanvas, const FText& InText);
+	static UE_API void DrawEnabledTextNotice(FCanvas* const InCanvas, const FText& InText);
 
 	FEditorViewportSelectability() = delete;
-	FEditorViewportSelectability(const FOnGetWorld& InOnGetWorld, const FOnIsObjectSelectableInViewport& InOnIsObjectSelectableInViewport);
+	UE_API FEditorViewportSelectability(const FOnGetWorld& InOnGetWorld, const FOnIsObjectSelectableInViewport& InOnIsObjectSelectableInViewport);
 
 	/** Enables or disables the selectability tool */
-	void EnableLimitedSelection(const bool bInEnabled);
+	UE_API void EnableLimitedSelection(const bool bInEnabled);
 
 	bool IsSelectionLimited() const
     {
@@ -65,33 +67,33 @@ public:
     }
 
 	/** @return True if the specified object is selectable in the viewport and not made unselectable by the Sequencer selection limiting. */
-	bool IsObjectSelectableInViewport(UObject* const InObject) const;
+	UE_API bool IsObjectSelectableInViewport(UObject* const InObject) const;
 
 	/** Updates hover visual states based on current selection limiting settings */
-	void UpdateSelectionLimitedVisuals(const bool bInClearHovered);
+	UE_API void UpdateSelectionLimitedVisuals(const bool bInClearHovered);
 
-	void DeselectNonSelectableActors();
+	UE_API void DeselectNonSelectableActors();
 
-	bool GetCursorForHovered(EMouseCursor::Type& OutCursor) const;
+	UE_API bool GetCursorForHovered(EMouseCursor::Type& OutCursor) const;
 
-	void UpdateHoverFromHitProxy(HHitProxy* const InHitProxy);
+	UE_API void UpdateHoverFromHitProxy(HHitProxy* const InHitProxy);
 
-	bool HandleClick(FEditorViewportClient* const InViewportClient, HHitProxy* const InHitProxy, const FViewportClick& InClick);
+	UE_API bool HandleClick(FEditorViewportClient* const InViewportClient, HHitProxy* const InHitProxy, const FViewportClick& InClick);
 
-	void StartTracking(FEditorViewportClient* const InViewportClient, FViewport* const InViewport);
-	void EndTracking(FEditorViewportClient* const InViewportClient, FViewport* const InViewport);
+	UE_API void StartTracking(FEditorViewportClient* const InViewportClient, FViewport* const InViewport);
+	UE_API void EndTracking(FEditorViewportClient* const InViewportClient, FViewport* const InViewport);
 
 	/** Selects or deselects all actors in a level world that are inside a defined box. */
-	bool BoxSelectWorldActors(FBox& InBox, FEditorViewportClient* const InEditorViewportClient, const bool bInSelect);
+	UE_API bool BoxSelectWorldActors(FBox& InBox, FEditorViewportClient* const InEditorViewportClient, const bool bInSelect);
 	/** Selects or deselects all actors in a level world that are inside a defined convex volume. */
-	bool FrustumSelectWorldActors(const FConvexVolume& InFrustum, FEditorViewportClient* const InEditorViewportClient, const bool bInSelect);
+	UE_API bool FrustumSelectWorldActors(const FConvexVolume& InFrustum, FEditorViewportClient* const InEditorViewportClient, const bool bInSelect);
 
 protected:
-	static bool IsActorSelectableClass(const AActor& InActor);
+	static UE_API bool IsActorSelectableClass(const AActor& InActor);
 
-	static bool IsActorInLevelHiddenLayer(const AActor& InActor, FLevelEditorViewportClient* const InLevelEditorViewportClient);
+	static UE_API bool IsActorInLevelHiddenLayer(const AActor& InActor, FLevelEditorViewportClient* const InLevelEditorViewportClient);
 
-	static TTypedElement<ITypedElementWorldInterface> GetTypedWorldElementFromActor(const AActor& InActor);
+	static UE_API TTypedElement<ITypedElementWorldInterface> GetTypedWorldElementFromActor(const AActor& InActor);
 
 	/**
 	 * Selects or deselects actors in a world. If no actors are specified, uses all the actors in the level.
@@ -102,20 +104,20 @@ protected:
 	 * @param bInClearSelection If true, clears the current selection before selecting the new actors
 	 * @return True if atleast one new actor was selected/deselected
 	 */
-	static bool SelectActorsByPredicate(UWorld* const InWorld
+	static UE_API bool SelectActorsByPredicate(UWorld* const InWorld
 		, const bool bInSelect
 		, const bool bInClearSelection
 		, const TFunctionRef<bool(AActor*)> InPredicate
 		, const TArray<AActor*>& InActors = {});
 
 	/** Updates an actors hovered state and visuals. */
-	void UpdateHoveredActorPrimitives(AActor* const InActor);
+	UE_API void UpdateHoveredActorPrimitives(AActor* const InActor);
 
-	static UTypedElementSelectionSet* GetLevelEditorSelectionSet();
+	static UE_API UTypedElementSelectionSet* GetLevelEditorSelectionSet();
 
-	bool IsTypedElementSelectable(const FTypedElementHandle& InElementHandle) const;
+	UE_API bool IsTypedElementSelectable(const FTypedElementHandle& InElementHandle) const;
 
-	void GetSelectionElements(AActor* const InActor, const TFunctionRef<void(const TTypedElement<ITypedElementWorldInterface>&)> InPredicate);
+	UE_API void GetSelectionElements(AActor* const InActor, const TFunctionRef<void(const TTypedElement<ITypedElementWorldInterface>&)> InPredicate);
 
 	bool bSelectionLimited = false;
 
@@ -134,3 +136,5 @@ protected:
 	FVector DragEndPosition;
 	FIntRect DragSelectionRect;
 };
+
+#undef UE_API
diff --git a/Engine/Source/Editor/UnrealEd/Public/EditorViewportSelectabilityBridge.h b/Engine/Source/Editor/UnrealEd/Public/EditorViewportSelectabilityBridge.h
--- a/Engine/Source/Editor/UnrealEd/Public/EditorViewportSelectabilityBridge.h
+++ b/Engine/Source/Editor/UnrealEd/Public/EditorViewportSelectabilityBridge.h
@@ -5,6 +5,8 @@
 #include "Internationalization/Text.h"
 #include "Templates/SharedPointer.h"
 
+#define UE_API UNREALED_API
+
 class FEditorViewportClient;
 
 DECLARE_DELEGATE_RetVal(bool, FOnIsViewportSelectionLimited);
@@ -15,24 +17,24 @@ DECLARE_DELEGATE_RetVal(FText, FOnGetViewportSelectionLimitedText);
  * Creates a link between a viewport and an outside module without requiring extra dependencies.
  * This could be moved to Editor/UnrealEd module to allow other modules that may need this functionality to access.
  */
-class UNREALED_API FEditorViewportSelectabilityBridge
+class FEditorViewportSelectabilityBridge
 {
 public:
 	FEditorViewportSelectabilityBridge() = delete;
-	FEditorViewportSelectabilityBridge(const TWeakPtr<FEditorViewportClient>& InEditorViewportClientWeak);
+	UE_API FEditorViewportSelectabilityBridge(const TWeakPtr<FEditorViewportClient>& InEditorViewportClientWeak);
 
 	/** @return Delegate used to check if viewport selection is limited */
-	FOnIsViewportSelectionLimited& OnIsViewportSelectionLimited();
-	bool IsViewportSelectionLimited() const;
+	UE_API FOnIsViewportSelectionLimited& OnIsViewportSelectionLimited();
+	UE_API bool IsViewportSelectionLimited() const;
 
 	/** @return Delegate used to check if an object is selectable in the viewport */
-	FOnIsObjectSelectableInViewport& OnGetIsObjectSelectableInViewport();
+	UE_API FOnIsObjectSelectableInViewport& OnGetIsObjectSelectableInViewport();
 	/** @return True if the specified object is selectable in the viewport and not made unselectable by the Sequencer selection limiting */
-	bool IsObjectSelectableInViewport(UObject* const InObject) const;
+	UE_API bool IsObjectSelectableInViewport(UObject* const InObject) const;
 
 	/** @return Delegate used to get the text to display in the viewport when selection is limited */
-	FOnGetViewportSelectionLimitedText& OnGetViewportSelectionLimitedText();
-	FText GetViewportSelectionLimitedText() const;
+	UE_API FOnGetViewportSelectionLimitedText& OnGetViewportSelectionLimitedText();
+	UE_API FText GetViewportSelectionLimitedText() const;
 
 private:
 	TWeakPtr<FEditorViewportClient> EditorViewportClientWeak;
@@ -41,3 +43,5 @@ private:
 	FOnIsObjectSelectableInViewport GetIsObjectSelectableInViewportDelegate;
 	FOnGetViewportSelectionLimitedText GetViewportSelectionLimitedTextDelegate;
 };
+
+#undef UE_API
diff --git a/Engine/Source/Editor/UnrealEd/Public/Kismet2/CompilerResultsLog.h b/Engine/Source/Editor/UnrealEd/Public/Kismet2/CompilerResultsLog.h
--- a/Engine/Source/Editor/UnrealEd/Public/Kismet2/CompilerResultsLog.h
+++ b/Engine/Source/Editor/UnrealEd/Public/Kismet2/CompilerResultsLog.h
@@ -339,7 +339,7 @@ public:
 	static UNREALED_API TSharedRef<IMessageLogListing> GetBlueprintMessageLog(UBlueprint* InBlueprint);
 
 	/** ICompilerResultsLog implementation */
-	UNREALED_API void SetSilentMode(bool bValue) { bSilentMode = bValue; };
+	void SetSilentMode(bool bValue) { bSilentMode = bValue; };
 
 protected:
 	/** Helper method to add a child event to the given parent event scope */
diff --git a/Engine/Source/Editor/UnrealEd/Public/Kismet2/ReloadUtilities.h b/Engine/Source/Editor/UnrealEd/Public/Kismet2/ReloadUtilities.h
--- a/Engine/Source/Editor/UnrealEd/Public/Kismet2/ReloadUtilities.h
+++ b/Engine/Source/Editor/UnrealEd/Public/Kismet2/ReloadUtilities.h
@@ -45,12 +45,12 @@ public:
 	UNREALED_API virtual UObject* GetReinstancedCDO(UObject* CDO) override;
 	UNREALED_API virtual const UObject* GetReinstancedCDO(const UObject* CDO) override;
 
-	UNREALED_API virtual const TMap<UClass*, UClass*>& GetReinstancedClasses() override
+	virtual const TMap<UClass*, UClass*>& GetReinstancedClasses() override
 	{
 		return ReinstancedClasses;
 	}
 
-	UNREALED_API const TArray<UClass*>& GetNewClasses() override
+	const TArray<UClass*>& GetNewClasses() override
 	{
 		return NewClasses;
 	}
diff --git a/Engine/Source/Editor/UnrealEd/Public/SCommonEditorViewportToolbarBase.h b/Engine/Source/Editor/UnrealEd/Public/SCommonEditorViewportToolbarBase.h
--- a/Engine/Source/Editor/UnrealEd/Public/SCommonEditorViewportToolbarBase.h
+++ b/Engine/Source/Editor/UnrealEd/Public/SCommonEditorViewportToolbarBase.h
@@ -17,6 +17,8 @@
 
 #include "SCommonEditorViewportToolbarBase.generated.h"
 
+#define UE_API UNREALED_API
+
 class SComboButton;
 
 // This is the interface that the host of a SCommonEditorViewportToolbarBase must implement
@@ -71,18 +73,18 @@ namespace CommonEditorViewportUtils
 	}
 }
 
-UCLASS()
-class UNREALED_API UCommonViewportToolbarBaseMenuContext : public UUnrealEdViewportToolbarContext
+UCLASS(MinimalAPI)
+class UCommonViewportToolbarBaseMenuContext : public UUnrealEdViewportToolbarContext
 {
 	GENERATED_BODY()
 
 public:
 	TWeakPtr<const class SCommonEditorViewportToolbarBase> ToolbarWidget;
 	
-	virtual TSharedPtr<IPreviewProfileController> GetPreviewProfileController() const override;
+	UE_API virtual TSharedPtr<IPreviewProfileController> GetPreviewProfileController() const override;
 };
 
-class UNREALED_API SPreviewSceneProfileSelector : public SCompoundWidget
+class SPreviewSceneProfileSelector : public SCompoundWidget
 {
 public:
 	virtual ~SPreviewSceneProfileSelector() override
@@ -94,7 +96,7 @@ public:
 	SLATE_ARGUMENT(TSharedPtr<IPreviewProfileController>, PreviewProfileController)
 	SLATE_END_ARGS()
 	
-	void Construct(const FArguments& InArgs);
+	UE_API void Construct(const FArguments& InArgs);
 
 protected:
 	UE_DEPRECATED(5.5, "Unused")
@@ -113,7 +115,7 @@ protected:
 	}
 
 	/** Creates and returns the asset viewer profile combo box.*/
-	TSharedRef<SWidget> MakeAssetViewerProfileComboBox();
+	UE_API TSharedRef<SWidget> MakeAssetViewerProfileComboBox();
 
 private:
 	/** Interface to set/get/list the preview profiles. */
@@ -123,7 +125,7 @@ private:
 	TSharedPtr<SComboButton> AssetViewerProfileComboButton;
 
 	/** Builds the drop-down list for selecting a viewer profile. */
-	TSharedRef<SWidget> BuildComboMenu();
+	UE_API TSharedRef<SWidget> BuildComboMenu();
 };
 
 /**
@@ -141,7 +143,7 @@ public:
 		SLATE_ARGUMENT(TSharedPtr<IPreviewProfileController>, PreviewProfileController) // Should be null if the Preview doesn't require profile.
 	SLATE_END_ARGS()
 
-	UNREALED_API virtual ~SCommonEditorViewportToolbarBase(){};
+	virtual ~SCommonEditorViewportToolbarBase(){};
 
 	UNREALED_API void Construct(const FArguments& InArgs, TSharedPtr<class ICommonEditorViewportToolbarInfoProvider> InInfoProvider);
 
@@ -149,7 +151,7 @@ public:
 	static UNREALED_API void ConstructScreenPercentageMenu(FMenuBuilder& MenuBuilder, class FEditorViewportClient* ViewportClient);
 
 	/** Gets the preview profile controller explicitly set to this toolbar */
-	UNREALED_API const TSharedPtr<IPreviewProfileController>& GetPreviewProfileController() const { return PreviewProfileController; }
+	const TSharedPtr<IPreviewProfileController>& GetPreviewProfileController() const { return PreviewProfileController; }
 
 private:
 	/**
@@ -285,3 +287,4 @@ private:
 	bool ShouldCreateOptionsMenu() const;
 };
 
+#undef UE_API
diff --git a/Engine/Source/Editor/UnrealEd/Public/SCurveEditor.h b/Engine/Source/Editor/UnrealEd/Public/SCurveEditor.h
--- a/Engine/Source/Editor/UnrealEd/Public/SCurveEditor.h
+++ b/Engine/Source/Editor/UnrealEd/Public/SCurveEditor.h
@@ -254,7 +254,7 @@ public:
 	/** Set the curve that is being edited by this track widget. Also provide an option to enable/disable editing */
 	UNREALED_API void SetCurveOwner(FCurveOwnerInterface* InCurveOwner, bool bCanEdit = true);
 	
-	UNREALED_API void SetPropertyUtils(const TSharedPtr<IPropertyUtilities>& InPropertyUtils)
+	void SetPropertyUtils(const TSharedPtr<IPropertyUtilities>& InPropertyUtils)
 	{
 		PropertyUtilitiesWeak = InPropertyUtils;
 	}
diff --git a/Engine/Source/Editor/UnrealEd/Public/SEditorViewport.h b/Engine/Source/Editor/UnrealEd/Public/SEditorViewport.h
--- a/Engine/Source/Editor/UnrealEd/Public/SEditorViewport.h
+++ b/Engine/Source/Editor/UnrealEd/Public/SEditorViewport.h
@@ -169,7 +169,7 @@ protected:
 	// Implement this to add an arbitrary set of toolbars or other overlays to the inside of the viewport
 	virtual void PopulateViewportOverlays(TSharedRef<SOverlay> Overlay) { }
 
-	UNREALED_API virtual TSharedPtr<IPreviewProfileController> CreatePreviewProfileController() { return nullptr; }
+	virtual TSharedPtr<IPreviewProfileController> CreatePreviewProfileController() { return nullptr; }
 
 	UNREALED_API virtual void BindCommands();
 	UNREALED_API virtual void BindShowCommands( FUICommandList& OutCommandList );
diff --git a/Engine/Source/Editor/UnrealEd/Public/Subsystems/CollectionManagerScriptingSubsystem.h b/Engine/Source/Editor/UnrealEd/Public/Subsystems/CollectionManagerScriptingSubsystem.h
--- a/Engine/Source/Editor/UnrealEd/Public/Subsystems/CollectionManagerScriptingSubsystem.h
+++ b/Engine/Source/Editor/UnrealEd/Public/Subsystems/CollectionManagerScriptingSubsystem.h
@@ -8,6 +8,8 @@
 
 #include "CollectionManagerScriptingSubsystem.generated.h"
 
+#define UE_API UNREALED_API
+
 USTRUCT(BlueprintType, DisplayName="Collection Container Source", meta=(ScriptName="CollectionContainerSource"))
 struct FCollectionScriptingContainerSource
 {
@@ -40,8 +42,8 @@ struct FCollectionScriptingRef
 	ECollectionScriptingShareType ShareType = ECollectionScriptingShareType::Local;
 };
 
-UCLASS(DisplayName="Collection Manager Subsystem", meta=(ScriptName="CollectionManagerSubsystem"))
-class UNREALED_API UCollectionManagerScriptingSubsystem final : public UEditorSubsystem
+UCLASS(MinimalAPI, DisplayName="Collection Manager Subsystem", meta=(ScriptName="CollectionManagerSubsystem"))
+class UCollectionManagerScriptingSubsystem final : public UEditorSubsystem
 {
 	GENERATED_BODY()
 
@@ -52,7 +54,7 @@ public:
 	 * @return Array of all collection containers.
 	 */
 	UFUNCTION(BlueprintCallable, Category="CollectionManager")
-	TArray<FCollectionScriptingContainerSource> GetCollectionContainers();
+	UE_API TArray<FCollectionScriptingContainerSource> GetCollectionContainers();
 
 	/**
 	 * Create a new collection with the given name and share type in the provided collection container.
@@ -65,7 +67,7 @@ public:
 	 * @return True if the collection was created, false otherwise (see the output log for details on error).
 	 */
 	UFUNCTION(BlueprintCallable, Category="CollectionManager")
-	bool CreateCollection(const FCollectionScriptingContainerSource Container, const FName Collection, const ECollectionScriptingShareType ShareType, FCollectionScriptingRef& OutNewCollection);
+	UE_API bool CreateCollection(const FCollectionScriptingContainerSource Container, const FName Collection, const ECollectionScriptingShareType ShareType, FCollectionScriptingRef& OutNewCollection);
 
 	/**
 	 * Create a new collection with the given name and share type in the provided collection container if it doesn't already exist, or empty the existing collection if it does.
@@ -78,7 +80,7 @@ public:
 	 * @return True if the collection was created or emptied, false otherwise (see the output log for details on error).
 	 */
 	UFUNCTION(BlueprintCallable, Category="CollectionManager")
-	bool CreateOrEmptyCollection(const FCollectionScriptingContainerSource Container, const FName Collection, const ECollectionScriptingShareType ShareType, FCollectionScriptingRef& OutNewOrEmptyCollection);
+	UE_API bool CreateOrEmptyCollection(const FCollectionScriptingContainerSource Container, const FName Collection, const ECollectionScriptingShareType ShareType, FCollectionScriptingRef& OutNewOrEmptyCollection);
 
 	/**
 	 * Get all available collections in the specified container.
@@ -89,7 +91,7 @@ public:
 	 * @return True if OutCollections contains at least one collection, false otherwise (see the output log for details on error).
 	 */
 	UFUNCTION(BlueprintCallable, Category="CollectionManager")
-	bool GetCollections(const FCollectionScriptingContainerSource Container, TArray<FCollectionScriptingRef>& OutCollections);
+	UE_API bool GetCollections(const FCollectionScriptingContainerSource Container, TArray<FCollectionScriptingRef>& OutCollections);
 
 	/**
 	 * Destroy the given collection.
@@ -99,7 +101,7 @@ public:
 	 * @return True if the collection was destroyed, false otherwise (see the output log for details on error).
 	 */
 	UFUNCTION(BlueprintCallable, Category="CollectionManager")
-	bool DestroyCollection(const FCollectionScriptingRef& Collection);
+	UE_API bool DestroyCollection(const FCollectionScriptingRef& Collection);
 
 	/**
 	 * Rename the given collection.
@@ -111,7 +113,7 @@ public:
 	 * @return True if the collection was renamed, false otherwise (see the output log for details on error).
 	 */
 	UFUNCTION(BlueprintCallable, Category="CollectionManager")
-	bool RenameCollection(const FCollectionScriptingRef& Collection, const FName NewName, const ECollectionScriptingShareType NewShareType);
+	UE_API bool RenameCollection(const FCollectionScriptingRef& Collection, const FName NewName, const ECollectionScriptingShareType NewShareType);
 
 	/**
 	 * Re-parent the given collection.
@@ -122,7 +124,7 @@ public:
 	 * @return True if the collection was re-parented, false otherwise (see the output log for details on error).
 	 */
 	UFUNCTION(BlueprintCallable, Category="CollectionManager")
-	bool ReparentCollection(const FCollectionScriptingRef& Collection, const FCollectionScriptingRef NewParentCollection);
+	UE_API bool ReparentCollection(const FCollectionScriptingRef& Collection, const FCollectionScriptingRef NewParentCollection);
 
 	/**
 	 * Remove all assets from the given collection.
@@ -132,7 +134,7 @@ public:
 	 * @return True if the collection was modified, false otherwise (see the output log for details on error).
 	 */
 	UFUNCTION(BlueprintCallable, Category="CollectionManager")
-	bool EmptyCollection(const FCollectionScriptingRef& Collection);
+	UE_API bool EmptyCollection(const FCollectionScriptingRef& Collection);
 
 	/**
 	 * Add the given asset to the given collection.
@@ -143,7 +145,7 @@ public:
 	 * @return True if the collection was modified, false otherwise (see the output log for details on error).
 	 */
 	UFUNCTION(BlueprintCallable, Category="CollectionManager")
-	bool AddAssetToCollection(const FCollectionScriptingRef& Collection, const FSoftObjectPath& AssetPath);
+	UE_API bool AddAssetToCollection(const FCollectionScriptingRef& Collection, const FSoftObjectPath& AssetPath);
 
 	/**
 	 * Add the given asset to the given collection.
@@ -154,7 +156,7 @@ public:
 	 * @return True if the collection was modified, false otherwise (see the output log for details on error).
 	 */
 	UFUNCTION(BlueprintCallable, Category="CollectionManager")
-	bool AddAssetDataToCollection(const FCollectionScriptingRef& Collection, const FAssetData& AssetData);
+	UE_API bool AddAssetDataToCollection(const FCollectionScriptingRef& Collection, const FAssetData& AssetData);
 
 	/**
 	 * Add the given asset to the given collection.
@@ -165,7 +167,7 @@ public:
 	 * @return True if the collection was modified, false otherwise (see the output log for details on error).
 	 */
 	UFUNCTION(BlueprintCallable, Category="CollectionManager")
-	bool AddAssetPtrToCollection(const FCollectionScriptingRef& Collection, const UObject* AssetPtr);
+	UE_API bool AddAssetPtrToCollection(const FCollectionScriptingRef& Collection, const UObject* AssetPtr);
 
 	/**
 	 * Add the given assets to the given collection.
@@ -176,7 +178,7 @@ public:
 	 * @return True if the collection was modified, false otherwise (see the output log for details on error).
 	 */
 	UFUNCTION(BlueprintCallable, Category="CollectionManager")
-	bool AddAssetsToCollection(const FCollectionScriptingRef& Collection, const TArray<FSoftObjectPath>& AssetPaths);
+	UE_API bool AddAssetsToCollection(const FCollectionScriptingRef& Collection, const TArray<FSoftObjectPath>& AssetPaths);
 
 	/**
 	 * Add the given assets to the given collection.
@@ -187,7 +189,7 @@ public:
 	 * @return True if the collection was modified, false otherwise (see the output log for details on error).
 	 */
 	UFUNCTION(BlueprintCallable, Category="CollectionManager")
-	bool AddAssetDatasToCollection(const FCollectionScriptingRef& Collection, const TArray<FAssetData>& AssetDatas);
+	UE_API bool AddAssetDatasToCollection(const FCollectionScriptingRef& Collection, const TArray<FAssetData>& AssetDatas);
 
 	/**
 	 * Add the given assets to the given collection.
@@ -198,7 +200,7 @@ public:
 	 * @return True if the collection was modified, false otherwise (see the output log for details on error).
 	 */
 	UFUNCTION(BlueprintCallable, Category="CollectionManager")
-	bool AddAssetPtrsToCollection(const FCollectionScriptingRef& Collection, const TArray<UObject*>& AssetPtrs);
+	UE_API bool AddAssetPtrsToCollection(const FCollectionScriptingRef& Collection, const TArray<UObject*>& AssetPtrs);
 
 	/**
 	 * Remove the given asset from the given collection.
@@ -209,7 +211,7 @@ public:
 	 * @return True if the collection was modified, false otherwise (see the output log for details on error).
 	 */
 	UFUNCTION(BlueprintCallable, Category="CollectionManager")
-	bool RemoveAssetFromCollection(const FCollectionScriptingRef& Collection, const FSoftObjectPath& AssetPath);
+	UE_API bool RemoveAssetFromCollection(const FCollectionScriptingRef& Collection, const FSoftObjectPath& AssetPath);
 
 	/**
 	 * Remove the given asset from the given collection.
@@ -220,7 +222,7 @@ public:
 	 * @return True if the collection was modified, false otherwise (see the output log for details on error).
 	 */
 	UFUNCTION(BlueprintCallable, Category="CollectionManager")
-	bool RemoveAssetDataFromCollection(const FCollectionScriptingRef& Collection, const FAssetData& AssetData);
+	UE_API bool RemoveAssetDataFromCollection(const FCollectionScriptingRef& Collection, const FAssetData& AssetData);
 
 	/**
 	 * Remove the given asset from the given collection.
@@ -231,7 +233,7 @@ public:
 	 * @return True if the collection was modified, false otherwise (see the output log for details on error).
 	 */
 	UFUNCTION(BlueprintCallable, Category="CollectionManager")
-	bool RemoveAssetPtrFromCollection(const FCollectionScriptingRef& Collection, const UObject* AssetPtr);
+	UE_API bool RemoveAssetPtrFromCollection(const FCollectionScriptingRef& Collection, const UObject* AssetPtr);
 
 	/**
 	 * Remove the given assets from the given collection.
@@ -242,7 +244,7 @@ public:
 	 * @return True if the collection was modified, false otherwise (see the output log for details on error).
 	 */
 	UFUNCTION(BlueprintCallable, Category="CollectionManager")
-	bool RemoveAssetsFromCollection(const FCollectionScriptingRef& Collection, const TArray<FSoftObjectPath>& AssetPaths);
+	UE_API bool RemoveAssetsFromCollection(const FCollectionScriptingRef& Collection, const TArray<FSoftObjectPath>& AssetPaths);
 
 	/**
 	 * Remove the given assets from the given collection.
@@ -253,7 +255,7 @@ public:
 	 * @return True if the collection was modified, false otherwise (see the output log for details on error).
 	 */
 	UFUNCTION(BlueprintCallable, Category="CollectionManager")
-	bool RemoveAssetDatasFromCollection(const FCollectionScriptingRef& Collection, const TArray<FAssetData>& AssetDatas);
+	UE_API bool RemoveAssetDatasFromCollection(const FCollectionScriptingRef& Collection, const TArray<FAssetData>& AssetDatas);
 
 	/**
 	 * Remove the given assets from the given collection.
@@ -264,7 +266,7 @@ public:
 	 * @return True if the collection was modified, false otherwise (see the output log for details on error).
 	 */
 	UFUNCTION(BlueprintCallable, Category="CollectionManager")
-	bool RemoveAssetPtrsFromCollection(const FCollectionScriptingRef& Collection, const TArray<UObject*>& AssetPtrs);
+	UE_API bool RemoveAssetPtrsFromCollection(const FCollectionScriptingRef& Collection, const TArray<UObject*>& AssetPtrs);
 
 	/**
 	 * Check whether the given collection exists in the given container (matching both name and share type).
@@ -276,7 +278,7 @@ public:
 	 * @return True if the collection exists, false otherwise (if false is due to an error, see the output log for details).
 	 */
 	UFUNCTION(BlueprintCallable, Category="CollectionManager")
-	bool CollectionExists(const FCollectionScriptingContainerSource Container, const FName Collection, const ECollectionScriptingShareType ShareType);
+	UE_API bool CollectionExists(const FCollectionScriptingContainerSource Container, const FName Collection, const ECollectionScriptingShareType ShareType);
 
 	/**
 	 * Gets the given collections in the given container (matching only by name).
@@ -288,7 +290,7 @@ public:
 	 * @return True if OutCollections contains at least one collection, false otherwise (see the output log for details on error).
 	 */
 	UFUNCTION(BlueprintCallable, Category = "CollectionManager")
-	bool GetCollectionsByName(const FCollectionScriptingContainerSource Container, const FName Collection, TArray<FCollectionScriptingRef>& OutCollections);
+	UE_API bool GetCollectionsByName(const FCollectionScriptingContainerSource Container, const FName Collection, TArray<FCollectionScriptingRef>& OutCollections);
 
 	/**
 	 * Get the assets in the given collection.
@@ -299,7 +301,7 @@ public:
 	 * @return True if the collection exists, false otherwise (see the output log for details on error).
 	 */
 	UFUNCTION(BlueprintCallable, Category="CollectionManager")
-	bool GetAssetsInCollection(const FCollectionScriptingRef& Collection, TArray<FAssetData>& OutAssets);
+	UE_API bool GetAssetsInCollection(const FCollectionScriptingRef& Collection, TArray<FAssetData>& OutAssets);
 
 	/**
 	 * Get the collections in the specified container that contain the given asset.
@@ -311,7 +313,7 @@ public:
 	 * @return True if the container exists, false otherwise (see the output log for details on error).
 	 */
 	UFUNCTION(BlueprintCallable, Category="CollectionManager")
-	bool GetCollectionsContainingAsset(const FCollectionScriptingContainerSource Container, const FSoftObjectPath& AssetPath, TArray<FCollectionScriptingRef>& OutCollections);
+	UE_API bool GetCollectionsContainingAsset(const FCollectionScriptingContainerSource Container, const FSoftObjectPath& AssetPath, TArray<FCollectionScriptingRef>& OutCollections);
 
 	/**
 	 * Get the collections in the specified container that contain the given asset.
@@ -323,7 +325,7 @@ public:
 	 * @return True if the container exists, false otherwise (see the output log for details on error).
 	 */
 	UFUNCTION(BlueprintCallable, Category="CollectionManager")
-	bool GetCollectionsContainingAssetData(const FCollectionScriptingContainerSource Container, const FAssetData& AssetData, TArray<FCollectionScriptingRef>& OutCollections);
+	UE_API bool GetCollectionsContainingAssetData(const FCollectionScriptingContainerSource Container, const FAssetData& AssetData, TArray<FCollectionScriptingRef>& OutCollections);
 
 	/**
 	 * Get the collections in the specified container that contain the given asset.
@@ -335,7 +337,7 @@ public:
 	 * @return True if the container exists, false otherwise (see the output log for details on error).
 	 */
 	UFUNCTION(BlueprintCallable, Category="CollectionManager")
-	bool GetCollectionsContainingAssetPtr(const FCollectionScriptingContainerSource Container, const UObject* AssetPtr, TArray<FCollectionScriptingRef>& OutCollections);
+	UE_API bool GetCollectionsContainingAssetPtr(const FCollectionScriptingContainerSource Container, const UObject* AssetPtr, TArray<FCollectionScriptingRef>& OutCollections);
 
 	/**
 	 * Get the collection container for the base game.
@@ -343,5 +345,7 @@ public:
 	 * @return The collection container for the base game.
 	 */
 	UFUNCTION(BlueprintCallable, Category="CollectionManager")
-	FCollectionScriptingContainerSource GetBaseGameCollectionContainer() const;
+	UE_API FCollectionScriptingContainerSource GetBaseGameCollectionContainer() const;
 };
+
+#undef UE_API
diff --git a/Engine/Source/Editor/UnrealEd/Public/Subsystems/PanelExtensionSubsystem.h b/Engine/Source/Editor/UnrealEd/Public/Subsystems/PanelExtensionSubsystem.h
--- a/Engine/Source/Editor/UnrealEd/Public/Subsystems/PanelExtensionSubsystem.h
+++ b/Engine/Source/Editor/UnrealEd/Public/Subsystems/PanelExtensionSubsystem.h
@@ -13,9 +13,11 @@
 
 #include "PanelExtensionSubsystem.generated.h"
 
+#define UE_API UNREALED_API
+
 PRAGMA_DISABLE_DEPRECATION_WARNINGS
 
-struct UNREALED_API FPanelExtensionFactory
+struct FPanelExtensionFactory
 {
 public:
 	/** An identifier to allow removal later on. */
@@ -41,10 +43,10 @@ public:
 
 PRAGMA_ENABLE_DEPRECATION_WARNINGS
 
-class UNREALED_API SExtensionPanel : public SCompoundWidget
+class SExtensionPanel : public SCompoundWidget
 {
 public:
-	~SExtensionPanel();
+	UE_API ~SExtensionPanel();
 
 	SLATE_BEGIN_ARGS(SExtensionPanel)
 		: _ExtensionPanelID()
@@ -66,7 +68,7 @@ public:
 	SLATE_END_ARGS()
 
 	/** Constructs the widget */
-	void Construct(const FArguments& InArgs);
+	UE_API void Construct(const FArguments& InArgs);
 
 	UObject* GetExtensionContext() const
 	{
@@ -79,7 +81,7 @@ public:
 	}
 
 private:
-	void RebuildWidget();
+	UE_API void RebuildWidget();
 
 	FName ExtensionPanelID;
 	TSharedPtr<SWidget> DefaultWidget;
@@ -120,3 +122,5 @@ private:
 	TMap<FName, FPanelFactoryRegistryChanged> PanelFactoryRegistryChangedCallbackMap;
 	TMap<FName, FText> DeprecatedExtensions;
 };
+
+#undef UE_API
diff --git a/Engine/Source/Editor/UnrealEd/Public/Toolkits/AssetEditorToolkit.h b/Engine/Source/Editor/UnrealEd/Public/Toolkits/AssetEditorToolkit.h
--- a/Engine/Source/Editor/UnrealEd/Public/Toolkits/AssetEditorToolkit.h
+++ b/Engine/Source/Editor/UnrealEd/Public/Toolkits/AssetEditorToolkit.h
@@ -172,7 +172,7 @@ public:
 	UNREALED_API virtual TSharedPtr<FTabManager> GetAssociatedTabManager() override;
 	UNREALED_API virtual double GetLastActivationTime() override;
 	UNREALED_API virtual void RemoveEditingAsset(UObject* Asset) override;
-	UNREALED_API virtual EAssetOpenMethod GetOpenMethod() const override { return OpenMethod; }
+	virtual EAssetOpenMethod GetOpenMethod() const override { return OpenMethod; }
 
 	/**
 	 * Fills in the supplied menu with commands for working with this asset file
@@ -239,7 +239,7 @@ public:
 	UNREALED_API virtual FName GetToolMenuToolbarName(FName& OutParentName) const;
 	UNREALED_API virtual void InitToolMenuContext(FToolMenuContext& MenuContext);
 	/** Method that can be overriden by derived classes to provide a custom MenuBar to the Standalone Toolkit Host. */
-	UNREALED_API virtual TSharedPtr<SWidget> CreateMenuBar(const TSharedPtr<FTabManager>& InTabManager, const FName InMenuName, FToolMenuContext& InToolMenuContext) { return nullptr; }
+	virtual TSharedPtr<SWidget> CreateMenuBar(const TSharedPtr<FTabManager>& InTabManager, const FName InMenuName, FToolMenuContext& InToolMenuContext) { return nullptr; }
 	UNREALED_API FName GetToolMenuToolbarName() const;
 	UNREALED_API FName GetToolMenuAppName() const;
 	UNREALED_API FName GetToolMenuName() const;
@@ -447,10 +447,10 @@ protected:
 	UNREALED_API virtual void CreateEditorModeManager() override;
 
 	/** Specify the permission lists to use in read only mode */
-	UNREALED_API virtual void SetupReadOnlyMenuProfiles(FReadOnlyAssetEditorCustomization& OutReadOnlyCustomization) { }
+	virtual void SetupReadOnlyMenuProfiles(FReadOnlyAssetEditorCustomization& OutReadOnlyCustomization) { }
 
 	/** Get the name of the profile registered with UToolMenus for read only menu customizations */
-	UNREALED_API virtual FName GetReadOnlyMenuProfileName() { return ReadOnlyMenuProfileName; }
+	virtual FName GetReadOnlyMenuProfileName() { return ReadOnlyMenuProfileName; }
 
 	/** Can be overriden to prevent registering the default commands.  */
 	UNREALED_API virtual void MapToolkitCommands();
diff --git a/Engine/Source/Editor/UnrealEd/Public/Tools/EdModeInteractiveToolsContext.h b/Engine/Source/Editor/UnrealEd/Public/Tools/EdModeInteractiveToolsContext.h
--- a/Engine/Source/Editor/UnrealEd/Public/Tools/EdModeInteractiveToolsContext.h
+++ b/Engine/Source/Editor/UnrealEd/Public/Tools/EdModeInteractiveToolsContext.h
@@ -25,6 +25,8 @@
 
 #include "EdModeInteractiveToolsContext.generated.h"
 
+#define UE_API UNREALED_API
+
 class FCanvas;
 class FEdMode;
 class FEditorModeTools;
@@ -59,34 +61,34 @@ struct FToolBuilderState;
  * called from the FEdMode functions of the same name. These will handle the data type
  * conversions and forwarding calls necessary to operate the ToolsContext
  */
-UCLASS(Transient)
-class UNREALED_API UEditorInteractiveToolsContext : public UInteractiveToolsContext
+UCLASS(MinimalAPI, Transient)
+class UEditorInteractiveToolsContext : public UInteractiveToolsContext
 {
 	GENERATED_BODY()
 
 public:
-	UEditorInteractiveToolsContext();
+	UE_API UEditorInteractiveToolsContext();
 
 	/**
 	 * Initialize a new ToolsContext for an EdMode owned by the given InEditorModeManager
 	 * @param 
 	 */
-	void InitializeContextWithEditorModeManager(FEditorModeTools* InEditorModeManager, UInputRouter* UseInputRouter = nullptr);
+	UE_API void InitializeContextWithEditorModeManager(FEditorModeTools* InEditorModeManager, UInputRouter* UseInputRouter = nullptr);
 
 	/** Shutdown ToolsContext and clean up any connections/etc */
-	virtual void ShutdownContext();
+	UE_API virtual void ShutdownContext();
 
 	// default behavior is to accept active tool
-	virtual void TerminateActiveToolsOnPIEStart();
+	UE_API virtual void TerminateActiveToolsOnPIEStart();
 
 	// default behavior is to accept active tool
-	virtual void TerminateActiveToolsOnSaveWorld();
+	UE_API virtual void TerminateActiveToolsOnSaveWorld();
 
 	// default behavior is to cancel active tool
-	virtual void TerminateActiveToolsOnWorldTearDown();
+	UE_API virtual void TerminateActiveToolsOnWorldTearDown();
 
 	// default behavior is to cancel active tool
-	virtual void TerminateActiveToolsOnLevelChange();
+	UE_API virtual void TerminateActiveToolsOnLevelChange();
 
 	FEditorModeTools* GetParentEditorModeManager() const { return EditorModeManager; }
 
@@ -94,15 +96,15 @@ public:
 	IToolsContextTransactionsAPI* GetTransactionAPI() const { return TransactionAPI; }
 
 	/** Call this to notify the Editor that the viewports this ToolsContext is related to may need a repaint, ie during interactive tool usage */
-	virtual void PostInvalidation();
+	UE_API virtual void PostInvalidation();
 
 	// UObject Interface
-	virtual UWorld* GetWorld() const override;
+	UE_API virtual UWorld* GetWorld() const override;
 
 	// call functions of the same name on the ToolManager and GizmoManager
-	virtual void Tick(FEditorViewportClient* ViewportClient, float DeltaTime);
-	virtual void Render(const FSceneView* View, FViewport* Viewport, FPrimitiveDrawInterface* PDI);
-	virtual void DrawHUD(FViewportClient* ViewportClient,FViewport* Viewport,const FSceneView* View, FCanvas* Canvas);
+	UE_API virtual void Tick(FEditorViewportClient* ViewportClient, float DeltaTime);
+	UE_API virtual void Render(const FSceneView* View, FViewport* Viewport, FPrimitiveDrawInterface* PDI);
+	UE_API virtual void DrawHUD(FViewportClient* ViewportClient,FViewport* Viewport,const FSceneView* View, FCanvas* Canvas);
 
 	// These delegates can be used to hook into the Render() / DrawHUD() / Tick() calls above. In particular, non-legacy UEdMode's
 	// don't normally receive Render() and DrawHUD() calls from the mode manager, but can attach to these.
@@ -114,23 +116,23 @@ public:
 	FOnTick OnTick;
 
 	/** @return true if selected actors/components can be deleted */
-	virtual bool ProcessEditDelete();
+	UE_API virtual bool ProcessEditDelete();
 
 	//
 	// Utility functions useful for hooking up to UICommand/etc
 	//
 
-	virtual bool CanStartTool(const FString ToolTypeIdentifier) const;
-	virtual bool HasActiveTool() const;
-	virtual FString GetActiveToolName() const;
-	virtual bool ActiveToolHasAccept() const;
-	virtual bool CanAcceptActiveTool() const;
-	virtual bool CanCancelActiveTool() const;
-	virtual bool CanCompleteActiveTool() const;
-	virtual void StartTool(const FString ToolTypeIdentifier);
-	virtual void EndTool(EToolShutdownType ShutdownType);
-	void Activate();
-	void Deactivate();
+	UE_API virtual bool CanStartTool(const FString ToolTypeIdentifier) const;
+	UE_API virtual bool HasActiveTool() const;
+	UE_API virtual FString GetActiveToolName() const;
+	UE_API virtual bool ActiveToolHasAccept() const;
+	UE_API virtual bool CanAcceptActiveTool() const;
+	UE_API virtual bool CanCancelActiveTool() const;
+	UE_API virtual bool CanCompleteActiveTool() const;
+	UE_API virtual void StartTool(const FString ToolTypeIdentifier);
+	UE_API virtual void EndTool(EToolShutdownType ShutdownType);
+	UE_API void Activate();
+	UE_API void Deactivate();
 
 
 	/** @return Ray into 3D scene at last mouse event */
@@ -148,7 +150,7 @@ public:
 	 * Configure whether ::Render() should early-out for HitProxy rendering passes.
 	 * If the Mode does not use HitProxy, and the Tools/Gizmos have expensive Render() calls, this can help with interactive performance.
 	 */
-	void SetEnableRenderingDuringHitProxyPass(bool bEnabled);
+	UE_API void SetEnableRenderingDuringHitProxyPass(bool bEnabled);
 
 	/** @return true if HitProxy rendering will be allowed in ::Render() */
 	bool GetEnableRenderingDuringHitProxyPass() const { return bEnableRenderingDuringHitProxyPass; }
@@ -158,7 +160,7 @@ public:
 	 * Configure whether Transform Gizmos created by the ITF (eg CombinedTransformGizmo) should prefer to show in 'Combined' mode.
 	 * If this is disabled, the Gizmo should respect the active Editor Gizmo setting (eg in the Level Viewport)
 	 */
-	void SetForceCombinedGizmoMode(bool bEnabled);
+	UE_API void SetForceCombinedGizmoMode(bool bEnabled);
 
 	/** @return true if Force Combined Gizmo mode is Enabled */
 	bool GetForceCombinedGizmoModeEnabled() const { return bForceCombinedGizmoMode; }
@@ -169,7 +171,7 @@ public:
 	 * snap to an Absolute world-aligned grid, or snap Relative to the initial position of any particular gizmo transform.
 	 * Relative is the default and is also the behavior of the standard UE Gizmo.
 	 */
-	void SetAbsoluteWorldSnappingEnabled(bool bEnabled);
+	UE_API void SetAbsoluteWorldSnappingEnabled(bool bEnabled);
 
 	/** @return true if Absolute World Snapping mode is Enabled */
 	bool GetAbsoluteWorldSnappingEnabled() const { return bEnableAbsoluteWorldSnapping; }
@@ -180,7 +182,7 @@ public:
 	 * Setting this to true will prevent this, but the system makes no promises about whether tools work properly across PIE start/run/shutdown.
 	 * Therefore, do not use this if you do not know for certain that it is safe to persist tools across PIE start in your situation.
 	 */
-	void SetDeactivateToolsOnPIEStart(bool bDeactivateTools);
+	UE_API void SetDeactivateToolsOnPIEStart(bool bDeactivateTools);
 
 	/** @return true if tools should be deactivated when PIE is started */
 	bool GetDeactivateToolsOnPIEStart() const { return bDeactivateOnPIEStart; }
@@ -191,7 +193,7 @@ public:
 	 * Setting this to true will prevent this, but the system makes no promises about whether tools work properly across a save.
 	 * Therefore, do not use this if you do not know for certain that it is safe to persist tools across a save in your situation.
 	 */
-	void SetDeactivateToolsOnSaveWorld(bool bDeactivateTools);
+	UE_API void SetDeactivateToolsOnSaveWorld(bool bDeactivateTools);
 
 	/** @return true if tools should be deactivated when save is started */
 	bool GetDeactivateToolsOnSaveWorld() const { return bDeactivateOnSaveWorld; }
@@ -199,11 +201,11 @@ public:
 protected:
 
 	// we hide these 
-	virtual void Initialize(IToolsContextQueriesAPI* QueriesAPI, IToolsContextTransactionsAPI* TransactionsAPI) override;
-	virtual void Shutdown() override;
+	UE_API virtual void Initialize(IToolsContextQueriesAPI* QueriesAPI, IToolsContextTransactionsAPI* TransactionsAPI) override;
+	UE_API virtual void Shutdown() override;
 
-	virtual void DeactivateActiveTool(EToolSide WhichSide, EToolShutdownType ShutdownType);
-	virtual void DeactivateAllActiveTools(EToolShutdownType ShutdownType);
+	UE_API virtual void DeactivateActiveTool(EToolSide WhichSide, EToolShutdownType ShutdownType);
+	UE_API virtual void DeactivateAllActiveTools(EToolShutdownType ShutdownType);
 
 public:
 	UPROPERTY()
@@ -230,15 +232,15 @@ protected:
 
 	// Utility function to convert viewport x/y from mouse events (and others?) into scene ray.
 	// Copy-pasted from other Editor code, seems kind of expensive?
-	static FRay GetRayFromMousePos(FEditorViewportClient* ViewportClient, FViewport* Viewport, int MouseX, int MouseY);
+	static UE_API FRay GetRayFromMousePos(FEditorViewportClient* ViewportClient, FViewport* Viewport, int MouseX, int MouseY);
 
 	// editor UI state that we set before starting tool and when exiting tool
 	// Currently disabling anti-aliasing during active Tools because it causes PDI flickering
-	void SetEditorStateForTool();
-	void RestoreEditorState();
+	UE_API void SetEditorStateForTool();
+	UE_API void RestoreEditorState();
 
-	void OnToolEnded(UInteractiveToolManager* InToolManager, UInteractiveTool* InEndedTool);
-	void OnToolPostBuild(UInteractiveToolManager* InToolManager, EToolSide InSide, UInteractiveTool* InBuiltTool, UInteractiveToolBuilder* InToolBuilder, const FToolBuilderState& ToolState);
+	UE_API void OnToolEnded(UInteractiveToolManager* InToolManager, UInteractiveTool* InEndedTool);
+	UE_API void OnToolPostBuild(UInteractiveToolManager* InToolManager, EToolSide InSide, UInteractiveTool* InBuiltTool, UInteractiveToolBuilder* InToolBuilder, const FToolBuilderState& ToolState);
 
 	TOptional<FString> PendingToolToStart = {};
 	TOptional<EToolShutdownType> PendingToolShutdownType = {};
@@ -263,8 +265,8 @@ private:
  * UModeManagerInteractiveToolsContext extends UEditorInteractiveToolsContext with various functions for handling 
  * device (mouse) input. These functions are currently called by the EdMode Manager (FEditorModeTools).
  */
-UCLASS(Transient)
-class UNREALED_API UModeManagerInteractiveToolsContext : public UEditorInteractiveToolsContext
+UCLASS(MinimalAPI, Transient)
+class UModeManagerInteractiveToolsContext : public UEditorInteractiveToolsContext
 {
 	GENERATED_BODY()
 
@@ -272,43 +274,43 @@ class UNREALED_API UModeManagerInteractiveToolsContext : public UEditorInteracti
 	// UEditorInteractiveToolsContext API implementations that also forward calls to any child EdMode ToolsContexts
 	//
 public:
-	virtual void Tick(FEditorViewportClient* ViewportClient, float DeltaTime) override;
-	virtual void Render(const FSceneView* View, FViewport* Viewport, FPrimitiveDrawInterface* PDI) override;
-	virtual void DrawHUD(FViewportClient* ViewportClient,FViewport* Viewport,const FSceneView* View, FCanvas* Canvas) override;
+	UE_API virtual void Tick(FEditorViewportClient* ViewportClient, float DeltaTime) override;
+	UE_API virtual void Render(const FSceneView* View, FViewport* Viewport, FPrimitiveDrawInterface* PDI) override;
+	UE_API virtual void DrawHUD(FViewportClient* ViewportClient,FViewport* Viewport,const FSceneView* View, FCanvas* Canvas) override;
 
-	virtual bool ProcessEditDelete() override;
+	UE_API virtual bool ProcessEditDelete() override;
 
 	//
 	// Input handling, these functions forward ViewportClient events to the UInputRouter
 	//
 public:
-	bool InputKey(FEditorViewportClient* ViewportClient, FViewport* Viewport, FKey Key, EInputEvent Event);
+	UE_API bool InputKey(FEditorViewportClient* ViewportClient, FViewport* Viewport, FKey Key, EInputEvent Event);
 	/**
 	 * This updates internal state like InputKey, but doesn't route the results to the input router. 
 	 * Use this if the input is captured by some higher system, to avoid this class from having an
 	 * incorrect view of e.g. the mouse state because it did not receive a mouse release event.
 	 */
-	void UpdateStateWithoutRoutingInputKey(FEditorViewportClient* ViewportClient, FViewport* Viewport, FKey Key, EInputEvent Event);
+	UE_API void UpdateStateWithoutRoutingInputKey(FEditorViewportClient* ViewportClient, FViewport* Viewport, FKey Key, EInputEvent Event);
 
-	bool MouseEnter(FEditorViewportClient* ViewportClient, FViewport* Viewport, int32 x, int32 y);
-	bool MouseLeave(FEditorViewportClient* ViewportClient, FViewport* Viewport);
-	bool MouseMove(FEditorViewportClient* ViewportClient, FViewport* Viewport, int32 x, int32 y);
-	bool LostFocus(FEditorViewportClient* InViewportClient, FViewport* Viewport) const;
+	UE_API bool MouseEnter(FEditorViewportClient* ViewportClient, FViewport* Viewport, int32 x, int32 y);
+	UE_API bool MouseLeave(FEditorViewportClient* ViewportClient, FViewport* Viewport);
+	UE_API bool MouseMove(FEditorViewportClient* ViewportClient, FViewport* Viewport, int32 x, int32 y);
+	UE_API bool LostFocus(FEditorViewportClient* InViewportClient, FViewport* Viewport) const;
 
-	bool StartTracking(FEditorViewportClient* InViewportClient, FViewport* InViewport);
-	bool CapturedMouseMove(FEditorViewportClient* InViewportClient, FViewport* InViewport, int32 InMouseX, int32 InMouseY);
-	bool EndTracking(FEditorViewportClient* InViewportClient, FViewport* InViewport);
+	UE_API bool StartTracking(FEditorViewportClient* InViewportClient, FViewport* InViewport);
+	UE_API bool CapturedMouseMove(FEditorViewportClient* InViewportClient, FViewport* InViewport, int32 InMouseX, int32 InMouseY);
+	UE_API bool EndTracking(FEditorViewportClient* InViewportClient, FViewport* InViewport);
 
 	/** @return True if the context has overriden the cursor style, false if not. */
-	bool GetCursor(EMouseCursor::Type& OutCursor) const;
+	UE_API bool GetCursor(EMouseCursor::Type& OutCursor) const;
 
 	/** @return Ray into 3D scene at last mouse event */
-	virtual FRay GetLastWorldRay() const override;
+	UE_API virtual FRay GetLastWorldRay() const override;
 
 protected:
-	virtual void DeactivateAllActiveTools(EToolShutdownType ShutdownType) override;
-	virtual void Initialize(IToolsContextQueriesAPI* QueriesAPIIn, IToolsContextTransactionsAPI* TransactionsAPIIn) override;
-	virtual void Shutdown() override;
+	UE_API virtual void DeactivateAllActiveTools(EToolShutdownType ShutdownType) override;
+	UE_API virtual void Initialize(IToolsContextQueriesAPI* QueriesAPIIn, IToolsContextTransactionsAPI* TransactionsAPIIn) override;
+	UE_API virtual void Shutdown() override;
 
 	/** Input event instance used to keep track of various button states, etc, that we cannot directly query on-demand */
 	FInputDeviceState CurrentMouseState;
@@ -339,19 +341,19 @@ public:
 	* When the child ToolsContext is shut down, OnChildEdModeToolsContextShutdown() must be called to clean up
 	* @return new ToolsContext
 	*/
-	UEdModeInteractiveToolsContext* CreateNewChildEdModeToolsContext();
+	UE_API UEdModeInteractiveToolsContext* CreateNewChildEdModeToolsContext();
 
 	/**
 	 * Call to add a child EdMode ToolsContext created using the above function
 	 * @return true if child was added
 	 */
-	bool OnChildEdModeActivated(UEdModeInteractiveToolsContext* ChildToolsContext);
+	UE_API bool OnChildEdModeActivated(UEdModeInteractiveToolsContext* ChildToolsContext);
 
 	/**
 	 * Call to release a child EdMode ToolsContext created using the above function
 	 * @return true if child was found and removed
 	 */
-	bool OnChildEdModeDeactivated(UEdModeInteractiveToolsContext* ChildToolsContext);
+	UE_API bool OnChildEdModeDeactivated(UEdModeInteractiveToolsContext* ChildToolsContext);
 
 	
 	/**
@@ -375,17 +377,17 @@ protected:
 	/**
 	 * Activating Drag Tools if current Mode Manager requires them
 	 */
-	void RegisterDragTools();
+	UE_API void RegisterDragTools();
 
 	/**
 	 * Deactivating Drag Tools
 	 */
-	void UnregisterDragTools();
+	UE_API void UnregisterDragTools();
 
 	/**
 	 * Responding to Mode Changes in order to deactivate/activate Drag Tools accordingly
 	 */
-	void OnEditorModeChanged(const FEditorModeID& InModeID, bool bInIsEnteringMode);
+	UE_API void OnEditorModeChanged(const FEditorModeID& InModeID, bool bInIsEnteringMode);
 
 	//~Begin ITF Alt processing Bypass
 
@@ -395,7 +397,7 @@ protected:
 	 */
 
 	/** Enable/Disable Alt modifier bypass */
-	void SetBypassAltModifier(bool bInBypassAltModifier);
+	UE_API void SetBypassAltModifier(bool bInBypassAltModifier);
 
 	/** Is this Interactive Tools Context skipping Alt processing? */
 	bool BypassAltModifier() const { return bBypassAltModifier; }
@@ -419,8 +421,8 @@ protected:
  * UEdModeInteractiveToolsContext is an UEditorInteractiveToolsContext intended for use/lifetime in the context of a UEdMode.
  * This ITC subclass is dependent on a UModeManagerInteractiveToolsContext to provide an InputRouter.
  */
-UCLASS(Transient)
-class UNREALED_API UEdModeInteractiveToolsContext : public UEditorInteractiveToolsContext
+UCLASS(MinimalAPI, Transient)
+class UEdModeInteractiveToolsContext : public UEditorInteractiveToolsContext
 {
 	friend class UModeManagerInteractiveToolsContext;
 
@@ -430,12 +432,14 @@ public:
 	 * Initialize a new EdModeToolsContext that is derived from a ModeManagerToolsContext.
 	 * This new ToolsContext will not have it's own InputRouter, it will share the InputRouter with the ModeManagerToolsContext
 	 */
-	void InitializeContextFromModeManagerContext(UModeManagerInteractiveToolsContext* ModeManagerToolsContext);
+	UE_API void InitializeContextFromModeManagerContext(UModeManagerInteractiveToolsContext* ModeManagerToolsContext);
 
 	/** @return Ray into 3D scene at last mouse event */
-	virtual FRay GetLastWorldRay() const override;
+	UE_API virtual FRay GetLastWorldRay() const override;
 
 protected:
 	UPROPERTY()
 	TObjectPtr<UModeManagerInteractiveToolsContext> ParentModeManagerToolsContext = nullptr;
 };
+
+#undef UE_API
diff --git a/Engine/Source/Editor/UnrealEd/Public/ViewportToolbar/UnrealEdViewportToolbar.h b/Engine/Source/Editor/UnrealEd/Public/ViewportToolbar/UnrealEdViewportToolbar.h
--- a/Engine/Source/Editor/UnrealEd/Public/ViewportToolbar/UnrealEdViewportToolbar.h
+++ b/Engine/Source/Editor/UnrealEd/Public/ViewportToolbar/UnrealEdViewportToolbar.h
@@ -10,6 +10,8 @@
 #include "UnrealEdViewportToolbarContext.h"
 #include "UnrealWidgetFwd.h"
 
+#define UE_API UNREALED_API
+
 class FEditorViewportClient;
 class FTabManager;
 class FText;
@@ -124,7 +126,7 @@ UNREALED_API TSharedRef<SWidget> CreateLocationGridSnapMenu(
 	const TSharedPtr<FUICommandList>& InCommandList = {}
 );
 
-struct UNREALED_API FLocationGridSnapMenuOptions
+struct FLocationGridSnapMenuOptions
 {
 	FName MenuName;
 	FLocationGridCheckboxListExecuteActionDelegate ExecuteDelegate;
@@ -189,17 +191,17 @@ UNREALED_API FName GetCameraSubmenuIconFNameFromViewportType(const ELevelViewpor
 UE_DEPRECATED(5.6, "Use CreateCameraSubmenu() instead.")
 UNREALED_API FToolMenuEntry CreateViewportToolbarCameraSubmenu();
 
-struct UNREALED_API FViewportCameraMenuOptions
+struct FViewportCameraMenuOptions
 {
-	FViewportCameraMenuOptions();
+	UE_API FViewportCameraMenuOptions();
 
 	bool bShowCameraMovement: 1;
 	bool bShowFieldOfView: 1;
 	bool bShowNearAndFarPlanes: 1;
 	
-	FViewportCameraMenuOptions& ShowAll();
-	FViewportCameraMenuOptions& ShowCameraMovement();
-	FViewportCameraMenuOptions& ShowLensControls();
+	UE_API FViewportCameraMenuOptions& ShowAll();
+	UE_API FViewportCameraMenuOptions& ShowCameraMovement();
+	UE_API FViewportCameraMenuOptions& ShowLensControls();
 };
 
 UNREALED_API FToolMenuEntry CreateCameraSubmenu(const FViewportCameraMenuOptions& InOptions = {});
@@ -298,3 +300,5 @@ TSharedRef<SWidget> CreateCurrentScreenPercentageSettingWidget(FEditorViewportCl
 TSharedRef<SWidget> CreateCurrentScreenPercentageWidget(FEditorViewportClient& InViewportClient);
 
 } // namespace UE::UnrealEd
+
+#undef UE_API
diff --git a/Engine/Source/Editor/UnrealEd/Public/ViewportToolbar/UnrealEdViewportToolbarContext.h b/Engine/Source/Editor/UnrealEd/Public/ViewportToolbar/UnrealEdViewportToolbarContext.h
--- a/Engine/Source/Editor/UnrealEd/Public/ViewportToolbar/UnrealEdViewportToolbarContext.h
+++ b/Engine/Source/Editor/UnrealEd/Public/ViewportToolbar/UnrealEdViewportToolbarContext.h
@@ -1,4 +1,4 @@
-﻿// Copyright Epic Games, Inc. All Rights Reserved.
+// Copyright Epic Games, Inc. All Rights Reserved.
 
 #pragma once
 
@@ -10,6 +10,8 @@
 
 #include "UnrealEdViewportToolbarContext.generated.h"
 
+#define UE_API UNREALED_API
+
 class FAssetEditorToolkit;
 class IPreviewProfileController;
 enum ERotationGridMode : int;
@@ -26,8 +28,8 @@ DECLARE_DELEGATE_RetVal_OneParam(bool, DoesViewModeMenuShowSectionDelegate, EHid
 
 class SEditorViewport;
 
-UCLASS()
-class UNREALED_API UUnrealEdViewportToolbarContext : public UObject
+UCLASS(MinimalAPI)
+class UUnrealEdViewportToolbarContext : public UObject
 {
 	GENERATED_BODY()
 
@@ -56,23 +58,25 @@ public:
 	 */
 	FTabId PreviewSettingsTabId;
 	
-	virtual TSharedPtr<IPreviewProfileController> GetPreviewProfileController() const;
+	UE_API virtual TSharedPtr<IPreviewProfileController> GetPreviewProfileController() const;
 	
 	bool bShowSurfaceSnap = true;
 	
-	virtual void RefreshViewport();
+	UE_API virtual void RefreshViewport();
 	
-	virtual FText GetGridSnapLabel() const;
-	virtual TArray<float> GetGridSnapSizes() const;
-	virtual bool IsGridSnapSizeActive(int32 GridSizeIndex) const;
-	virtual void SetGridSnapSize(int32 GridSizeIndex);
+	UE_API virtual FText GetGridSnapLabel() const;
+	UE_API virtual TArray<float> GetGridSnapSizes() const;
+	UE_API virtual bool IsGridSnapSizeActive(int32 GridSizeIndex) const;
+	UE_API virtual void SetGridSnapSize(int32 GridSizeIndex);
 	
-	virtual FText GetRotationSnapLabel() const;
-	virtual bool IsRotationSnapActive(int32 RotationIndex, ERotationGridMode RotationMode) const;
-	virtual void SetRotationSnapSize(int32 RotationIndex, ERotationGridMode RotationMode);
+	UE_API virtual FText GetRotationSnapLabel() const;
+	UE_API virtual bool IsRotationSnapActive(int32 RotationIndex, ERotationGridMode RotationMode) const;
+	UE_API virtual void SetRotationSnapSize(int32 RotationIndex, ERotationGridMode RotationMode);
 	
-	virtual FText GetScaleSnapLabel() const;
-	virtual TArray<float> GetScaleSnapSizes() const;
-	virtual bool IsScaleSnapActive(int32 ScaleIndex) const;
-	virtual void SetScaleSnapSize(int32 ScaleIndex);
+	UE_API virtual FText GetScaleSnapLabel() const;
+	UE_API virtual TArray<float> GetScaleSnapSizes() const;
+	UE_API virtual bool IsScaleSnapActive(int32 ScaleIndex) const;
+	UE_API virtual void SetScaleSnapSize(int32 ScaleIndex);
 };
+
+#undef UE_API
diff --git a/Engine/Source/Editor/UnrealEd/Public/WorldPartition/WorldPartitionBuilder.h b/Engine/Source/Editor/UnrealEd/Public/WorldPartition/WorldPartitionBuilder.h
--- a/Engine/Source/Editor/UnrealEd/Public/WorldPartition/WorldPartitionBuilder.h
+++ b/Engine/Source/Editor/UnrealEd/Public/WorldPartition/WorldPartitionBuilder.h
@@ -50,8 +50,8 @@ public:
 
 	UNREALED_API bool RunBuilder(UWorld* World);
 
-	UNREALED_API virtual bool RequiresCommandletRendering() const PURE_VIRTUAL(UWorldPartitionBuilder::RequiresCommandletRendering, return false;);
-	UNREALED_API virtual ELoadingMode GetLoadingMode() const PURE_VIRTUAL(UWorldPartitionBuilder::GetLoadingMode, return ELoadingMode::Custom;);
+	virtual bool RequiresCommandletRendering() const PURE_VIRTUAL(UWorldPartitionBuilder::RequiresCommandletRendering, return false;);
+	virtual ELoadingMode GetLoadingMode() const PURE_VIRTUAL(UWorldPartitionBuilder::GetLoadingMode, return ELoadingMode::Custom;);
 	
 	UNREALED_API bool Run(UWorld* World, FPackageSourceControlHelper& PackageHelper);
 
@@ -68,7 +68,7 @@ public:
 	 * Overridable method for derived classes to add additional world to be process by world builder
 	 * This is called after `RunInternal` is executed.
 	 */
-	UNREALED_API virtual bool ShouldProcessAdditionalWorlds(UWorld* InWorld, TArray<FString>& OutPackageNames) const { return false; }
+	virtual bool ShouldProcessAdditionalWorlds(UWorld* InWorld, TArray<FString>& OutPackageNames) const { return false; }
 
 protected:
 	/**
@@ -77,7 +77,7 @@ protected:
 	 */
 	virtual bool PreRun(UWorld* World, FPackageSourceControlHelper& PackageHelper) { return true; }
 
-	UNREALED_API virtual bool RunInternal(UWorld* World, const FCellInfo& InCellInfo, FPackageSourceControlHelper& PackageHelper) PURE_VIRTUAL(UWorldPartition::RunInternal, return false;);
+	virtual bool RunInternal(UWorld* World, const FCellInfo& InCellInfo, FPackageSourceControlHelper& PackageHelper) PURE_VIRTUAL(UWorldPartition::RunInternal, return false;);
 
 	/**
 	 * Overridable method for derived classes to perform operations when world builder process completes.
diff --git a/Engine/Source/Editor/VREditor/Public/Teleporter/VREditorTeleporter.h b/Engine/Source/Editor/VREditor/Public/Teleporter/VREditorTeleporter.h
--- a/Engine/Source/Editor/VREditor/Public/Teleporter/VREditorTeleporter.h
+++ b/Engine/Source/Editor/VREditor/Public/Teleporter/VREditorTeleporter.h
@@ -9,6 +9,8 @@
 #include "Engine/EngineBaseTypes.h"
 #include "VREditorTeleporter.generated.h"
 
+#define UE_API VREDITOR_API
+
 
 class UMaterialInstanceDynamic;
 class UVREditorInteractor;
@@ -18,58 +20,58 @@ class UVREditorMode;
 /**
  * VR Editor teleport manager and the visual representation of the teleport
  */
-UCLASS( Blueprintable, BlueprintType )
-class VREDITOR_API AVREditorTeleporter: public AActor
+UCLASS(MinimalAPI,  Blueprintable, BlueprintType )
+class AVREditorTeleporter: public AActor
 {
 	GENERATED_BODY()
 
 public:
 
 	/** Default constructor that sets up CDO properties */
-	AVREditorTeleporter();
+	UE_API AVREditorTeleporter();
 
 	/** Initializes the teleporter */
 	UFUNCTION( BlueprintNativeEvent, CallInEditor, Category = "Teleporter" )
-	void Init( class UVREditorMode* InMode );
+	UE_API void Init( class UVREditorMode* InMode );
 
 	/** Shuts down the teleporter */
 	UFUNCTION( BlueprintNativeEvent, CallInEditor, Category = "Teleporter" )
-	void Shutdown();
+	UE_API void Shutdown();
 
 	/** Whether we are currently aiming to teleport. */
 	UFUNCTION( BlueprintCallable, Category = "Teleporter" )
-	bool IsAiming() const;
+	UE_API bool IsAiming() const;
 
 	UFUNCTION( BlueprintCallable, Category = "Teleporter" )
-	bool IsTeleporting() const;
+	UE_API bool IsTeleporting() const;
 
 	UFUNCTION( BlueprintCallable, Category = "Teleporter" )
-	class UVREditorMode* GetVRMode() const;
+	UE_API class UVREditorMode* GetVRMode() const;
 
 	/** Start teleporting, does a ray trace with the hand passed and calculates the locations for lerp movement in Teleport */
 	UFUNCTION( BlueprintNativeEvent, CallInEditor, Category = "Teleporter" )
-	void StartTeleport();
+	UE_API void StartTeleport();
 
 	/** Called when teleport is done for cleanup */
 	UFUNCTION( BlueprintNativeEvent, CallInEditor, Category = "Teleporter" )
-	void TeleportDone();
+	UE_API void TeleportDone();
 
 	/** Hide or show the teleporter visuals */
 	UFUNCTION( BlueprintCallable, Category = "Teleporter" )
-	void SetVisibility( const bool bVisible );
+	UE_API void SetVisibility( const bool bVisible );
 
 	/** Sets the color for the teleporter visuals */
 	UFUNCTION( BlueprintCallable, Category = "Teleporter" )
-	void SetColor( const FLinearColor& Color );
+	UE_API void SetColor( const FLinearColor& Color );
 
 	/** Get slide delta to push/pull or scale the teleporter */
 	UFUNCTION( BlueprintNativeEvent, Category = "Teleporter" )
-	float GetSlideDelta( UVREditorInteractor* Interactor, const bool Axis );
+	UE_API float GetSlideDelta( UVREditorInteractor* Interactor, const bool Axis );
 
 private:
 
 	//~ Begin AActor interface
-	virtual void Tick(const float DeltaTime) override;
+	UE_API virtual void Tick(const float DeltaTime) override;
 	virtual bool IsEditorOnly() const final
 	{
 		return true;
@@ -78,45 +80,45 @@ private:
 
 	/** Functions we call to handle teleporting in navigation mode */
 	UFUNCTION( BlueprintCallable, Category = "Teleporter" )
-	void StartAiming( class UViewportInteractor* Interactor );
+	UE_API void StartAiming( class UViewportInteractor* Interactor );
 
 	/** Cancel teleport aiming mode without doing the teleport */
 	UFUNCTION( BlueprintCallable, Category = "Teleporter" )
-	void StopAiming();
+	UE_API void StopAiming();
 
 	/** Do and finalize teleport.  	*/
 	UFUNCTION( BlueprintCallable, Category = "Teleporter" )
-	void DoTeleport();
+	UE_API void DoTeleport();
 
 	/** Get the actor we're currently trying to teleport with.
 	* Valid during aiming and teleporting.
 	*/
 	UFUNCTION(BlueprintCallable, Category = "Teleporter")
-	UViewportInteractor* GetInteractorTryingTeleport() const;
+	UE_API UViewportInteractor* GetInteractorTryingTeleport() const;
 
 	/** Called when the user presses a button on their motion controller device */
-	void OnPreviewInputAction(class FEditorViewportClient& ViewportClient, UViewportInteractor* Interactor,
+	UE_API void OnPreviewInputAction(class FEditorViewportClient& ViewportClient, UViewportInteractor* Interactor,
 		const struct FViewportActionKeyInput& Action, bool& bOutIsInputCaptured, bool& bWasHandled);
 
 	/** Move the roomspace using lerp towards the new location */
-	void Teleport(const float DeltaTime);
+	UE_API void Teleport(const float DeltaTime);
 
 	/** Updating aiming with teleport to end of laser and pulling and pushing */
-	void UpdateTeleportAim(const float DeltaTime);
+	UE_API void UpdateTeleportAim(const float DeltaTime);
 
 	/** Helper function to push and pull the teleporter along the laser */
-	FVector UpdatePushPullTeleporter(class UVREditorInteractor* VREditorInteractor, const FVector& LaserPointerStart, const FVector& LaserPointerEnd, const bool bEnablePushPull = true);
+	UE_API FVector UpdatePushPullTeleporter(class UVREditorInteractor* VREditorInteractor, const FVector& LaserPointerStart, const FVector& LaserPointerEnd, const bool bEnablePushPull = true);
 
 	/** Figures out the new transforms for all the visuals based on the new location and the transforms of the HMD and motioncontrollers */
-	void UpdateVisuals(const FVector& NewLocation);
+	UE_API void UpdateVisuals(const FVector& NewLocation);
 
 	/** If we want to start showing or hiding the meshes */
-	void Show(const bool bShow);
+	UE_API void Show(const bool bShow);
 
-	void UpdateFadingState(const float DeltaTime);
+	UE_API void UpdateFadingState(const float DeltaTime);
 
 	/** Calculated the scale for animation */
-	float CalculateAnimatedScaleFactor() const;
+	UE_API float CalculateAnimatedScaleFactor() const;
 
 	/** The owning VR mode */
 	UPROPERTY()
@@ -191,4 +193,6 @@ private:
 
 	/** Delay to start the actual moving to the end location */
 	uint32 TeleportTickDelay;
-};
\ No newline at end of file
+};
+
+#undef UE_API
diff --git a/Engine/Source/Editor/VREditor/Public/VREditorAssetContainer.h b/Engine/Source/Editor/VREditor/Public/VREditorAssetContainer.h
--- a/Engine/Source/Editor/VREditor/Public/VREditorAssetContainer.h
+++ b/Engine/Source/Editor/VREditor/Public/VREditorAssetContainer.h
@@ -18,8 +18,8 @@ class UFont;
 /**
  * Asset container for VREditor.
  */
-UCLASS()
-class VREDITOR_API UVREditorAssetContainer : public UDataAsset
+UCLASS(MinimalAPI)
+class UVREditorAssetContainer : public UDataAsset
 {
 	GENERATED_BODY()
 
diff --git a/Engine/Source/Editor/VREditor/Public/VREditorInteractor.h b/Engine/Source/Editor/VREditor/Public/VREditorInteractor.h
--- a/Engine/Source/Editor/VREditor/Public/VREditorInteractor.h
+++ b/Engine/Source/Editor/VREditor/Public/VREditorInteractor.h
@@ -9,6 +9,8 @@
 #include "UObject/ObjectMacros.h"
 #include "VREditorInteractor.generated.h"
 
+#define UE_API VREDITOR_API
+
 class AActor;
 class UWidgetComponent;
 class UStaticMeshSocket;
@@ -38,15 +40,15 @@ enum class ETouchSwipeDirection : uint8
 /**
  * VREditor default interactor
  */
-UCLASS( Blueprintable, BlueprintType )
-class VREDITOR_API UVREditorInteractor : public UViewportInteractor
+UCLASS(MinimalAPI,  Blueprintable, BlueprintType )
+class UVREditorInteractor : public UViewportInteractor
 {
 	GENERATED_BODY()
 
 public:
 
 	/** Default constructor */
-	UVREditorInteractor();
+	UE_API UVREditorInteractor();
 
 	/** Gets the owner of this system */
 	class UVREditorMode& GetVRMode()
@@ -66,67 +68,67 @@ public:
 		return KeyToActionMap;
 	}
 
-	virtual TMap<FViewportActionKeyInput, TArray<FKey>> GetKnownActionMappings(EControllerHand InHand = EControllerHand::AnyHand, FName InHMDDeviceType = NAME_None) const;
+	UE_API virtual TMap<FViewportActionKeyInput, TArray<FKey>> GetKnownActionMappings(EControllerHand InHand = EControllerHand::AnyHand, FName InHMDDeviceType = NAME_None) const;
 
 	/** Initialize default values */
 	UFUNCTION( BlueprintNativeEvent, CallInEditor, Category = "UVREditorInteractor" )
-	void Init( class UVREditorMode* InVRMode );
+	UE_API void Init( class UVREditorMode* InVRMode );
 
 	/** Sets up all components */
 	UFUNCTION( BlueprintNativeEvent, CallInEditor, Category = "UVREditorInteractor" )
-	void SetupComponent( AActor* OwningActor );
+	UE_API void SetupComponent( AActor* OwningActor );
 
 
 	// ViewportInteractorInterface overrides
-	virtual void Shutdown_Implementation() override;
-	virtual void Tick_Implementation( const float DeltaTime ) override;
-	virtual void CalculateDragRay( double& InOutDragRayLength, double& InOutDragRayVelocity ) override;
+	UE_API virtual void Shutdown_Implementation() override;
+	UE_API virtual void Tick_Implementation( const float DeltaTime ) override;
+	UE_API virtual void CalculateDragRay( double& InOutDragRayLength, double& InOutDragRayVelocity ) override;
 
 	/** @return Returns the type of HMD we're dealing with */
 	UFUNCTION( BlueprintCallable, Category = "UVREditorInteractor" )
-	FName GetHMDDeviceType() const;
+	UE_API FName GetHMDDeviceType() const;
 
-	virtual FHitResult GetHitResultFromLaserPointer( TArray<AActor*>* OptionalListOfIgnoredActors = nullptr, const EHitResultGizmoFilterMode GizmoFilerMode = EHitResultGizmoFilterMode::All,
+	UE_API virtual FHitResult GetHitResultFromLaserPointer( TArray<AActor*>* OptionalListOfIgnoredActors = nullptr, const EHitResultGizmoFilterMode GizmoFilerMode = EHitResultGizmoFilterMode::All,
 	TArray<UClass*>* ObjectsInFrontOfGizmo = nullptr, const bool bEvenIfBlocked = false, const float LaserLengthOverride = 0.0f ) override;
-	virtual void ResetHoverState() override;
-	virtual bool IsModifierPressed() const override;
-	virtual void PreviewInputKey( class FEditorViewportClient& ViewportClient, FViewportActionKeyInput& Action, const FKey Key, const EInputEvent Event, bool& bOutWasHandled ) override;
-	virtual void HandleInputKey( class FEditorViewportClient& ViewportClient, FViewportActionKeyInput& Action, const FKey Key, const EInputEvent Event, bool& bOutWasHandled ) override;
-	virtual bool GetTransformAndForwardVector( FTransform& OutHandTransform, FVector& OutForwardVector ) const override;
-	virtual void HandleInputAxis( FEditorViewportClient& ViewportClient, FViewportActionKeyInput& Action, const FKey Key, const float Delta, const float DeltaTime, bool& bOutWasHandled ) override;
+	UE_API virtual void ResetHoverState() override;
+	UE_API virtual bool IsModifierPressed() const override;
+	UE_API virtual void PreviewInputKey( class FEditorViewportClient& ViewportClient, FViewportActionKeyInput& Action, const FKey Key, const EInputEvent Event, bool& bOutWasHandled ) override;
+	UE_API virtual void HandleInputKey( class FEditorViewportClient& ViewportClient, FViewportActionKeyInput& Action, const FKey Key, const EInputEvent Event, bool& bOutWasHandled ) override;
+	UE_API virtual bool GetTransformAndForwardVector( FTransform& OutHandTransform, FVector& OutForwardVector ) const override;
+	UE_API virtual void HandleInputAxis( FEditorViewportClient& ViewportClient, FViewportActionKeyInput& Action, const FKey Key, const float Delta, const float DeltaTime, bool& bOutWasHandled ) override;
 
 	/** Toggles whether or not this controller is being used to scrub sequencer */
-	void ToggleSequencerScrubbingMode();
+	UE_API void ToggleSequencerScrubbingMode();
 
 	/** Returns whether or not this controller is being used to scrub sequencer */
-	bool IsScrubbingSequencer() const;
+	UE_API bool IsScrubbingSequencer() const;
 
 	/** Get the motioncontroller component of this interactor */
 	UFUNCTION( BlueprintCallable, Category = "UVREditorInteractor" )
-	class UMotionControllerComponent* GetMotionControllerComponent() const;
+	UE_API class UMotionControllerComponent* GetMotionControllerComponent() const;
 
 
 	/** Sets the EControllerHand for this motioncontroller */
 	UFUNCTION( BlueprintCallable, Category = "UVREditorInteractor" )
-	void SetControllerHandSide( const FName InControllerHandSide );
+	UE_API void SetControllerHandSide( const FName InControllerHandSide );
 
 	/** Sets the EControllerHand for this motioncontroller */
 	UFUNCTION( BlueprintCallable, Category = "UVREditorInteractor" )
-	FName GetControllerHandSide() const;
+	UE_API FName GetControllerHandSide() const;
 
 	/** Returns the slide delta for pushing and pulling objects. Needs to be implemented by derived classes (e.g. touchpad for vive controller or scrollweel for mouse ) */
 	UFUNCTION( CallInEditor, BlueprintNativeEvent, Category = "UVREditorInteractor" )
-	float GetSlideDelta() const;
+	UE_API float GetSlideDelta() const;
 
 	/** Starts haptic feedback for physical motion controller */
-	virtual void PlayHapticEffect( const float Strength ) override;
+	UE_API virtual void PlayHapticEffect( const float Strength ) override;
 
 	/** Set if we want to force to show the laser*/
 	UFUNCTION( BlueprintCallable, Category = "UVREditorInteractor" )
-	void SetForceShowLaser( const bool bInForceShow );
+	UE_API void SetForceShowLaser( const bool bInForceShow );
 
 	// Tells us if thie interator is carrying a "carryable" actor, like a camera, which matches the interactor motions instead of using manipulators.
-	bool IsCarrying() const;
+	UE_API bool IsCarrying() const;
 
 	//
 	// Trackpad
@@ -137,60 +139,60 @@ public:
 	 *
 	 * @param Axis	The axis of which we want the slide delta. 0 is X axis and 1 is Y axis.Default is axis Y
 	 */
-	float GetTrackpadSlideDelta( const bool Axis = 1 ) const;
+	UE_API float GetTrackpadSlideDelta( const bool Axis = 1 ) const;
 
 	/** Resets all the trackpad related values to default. */
-	void ResetTrackpad();
+	UE_API void ResetTrackpad();
 
 	/** Check if the touchpad is currently touched */
 	UFUNCTION( BlueprintCallable, Category = "UVREditorInteractor" )
-	bool IsTouchingTrackpad() const;
+	UE_API bool IsTouchingTrackpad() const;
 
 	/** Get the current position of the trackpad or analog stick */
 	UFUNCTION( BlueprintCallable, Category = "UVREditorInteractor" )
-	FVector2D GetTrackpadPosition() const;
+	UE_API FVector2D GetTrackpadPosition() const;
 
 	/** Get the last position of the trackpad or analog stick */
 	UFUNCTION( BlueprintCallable, Category = "UVREditorInteractor" )
-	FVector2D GetLastTrackpadPosition() const;
+	UE_API FVector2D GetLastTrackpadPosition() const;
 
 	/** If the trackpad values are valid */
-	bool IsTrackpadPositionValid( const int32 AxisIndex ) const;
+	UE_API bool IsTrackpadPositionValid( const int32 AxisIndex ) const;
 
 	/** Get when the last time the trackpad position was updated */
-	FTimespan& GetLastTrackpadPositionUpdateTime();
+	UE_API FTimespan& GetLastTrackpadPositionUpdateTime();
 
 	/** Get when the last time the trackpad position was updated */
-	FTimespan& GetLastActiveTrackpadUpdateTime();
+	UE_API FTimespan& GetLastActiveTrackpadUpdateTime();
 
 
 	//
 	// Getters and setters
 	//
 	UFUNCTION( BlueprintCallable, Category = "UVREditorInteractor" )
-	const FVector& GetLaserStart() const;
+	UE_API const FVector& GetLaserStart() const;
 
 	UFUNCTION( BlueprintCallable, Category = "UVREditorInteractor" )
-	const FVector& GetLaserEnd() const;
+	UE_API const FVector& GetLaserEnd() const;
 
 	/** Next frame this will be used as color for the laser */
 	UFUNCTION( BlueprintCallable, Category = "UVREditorInteractor" )
-	void SetForceLaserColor( const FLinearColor& InColor );
+	UE_API void SetForceLaserColor( const FLinearColor& InColor );
 
 	UFUNCTION( BlueprintCallable, Category = "UVREditorInteractor" )
-	AVREditorTeleporter* GetTeleportActor();
+	UE_API AVREditorTeleporter* GetTeleportActor();
 
 	/** Get the side of the controller */
 	UFUNCTION( BlueprintCallable, Category = "UVREditorInteractor" )
-	EControllerHand GetControllerSide() const;
+	UE_API EControllerHand GetControllerSide() const;
 
 	/** Returns what controller type this is for asymmetric control schemes */
 	UFUNCTION( BlueprintCallable, Category = "VREditorInteractor" )
-	EControllerType GetControllerType() const;
+	UE_API EControllerType GetControllerType() const;
 
 	/** Set what controller type this is for asymmetric control schemes */
 	UFUNCTION( BlueprintCallable, Category = "VREditorInteractor" )
-	void SetControllerType( const EControllerType InControllerType );
+	UE_API void SetControllerType( const EControllerType InControllerType );
 
 	/**
 	 * Temporary set what controller type this is for asymmetric control schemes.
@@ -199,69 +201,69 @@ public:
 	 * @return true if the controller type was changed
 	 */
 	UFUNCTION( BlueprintCallable, Category = "VREditorInteractor" )
-	bool TryOverrideControllerType( const EControllerType InControllerType );
+	UE_API bool TryOverrideControllerType( const EControllerType InControllerType );
 
 	/** Gets if this interactor is hovering over UI */
 	UFUNCTION( BlueprintCallable, Category = "VREditorInteractor" )
-	bool IsHoveringOverUI() const;
+	UE_API bool IsHoveringOverUI() const;
 
 	/** Sets if the quick menu is on this interactor */
-	void SetHasUIOnForearm( const bool bInHasUIOnForearm );
+	UE_API void SetHasUIOnForearm( const bool bInHasUIOnForearm );
 
 	/** Check if the quick menu is on this interactor */
-	bool HasUIOnForearm() const;
+	UE_API bool HasUIOnForearm() const;
 
 	/** Gets the current hovered widget component if any */
-	UWidgetComponent* GetLastHoveredWidgetComponent() const;
+	UE_API UWidgetComponent* GetLastHoveredWidgetComponent() const;
 
 	/** Sets the current hovered widget component */
-	void SetLastHoveredWidgetComponent( UWidgetComponent* NewHoveringOverWidgetComponent );
+	UE_API void SetLastHoveredWidgetComponent( UWidgetComponent* NewHoveringOverWidgetComponent );
 
 	/** Sets if the interactor is clicking on any UI */
-	void SetIsClickingOnUI( const bool bInIsClickingOnUI );
+	UE_API void SetIsClickingOnUI( const bool bInIsClickingOnUI );
 
 	/** Gets if the interactor is clicking on any UI */
 	UFUNCTION( BlueprintCallable, Category = "VREditorInteractor" )
-	bool IsClickingOnUI() const;
+	UE_API bool IsClickingOnUI() const;
 
 	/** Sets if the interactor is hovering over any UI */
-	void SetIsHoveringOverUI( const bool bInIsHoveringOverUI );
+	UE_API void SetIsHoveringOverUI( const bool bInIsHoveringOverUI );
 
 	/** Sets if the interactor is right  hover over any UI */
-	void SetIsRightClickingOnUI( const bool bInIsRightClickingOnUI );
+	UE_API void SetIsRightClickingOnUI( const bool bInIsRightClickingOnUI );
 
 	/** Gets if the interactor is right clicking on UI */
-	bool IsRightClickingOnUI() const;
+	UE_API bool IsRightClickingOnUI() const;
 
 	/** Sets the time the interactor last pressed on UI */
-	void SetLastUIPressTime( const double InLastUIPressTime );
+	UE_API void SetLastUIPressTime( const double InLastUIPressTime );
 
 	/** Gets last time the interactor pressed on UI */
-	double GetLastUIPressTime() const;
+	UE_API double GetLastUIPressTime() const;
 
 	/** Sets the UI scroll velocity */
-	void SetUIScrollVelocity( const float InUIScrollVelocity );
+	UE_API void SetUIScrollVelocity( const float InUIScrollVelocity );
 
 	/** Gets the UI scroll velocity */
-	float GetUIScrollVelocity() const;
+	UE_API float GetUIScrollVelocity() const;
 
 	/** Gets the trigger value */
 	UFUNCTION( BlueprintCallable, Category = "VREditorInteractor" )
-	float GetSelectAndMoveTriggerValue() const;
+	UE_API float GetSelectAndMoveTriggerValue() const;
 
 	/* ViewportInteractor overrides, checks if the laser is blocked by UI */
-	virtual bool GetIsLaserBlocked() const override;
+	UE_API virtual bool GetIsLaserBlocked() const override;
 
 	/** Replace the default VR controller mesh with a custom one. */
 	UFUNCTION(BlueprintCallable, Category = "VREditorInteractor")
-	void ReplaceHandMeshComponent(UStaticMesh* NewMesh, FVector MeshScale = FVector(1.f, 1.f, 1.f));
+	UE_API void ReplaceHandMeshComponent(UStaticMesh* NewMesh, FVector MeshScale = FVector(1.f, 1.f, 1.f));
 
-	bool IsActionKeyPressed(FName ActionName) const;
+	UE_API bool IsActionKeyPressed(FName ActionName) const;
 
 protected:
 
 	/** Polls input for the motion controllers transforms */
-	virtual void PollInput() override;
+	UE_API virtual void PollInput() override;
 
 	/** Whether swiping left/right on touchpad/joystick triggers undo/redo  */
 	UPROPERTY(EditDefaultsOnly, Category = "VREditorInteractor")
@@ -284,7 +286,7 @@ protected:
 	TObjectPtr<class UStaticMeshComponent> HandMeshComponent;
 
 	UFUNCTION(BlueprintNativeEvent)
-	void UpdateHandMeshRelativeTransform();
+	UE_API void UpdateHandMeshRelativeTransform();
 
 	FVector3d HandMeshBaseScale;
 	FTransform HandMeshGripTransform;
@@ -328,13 +330,13 @@ private:
 	bool bHaveMotionController;
 
 	// Special key action names for motion controllers
-	static const FName TrackpadPositionX;
-	static const FName TrackpadPositionY;
-	static const FName TriggerAxis;
-	static const FName MotionController_Left_FullyPressedTriggerAxis;
-	static const FName MotionController_Right_FullyPressedTriggerAxis;
-	static const FName MotionController_Left_PressedTriggerAxis;
-	static const FName MotionController_Right_PressedTriggerAxis;
+	static UE_API const FName TrackpadPositionX;
+	static UE_API const FName TrackpadPositionY;
+	static UE_API const FName TriggerAxis;
+	static UE_API const FName MotionController_Left_FullyPressedTriggerAxis;
+	static UE_API const FName MotionController_Right_FullyPressedTriggerAxis;
+	static UE_API const FName MotionController_Left_PressedTriggerAxis;
+	static UE_API const FName MotionController_Right_PressedTriggerAxis;
 
 	/** Is the button for the specified action currently held down? */
 	TMap<FName, bool> ActionKeysPressed;
@@ -353,37 +355,37 @@ private:
 private:
 
 	/** Changes the color of the buttons on the handmesh */
-	void ApplyButtonPressColors( const FViewportActionKeyInput& Action );
+	UE_API void ApplyButtonPressColors( const FViewportActionKeyInput& Action );
 
 	/** Set the visuals for a button on the motion controller */
-	void SetMotionControllerButtonPressedVisuals( const EInputEvent Event, const FName& ParameterName, const float PressStrength );
+	UE_API void SetMotionControllerButtonPressedVisuals( const EInputEvent Event, const FName& ParameterName, const float PressStrength );
 
 	/** Pops up some help text labels for the controller in the specified hand, or hides it, if requested */
-	void ShowHelpForHand( const bool bShowIt );
+	UE_API void ShowHelpForHand( const bool bShowIt );
 
 	/** Called every frame to update the position of any floating help labels */
-	void UpdateHelpLabels();
+	UE_API void UpdateHelpLabels();
 
 	/** Given a mesh and a key name, tries to find a socket on the mesh that matches a supported key */
-	UStaticMeshSocket* FindMeshSocketForKey( UStaticMesh* StaticMesh, const FKey Key );
+	UE_API UStaticMeshSocket* FindMeshSocketForKey( UStaticMesh* StaticMesh, const FKey Key );
 
 	/** Updates all the segments of the curved laser */
-	void UpdateSplineLaser( const FVector& InStartLocation, const FVector& InEndLocation, const FVector& InForward );
+	UE_API void UpdateSplineLaser( const FVector& InStartLocation, const FVector& InEndLocation, const FVector& InForward );
 
 	/** Sets the visibility on all curved laser segments */
-	void SetLaserVisibility( const bool bVisible );
+	UE_API void SetLaserVisibility( const bool bVisible );
 
 	/** Sets the visuals of the LaserPointer */
-	void SetLaserVisuals( const FLinearColor& NewColor, const float CrawlFade, const float CrawlSpeed );
+	UE_API void SetLaserVisuals( const FLinearColor& NewColor, const float CrawlFade, const float CrawlSpeed );
 
 	/** Updates the radial menu */
-	void UpdateRadialMenuInput( const float DeltaTime );
+	UE_API void UpdateRadialMenuInput( const float DeltaTime );
 
 	/** Initialize shared HandMeshComponent settings */
-	void SetHandMeshComponentProperties();
+	UE_API void SetHandMeshComponentProperties();
 
 	/** Start undo or redo from swipe for the Vive */
-	void UndoRedoFromSwipe( const ETouchSwipeDirection InSwipeDirection );
+	UE_API void UndoRedoFromSwipe( const ETouchSwipeDirection InSwipeDirection );
 
 	//
 	// General input @todo: VREditor: Should this be general (non-UI) in interactordata ?
@@ -508,3 +510,5 @@ protected:
 	UPROPERTY()
 	TObjectPtr<class UVREditorMode> VRMode;
 };
+
+#undef UE_API
diff --git a/Engine/Source/Editor/VREditor/Public/VREditorMode.h b/Engine/Source/Editor/VREditor/Public/VREditorMode.h
--- a/Engine/Source/Editor/VREditor/Public/VREditorMode.h
+++ b/Engine/Source/Editor/VREditor/Public/VREditorMode.h
@@ -17,6 +17,8 @@
 #include "VREditorModeBase.h"
 #include "VREditorMode.generated.h"
 
+#define UE_API VREDITOR_API
+
 class AActor;
 class FEditorViewportClient;
 class SLevelViewport;
@@ -50,29 +52,29 @@ namespace VRActionTypes
 /**
  * VR Editor Mode. Extends editor viewports with functionality for VR controls and object manipulation
  */
-UCLASS( Abstract, Transient )
-class VREDITOR_API UVREditorMode : public UVREditorModeBase
+UCLASS(MinimalAPI,  Abstract, Transient )
+class UVREditorMode : public UVREditorModeBase
 {
 	GENERATED_BODY()
 
 public:
 
 	/** Default constructor */
-	UVREditorMode();
+	UE_API UVREditorMode();
 
 	/** Overrides the HMD device type, which is otherwise derived from the XR system name. */
-	void SetHMDDeviceTypeOverride( FName InOverrideType );
+	UE_API void SetHMDDeviceTypeOverride( FName InOverrideType );
 
 	//~ Begin UEditorWorldExtension interface
 
 	/** Initialize the VREditor */
-	virtual void Init() override;
+	UE_API virtual void Init() override;
 
 	/** Shutdown the VREditor */
-	virtual void Shutdown() override;
+	UE_API virtual void Shutdown() override;
 
 protected:
-	virtual void TransitionWorld(UWorld* NewWorld, EEditorWorldExtensionTransitionState TransitionState) override;
+	UE_API virtual void TransitionWorld(UWorld* NewWorld, EEditorWorldExtensionTransitionState TransitionState) override;
 
 	//~ End UEditorWorldExtension interface
 
@@ -86,16 +88,16 @@ public:
 	virtual bool ShouldDisplayExperimentalWarningOnEntry() const { return true; }
 
 	/** When the user actually enters the VR Editor mode */
-	virtual void Enter() override;
+	UE_API virtual void Enter() override;
 
 	/** When the user leaves the VR Editor mode */
-	virtual void Exit( bool bShouldDisableStereo ) override;
+	UE_API virtual void Exit( bool bShouldDisableStereo ) override;
 
 	/** Tick before the ViewportWorldInteraction is ticked */
-	void PreTick( const float DeltaTime );
+	UE_API void PreTick( const float DeltaTime );
 
 	/** Tick after the ViewportWorldInteraction is ticked */
-	void PostTick( const float DeltaTime );
+	UE_API void PostTick( const float DeltaTime );
 
 	/** Returns true if the user wants to exit this mode */
 	virtual bool WantsToExitMode() const override
@@ -103,38 +105,38 @@ public:
 		return bWantsToExitMode;
 	}
 
-	void AllocateInteractors();
+	UE_API void AllocateInteractors();
 
 	/** Call this to start exiting VR mode */
-	void StartExitingVRMode();
+	UE_API void StartExitingVRMode();
 
-	virtual bool GetLaserForHand(EControllerHand InHand, FVector& OutLaserStart, FVector& OutLaserEnd) const override;
+	UE_API virtual bool GetLaserForHand(EControllerHand InHand, FVector& OutLaserStart, FVector& OutLaserEnd) const override;
 
 	/** Gets the world space transform of the calibrated VR room origin.  When using a seated VR device, this will feel like the
 	camera's world transform (before any HMD positional or rotation adjustments are applied.) */
-	virtual FTransform GetRoomTransform() const override;
+	UE_API virtual FTransform GetRoomTransform() const override;
 
 	/** Sets a new transform for the room, in world space.  This is basically setting the editor's camera transform for the viewport */
-	void SetRoomTransform( const FTransform& NewRoomTransform );
+	UE_API void SetRoomTransform( const FTransform& NewRoomTransform );
 
 	/** Gets the transform of the user's HMD in room space */
-	FTransform GetRoomSpaceHeadTransform() const;
+	UE_API FTransform GetRoomSpaceHeadTransform() const;
 
 	/**
 	 * Gets the world space transform of the HMD (head)
 	 *
 	 * @return	World space space HMD transform
 	 */
-	virtual FTransform GetHeadTransform() const override;
+	UE_API virtual FTransform GetHeadTransform() const override;
 
 	/** Gets access to the world interaction system (const) */
-	const class UViewportWorldInteraction& GetWorldInteraction() const;
+	UE_API const class UViewportWorldInteraction& GetWorldInteraction() const;
 
 	/** Gets access to the world interaction system */
-	class UViewportWorldInteraction& GetWorldInteraction();
+	UE_API class UViewportWorldInteraction& GetWorldInteraction();
 
 	/** If the mode was completely initialized */
-	bool IsFullyInitialized() const;
+	UE_API bool IsFullyInitialized() const;
 
 	/** * Gets the tick handle to give external systems the change to be ticked right after the ViewportWorldInteraction is ticked */
 	DECLARE_EVENT_OneParam(UVREditorMode, FOnVRTickHandle, const float /* DeltaTime */);
@@ -189,40 +191,40 @@ public:
 	}
 
 	/** Lets other modules know if the radial menu is visible on a given interactor so input should be handled differently */
-	bool IsShowingRadialMenu( const class UVREditorInteractor* Interactor ) const;
+	UE_API bool IsShowingRadialMenu( const class UVREditorInteractor* Interactor ) const;
 
 	/** Display the scene more closely to how it would appear at runtime (as opposed to edit time). */
 	UFUNCTION(BlueprintCallable, Category="VREditorMode")
-	void SetGameView(bool bGameView);
+	UE_API void SetGameView(bool bGameView);
 
 	/** Returns whether game view is currently active. */
 	UFUNCTION(BlueprintCallable, Category="VREditorMode")
-	bool IsInGameView() const;
+	UE_API bool IsInGameView() const;
 
 	/** Gets the world scale factor, which can be multiplied by a scale vector to convert to room space */
 	UFUNCTION( BlueprintCallable, Category = "VREditorMode" )
-	float GetWorldScaleFactor() const;
+	UE_API float GetWorldScaleFactor() const;
 
 	/** Spawns a flashlight on the specified hand */
-	void ToggleFlashlight( class UVREditorInteractor* Interactor );
+	UE_API void ToggleFlashlight( class UVREditorInteractor* Interactor );
 
 	/** Will update the TransformGizmo Actor with the next Gizmo type  */
-	void CycleTransformGizmoHandleType();
+	UE_API void CycleTransformGizmoHandleType();
 
 	/** Gets the current Gizmo handle type */
-	EGizmoHandleTypes GetCurrentGizmoType() const;
+	UE_API EGizmoHandleTypes GetCurrentGizmoType() const;
 
 	/** @return Returns the type of HMD we're dealing with */
-	FName GetHMDDeviceType() const;
+	UE_API FName GetHMDDeviceType() const;
 
 	/** @return Checks to see if the specified interactor is aiming roughly toward the specified capsule */
-	bool IsHandAimingTowardsCapsule(class UViewportInteractor* Interactor, const FTransform& CapsuleTransform, const FVector CapsuleStart, const FVector CapsuleEnd, const float CapsuleRadius, const float MinDistanceToCapsule, const FVector CapsuleFrontDirection, const float MinDotForAimingAtCapsule) const;
+	UE_API bool IsHandAimingTowardsCapsule(class UViewportInteractor* Interactor, const FTransform& CapsuleTransform, const FVector CapsuleStart, const FVector CapsuleEnd, const float CapsuleRadius, const float MinDistanceToCapsule, const FVector CapsuleFrontDirection, const float MinDotForAimingAtCapsule) const;
 
 	/** Gets the hand interactor  */
-	class UVREditorInteractor* GetHandInteractor( const EControllerHand ControllerHand ) const;
+	UE_API class UVREditorInteractor* GetHandInteractor( const EControllerHand ControllerHand ) const;
 
 	/** Snaps the current selected actor to the ground */
-	void SnapSelectedActorsToGround();
+	UE_API void SnapSelectedActorsToGround();
 
 	/** Saved information about the editor and viewport we possessed, so we can restore it after exiting VR mode */
 	struct FSavedEditorState : public FBaseSavedEditorState
@@ -236,29 +238,29 @@ public:
 	virtual TSharedRef<FBaseSavedEditorState> CreateSavedState() override { return MakeShared<FSavedEditorState>(); }
 
 	/** Gets the saved editor state from entering the mode */
-	const FSavedEditorState& GetSavedEditorState() const;
+	UE_API const FSavedEditorState& GetSavedEditorState() const;
 
 	DECLARE_DELEGATE(FOnVREditingModeExit);
 	/** Used to override dockable area restoration behavior */
 	FOnVREditingModeExit OnVREditingModeExit_Handler;
 
-	void SaveSequencerSettings(bool bInKeyAllEnabled, EAutoChangeMode InAutoChangeMode, const class USequencerSettings& InSequencerSettings);
+	UE_API void SaveSequencerSettings(bool bInKeyAllEnabled, EAutoChangeMode InAutoChangeMode, const class USequencerSettings& InSequencerSettings);
 
 	/** Create a static motion controller mesh for the current HMD platform */
-	UStaticMeshComponent* CreateMotionControllerMesh( AActor* OwningActor, USceneComponent* AttachmentToComponent, UStaticMesh* OptionalControllerMesh = nullptr );
+	UE_API UStaticMeshComponent* CreateMotionControllerMesh( AActor* OwningActor, USceneComponent* AttachmentToComponent, UStaticMesh* OptionalControllerMesh = nullptr );
 
 	/** Helper functions to create a static mesh */
-	UStaticMeshComponent* CreateMesh( AActor* OwningActor, const FString& MeshName, USceneComponent* AttachmentToComponent /*= nullptr */ );
-	UStaticMeshComponent* CreateMesh(AActor* OwningActor, UStaticMesh* Mesh, USceneComponent* AttachmentToComponent /*= nullptr */);
+	UE_API UStaticMeshComponent* CreateMesh( AActor* OwningActor, const FString& MeshName, USceneComponent* AttachmentToComponent /*= nullptr */ );
+	UE_API UStaticMeshComponent* CreateMesh(AActor* OwningActor, UStaticMesh* Mesh, USceneComponent* AttachmentToComponent /*= nullptr */);
 
 	/** Gets the container for all the assets of VREditor. */
-	const class UVREditorAssetContainer& GetAssetContainer() const;
+	UE_API const class UVREditorAssetContainer& GetAssetContainer() const;
 
 	/** Loads and returns the container for all the assets of VREditor. */
-	static class UVREditorAssetContainer& LoadAssetContainer();
+	static UE_API class UVREditorAssetContainer& LoadAssetContainer();
 
 	/** Plays sound at location. */
-	void PlaySound(USoundBase* SoundBase, const FVector& InWorldLocation, const float InVolume = 1.0f);
+	UE_API void PlaySound(USoundBase* SoundBase, const FVector& InWorldLocation, const float InVolume = 1.0f);
 
 	/** Delegate to be called when a material is placed **/
 	DECLARE_EVENT_ThreeParams( UVREditorPlacement, FOnPlaceDraggedMaterial, UPrimitiveComponent*, UMaterialInterface*, bool& );
@@ -269,17 +271,17 @@ public:
 	FOnPlacePreviewActor& OnPlacePreviewActor() { return OnPlacePreviewActorEvent; };
 
 	/** Return true if currently aiming to teleport. */
-	bool IsAimingTeleport() const;
-	bool IsTeleporting() const;
+	UE_API bool IsAimingTeleport() const;
+	UE_API bool IsTeleporting() const;
 
 	/** Toggles the debug mode. */
-	static void ToggleDebugMode();
+	static UE_API void ToggleDebugMode();
 
 	/** Returns if the VR Mode is in debug mode. */
-	static bool IsDebugModeEnabled();
+	static UE_API bool IsDebugModeEnabled();
 
 	/** Get Teleporter */
-	class AVREditorTeleporter* GetTeleportActor();
+	UE_API class AVREditorTeleporter* GetTeleportActor();
 
 	/** Delegate to be called when the debug mode is toggled. */
 	DECLARE_EVENT_OneParam(UVREditorMode, FOnToggleVRModeDebug, bool);
@@ -293,13 +295,13 @@ public:
 private:
 
 	/** Called when the editor is closed */
-	void OnEditorClosed();
+	UE_API void OnEditorClosed();
 
 	/** Called when someone closes a standalone VR Editor window */
-	void OnVREditorWindowClosed( const TSharedRef<SWindow>& ClosedWindow );
+	UE_API void OnVREditorWindowClosed( const TSharedRef<SWindow>& ClosedWindow );
 
 	/** Restore the world to meters to the saved one when entering VR Editor */
-	void RestoreWorldToMeters();
+	UE_API void RestoreWorldToMeters();
 
 protected:
 
@@ -307,9 +309,9 @@ protected:
 	// Startup/Shutdown
 	//
 
-	void BeginEntry();
-	void SetupSubsystems();
-	void FinishEntry();
+	UE_API void BeginEntry();
+	UE_API void SetupSubsystems();
+	UE_API void FinishEntry();
 
 	/** The VR editor window, if it's open right now */
 	TWeakPtr< class SWindow > VREditorWindowWeakPtr;
@@ -418,28 +420,28 @@ public:
 	};
 
 	// Gets the color
-	FLinearColor GetColor( const EColors Color ) const;
+	UE_API FLinearColor GetColor( const EColors Color ) const;
 
 	// Get the default near clipping plane for VR editing
-	float GetDefaultVRNearClipPlane() const;
+	UE_API float GetDefaultVRNearClipPlane() const;
 
 	/** Runtime and plugin modules can force VR Editor to refresh using this function */
-	void RefreshVREditorSequencer(class ISequencer* InCurrentSequencer);
+	UE_API void RefreshVREditorSequencer(class ISequencer* InCurrentSequencer);
 
 	/** Refresh the current actor preview widget on an in-world UI panel */
-	void RefreshActorPreviewWidget(TSharedRef<SWidget> InWidget, int32 Index, AActor *Actor, bool bIsPanelDetached = false);
+	UE_API void RefreshActorPreviewWidget(TSharedRef<SWidget> InWidget, int32 Index, AActor *Actor, bool bIsPanelDetached = false);
 
 	/** General way to spawn an external UMG UI from a radial menu */
-	void UpdateExternalUMGUI(const struct FVREditorFloatingUICreationContext& CreationContext);
+	UE_API void UpdateExternalUMGUI(const struct FVREditorFloatingUICreationContext& CreationContext);
 
 	/** General way to spawn an external Slate UI from a radial menu */
-	void UpdateExternalSlateUI(TSharedRef<SWidget> InWidget, FName Name, FVector2D InSize);
+	UE_API void UpdateExternalSlateUI(TSharedRef<SWidget> InWidget, FName Name, FVector2D InSize);
 
 	/** Returns the currently active sequencer */
-	class ISequencer* GetCurrentSequencer();
+	UE_API class ISequencer* GetCurrentSequencer();
 
 	/** The asset container path */
-	static const TCHAR* AssetContainerPath;
+	static UE_API const TCHAR* AssetContainerPath;
 
 	/** The controller to use when UnrealEd is in VR mode. Use VREditorInteractor get default editor behavior, or select a custom controller for special behavior */
 	UPROPERTY( EditAnywhere, BlueprintReadWrite, NoClear, Category="Classes" )
@@ -466,7 +468,7 @@ private:
 	TObjectPtr<class UVREditorAssetContainer> AssetContainer;
 
 	/** Whether currently in debug mode or not. */
-	static bool bDebugModeEnabled;
+	static UE_API bool bDebugModeEnabled;
 
 	/** Event that gets broadcasted when debug mode is toggled. */
 	FOnToggleVRModeDebug OnToggleDebugModeEvent;
@@ -476,3 +478,5 @@ private:
 
 	bool bAddedViewportWorldInteractionExtension = false;
 };
+
+#undef UE_API
