Curve Editor : fix drag handle not setting interactive flag correctly

diff --git a/Engine/Source/Editor/UnrealEd/Private/SColorGradientEditor.cpp b/Engine/Source/Editor/UnrealEd/Private/SColorGradientEditor.cpp
--- a/Engine/Source/Editor/UnrealEd/Private/SColorGradientEditor.cpp
+++ b/Engine/Source/Editor/UnrealEd/Private/SColorGradientEditor.cpp
@@ -592,6 +592,9 @@ FReply SColorGradientEditor::OnMouseButtonUp( const FGeometry& MyGeometry, const
 			}
 			DistanceDragged = 0;
 			bDraggingStop = false;
+			// during interactive drag we called OnCurveChanged with EPropertyChangeType::Interactive
+			//	now the drag is over, do OnCurveChanged with EPropertyChangeType::ValueSet
+			CurveOwner->OnCurveChanged(CurveOwner->GetCurves());
 			return FReply::Handled().ReleaseMouseCapture();
 		}
 		else if( MouseEvent.GetEffectingButton() == EKeys::RightMouseButton && !bDraggingStop)
@@ -1114,7 +1117,13 @@ void SColorGradientEditor::MoveStop( FGradientStopMark& Mark, float NewTime )
 {
 	CurveOwner->ModifyOwner();
 	Mark.SetTime( NewTime, *CurveOwner );
+	EPropertyChangeType::Type ChangeType = bDraggingStop ? EPropertyChangeType::Interactive : EPropertyChangeType::ValueSet;
+	// we want to be able to pass ChangeType as an arg to OnCurveChanged for interactive drags, but OnCurveChanged doesn't take that
+	// there are lots of implementations of OnCurveChanged so changing the function now is not great
+	//	instead set a bool on the curve to indicate the next OnCurveChanged type
+	CurveOwner->SetOnCurveChangedIsInteractive(bDraggingStop);
 	CurveOwner->OnCurveChanged(CurveOwner->GetCurves());
+	CurveOwner->SetOnCurveChangedIsInteractive(false);
 }
 
 
diff --git a/Engine/Source/Runtime/Engine/Classes/Curves/CurveBase.h b/Engine/Source/Runtime/Engine/Classes/Curves/CurveBase.h
--- a/Engine/Source/Runtime/Engine/Classes/Curves/CurveBase.h
+++ b/Engine/Source/Runtime/Engine/Classes/Curves/CurveBase.h
@@ -75,7 +75,8 @@ public:
 	ENGINE_API virtual TArray<const UObject*> GetOwners() const override;
 	ENGINE_API virtual void MakeTransactional() override;
 	ENGINE_API virtual void OnCurveChanged(const TArray<FRichCurveEditInfo>& ChangedCurveEditInfos) override;
-	
+	ENGINE_API virtual void SetOnCurveChangedIsInteractive(bool bInteractive) override;
+
 	virtual bool RepointCurveOwner(const FPackageReloadedEvent& InPackageReloadedEvent, FCurveOwnerInterface*& OutNewCurveOwner) const override
 	{
 		return RepointCurveOwnerAsset(InPackageReloadedEvent, this, OutNewCurveOwner);
@@ -115,5 +116,7 @@ public:
 	// Delegate called whenever the curve data is updated
 	DECLARE_MULTICAST_DELEGATE_TwoParams(FOnUpdateCurve, UCurveBase* /*Curve*/, EPropertyChangeType::Type /*ChangeType*/);
 	FOnUpdateCurve OnUpdateCurve;
+
+	bool bOnCurveUpdateInteractive = false;
 #endif
 };
diff --git a/Engine/Source/Runtime/Engine/Classes/Curves/CurveLinearColor.h b/Engine/Source/Runtime/Engine/Classes/Curves/CurveLinearColor.h
--- a/Engine/Source/Runtime/Engine/Classes/Curves/CurveLinearColor.h
+++ b/Engine/Source/Runtime/Engine/Classes/Curves/CurveLinearColor.h
@@ -9,8 +9,6 @@
 #include "CurveLinearColor.generated.h"
 
 class UCurveLinearColor;
-// Delegate called right before rendering each slot. This allows the gradient to be modified dynamically per slot.
-DECLARE_MULTICAST_DELEGATE_OneParam(FOnUpdateGradient, UCurveLinearColor* /*GradientAtlas*/);
 
 USTRUCT(BlueprintType)
 struct FRuntimeCurveLinearColor
@@ -102,11 +100,6 @@ public:
 	UPROPERTY(EditAnywhere, Category = "Color", meta = (ClampMin = "0.0", ClampMax = "1.0"))
 	float AdjustMaxAlpha;
 
-#if WITH_EDITORONLY_DATA
-	UE_DEPRECATED(4.26, "OnUpdateGradient is deprecated. Please use UCurveBase::OnUpdateCurve instead")
-	FOnUpdateGradient OnUpdateGradient;
-#endif
-
 private:
 
 	// if bClampOver1Always is true, clamping in [0,1] is always done
diff --git a/Engine/Source/Runtime/Engine/Classes/Curves/CurveOwnerInterface.h b/Engine/Source/Runtime/Engine/Classes/Curves/CurveOwnerInterface.h
--- a/Engine/Source/Runtime/Engine/Classes/Curves/CurveOwnerInterface.h
+++ b/Engine/Source/Runtime/Engine/Classes/Curves/CurveOwnerInterface.h
@@ -36,6 +36,9 @@ public:
 	/** Called when any of the curves have been changed */
 	virtual void OnCurveChanged(const TArray<FRichCurveEditInfo>& ChangedCurveEditInfos) = 0;
 
+	/** Tell the curve owner that calls to OnCurveChanged are from an interactive drag **/
+	virtual void SetOnCurveChangedIsInteractive(bool bInteractive) { };
+
 	/** Whether the curve represents a linear color */
 	virtual bool IsLinearColorCurve() const
 	{
diff --git a/Engine/Source/Runtime/Engine/Private/Curves/CurveBase.cpp b/Engine/Source/Runtime/Engine/Private/Curves/CurveBase.cpp
--- a/Engine/Source/Runtime/Engine/Private/Curves/CurveBase.cpp
+++ b/Engine/Source/Runtime/Engine/Private/Curves/CurveBase.cpp
@@ -94,10 +94,20 @@ void UCurveBase::MakeTransactional()
 }
 
 
+void UCurveBase::SetOnCurveChangedIsInteractive(bool bInteractive)
+{
+#if WITH_EDITOR
+	bOnCurveUpdateInteractive = bInteractive;
+#endif
+}
+
 void UCurveBase::OnCurveChanged(const TArray<FRichCurveEditInfo>& ChangedCurveEditInfos)
 {
 #if WITH_EDITOR
-	OnUpdateCurve.Broadcast(this, EPropertyChangeType::ValueSet);
+	// ideally OnCurveChanged should have taken a EPropertyChangeType as a function arg
+	//	instead we pick it up from an extra call to SetOnCurveChangedIsInteractive
+	EPropertyChangeType::Type ChangeType = bOnCurveUpdateInteractive ? EPropertyChangeType::Interactive : EPropertyChangeType::ValueSet;
+	OnUpdateCurve.Broadcast(this, ChangeType);
 #endif // WITH_EDITOR
 }
 
diff --git a/Engine/Source/Runtime/Engine/Private/Curves/CurveLinearColor.cpp b/Engine/Source/Runtime/Engine/Private/Curves/CurveLinearColor.cpp
--- a/Engine/Source/Runtime/Engine/Private/Curves/CurveLinearColor.cpp
+++ b/Engine/Source/Runtime/Engine/Private/Curves/CurveLinearColor.cpp
@@ -219,12 +219,6 @@ bool UCurveLinearColor::IsValidCurve( FRichCurveEditInfo CurveInfo )
 
 void UCurveLinearColor::PostEditChangeProperty(struct FPropertyChangedEvent& PropertyChangedEvent)
 {
-	if (PropertyChangedEvent.ChangeType != EPropertyChangeType::Interactive)
-	{
-		PRAGMA_DISABLE_DEPRECATION_WARNINGS
-		OnUpdateGradient.Broadcast(this);
-		PRAGMA_ENABLE_DEPRECATION_WARNINGS
-	}
 	Super::PostEditChangeProperty(PropertyChangedEvent);
 }
 
@@ -270,12 +264,8 @@ void UCurveLinearColor::PushUnadjustedToSourceData(TArrayView<FFloat16Color>& Sr
 
 void UCurveLinearColor::OnCurveChanged(const TArray<FRichCurveEditInfo>& ChangedCurveEditInfos)
 {
+	// this triggers OnUpdateCurve delegate, which goes to UCurveLinearColorAtlas::OnCurveUpdated to redraw the atlas texture 
 	Super::OnCurveChanged(ChangedCurveEditInfos);
-
-	PRAGMA_DISABLE_DEPRECATION_WARNINGS
-	// this redraws any UCurveLinearColorAtlas that holds me
-	OnUpdateGradient.Broadcast(this);
-	PRAGMA_ENABLE_DEPRECATION_WARNINGS
 }
 #endif
 
