Cooker: Optimize the package build dependency tracker by by adding an additional IsNative() condition to ShouldSkipDependency().
Avoids more expensive GetOutermost()->HasAnyPackageFlags(PKG_CompiledIn) calls for frequently referenced native types like enums and script structs.
FPackageBuildDependencyTracker::StaticOnObjectHandleRead down from 41s to 33s on when cooking CitySample Win64 Game.

#rnx

diff --git a/Engine/Source/Editor/UnrealEd/Private/Cooker/PackageBuildDependencyTracker.cpp b/Engine/Source/Editor/UnrealEd/Private/Cooker/PackageBuildDependencyTracker.cpp
--- a/Engine/Source/Editor/UnrealEd/Private/Cooker/PackageBuildDependencyTracker.cpp
+++ b/Engine/Source/Editor/UnrealEd/Private/Cooker/PackageBuildDependencyTracker.cpp
@@ -97,18 +97,37 @@ FPackageBuildDependencyTracker::~FPackageBuildDependencyTracker()
 
 static bool ShouldSkipDependency(const UObject* Object)
 {
-	return !Object ||
-		!GUObjectArray.IsValidIndex(Object) ||
-		!Object->HasAnyFlags(RF_Public) ||
-		(Object->GetClass() == UClass::StaticClass());
+	if (!Object)
+	{
+		return true;
+	}
+	if (!GUObjectArray.IsValidIndex(Object))
+	{
+		return true;
+	}
+	if (!Object->HasAnyFlags(RF_Public))
+	{
+		return true;
+	}
+	if (Object->GetClass() == UClass::StaticClass())
+	{
+		return true;
+	}
+	// Optimization to skip GetOutermost()->HasAnyPackageFlags(PKG_CompiledIn) for other native types like enums and script structs.
+	// CDOs and their sub objects will still be handled in the slower path.
+	if (Object->IsNative()) 
+	{
+		return true;
+	}
+	return false;
 }
 
 static FName NAME_EngineTransient(TEXT("/Engine/Transient"));
 
 void FPackageBuildDependencyTracker::StaticOnObjectHandleRead(const TArrayView<const UObject* const>& Objects)
 {
-	int Count = Objects.Num();
-	if(Count == 0 || (Count == 1 && ShouldSkipDependency(Objects[0])))
+	const int32 Count = Objects.Num();
+	if (Count == 0 || (Count == 1 && ShouldSkipDependency(Objects[0])))
 	{
 		return;
 	}
@@ -131,7 +150,8 @@ void FPackageBuildDependencyTracker::StaticOnObjectHandleRead(const TArrayView<c
 
 	for (const UObject* ReadObject : Objects)
 	{
-		if (ShouldSkipDependency(ReadObject))
+		// No need to re-evaluate ShouldSkipDependency when Count == 1
+		if (Count > 1 && ShouldSkipDependency(ReadObject))
 		{
 			continue;
 		}
