Fixed graph schema EMaterialValueType implicit truncation

diff --git a/Engine/Plugins/TextureGraph/Source/TextureGraphEditor/Private/EdGraph/TG_EdGraphSchema.cpp b/Engine/Plugins/TextureGraph/Source/TextureGraphEditor/Private/EdGraph/TG_EdGraphSchema.cpp
--- a/Engine/Plugins/TextureGraph/Source/TextureGraphEditor/Private/EdGraph/TG_EdGraphSchema.cpp
+++ b/Engine/Plugins/TextureGraph/Source/TextureGraphEditor/Private/EdGraph/TG_EdGraphSchema.cpp
@@ -205,7 +205,7 @@ const FPinConnectionResponse UTG_EdGraphSchema::CanCreateConnection(const UEdGra
 	// For non-exec pins, break existing connections on inputs only - multiple output connections are acceptable
 	if (InputPin->LinkedTo.Num() > 0)
 	{
-		const uint32 InputType = 1; //GetMaterialValueType(InputPin);
+		const EMaterialValueType InputType = MCT_Float1; //GetMaterialValueType(InputPin);
 		if (!(InputType & MCT_Execution))
 		{
 			ECanCreateConnectionResponse ReplyBreakOutputs;
@@ -227,7 +227,7 @@ const FPinConnectionResponse UTG_EdGraphSchema::CanCreateConnection(const UEdGra
 	// For exec pins, reverse is true - multiple input connections are acceptable
 	if (OutputPin->LinkedTo.Num() > 0)
 	{
-		const uint32 OutputType = 1; //GetMaterialValueType(InputPin);
+		const EMaterialValueType OutputType = MCT_Float1; //GetMaterialValueType(InputPin);
 		if (OutputType & MCT_Execution)
 		{
 			ECanCreateConnectionResponse ReplyBreakInputs;
diff --git a/Engine/Source/Editor/MaterialEditor/Private/MaterialEditorUtilities.cpp b/Engine/Source/Editor/MaterialEditor/Private/MaterialEditorUtilities.cpp
--- a/Engine/Source/Editor/MaterialEditor/Private/MaterialEditorUtilities.cpp
+++ b/Engine/Source/Editor/MaterialEditor/Private/MaterialEditorUtilities.cpp
@@ -692,10 +692,10 @@ void FMaterialEditorUtilities::AddMaterialExpressionCategory(FGraphActionMenuBui
 void FMaterialEditorUtilities::AddMaterialExpressionCategory(FGraphActionMenuBuilder& ActionMenuBuilder, FText CategoryName, TArray<struct FMaterialExpression>* MaterialExpressions, const UObject* MaterialOrFunction)
 {
 	// Get type of dragged pin
-	uint32 FromPinType = 0;
+	EMaterialValueType FromPinType{};
 	if (ActionMenuBuilder.FromPin)
 	{
-		FromPinType = UMaterialGraphSchema::GetMaterialValueType(ActionMenuBuilder.FromPin);
+		FromPinType = UMaterialGraphSchema::GetMaterialIOValueType(ActionMenuBuilder.FromPin);
 	}
 
 	bool bMaterialFunction = Cast<UMaterialFunction>(MaterialOrFunction) != nullptr;
@@ -741,6 +741,11 @@ void FMaterialEditorUtilities::AddMaterialExpressionCategory(FGraphActionMenuBui
 }
 
 bool FMaterialEditorUtilities::HasCompatibleConnection(UClass* ExpressionClass, uint32 TestType, EEdGraphPinDirection TestDirection, bool bMaterialFunction)
+{
+	return HasCompatibleConnection(ExpressionClass, static_cast<EMaterialValueType>(TestType), TestDirection, bMaterialFunction);
+}
+
+bool FMaterialEditorUtilities::HasCompatibleConnection(UClass* ExpressionClass, EMaterialValueType TestType, EEdGraphPinDirection TestDirection, bool bMaterialFunction)
 {
 	if (TestType != 0)
 	{
diff --git a/Engine/Source/Editor/MaterialEditor/Private/MaterialNodes/SGraphSubstrateMaterial.cpp b/Engine/Source/Editor/MaterialEditor/Private/MaterialNodes/SGraphSubstrateMaterial.cpp
--- a/Engine/Source/Editor/MaterialEditor/Private/MaterialNodes/SGraphSubstrateMaterial.cpp
+++ b/Engine/Source/Editor/MaterialEditor/Private/MaterialNodes/SGraphSubstrateMaterial.cpp
@@ -455,7 +455,7 @@ void FSubstrateWidget::GetPinColor(TSharedPtr<SGraphPin>& Out, const UMaterialGr
 			 InNode->MaterialExpression->IsA(UMaterialExpressionFunctionInput::StaticClass()) ||
 			 InNode->MaterialExpression->IsA(UMaterialExpressionFunctionOutput::StaticClass()))
 	{
-		const uint32 PinType = UMaterialGraphSchema::GetMaterialValueType(Pin);
+		const EMaterialValueType PinType = UMaterialGraphSchema::GetMaterialIOValueType(Pin);
 		if ((Out->GetDirection() == EGPD_Input || Out->GetDirection() == EGPD_Output) && (PinType & MCT_Substrate) != 0)
 		{
 			bHasColorModifier = true;
@@ -463,12 +463,12 @@ void FSubstrateWidget::GetPinColor(TSharedPtr<SGraphPin>& Out, const UMaterialGr
 		}
 	}
 
-	if (InNode->MaterialExpression->IsA(UMaterialExpressionSubstrateBSDF::StaticClass()) && Out->GetDirection() == EGPD_Output && UMaterialGraphSchema::GetMaterialValueType(Pin) == MCT_Substrate)
+	if (InNode->MaterialExpression->IsA(UMaterialExpressionSubstrateBSDF::StaticClass()) && Out->GetDirection() == EGPD_Output && UMaterialGraphSchema::GetMaterialIOValueType(Pin) == MCT_Substrate)
 	{
 		bHasColorModifier = true;
 		ColorModifier = FSubstrateWidget::GetConnectionColor();
 	}
-	else if (Pin && UMaterialGraphSchema::GetMaterialValueType(Pin) == MCT_Substrate)
+	else if (Pin && UMaterialGraphSchema::GetMaterialIOValueType(Pin) == MCT_Substrate)
 	{
 		if (!Out->IsConnected())
 		{
diff --git a/Engine/Source/Editor/MaterialEditor/Private/MaterialPins/SGraphPinMaterialInput.cpp b/Engine/Source/Editor/MaterialEditor/Private/MaterialPins/SGraphPinMaterialInput.cpp
--- a/Engine/Source/Editor/MaterialEditor/Private/MaterialPins/SGraphPinMaterialInput.cpp
+++ b/Engine/Source/Editor/MaterialEditor/Private/MaterialPins/SGraphPinMaterialInput.cpp
@@ -17,7 +17,7 @@ void SGraphPinMaterialInput::Construct(const FArguments& InArgs, UEdGraphPin* In
 	bool bUsePinColor = true;
 	if (Substrate::IsSubstrateEnabled())
 	{
-		if (UMaterialGraphSchema::GetMaterialValueType(InGraphPinObj) == MCT_Substrate)
+		if (UMaterialGraphSchema::GetMaterialIOValueType(InGraphPinObj) == MCT_Substrate)
 		{
 			bUsePinColor = false;
 		}			
@@ -35,7 +35,7 @@ FSlateColor SGraphPinMaterialInput::GetPinColor() const
 	{
 		if (Substrate::IsSubstrateEnabled())
 		{
-			if (UMaterialGraphSchema::GetMaterialValueType(GraphPinObj) == MCT_Substrate)
+			if (UMaterialGraphSchema::GetMaterialIOValueType(GraphPinObj) == MCT_Substrate)
 			{
 				return FSlateColor(FSubstrateWidget::GetConnectionColor());
 			}			
diff --git a/Engine/Source/Editor/MaterialEditor/Public/MaterialEditorUtilities.h b/Engine/Source/Editor/MaterialEditor/Public/MaterialEditorUtilities.h
--- a/Engine/Source/Editor/MaterialEditor/Public/MaterialEditorUtilities.h
+++ b/Engine/Source/Editor/MaterialEditor/Public/MaterialEditorUtilities.h
@@ -302,7 +302,9 @@ public:
 	 * @param	TestDirection		Pin Direction we are testing.
 	 * @param	bMaterialFunction	Whether we are testing for a material function.
 	*/
+	UE_DEPRECATED(5.8, "Use the HasCompatibleConnection(..., EMaterialValueType, ...) overload instead")
 	static UE_API bool HasCompatibleConnection(UClass* ExpressionClass, uint32 TestType, EEdGraphPinDirection TestDirection, bool bMaterialFunction);
+	static UE_API bool HasCompatibleConnection(UClass* ExpressionClass, EMaterialValueType TestType, EEdGraphPinDirection TestDirection, bool bMaterialFunction);
 
 	/** Constructor */
 	FMaterialEditorUtilities() {}
diff --git a/Engine/Source/Editor/UnrealEd/Classes/MaterialGraph/MaterialGraphSchema.h b/Engine/Source/Editor/UnrealEd/Classes/MaterialGraph/MaterialGraphSchema.h
--- a/Engine/Source/Editor/UnrealEd/Classes/MaterialGraph/MaterialGraphSchema.h
+++ b/Engine/Source/Editor/UnrealEd/Classes/MaterialGraph/MaterialGraphSchema.h
@@ -6,6 +6,7 @@
 #include "Delegates/Delegate.h"
 #include "UObject/ObjectMacros.h"
 #include "EdGraph/EdGraphSchema.h"
+#include "MaterialValueType.h"
 #include "MaterialGraphSchema.generated.h"
 
 class UEdGraph;
@@ -50,7 +51,9 @@ struct FMaterialGraphSchemaAction_NewNode : public FEdGraphSchemaAction
 	 * @param	FunctionInput		The function input to set.
 	 * @param	MaterialValueType	Value type we want input to accept.
 	 */
+	UE_DEPRECATED(5.8, "Use SetFunctionInputType(..., EMaterialValueType) overload instead")
 	UNREALED_API void SetFunctionInputType(class UMaterialExpressionFunctionInput* FunctionInput, uint32 MaterialValueType) const;
+	UNREALED_API void SetFunctionInputType(class UMaterialExpressionFunctionInput* FunctionInput, EMaterialValueType MaterialValueType) const;
 };
 
 /** Action to add a Material Function call to the graph */
@@ -254,7 +257,9 @@ class UMaterialGraphSchema : public UEdGraphSchema
 	UNREALED_API void UpdateDetailView(const UEdGraph* Graph) const;
 
 	/** Gets the type of this pin (must be part of a UMaterialGraphNode_Base) */
+	UE_DEPRECATED(5.8, "Use GetMaterialIOValueType instead")
 	UNREALED_API static uint32 GetMaterialValueType(const UEdGraphPin* MaterialPin);
+	UNREALED_API static EMaterialValueType GetMaterialIOValueType(const UEdGraphPin* MaterialPin);
 
 	//~ Begin UEdGraphSchema Interface
 	virtual void GetGraphContextActions(FGraphContextMenuBuilder& ContextMenuBuilder) const override;
@@ -301,7 +306,7 @@ private:
 	 * @param	TestType			Material Value Type we are testing.
 	 * @param	TestDirection		Pin Direction we are testing.
 	*/
-	bool HasCompatibleConnection(const FAssetData& FunctionAssetData, uint32 TestType, EEdGraphPinDirection TestDirection) const;
+	bool HasCompatibleConnection(const FAssetData& FunctionAssetData, EMaterialValueType TestType, EEdGraphPinDirection TestDirection) const;
 
 	/** Will promote selected pin to a parameter of the pin type */
 	void OnPromoteToParameter(const FToolMenuContext& InMenuContext);
diff --git a/Engine/Source/Editor/UnrealEd/Private/MaterialGraphSchema.cpp b/Engine/Source/Editor/UnrealEd/Private/MaterialGraphSchema.cpp
--- a/Engine/Source/Editor/UnrealEd/Private/MaterialGraphSchema.cpp
+++ b/Engine/Source/Editor/UnrealEd/Private/MaterialGraphSchema.cpp
@@ -92,7 +92,7 @@ UEdGraphNode* FMaterialGraphSchemaAction_NewNode::PerformAction(class UEdGraph*
 		if (MaterialExpressionClass == UMaterialExpressionFunctionInput::StaticClass() && FromPin)
 		{
 			// Set this to be an input of the type we dragged from
-			SetFunctionInputType(CastChecked<UMaterialExpressionFunctionInput>(NewExpression), UMaterialGraphSchema::GetMaterialValueType(FromPin));
+			SetFunctionInputType(CastChecked<UMaterialExpressionFunctionInput>(NewExpression), UMaterialGraphSchema::GetMaterialIOValueType(FromPin));
 		}
 
 		NewExpression->GraphNode->AutowireNewNode(FromPin);
@@ -104,6 +104,11 @@ UEdGraphNode* FMaterialGraphSchemaAction_NewNode::PerformAction(class UEdGraph*
 }
 
 void FMaterialGraphSchemaAction_NewNode::SetFunctionInputType(UMaterialExpressionFunctionInput* FunctionInput, uint32 MaterialValueType) const
+{
+	return SetFunctionInputType(FunctionInput, static_cast<EMaterialValueType>(MaterialValueType));
+}
+
+void FMaterialGraphSchemaAction_NewNode::SetFunctionInputType(UMaterialExpressionFunctionInput* FunctionInput, EMaterialValueType MaterialValueType) const
 {
 	switch (MaterialValueType)
 	{
@@ -345,8 +350,8 @@ bool UMaterialGraphSchema::ConnectionCausesLoop(const UEdGraphPin* InputPin, con
 
 bool UMaterialGraphSchema::ArePinsCompatible_Internal(const UEdGraphPin* InputPin, const UEdGraphPin* OutputPin, FText& ResponseMessage) const
 {
-	uint32 InputType = GetMaterialValueType(InputPin);
-	uint32 OutputType = GetMaterialValueType(OutputPin);
+	EMaterialValueType InputType = GetMaterialIOValueType(InputPin);
+	EMaterialValueType OutputType = GetMaterialIOValueType(OutputPin);
 
 	if (InputPin->bNotConnectable)
 	{
@@ -393,6 +398,11 @@ bool UMaterialGraphSchema::ArePinsCompatible_Internal(const UEdGraphPin* InputPi
 }
 
 uint32 UMaterialGraphSchema::GetMaterialValueType(const UEdGraphPin* MaterialPin)
+{
+	return static_cast<uint32>(GetMaterialIOValueType(MaterialPin));
+}
+
+EMaterialValueType UMaterialGraphSchema::GetMaterialIOValueType(const UEdGraphPin* MaterialPin)
 {
 	const UMaterialGraphNode_Base* OwningNode = CastChecked<UMaterialGraphNode_Base>(MaterialPin->GetOwningNode());
 	if (MaterialPin->Direction == EGPD_Output)
@@ -545,7 +555,7 @@ const FPinConnectionResponse UMaterialGraphSchema::CanCreateConnection(const UEd
 	// For non-exec pins, break existing connections on inputs only - multiple output connections are acceptable
 	if (InputPin->LinkedTo.Num() > 0)
 	{
-		const uint32 InputType = GetMaterialValueType(InputPin);
+		const EMaterialValueType InputType = GetMaterialIOValueType(InputPin);
 		if (!(InputType & MCT_Execution))
 		{
 			ECanCreateConnectionResponse ReplyBreakOutputs;
@@ -568,7 +578,7 @@ const FPinConnectionResponse UMaterialGraphSchema::CanCreateConnection(const UEd
 	// For exec pins, reverse is true - multiple input connections are acceptable
 	if (OutputPin->LinkedTo.Num() > 0)
 	{
-		const uint32 OutputType = GetMaterialValueType(InputPin);
+		const EMaterialValueType OutputType = GetMaterialIOValueType(InputPin);
 		if (OutputType & MCT_Execution)
 		{
 			ECanCreateConnectionResponse ReplyBreakInputs;
@@ -900,10 +910,10 @@ TSharedPtr<FEdGraphSchemaAction> UMaterialGraphSchema::GetCreateCommentAction()
 void UMaterialGraphSchema::GetMaterialFunctionActions(FGraphActionMenuBuilder& ActionMenuBuilder) const
 {
 	// Get type of dragged pin
-	uint32 FromPinType = 0;
+	EMaterialValueType FromPinType{};
 	if (ActionMenuBuilder.FromPin)
 	{
-		FromPinType = GetMaterialValueType(ActionMenuBuilder.FromPin);
+		FromPinType = GetMaterialIOValueType(ActionMenuBuilder.FromPin);
 	}
 
 	// Load the asset registry module
@@ -1078,12 +1088,12 @@ void UMaterialGraphSchema::GetNamedRerouteActions(FGraphActionMenuBuilder& Actio
 	}
 }
 
-bool UMaterialGraphSchema::HasCompatibleConnection(const FAssetData& FunctionAssetData, uint32 TestType, EEdGraphPinDirection TestDirection) const
+bool UMaterialGraphSchema::HasCompatibleConnection(const FAssetData& FunctionAssetData, EMaterialValueType TestType, EEdGraphPinDirection TestDirection) const
 {
 	if (TestType != 0)
 	{
-		uint32 CombinedInputTypes = FunctionAssetData.GetTagValueRef<uint32>(GET_MEMBER_NAME_CHECKED(UMaterialFunctionInterface, CombinedInputTypes));
-		uint32 CombinedOutputTypes = FunctionAssetData.GetTagValueRef<uint32>(GET_MEMBER_NAME_CHECKED(UMaterialFunctionInterface, CombinedOutputTypes));
+		auto CombinedInputTypes = static_cast<EMaterialValueType>(FunctionAssetData.GetTagValueRef<uint64>(GET_MEMBER_NAME_CHECKED(UMaterialFunctionInterface, CombinedInputTypes)));
+		auto CombinedOutputTypes = static_cast<EMaterialValueType>(FunctionAssetData.GetTagValueRef<uint64>(GET_MEMBER_NAME_CHECKED(UMaterialFunctionInterface, CombinedOutputTypes)));
 
 		if (CombinedOutputTypes == 0)
 		{
@@ -1091,8 +1101,8 @@ bool UMaterialGraphSchema::HasCompatibleConnection(const FAssetData& FunctionAss
 			UMaterialFunctionInterface* MaterialFunction = Cast<UMaterialFunctionInterface>(FunctionAssetData.GetAsset());
 			if (MaterialFunction != nullptr)
 			{
-				CombinedInputTypes = MaterialFunction->CombinedInputTypes;
-				CombinedOutputTypes = MaterialFunction->CombinedOutputTypes;
+				CombinedInputTypes = static_cast<EMaterialValueType>(MaterialFunction->CombinedInputTypes);
+				CombinedOutputTypes = static_cast<EMaterialValueType>(MaterialFunction->CombinedOutputTypes);
 			}
 		}
 
diff --git a/Engine/Source/Runtime/Engine/Private/Materials/MaterialExpressions.cpp b/Engine/Source/Runtime/Engine/Private/Materials/MaterialExpressions.cpp
--- a/Engine/Source/Runtime/Engine/Private/Materials/MaterialExpressions.cpp
+++ b/Engine/Source/Runtime/Engine/Private/Materials/MaterialExpressions.cpp
@@ -391,10 +391,15 @@ void ConvertToMultilineToolTip(const FString& InToolTip, const int32 TargetLineL
 	}
 }
 
-void GetMaterialValueTypeDescriptions(const uint32 MaterialValueType, TArray<FText>& OutDescriptions)
+extern ENGINE_API void GetMaterialValueTypeDescriptions(const uint32 MaterialValueType, TArray<FText>& OutDescriptions)
+{
+	return GetMaterialValueTypeDescriptions(static_cast<EMaterialValueType>(MaterialValueType), OutDescriptions);
+}
+
+extern ENGINE_API void GetMaterialValueTypeDescriptions(const EMaterialValueType MaterialValueType, TArray<FText>& OutDescriptions)
 {
 	// Get exact float type if possible
-	uint32 MaskedFloatType = MaterialValueType & MCT_Float;
+	EMaterialValueType MaskedFloatType = static_cast<EMaterialValueType>(MaterialValueType & MCT_Float);
 	if (MaskedFloatType)
 	{
 		switch (MaskedFloatType)
@@ -409,7 +414,7 @@ void GetMaterialValueTypeDescriptions(const uint32 MaterialValueType, TArray<FTe
 	}
 
 	// Get exact int type if possible
-	uint32 MaskedUIntType = MaterialValueType & MCT_UInt;
+	EMaterialValueType MaskedUIntType = static_cast<EMaterialValueType>(MaterialValueType & MCT_UInt);
 	if (MaskedUIntType)
 	{
 		switch (MaskedUIntType)
@@ -423,7 +428,7 @@ void GetMaterialValueTypeDescriptions(const uint32 MaterialValueType, TArray<FTe
 		}
 	}
 	// Get exact texture type if possible
-	uint32 MaskedTextureType = MaterialValueType & MCT_Texture;
+	EMaterialValueType MaskedTextureType = static_cast<EMaterialValueType>(MaterialValueType & MCT_Texture);
 	if (MaskedTextureType)
 	{
 		switch (MaskedTextureType)
@@ -466,6 +471,11 @@ void GetMaterialValueTypeDescriptions(const uint32 MaterialValueType, TArray<FTe
 }
 
 bool CanConnectMaterialValueTypes(const uint32 InputType, const uint32 OutputType)
+{
+	return CanConnectMaterialValueTypes(static_cast<EMaterialValueType>(InputType), static_cast<EMaterialValueType>(OutputType));
+}
+
+bool CanConnectMaterialValueTypes(const EMaterialValueType InputType, const EMaterialValueType OutputType)
 {
 	if ((InputType & MCT_Execution) || (OutputType & MCT_Execution))
 	{
@@ -13737,8 +13747,8 @@ UMaterialInterface* UMaterialFunction::GetPreviewMaterial()
 
 void UMaterialFunction::UpdateInputOutputTypes()
 {
-	CombinedInputTypes = 0;
-	CombinedOutputTypes = 0;
+	CombinedInputTypes = {};
+	CombinedOutputTypes = {};
 
 	for (UMaterialExpression* CurrentExpression : GetExpressions())
 	{
@@ -13751,7 +13761,7 @@ void UMaterialFunction::UpdateInputOutputTypes()
 		}
 		else if (OutputExpression)
 		{
-			CombinedOutputTypes |= OutputExpression->GetOutputValueType(0);
+			CombinedInputTypes |= OutputExpression->GetOutputValueType(0);
 		}
 	}
 }
diff --git a/Engine/Source/Runtime/Engine/Public/MaterialShared.h b/Engine/Source/Runtime/Engine/Public/MaterialShared.h
--- a/Engine/Source/Runtime/Engine/Public/MaterialShared.h
+++ b/Engine/Source/Runtime/Engine/Public/MaterialShared.h
@@ -1963,10 +1963,14 @@ extern ENGINE_API bool IsAllowedExpressionType(const UClass* Class, const bool b
 extern ENGINE_API void ConvertToMultilineToolTip(const FString& InToolTip, const int32 TargetLineLength, TArray<FString>& OutToolTip);
 
 /** Given a combination of EMaterialValueType flags, get text descriptions of all types */
+UE_DEPRECATED(5.8, "Use the GetMaterialValueTypeDescriptions(EMaterialValueType...) overload instead")
 extern ENGINE_API void GetMaterialValueTypeDescriptions(const uint32 MaterialValueType, TArray<FText>& OutDescriptions);
+extern ENGINE_API void GetMaterialValueTypeDescriptions(const EMaterialValueType MaterialValueType, TArray<FText>& OutDescriptions);
 
 /** Check whether a combination of EMaterialValueType flags can be connected */
+UE_DEPRECATED(5.8, "Use the CanConnectMaterialValueTypes(EMaterialValueType, EMaterialValueType) overload instead")
 extern ENGINE_API bool CanConnectMaterialValueTypes(const uint32 InputType, const uint32 OutputType);
+extern ENGINE_API bool CanConnectMaterialValueTypes(const EMaterialValueType InputType, const EMaterialValueType OutputType);
 
 /**
  * FMaterial serves 3 intertwined purposes:
diff --git a/Engine/Source/Runtime/Engine/Public/Materials/MaterialFunctionInterface.h b/Engine/Source/Runtime/Engine/Public/Materials/MaterialFunctionInterface.h
--- a/Engine/Source/Runtime/Engine/Public/Materials/MaterialFunctionInterface.h
+++ b/Engine/Source/Runtime/Engine/Public/Materials/MaterialFunctionInterface.h
@@ -164,10 +164,10 @@ public:
 
 #if WITH_EDITORONLY_DATA
 	UPROPERTY(AssetRegistrySearchable)
-	uint32 CombinedInputTypes;
+	uint64 CombinedInputTypes;
 
 	UPROPERTY(AssetRegistrySearchable)
-	uint32 CombinedOutputTypes;
+	uint64 CombinedOutputTypes;
 
 	/** Information for thumbnail rendering */
 	UPROPERTY(VisibleAnywhere, Instanced, Category = Thumbnail)
