AssetRegistry request tracking:
- Add new event OnEnumerateAssetsEvent called when a request is made to the asset registry.
- CookAssetRegistryAccessTracker : new singleton listening to the event OnEnumerateAssetsEvent and keeping track of requests per package.
- TryCollectInternal : Make AssetRegistryQuery dependencies out of the accesses tracked by FCookAssetRegistryAccessTracker.
- FARFilter : new functions to compute hash and convert from FARCompiledFilter to FARFilter.

#rnx
#incrementalcook

diff --git a/Engine/Source/Editor/UnrealEd/Private/CookOnTheFlyServer.cpp b/Engine/Source/Editor/UnrealEd/Private/CookOnTheFlyServer.cpp
--- a/Engine/Source/Editor/UnrealEd/Private/CookOnTheFlyServer.cpp
+++ b/Engine/Source/Editor/UnrealEd/Private/CookOnTheFlyServer.cpp
@@ -21,6 +21,7 @@
 #include "Containers/DirectoryTree.h"
 #include "Containers/RingBuffer.h"
 #include "Cooker/AsyncIODelete.h"
+#include "Cooker/CookAssetRegistryAccessTracker.h"
 #include "Cooker/CookConfigAccessTracker.h"
 #include "Cooker/CookDiagnostics.h"
 #include "Cooker/CookDirector.h"
@@ -5606,6 +5607,13 @@ double UCookOnTheFlyServer::GetIdleTimeToGC() const
 
 void UCookOnTheFlyServer::BeginDestroy()
 {
+	// BeginDestroy will be called for the CDO of UCookOnTheFlyServer and I shouldn't call anything in this case.
+	// So check if Initialize was called using PackageDatas.IsValid() and only call Shutdown in this case.
+	if (PackageDatas.IsValid())
+	{
+		UE::CookAssetRegistryAccessTracker::FCookAssetRegistryAccessTracker::Get().Shutdown();
+	}
+
 	ShutdownCookOnTheFly();
 
 	Super::BeginDestroy();
@@ -5945,6 +5953,8 @@ void UCookOnTheFlyServer::Initialize( ECookMode::Type DesiredCookMode, ECookInit
 #if WITH_ADDITIONAL_CRASH_CONTEXTS
 	FGenericCrashContext::OnAdditionalCrashContextDelegate().AddUObject(this, &UCookOnTheFlyServer::DumpCrashContext);
 #endif
+
+	UE::CookAssetRegistryAccessTracker::FCookAssetRegistryAccessTracker::Get().Init();
 }
 
 void UCookOnTheFlyServer::InitializeAtFirstSession()
diff --git a/Engine/Source/Editor/UnrealEd/Private/Cooker/CookAssetRegistryAccessTracker.cpp b/Engine/Source/Editor/UnrealEd/Private/Cooker/CookAssetRegistryAccessTracker.cpp
new file mode 100644
--- /dev/null
+++ b/Engine/Source/Editor/UnrealEd/Private/Cooker/CookAssetRegistryAccessTracker.cpp
@@ -0,0 +1,158 @@
+// Copyright Epic Games, Inc. All Rights Reserved.
+
+#include "Cooker/CookAssetRegistryAccessTracker.h"
+
+#include "Algo/Unique.h"
+#include "AssetRegistry/AssetData.h"
+#include "AssetRegistry/AssetRegistryModule.h"
+#include "Interfaces/ITargetPlatform.h"
+#include "Misc/PackageAccessTracking.h"
+
+namespace UE::CookAssetRegistryAccessTracker
+{
+	uint32 GetTypeHash(const FRecord& Data)
+	{
+		uint32 Hash = 0;
+		Hash = HashCombineFast(Hash, GetTypeHash(Data.Platform));
+		Hash = HashCombineFast(Hash, UE::AssetRegistry::GetTypeHash(Data.Filter));
+		return Hash;
+	}
+
+	bool FRecord::operator==(const FRecord& Other) const
+	{
+		return Filter == Other.Filter && Platform == Other.Platform;
+	}
+
+	bool FRecord::operator<(const FRecord& Other) const
+	{
+		if (Filter != Other.Filter)
+		{
+			return Filter < Other.Filter;
+		}
+
+		if (Platform != Other.Platform)
+		{
+			if (Platform == nullptr)
+			{
+				return true;
+			}
+			if (Other.Platform == nullptr)
+			{
+				return false;
+			}
+
+			return Platform->PlatformName() < Other.Platform->PlatformName();
+		}
+
+		return false;
+	}
+
+	FCookAssetRegistryAccessTracker FCookAssetRegistryAccessTracker::Singleton;
+
+	FCookAssetRegistryAccessTracker& FCookAssetRegistryAccessTracker::Get()
+	{ 
+		return Singleton; 
+	}
+
+	void FCookAssetRegistryAccessTracker::Init()
+	{
+		if (bEnabled)
+		{
+			return;
+		}
+
+		FAssetRegistryModule& AssetRegistryModule = FModuleManager::LoadModuleChecked<FAssetRegistryModule>(TEXT("AssetRegistry"));
+		IAssetRegistry& AssetRegistry = AssetRegistryModule.GetRegistry();
+		Handle = AssetRegistry.OnEnumerateAssetsEvent().AddLambda([this](const FARCompiledFilter& Filter, UE::AssetRegistry::EEnumerateAssetsFlags Flag) 
+			{ 
+				OnEnumerateAssets(Filter, Flag);
+			});
+
+		bEnabled = true;
+	}
+
+	void FCookAssetRegistryAccessTracker::Shutdown()
+	{
+		if (!bEnabled)
+		{
+			return;
+		}
+
+		if (FAssetRegistryModule* AssetRegistryModule = FModuleManager::GetModulePtr<FAssetRegistryModule>(TEXT("AssetRegistry")))
+		{
+			IAssetRegistry& AssetRegistry = AssetRegistryModule->GetRegistry();
+			AssetRegistry.OnEnumerateAssetsEvent().Remove(Handle);
+		}
+
+		bEnabled = false;
+	}
+
+	TArray<FRecord> FCookAssetRegistryAccessTracker::GetRecords(const FName& PackageName, const ITargetPlatform* Platform) const
+	{
+		TArray<FRecord> PlatformRecords;
+
+		{
+			UE::TUniqueLock RecordsLock(RecordsMutex);
+
+			const TSet<FRecord>* FoundRecords = AccessRecords.Find(PackageName);
+			if (!FoundRecords)
+			{
+				return TArray<FRecord>();
+			}
+
+			PlatformRecords = FoundRecords->Array();
+		}
+
+		for (TArray<FRecord>::TIterator Iter(PlatformRecords); Iter; ++Iter)
+		{
+			FRecord& Record = *Iter;
+			if (Record.Platform != Platform && Record.Platform != nullptr)
+			{
+				Iter.RemoveCurrentSwap();
+			}
+		}
+
+		Algo::Sort(PlatformRecords);
+		PlatformRecords.SetNum(Algo::Unique(PlatformRecords));
+
+		return PlatformRecords;
+	}
+
+	FCookAssetRegistryAccessTracker::FCookAssetRegistryAccessTracker()
+		: AccessRecords()
+		, Handle()
+		, RecordsMutex()
+		, bEnabled(false)
+	{ }
+
+	FCookAssetRegistryAccessTracker::~FCookAssetRegistryAccessTracker()
+	{ }
+
+	void FCookAssetRegistryAccessTracker::OnEnumerateAssets(const FARCompiledFilter& CompiledFilter, UE::AssetRegistry::EEnumerateAssetsFlags Flag)
+	{
+		if (!bEnabled)
+		{
+			return;
+		}
+
+		PackageAccessTracking_Private::FTrackedData* AccumulatedScopeData = PackageAccessTracking_Private::FPackageAccessRefScope::GetCurrentThreadAccumulatedData();
+		if (!AccumulatedScopeData || AccumulatedScopeData->BuildOpName.IsNone())
+		{
+			return;
+		}
+
+		FName Referencer = AccumulatedScopeData->PackageName;
+
+		{
+			UE::TUniqueLock RecordsLock(RecordsMutex);
+
+			TSet<FRecord>& Records = AccessRecords.FindOrAdd(Referencer);
+			FRecord NewRecord;
+			NewRecord.Platform = AccumulatedScopeData->TargetPlatform;
+			NewRecord.Filter = UE::AssetRegistry::ConvertToNonCompiledFilter(CompiledFilter);
+			NewRecord.Filter.SortForSaving();
+
+			Records.Add(NewRecord);
+		}
+	}
+}
diff --git a/Engine/Source/Editor/UnrealEd/Private/Cooker/CookAssetRegistryAccessTracker.h b/Engine/Source/Editor/UnrealEd/Private/Cooker/CookAssetRegistryAccessTracker.h
new file mode 100644
--- /dev/null
+++ b/Engine/Source/Editor/UnrealEd/Private/Cooker/CookAssetRegistryAccessTracker.h
@@ -0,0 +1,53 @@
+// Copyright Epic Games, Inc. All Rights Reserved.
+
+#pragma once
+
+#include "AssetRegistry/ARFilter.h"
+#include "Async/Mutex.h"
+#include "Containers/Map.h"
+#include "Delegates/IDelegateInstance.h"
+#include "Misc/AssetRegistryInterface.h"
+#include "UObject/NameTypes.h"
+
+struct FAssetData;
+class ITargetPlatform;
+
+namespace UE::CookAssetRegistryAccessTracker
+{
+	struct FRecord
+	{
+		const ITargetPlatform* Platform;
+		FARFilter Filter;
+
+		bool operator==(const FRecord& Other) const;
+		bool operator<(const FRecord& Other) const;
+	};
+
+	class FCookAssetRegistryAccessTracker : public FNoncopyable
+	{
+	public:
+		static FCookAssetRegistryAccessTracker& Get();
+
+		void Init();
+		void Shutdown();
+
+		TArray<FRecord> GetRecords(const FName& PackageName, const ITargetPlatform* Platform) const;
+
+	private:
+		FCookAssetRegistryAccessTracker();
+		virtual ~FCookAssetRegistryAccessTracker();
+
+		void OnEnumerateAssets(const FARCompiledFilter& Filter, UE::AssetRegistry::EEnumerateAssetsFlags Flag);
+
+		static FCookAssetRegistryAccessTracker Singleton;
+
+		TMap<FName, TSet<FRecord>> AccessRecords;
+
+		FDelegateHandle Handle;
+
+		// Use a mutex rather than a critical section for synchronization.  Calls into system libraries, such as windows critical section
+		// functions, are 50 times more expensive on build farm VMs, radically affecting cook times, which this avoids. 
+		mutable UE::FMutex RecordsMutex;
+		bool bEnabled = false;
+	};
+}
diff --git a/Engine/Source/Editor/UnrealEd/Private/Cooker/CookPackageArtifacts.cpp b/Engine/Source/Editor/UnrealEd/Private/Cooker/CookPackageArtifacts.cpp
--- a/Engine/Source/Editor/UnrealEd/Private/Cooker/CookPackageArtifacts.cpp
+++ b/Engine/Source/Editor/UnrealEd/Private/Cooker/CookPackageArtifacts.cpp
@@ -5,6 +5,7 @@
 #include "Algo/Sort.h"
 #include "Algo/Unique.h"
 #include "AssetRegistry/IAssetRegistry.h"
+#include "Cooker/CookAssetRegistryAccessTracker.h"
 #include "Cooker/CookConfigAccessTracker.h"
 #include "Cooker/CookDependency.h"
 #include "Cooker/CookEvents.h"
@@ -470,6 +471,16 @@ bool FBuildDependencySet::TryCollectInternal(FBuildResultDependenciesMap& InOutR
 		}
 	}
 #endif
+
+	{
+		UE::CookAssetRegistryAccessTracker::FCookAssetRegistryAccessTracker& AssetRegistryTracker = UE::CookAssetRegistryAccessTracker::FCookAssetRegistryAccessTracker::Get();
+		TArray<UE::CookAssetRegistryAccessTracker::FRecord> Records = AssetRegistryTracker.GetRecords(PackageName, TargetPlatform);
+		for (const UE::CookAssetRegistryAccessTracker::FRecord& Record : Records)
+		{
+			DefaultResultDependencies.Add(FCookDependency::AssetRegistryQuery(Record.Filter));
+		}
+	}
+
 	if (!UntrackedSoftPackageReferences.IsEmpty())
 	{
 		TArray<FCookDependency>& SaveDependencies = InOutResultDependencies.FindOrAdd(BuildResult::NAME_Save);
diff --git a/Engine/Source/Runtime/AssetRegistry/Private/AssetRegistry.cpp b/Engine/Source/Runtime/AssetRegistry/Private/AssetRegistry.cpp
--- a/Engine/Source/Runtime/AssetRegistry/Private/AssetRegistry.cpp
+++ b/Engine/Source/Runtime/AssetRegistry/Private/AssetRegistry.cpp
@@ -2390,6 +2390,9 @@ bool UAssetRegistryImpl::GetAssets(const FARCompiledFilter& CompiledFilter, TArr
 	TRACE_CPUPROFILER_EVENT_SCOPE(UAssetRegistryImpl::GetAssets);
 	using namespace UE::AssetRegistry::Utils;
 
+	const UE::AssetRegistry::EEnumerateAssetsFlags Flags = bSkipARFilteredAssets ? UE::AssetRegistry::EEnumerateAssetsFlags::None : UE::AssetRegistry::EEnumerateAssetsFlags::AllowUnfilteredArAssets;
+	EnumerateAssetsEvent.Broadcast(CompiledFilter, Flags);
+
 	TSet<FName> PackagesToSkip;
 	if (!CompiledFilter.bIncludeOnlyOnDiskAssets)
 	{
@@ -2404,7 +2407,6 @@ bool UAssetRegistryImpl::GetAssets(const FARCompiledFilter& CompiledFilter, TArr
 	}
 
 	{
-		const UE::AssetRegistry::EEnumerateAssetsFlags Flags = bSkipARFilteredAssets ? UE::AssetRegistry::EEnumerateAssetsFlags::None : UE::AssetRegistry::EEnumerateAssetsFlags::AllowUnfilteredArAssets;
 		UE::AssetRegistry::FInterfaceReadScopeLock InterfaceScopeLock(InterfaceLock);
 		GuardedData.EnumerateDiskAssets(CompiledFilter, PackagesToSkip, [&OutAssetData](const FAssetData& AssetData)
 			{
@@ -2431,6 +2433,10 @@ bool UAssetRegistryImpl::GetInMemoryAssets(const FARFilter& InFilter, TArray<FAs
 bool UAssetRegistryImpl::GetInMemoryAssets(const FARCompiledFilter& CompiledFilter, TArray<FAssetData>& OutAssetData, bool bSkipARFilteredAssets) const 
 {
 	using namespace UE::AssetRegistry::Utils;
+
+	const UE::AssetRegistry::EEnumerateAssetsFlags Flags = bSkipARFilteredAssets ? UE::AssetRegistry::EEnumerateAssetsFlags::None : UE::AssetRegistry::EEnumerateAssetsFlags::AllowUnfilteredArAssets;
+	EnumerateAssetsEvent.Broadcast(CompiledFilter, Flags);
+
 	TSet<FName> PackagesToSkipUnused;
 	bool bStopIterationUnused;
 	EnumerateMemoryAssets(CompiledFilter, PackagesToSkipUnused, bStopIterationUnused,
@@ -2489,6 +2495,8 @@ bool UAssetRegistryImpl::EnumerateAssets(const FARCompiledFilter& InFilter, TFun
 		return false;
 	}
 
+	EnumerateAssetsEvent.Broadcast(InFilter, InEnumerateFlags);
+
 	TSet<FName> PackagesToSkip;
 	if (!InFilter.bIncludeOnlyOnDiskAssets)
 	{
@@ -4739,6 +4747,10 @@ UE::AssetRegistry::Impl::EGatherStatus UAssetRegistryImpl::TickOnBackgroundThrea
 	return Status;
 }
 
+UAssetRegistryImpl::FEnumerateAssetsEvent& UAssetRegistryImpl::OnEnumerateAssetsEvent()
+{
+	return EnumerateAssetsEvent;
+}
 
 void UAssetRegistryImpl::Tick(float DeltaTime)
 {
diff --git a/Engine/Source/Runtime/AssetRegistry/Private/AssetRegistry.h b/Engine/Source/Runtime/AssetRegistry/Private/AssetRegistry.h
--- a/Engine/Source/Runtime/AssetRegistry/Private/AssetRegistry.h
+++ b/Engine/Source/Runtime/AssetRegistry/Private/AssetRegistry.h
@@ -225,6 +225,9 @@ public:
 
 	UE::AssetRegistry::Impl::EGatherStatus TickOnBackgroundThread();
 
+	DECLARE_DERIVED_EVENT(UAssetRegistryImpl, IAssetRegistry::FEnumerateAssetsEvent, FEnumerateAssetsEvent)
+	virtual FEnumerateAssetsEvent& OnEnumerateAssetsEvent() override;
+
 protected:
 	virtual void SetManageReferences(const TMultiMap<FAssetIdentifier, FAssetIdentifier>& ManagerMap,
 		bool bClearExisting, UE::AssetRegistry::EDependencyCategory RecurseType,
@@ -428,6 +431,9 @@ private:
 	/** The delegate to execute when each batch of searching/gathering is complete */
 	FKnownGathersCompleteEvent KnownGathersCompleteEvent;
 
+	/** The delegate to execute when a query is made. */
+	FEnumerateAssetsEvent EnumerateAssetsEvent;
+
 	/**
 	 * Storage for events that will be broadcast later on the game thread. Only safe
 	 *  to access under the DeferredEventsCriticalSection. 
diff --git a/Engine/Source/Runtime/AssetRegistry/Public/AssetRegistry/IAssetRegistry.h b/Engine/Source/Runtime/AssetRegistry/Public/AssetRegistry/IAssetRegistry.h
--- a/Engine/Source/Runtime/AssetRegistry/Public/AssetRegistry/IAssetRegistry.h
+++ b/Engine/Source/Runtime/AssetRegistry/Public/AssetRegistry/IAssetRegistry.h
@@ -1165,6 +1165,10 @@ public:
 		~FPauseBackgroundProcessingScope() { IAssetRegistry::GetChecked().RequestResumeBackgroundProcessing(); }
 	};
 
+	/** Event triggered when a query is made. The event is triggered only if the filter is valid. */
+	DECLARE_MULTICAST_DELEGATE_TwoParams(FEnumerateAssetsEvent, const FARCompiledFilter& /*Filter*/, UE::AssetRegistry::EEnumerateAssetsFlags /*Flag*/)
+	virtual FEnumerateAssetsEvent& OnEnumerateAssetsEvent() = 0;
+
 protected:
 
 	// Implementation methods for FPauseBackgroundProcessingScope
diff --git a/Engine/Source/Runtime/CoreUObject/Private/AssetRegistry/ARFilter.cpp b/Engine/Source/Runtime/CoreUObject/Private/AssetRegistry/ARFilter.cpp
--- a/Engine/Source/Runtime/CoreUObject/Private/AssetRegistry/ARFilter.cpp
+++ b/Engine/Source/Runtime/CoreUObject/Private/AssetRegistry/ARFilter.cpp
@@ -424,3 +424,59 @@ bool FARFilter::TryLoad(const FCbFieldView& Field)
 	}
 	return true;
 }
+
+namespace UE::AssetRegistry
+{
+	uint32 GetTypeHash(const FARFilter& Filter)
+	{
+		uint32 Hash = 0;
+		Hash = HashCombineFast(Hash, GetTypeHash(Filter.PackageNames));
+		Hash = HashCombineFast(Hash, GetTypeHash(Filter.PackagePaths));
+		Hash = HashCombineFast(Hash, GetTypeHash(Filter.SoftObjectPaths));
+		Hash = HashCombineFast(Hash, GetTypeHash(Filter.ClassPaths));
+
+		for (TMultiMap<FName, TOptional<FString>>::TConstIterator It = Filter.TagsAndValues.CreateConstIterator(); It; ++It)
+		{
+			Hash = HashCombineFast(Hash, GetTypeHash(It.Key()));
+
+			const TOptional<FString>& Value = It.Value();
+			if (Value.IsSet())
+			{
+				Hash = HashCombineFast(Hash, GetTypeHash(Value.GetValue()));
+			}
+		}
+
+		for (const FTopLevelAssetPath& AssetPath : Filter.RecursiveClassPathsExclusionSet)
+		{
+			Hash = HashCombineFast(Hash, GetTypeHash(AssetPath));
+		}
+
+		Hash = HashCombineFast(Hash, ::GetTypeHash(Filter.bRecursivePaths));
+		Hash = HashCombineFast(Hash, ::GetTypeHash(Filter.bRecursiveClasses));
+		Hash = HashCombineFast(Hash, ::GetTypeHash(Filter.bIncludeOnlyOnDiskAssets));
+		Hash = HashCombineFast(Hash, ::GetTypeHash(Filter.WithoutPackageFlags));
+		Hash = HashCombineFast(Hash, ::GetTypeHash(Filter.WithPackageFlags));
+
+		return Hash;
+	}
+
+	FARFilter ConvertToNonCompiledFilter(const FARCompiledFilter& CompiledFilter)
+	{
+		FARFilter OutFilter;
+		OutFilter.PackageNames = CompiledFilter.PackageNames.Array();
+		OutFilter.PackagePaths = CompiledFilter.PackagePaths.Array();
+
+		OutFilter.SoftObjectPaths = CompiledFilter.SoftObjectPaths.Array();
+
+		OutFilter.ClassPaths = CompiledFilter.ClassPaths.Array();
+		OutFilter.TagsAndValues = CompiledFilter.TagsAndValues;
+		OutFilter.bIncludeOnlyOnDiskAssets = CompiledFilter.bIncludeOnlyOnDiskAssets;
+		OutFilter.WithoutPackageFlags = CompiledFilter.WithoutPackageFlags;
+		OutFilter.WithPackageFlags = CompiledFilter.WithPackageFlags;
+
+		OutFilter.bRecursivePaths = false;
+		OutFilter.bRecursiveClasses = false;
+
+		return OutFilter;
+	}
+}
diff --git a/Engine/Source/Runtime/CoreUObject/Public/AssetRegistry/ARFilter.h b/Engine/Source/Runtime/CoreUObject/Public/AssetRegistry/ARFilter.h
--- a/Engine/Source/Runtime/CoreUObject/Public/AssetRegistry/ARFilter.h
+++ b/Engine/Source/Runtime/CoreUObject/Public/AssetRegistry/ARFilter.h
@@ -237,3 +237,10 @@ PRAGMA_ENABLE_DEPRECATION_WARNINGS
 		ensure(IsEmpty());
 	}
 };
+
+namespace UE::AssetRegistry
+{
+	COREUOBJECT_API uint32 GetTypeHash(const FARFilter& Filter);
+
+	COREUOBJECT_API FARFilter ConvertToNonCompiledFilter(const FARCompiledFilter& CompiledFilter);
+}
