[UObject/General]
- Cleanup code and convert to the new ConditionalPreload function
- Fix a few thread-safety issue when resetting flags before preloading

diff --git a/Engine/Plugins/2D/Paper2D/Source/Paper2D/Private/PaperSprite.cpp b/Engine/Plugins/2D/Paper2D/Source/Paper2D/Private/PaperSprite.cpp
--- a/Engine/Plugins/2D/Paper2D/Source/Paper2D/Private/PaperSprite.cpp
+++ b/Engine/Plugins/2D/Paper2D/Source/Paper2D/Private/PaperSprite.cpp
@@ -1460,13 +1460,7 @@ UTexture2D* UPaperSprite::GetSourceTexture() const
 	// not safe to do in a game with EDL, but safe at editor time.
 	if (SourceTexturePtr)
 	{
-		if (SourceTexturePtr->HasAnyFlags(RF_NeedLoad))
-		{
-			if (FLinkerLoad* TextureLinker = SourceTexturePtr->GetLinker())
-			{
-				TextureLinker->Preload(SourceTexturePtr);
-			}
-		}
+		SourceTexturePtr->ConditionalPreload();
 		SourceTexturePtr->ConditionalPostLoad();
 	}
 
diff --git a/Engine/Plugins/Animation/AnimationData/Source/AnimationData/Private/AnimSequencerDataModel.cpp b/Engine/Plugins/Animation/AnimationData/Source/AnimationData/Private/AnimSequencerDataModel.cpp
--- a/Engine/Plugins/Animation/AnimationData/Source/AnimationData/Private/AnimSequencerDataModel.cpp
+++ b/Engine/Plugins/Animation/AnimationData/Source/AnimationData/Private/AnimSequencerDataModel.cpp
@@ -200,18 +200,11 @@ void UAnimationSequencerDataModel::PostLoad()
 	{
 		// Pre/post load any dependencies (Sequencer objects)
 		TArray<UObject*> ObjectReferences;
-		FReferenceFinder(ObjectReferences, this, false, true, true, true).FindReferences(this);	
+		FReferenceFinder(ObjectReferences, this, false, true, true, true).FindReferences(this);
 		for (UObject* Dependency : ObjectReferences)
 		{
-			if (Dependency->HasAnyFlags(RF_NeedLoad))
-			{
-				Dependency->GetLinker()->Preload(Dependency);
-			}
-
-			if (Dependency->HasAnyFlags(RF_NeedPostLoad))
-			{
-				Dependency->ConditionalPostLoad();
-			}
+			Dependency->ConditionalPreload();
+			Dependency->ConditionalPostLoad();
 		}
 
 		if (const UMovieSceneControlRigParameterSection* Section = GetFKControlRigSection())
diff --git a/Engine/Plugins/Chooser/Source/ChooserUncooked/Private/EvaluateChooserNode.cpp b/Engine/Plugins/Chooser/Source/ChooserUncooked/Private/EvaluateChooserNode.cpp
--- a/Engine/Plugins/Chooser/Source/ChooserUncooked/Private/EvaluateChooserNode.cpp
+++ b/Engine/Plugins/Chooser/Source/ChooserUncooked/Private/EvaluateChooserNode.cpp
@@ -97,10 +97,7 @@ void UK2Node_EvaluateChooser::PreloadRequiredAssets()
 {
 	if (Chooser)
 	{
-		if (FLinkerLoad* ObjLinker = GetLinker())
-		{
-			ObjLinker->Preload(Chooser);
-		}
+		Chooser->ConditionalPreload();
 	}
     		
 	Super::PreloadRequiredAssets();
@@ -387,10 +384,7 @@ void UK2Node_EvaluateChooser2::PreloadRequiredAssets()
 {
 	if (Chooser)
 	{
-		if (FLinkerLoad* ObjLinker = GetLinker())
-		{
-			ObjLinker->Preload(Chooser);
-		}
+		Chooser->ConditionalPreload();
 	}
     		
 	Super::PreloadRequiredAssets();
diff --git a/Engine/Plugins/Chooser/Source/ProxyTableUncooked/Private/EvaluateProxyNode.cpp b/Engine/Plugins/Chooser/Source/ProxyTableUncooked/Private/EvaluateProxyNode.cpp
--- a/Engine/Plugins/Chooser/Source/ProxyTableUncooked/Private/EvaluateProxyNode.cpp
+++ b/Engine/Plugins/Chooser/Source/ProxyTableUncooked/Private/EvaluateProxyNode.cpp
@@ -98,10 +98,7 @@ void UK2Node_EvaluateProxy::PreloadRequiredAssets()
 {
 	if (Proxy)
 	{
-		if (FLinkerLoad* ObjLinker = GetLinker())
-		{
-			ObjLinker->Preload(Proxy);
-		}
+		Proxy->ConditionalPreload();
 	}
     		
 	Super::PreloadRequiredAssets();
@@ -365,10 +362,7 @@ void UK2Node_EvaluateProxy2::PreloadRequiredAssets()
 {
 	if (Proxy)
 	{
-		if (FLinkerLoad* ObjLinker = GetLinker())
-		{
-			ObjLinker->Preload(Proxy);
-		}
+		Proxy->ConditionalPreload();
 	}
     		
 	Super::PreloadRequiredAssets();
diff --git a/Engine/Plugins/Mutable/Source/CustomizableObjectEditor/Private/MuCOE/CustomizableObjectEditorUtilities.cpp b/Engine/Plugins/Mutable/Source/CustomizableObjectEditor/Private/MuCOE/CustomizableObjectEditorUtilities.cpp
--- a/Engine/Plugins/Mutable/Source/CustomizableObjectEditor/Private/MuCOE/CustomizableObjectEditorUtilities.cpp
+++ b/Engine/Plugins/Mutable/Source/CustomizableObjectEditor/Private/MuCOE/CustomizableObjectEditorUtilities.cpp
@@ -1,4 +1,4 @@
-﻿// Copyright Epic Games, Inc. All Rights Reserved.
+// Copyright Epic Games, Inc. All Rights Reserved.
 
 #include "MuCOE/CustomizableObjectEditorUtilities.h"
 
@@ -28,10 +28,6 @@ bool CompareNames(const TSharedPtr<FString>& sp1, const TSharedPtr<FString>& sp2
 
 void ConditionalPostLoadReference(UObject& Object)
 {
-	if (Object.HasAnyFlags(RF_NeedLoad))
-	{
-		Object.GetLinker()->Preload(&Object);
-	}
-
+	Object.ConditionalPreload();
 	Object.ConditionalPostLoad();
 }
diff --git a/Engine/Plugins/VirtualProduction/DMX/DMXPixelMapping/Source/DMXPixelMappingRuntime/Private/DMXPixelMapping.cpp b/Engine/Plugins/VirtualProduction/DMX/DMXPixelMapping/Source/DMXPixelMappingRuntime/Private/DMXPixelMapping.cpp
--- a/Engine/Plugins/VirtualProduction/DMX/DMXPixelMapping/Source/DMXPixelMappingRuntime/Private/DMXPixelMapping.cpp
+++ b/Engine/Plugins/VirtualProduction/DMX/DMXPixelMapping/Source/DMXPixelMappingRuntime/Private/DMXPixelMapping.cpp
@@ -64,16 +64,13 @@ void UDMXPixelMapping::PostLoad()
 
 void UDMXPixelMapping::PreloadWithChildren()
 {
-	if (HasAnyFlags(RF_NeedLoad))
-	{
-		GetLinker()->Preload(this);
-	}
+	ConditionalPreload();
 
 	ForEachComponent([this](UDMXPixelMappingBaseComponent* InComponent)
 	{
-		if ((InComponent != nullptr) && InComponent->HasAnyFlags(RF_NeedLoad))
+		if (InComponent)
 		{
-			InComponent->GetLinker()->Preload(InComponent);
+			InComponent->ConditionalPreload();
 		}
 	});
 }
diff --git a/Engine/Source/Editor/AnimGraph/Private/AnimBlueprintPostCompileValidation.cpp b/Engine/Source/Editor/AnimGraph/Private/AnimBlueprintPostCompileValidation.cpp
--- a/Engine/Source/Editor/AnimGraph/Private/AnimBlueprintPostCompileValidation.cpp
+++ b/Engine/Source/Editor/AnimGraph/Private/AnimBlueprintPostCompileValidation.cpp
@@ -20,9 +20,9 @@ void UAnimBlueprintPostCompileValidation::DoPostCompileValidation(FAnimBPCompile
 // Ensures the specified object is preloaded.  ReferencedObject can be NULL.
 void UAnimBlueprintPostCompileValidation::PCV_PreloadObject(const UObject* const ReferencedObject)
 {
-	if ((ReferencedObject != nullptr) && ReferencedObject->HasAnyFlags(RF_NeedLoad))
+	if (ReferencedObject)
 	{
-		ReferencedObject->GetLinker()->Preload(const_cast<UObject*>(ReferencedObject));
+		const_cast<UObject*>(ReferencedObject)->ConditionalPreload();
 	}
 }
 
diff --git a/Engine/Source/Editor/BlueprintGraph/Classes/K2Node.h b/Engine/Source/Editor/BlueprintGraph/Classes/K2Node.h
--- a/Engine/Source/Editor/BlueprintGraph/Classes/K2Node.h
+++ b/Engine/Source/Editor/BlueprintGraph/Classes/K2Node.h
@@ -552,9 +552,9 @@ protected:
 	// Ensures the specified object is preloaded.  ReferencedObject can be NULL.
 	void PreloadObject(UObject* ReferencedObject)
 	{
-		if ((ReferencedObject != NULL) && ReferencedObject->HasAnyFlags(RF_NeedLoad))
+		if (ReferencedObject)
 		{
-			ReferencedObject->GetLinker()->Preload(ReferencedObject);
+			ReferencedObject->ConditionalPreload();
 		}
 	}
 
diff --git a/Engine/Source/Editor/BlueprintGraph/Private/K2Node.cpp b/Engine/Source/Editor/BlueprintGraph/Private/K2Node.cpp
--- a/Engine/Source/Editor/BlueprintGraph/Private/K2Node.cpp
+++ b/Engine/Source/Editor/BlueprintGraph/Private/K2Node.cpp
@@ -231,10 +231,7 @@ void UK2Node::FixupPinStringDataReferences(FArchive* SavingArchive)
 				if (Struct)
 				{
 					// While loading, our user struct may not have been linked yet
-					if (Struct->HasAnyFlags(RF_NeedLoad) && LinkerLoad)
-					{
-						LinkerLoad->Preload(Struct);
-					}
+					Struct->ConditionalPreload();
 					
 					FSoftObjectPathSerializationScope SetPackage(GetOutermost()->GetFName(), NAME_None, ESoftObjectPathCollectType::AlwaysCollect, ESoftObjectPathSerializeType::SkipSerializeIfArchiveHasSize);
 
diff --git a/Engine/Source/Editor/BlueprintGraph/Private/K2Node_FunctionEntry.cpp b/Engine/Source/Editor/BlueprintGraph/Private/K2Node_FunctionEntry.cpp
--- a/Engine/Source/Editor/BlueprintGraph/Private/K2Node_FunctionEntry.cpp
+++ b/Engine/Source/Editor/BlueprintGraph/Private/K2Node_FunctionEntry.cpp
@@ -1131,10 +1131,7 @@ bool UK2Node_FunctionEntry::ShouldUseConstRefParams() const
 				if (UK2Node_FunctionResult* ResultNode = Cast<UK2Node_FunctionResult>(Node))
 				{
 					// This might be called from the super's Serialize() method for older assets, so make sure the result node's pins have been loaded.
-					if (ResultNode->HasAnyFlags(RF_NeedLoad))
-					{
-						GetLinker()->Preload(ResultNode);
-					}
+					ResultNode->ConditionalPreload();
 
 					return ResultNode->UserDefinedPins.Num() == 0;
 				}
diff --git a/Engine/Source/Editor/BlueprintGraph/Private/K2Node_GetNumEnumEntries.cpp b/Engine/Source/Editor/BlueprintGraph/Private/K2Node_GetNumEnumEntries.cpp
--- a/Engine/Source/Editor/BlueprintGraph/Private/K2Node_GetNumEnumEntries.cpp
+++ b/Engine/Source/Editor/BlueprintGraph/Private/K2Node_GetNumEnumEntries.cpp
@@ -87,10 +87,7 @@ void UK2Node_GetNumEnumEntries::ExpandNode(class FKismetCompilerContext& Compile
 	}
 
 	// Force the enum to load its values if it hasn't already
-	if (Enum->HasAnyFlags(RF_NeedLoad))
-	{
-		Enum->GetLinker()->Preload(Enum);
-	}
+	Enum->ConditionalPreload();
 
 	//MAKE LITERAL
 	const FName FunctionName = GET_FUNCTION_NAME_CHECKED(UKismetSystemLibrary, MakeLiteralInt);
diff --git a/Engine/Source/Editor/MovieSceneTools/Private/K2Node_GetSequenceBinding.cpp b/Engine/Source/Editor/MovieSceneTools/Private/K2Node_GetSequenceBinding.cpp
--- a/Engine/Source/Editor/MovieSceneTools/Private/K2Node_GetSequenceBinding.cpp
+++ b/Engine/Source/Editor/MovieSceneTools/Private/K2Node_GetSequenceBinding.cpp
@@ -36,19 +36,11 @@ void EnsureFullyPreloaded(UObject* Object)
 		return;
 	}
 
-	FLinkerLoad* Linker = Object->GetLinker();
-	if (Object->HasAnyFlags(RF_NeedLoad))
-	{
-		if (ensure(Linker))
-		{
-			Linker->Preload(Object);
-			check(!Object->HasAnyFlags(RF_NeedLoad));
-		}
-	}
+	Object->ConditionalPreload();
 
 	// We only want to ensure that _loaded_ objects have RF_LoadCompleted set.
 	// Some objects can be created during postload, so we don't need to verify RF_LoadCompleted in those cases.
-	if (Linker)
+	if (FLinkerLoad* Linker = Object->GetLinker())
 	{
 		check(Object->HasAnyFlags(RF_LoadCompleted));
 	}
diff --git a/Engine/Source/Editor/UnrealEd/Private/Commandlets/PackageUtilities.cpp b/Engine/Source/Editor/UnrealEd/Private/Commandlets/PackageUtilities.cpp
--- a/Engine/Source/Editor/UnrealEd/Private/Commandlets/PackageUtilities.cpp
+++ b/Engine/Source/Editor/UnrealEd/Private/Commandlets/PackageUtilities.cpp
@@ -1877,10 +1877,7 @@ struct CompressAnimationsFunctor
 				continue;
 			}
 
-			if (Skeleton->HasAnyFlags(RF_NeedLoad))
-			{
-				Skeleton->GetLinker()->Preload(Skeleton);
-			}
+			Skeleton->ConditionalPreload();
 
 			float HighestRatio = 0.f;
 #if 0 // @todoanim: not sure why we need this here
diff --git a/Engine/Source/Editor/UnrealEd/Private/Kismet2/BlueprintEditorUtils.cpp b/Engine/Source/Editor/UnrealEd/Private/Kismet2/BlueprintEditorUtils.cpp
--- a/Engine/Source/Editor/UnrealEd/Private/Kismet2/BlueprintEditorUtils.cpp
+++ b/Engine/Source/Editor/UnrealEd/Private/Kismet2/BlueprintEditorUtils.cpp
@@ -836,6 +836,8 @@ void FBlueprintEditorUtils::RefreshGraphNodes(const UEdGraph* Graph)
 
 void FBlueprintEditorUtils::PreloadMembers(UObject* InObject)
 {
+	TRACE_CPUPROFILER_EVENT_SCOPE(FBlueprintEditorUtils::PreloadMembers);
+
 	// Collect a list of all things this element owns
 	TArray<UObject*> BPMemberReferences;
 	FReferenceFinder ComponentCollector(BPMemberReferences, InObject, false, true, true, true);
@@ -845,15 +847,8 @@ void FBlueprintEditorUtils::PreloadMembers(UObject* InObject)
 	for( TArray<UObject*>::TIterator it(BPMemberReferences); it; ++it )
 	{
 		UObject* CurrentObject = *it;
-		if( CurrentObject->HasAnyFlags(RF_NeedLoad) )
-		{
-			FLinkerLoad* Linker = CurrentObject->GetLinker();
-			if (Linker)
-			{
-				Linker->Preload(CurrentObject);
-			}
-			PreloadMembers(CurrentObject);
-		}
+		CurrentObject->ConditionalPreload();
+		PreloadMembers(CurrentObject);
 	}
 }
 
@@ -1043,12 +1038,13 @@ struct FRegenerationHelper
 	static void PreloadAndLinkIfNecessary(UStruct* Struct)
 	{
 		bool bChanged = false;
-		if (Struct->HasAnyFlags(RF_NeedLoad))
+		if (Struct->HasAnyFlags(RF_NeedLoad) || Struct->HasAnyInternalFlags(EInternalObjectFlags::AsyncLoadingPhase1))
 		{
 			if (FLinkerLoad* Linker = Struct->GetLinker())
 			{
+				UE::TUniqueLock LinkerLock(Linker->Mutex);
+				bChanged = Struct->HasAnyFlags(RF_NeedLoad);
 				Linker->Preload(Struct);
-				bChanged = true;
 			}
 		}
 
@@ -1113,13 +1109,8 @@ struct FRegenerationHelper
 		{
 			if (!BP->bHasBeenRegenerated)
 			{
-				if (BP->HasAnyFlags(RF_NeedLoad))
-				{
-					if (FLinkerLoad* Linker = BP->GetLinker())
-					{
-						Linker->Preload(BP);
-					}
-				}
+				BP->ConditionalPreload();
+
 				// at the point of blueprint regeneration (on load), we are guaranteed that blueprint dependencies (like this macro) have
 				// fully formed classes (meaning the blueprint class and all its direct dependencies have been loaded)... however, we do not 
 				// get the guarantee that all of that blueprint's graph dependencies are loaded (hence, why we have to force load 
@@ -1172,20 +1163,12 @@ struct FRegenerationHelper
 		{
 			if (UObject* TypeObject = NewVar.VarType.PinSubCategoryObject.Get())
 			{
-				FLinkerLoad* Linker = TypeObject->GetLinker();
-				if (Linker && TypeObject->HasAnyFlags(RF_NeedLoad))
-				{
-					Linker->Preload(TypeObject);
-				}
+				TypeObject->ConditionalPreload();
 			}
 
 			if (UClass* TypeClass = NewVar.VarType.PinSubCategoryMemberReference.GetMemberParentClass())
 			{
-				FLinkerLoad* Linker = TypeClass->GetLinker();
-				if (Linker && TypeClass->HasAnyFlags(RF_NeedLoad))
-				{
-					Linker->Preload(TypeClass);
-				}
+				TypeClass->ConditionalPreload();
 			}
 		}
 
diff --git a/Engine/Source/Runtime/CoreUObject/Private/Blueprint/BlueprintSupport.cpp b/Engine/Source/Runtime/CoreUObject/Private/Blueprint/BlueprintSupport.cpp
--- a/Engine/Source/Runtime/CoreUObject/Private/Blueprint/BlueprintSupport.cpp
+++ b/Engine/Source/Runtime/CoreUObject/Private/Blueprint/BlueprintSupport.cpp
@@ -542,9 +542,10 @@ struct FPreloadMembersHelper
 			if (!CurrentObject->HasAnyFlags(RF_LoadCompleted))
 			{
 				check(!GEventDrivenLoaderEnabled || !EVENT_DRIVEN_ASYNC_LOAD_ACTIVE_AT_RUNTIME);
-				CurrentObject->SetFlags(RF_NeedLoad);
 				if (FLinkerLoad* Linker = CurrentObject->GetLinker())
 				{
+					UE::TUniqueLock LinkerLock(Linker->Mutex);
+					CurrentObject->SetFlags(RF_NeedLoad);
 					Linker->Preload(CurrentObject);
 					PreloadMembers(CurrentObject);
 				}
@@ -584,9 +585,10 @@ struct FPreloadMembersHelper
 		if (InObject && !InObject->HasAnyFlags(RF_LoadCompleted))
 		{
 			check(!GEventDrivenLoaderEnabled || !EVENT_DRIVEN_ASYNC_LOAD_ACTIVE_AT_RUNTIME);
-			InObject->SetFlags(RF_NeedLoad);
 			if (FLinkerLoad* Linker = InObject->GetLinker())
 			{
+				UE::TUniqueLock LinkerLock(Linker->Mutex);
+				InObject->SetFlags(RF_NeedLoad);
 				Linker->Preload(InObject);
 			}
 		}
diff --git a/Engine/Source/Runtime/CoreUObject/Private/UObject/Class.cpp b/Engine/Source/Runtime/CoreUObject/Private/UObject/Class.cpp
--- a/Engine/Source/Runtime/CoreUObject/Private/UObject/Class.cpp
+++ b/Engine/Source/Runtime/CoreUObject/Private/UObject/Class.cpp
@@ -4980,10 +4980,7 @@ PRAGMA_DISABLE_DEPRECATION_WARNINGS
 						while (FieldIt && (FieldIt->GetOuter() == this))
 						{
 							// If we've had cyclic dependencies between classes here, we might need to preload to ensure that we load the rest of the property chain
-							if (FieldIt->HasAnyFlags(RF_NeedLoad))
-							{
-								ClassLinker->Preload(FieldIt);
-							}
+							FieldIt->ConditionalPreload();
 							FieldIt = FieldIt->Next;
 						}
 					}
diff --git a/Engine/Source/Runtime/CoreUObject/Private/UObject/EnumProperty.cpp b/Engine/Source/Runtime/CoreUObject/Private/UObject/EnumProperty.cpp
--- a/Engine/Source/Runtime/CoreUObject/Private/UObject/EnumProperty.cpp
+++ b/Engine/Source/Runtime/CoreUObject/Private/UObject/EnumProperty.cpp
@@ -184,10 +184,7 @@ bool TryLoadEnumValueByName(FStructuredArchive::FSlot Slot, FArchive& Underlying
 	if (Enum)
 	{
 		// Make sure enum is properly populated
-		if (Enum->HasAnyFlags(RF_NeedLoad))
-		{
-			UnderlyingArchive.Preload(Enum);
-		}
+		Enum->ConditionalPreload();
 
 		if (Enum->HasAnyEnumFlags(EEnumFlags::Flags))
 		{
diff --git a/Engine/Source/Runtime/CoreUObject/Private/UObject/Field.cpp b/Engine/Source/Runtime/CoreUObject/Private/UObject/Field.cpp
--- a/Engine/Source/Runtime/CoreUObject/Private/UObject/Field.cpp
+++ b/Engine/Source/Runtime/CoreUObject/Private/UObject/Field.cpp
@@ -317,11 +317,7 @@ FField::FField(UField* InField)
 	, MetaDataMap(nullptr)
 {
 	check(InField);
-	if (InField->HasAnyFlags(RF_NeedLoad))
-	{
-		// The source UField needs to be loaded, otherwise we'll be copying default property values
-		InField->GetLinker()->Preload(InField);
-	}
+	InField->ConditionalPreload();
 	FlagsPrivate = InField->GetFlags();
 	// Associate this FField with the UField we're constructing from so that next time something tries to convert it,
 	// it can already grab the cached new FField
diff --git a/Engine/Source/Runtime/CoreUObject/Private/UObject/LinkerPlaceholderBase.inl b/Engine/Source/Runtime/CoreUObject/Private/UObject/LinkerPlaceholderBase.inl
--- a/Engine/Source/Runtime/CoreUObject/Private/UObject/LinkerPlaceholderBase.inl
+++ b/Engine/Source/Runtime/CoreUObject/Private/UObject/LinkerPlaceholderBase.inl
@@ -36,6 +36,7 @@ int32 TLinkerImportPlaceholder<PlaceholderType>::ResolveAllPlaceholderReferences
 				// returned up the stack for it to be set yet... here we force 
 				// it to finish preloading (like we do in VerifyImport): 
 				check(!GEventDrivenLoaderEnabled || !EVENT_DRIVEN_ASYNC_LOAD_ACTIVE_AT_RUNTIME);
+				UE::TUniqueLock LinkerLock(ReplacementLinker->Mutex);
 				ReplacementRedirector->SetFlags(RF_NeedLoad);
 				ReplacementLinker->Preload(ReplacementRedirector);
 			}
diff --git a/Engine/Source/Runtime/CoreUObject/Private/UObject/ObjectRef.cpp b/Engine/Source/Runtime/CoreUObject/Private/UObject/ObjectRef.cpp
--- a/Engine/Source/Runtime/CoreUObject/Private/UObject/ObjectRef.cpp
+++ b/Engine/Source/Runtime/CoreUObject/Private/UObject/ObjectRef.cpp
@@ -156,10 +156,7 @@ UClass* FObjectRef::ResolveObjectRefClass(uint32 LoadFlags /*= LOAD_None*/) cons
 
 			if (ClassObject)
 			{
-				if (ClassObject->HasAnyFlags(RF_NeedLoad) && ClassPackage->GetLinker())
-				{
-					ClassPackage->GetLinker()->Preload(ClassObject);
-				}
+				ClassObject->ConditionalPreload();
 				ClassObject->GetDefaultObject(); // build the CDO if it isn't already built
 			}
 		}
diff --git a/Engine/Source/Runtime/CoreUObject/Private/UObject/Package.cpp b/Engine/Source/Runtime/CoreUObject/Private/UObject/Package.cpp
--- a/Engine/Source/Runtime/CoreUObject/Private/UObject/Package.cpp
+++ b/Engine/Source/Runtime/CoreUObject/Private/UObject/Package.cpp
@@ -281,12 +281,9 @@ TArray<UPackage*> UPackage::GetExternalPackages() const
 FMetaData& UPackage::GetMetaData()
 {
 	PRAGMA_DISABLE_DEPRECATION_WARNINGS
-	if (IsValid(DeprecatedMetaData) && DeprecatedMetaData->HasAllFlags(RF_NeedLoad))
+	if (IsValid(DeprecatedMetaData))
 	{
-		FLinkerLoad* MetaDataLinker = DeprecatedMetaData->GetLinker();
-		check(MetaDataLinker);
-
-		MetaDataLinker->Preload(DeprecatedMetaData);
+		DeprecatedMetaData->ConditionalPreload();
 	}
 	PRAGMA_ENABLE_DEPRECATION_WARNINGS
 
diff --git a/Engine/Source/Runtime/Engine/Classes/Sound/SoundWaveLoadingBehavior.cpp b/Engine/Source/Runtime/Engine/Classes/Sound/SoundWaveLoadingBehavior.cpp
--- a/Engine/Source/Runtime/Engine/Classes/Sound/SoundWaveLoadingBehavior.cpp
+++ b/Engine/Source/Runtime/Engine/Classes/Sound/SoundWaveLoadingBehavior.cpp
@@ -107,14 +107,8 @@ private:
 		while (InClass->ParentClass && InClass->Properties.LoadingBehavior == ESoundWaveLoadingBehavior::Inherited)
 		{
 			InClass = InClass->ParentClass;
-			if (InClass->HasAnyFlags(RF_NeedLoad))
-			{
-				InClass->GetLinker()->Preload(InClass);
-			}
-			if (InClass->HasAnyFlags(RF_NeedPostLoad))
-			{
-				InClass->ConditionalPostLoad();
-			}
+			InClass->ConditionalPreload();
+			InClass->ConditionalPostLoad();
 			Behavior.BaseData = FClassData(InClass);
 			Behavior.ClassHierarchy.Add(FName(InClass->GetPackage()->GetFName()));
 		}
diff --git a/Engine/Source/Runtime/Engine/Private/Animation/AnimCompressionTypes.cpp b/Engine/Source/Runtime/Engine/Private/Animation/AnimCompressionTypes.cpp
--- a/Engine/Source/Runtime/Engine/Private/Animation/AnimCompressionTypes.cpp
+++ b/Engine/Source/Runtime/Engine/Private/Animation/AnimCompressionTypes.cpp
@@ -1218,10 +1218,7 @@ template void FUECompressedAnimData::ByteSwapData(TArrayView<uint8> CompressedDa
 void ValidateUObjectLoaded(UObject* Obj, UObject* Source)
 {
 #if WITH_EDITOR
-	if (FLinkerLoad* ObjLinker = Obj->GetLinker())
-	{
-		ObjLinker->Preload(Obj);
-	}
+	Obj->ConditionalPreload();
 #endif
 	checkf(!Obj->HasAnyFlags(RF_NeedLoad), TEXT("Failed to load %s in %s"), *Obj->GetFullName(), *Source->GetFullName()); // in non editor should have been preloaded by GetPreloadDependencies
 }
diff --git a/Engine/Source/Runtime/Engine/Private/Animation/AnimSequenceBase.cpp b/Engine/Source/Runtime/Engine/Private/Animation/AnimSequenceBase.cpp
--- a/Engine/Source/Runtime/Engine/Private/Animation/AnimSequenceBase.cpp
+++ b/Engine/Source/Runtime/Engine/Private/Animation/AnimSequenceBase.cpp
@@ -111,10 +111,7 @@ void UAnimSequenceBase::PostLoad()
 		{
 			if (USkeleton* MySkeleton = GetSkeleton())
 			{
-				if (FLinkerLoad* SkeletonLinker = MySkeleton->GetLinker())
-				{
-					SkeletonLinker->Preload(MySkeleton);
-				}
+				MySkeleton->ConditionalPreload();
 				MySkeleton->ConditionalPostLoad();
 			}
 		};
@@ -162,11 +159,8 @@ void UAnimSequenceBase::PostLoad()
 				// In case a data model has already been created populate the new one with its data
 				if (DataModel != nullptr)
 				{
-					if (FLinkerLoad* DataModelLinker = DataModel->GetLinker())
-					{
-						DataModelLinker->Preload(DataModel);
-					}					
-					DataModel->PostLoad();
+					DataModel->ConditionalPreload();
+					DataModel->ConditionalPostLoad();
 					DataModel->ConditionalPostLoadSubobjects();
 					PopulateWithExistingModel(DataModel.Get());
 
diff --git a/Engine/Source/Runtime/Engine/Private/Animation/AnimStreamable.cpp b/Engine/Source/Runtime/Engine/Private/Animation/AnimStreamable.cpp
--- a/Engine/Source/Runtime/Engine/Private/Animation/AnimStreamable.cpp
+++ b/Engine/Source/Runtime/Engine/Private/Animation/AnimStreamable.cpp
@@ -411,10 +411,7 @@ void UAnimStreamable::PostLoad()
 #if WITH_EDITOR
 	if (UAnimSequence* NonConstSeq = const_cast<UAnimSequence*>(ToRawPtr(SourceSequence)))
 	{
-		if (FLinkerLoad* Linker = NonConstSeq->GetLinker())
-		{
-			Linker->Preload(NonConstSeq);
-		}
+		NonConstSeq->ConditionalPreload();
 		NonConstSeq->ConditionalPostLoad();
 	}
 	
diff --git a/Engine/Source/Runtime/Engine/Private/Animation/AnimationUtils.cpp b/Engine/Source/Runtime/Engine/Private/Animation/AnimationUtils.cpp
--- a/Engine/Source/Runtime/Engine/Private/Animation/AnimationUtils.cpp
+++ b/Engine/Source/Runtime/Engine/Private/Animation/AnimationUtils.cpp
@@ -858,8 +858,10 @@ static void EnsureDependenciesAreLoaded(UObject* Object)
 	if (Object->HasAnyFlags(RF_NeedLoad))
 	{
 		UE_LOG(LogAnimationCompression, Warning, TEXT("EnsureDependenciesAreLoaded: [Preload] %s"), *Object->GetPathName());
-		Object->GetLinker()->Preload(Object);
 	}
+	
+	// Always call conditional preload as the test is a little bit more involved than the one we did for the warning above.
+	Object->ConditionalPreload();
 
 	if (Object->HasAnyFlags(RF_NeedPostLoad))
 	{
@@ -875,9 +877,11 @@ static void EnsureDependenciesAreLoaded(UObject* Object)
 		if (Dependency->HasAnyFlags(RF_NeedLoad))
 		{
 			UE_LOG(LogAnimationCompression, Warning, TEXT("EnsureDependenciesAreLoaded: [Preload Dependency] %s"), *Dependency->GetPathName());
-			Dependency->GetLinker()->Preload(Dependency);
 		}
 
+		// Always call conditional preload as the test is a little bit more involved than the one we did for the warning above.
+		Dependency->ConditionalPreload();
+
 		if (Dependency->HasAnyFlags(RF_NeedPostLoad))
 		{
 			UE_LOG(LogAnimationCompression, Warning, TEXT("EnsureDependenciesAreLoaded: [ConditionalPostLoad Dependency] %s"), *Dependency->GetPathName());
diff --git a/Engine/Source/Runtime/Engine/Private/Animation/PoseAsset.cpp b/Engine/Source/Runtime/Engine/Private/Animation/PoseAsset.cpp
--- a/Engine/Source/Runtime/Engine/Private/Animation/PoseAsset.cpp
+++ b/Engine/Source/Runtime/Engine/Private/Animation/PoseAsset.cpp
@@ -941,13 +941,7 @@ void UPoseAsset::PostLoad()
 		if (SourceAnimation)
 		{
 			// Fully load the source animation to ensure its RawDataGUID is populated
-			if(SourceAnimation->HasAnyFlags(RF_NeedLoad))
-			{
-				if (FLinkerLoad* Linker = SourceAnimation->GetLinker())
-				{
-					Linker->Preload(SourceAnimation);
-				}
-			}
+			SourceAnimation->ConditionalPreload();
 			SourceAnimation->ConditionalPostLoad();
 
 			const FGuid DataModelGuid = GetSourceAnimationGuid();
diff --git a/Engine/Source/Runtime/Engine/Private/AnimationAsset.cpp b/Engine/Source/Runtime/Engine/Private/AnimationAsset.cpp
--- a/Engine/Source/Runtime/Engine/Private/AnimationAsset.cpp
+++ b/Engine/Source/Runtime/Engine/Private/AnimationAsset.cpp
@@ -315,10 +315,7 @@ void UAnimationAsset::PostLoad()
 	// skeleton is ready
 	if (Skeleton)
 	{
-		if (FLinkerLoad* SkeletonLinker = Skeleton->GetLinker())
-		{
-			SkeletonLinker->Preload(Skeleton);
-		}
+		Skeleton->ConditionalPreload();
 		Skeleton->ConditionalPostLoad();
 	}
 
@@ -326,10 +323,7 @@ void UAnimationAsset::PostLoad()
 	// Load Parent Asset, to make sure anything accessing from PostLoad has valid data to access
 	if (ParentAsset)
 	{
-		if (FLinkerLoad* ParentAssetLinker = ParentAsset->GetLinker())
-		{
-			ParentAssetLinker->Preload(ParentAsset);
-		}
+		ParentAsset->ConditionalPreload();
 		ParentAsset->ConditionalPostLoad();
 	}
 #endif
@@ -485,10 +479,7 @@ bool UAnimationAsset::ReplaceSkeleton(USkeleton* NewSkeleton, bool bConvertSpace
 		for (UAnimationAsset* AnimAsset : AnimAssetsToReplace)
 		{
 			//Make sure animation has finished loading before we start messing with it
-			if (FLinkerLoad* AnimLinker = AnimAsset->GetLinker())
-			{
-				AnimLinker->Preload(AnimAsset);
-			}
+			AnimAsset->ConditionalPreload();
 			AnimAsset->ConditionalPostLoad();
 
 			if (AnimAsset->GetSkeleton() != GetSkeleton())
diff --git a/Engine/Source/Runtime/Engine/Private/Blueprint.cpp b/Engine/Source/Runtime/Engine/Private/Blueprint.cpp
--- a/Engine/Source/Runtime/Engine/Private/Blueprint.cpp
+++ b/Engine/Source/Runtime/Engine/Private/Blueprint.cpp
@@ -453,9 +453,9 @@ void UBlueprint::Serialize(FArchive& Ar)
 		for (UClass* ClassIt = ParentClass; ClassIt && !ClassIt->HasAnyClassFlags(CLASS_Native); ClassIt = ClassIt->GetSuperClass())
 		{
 			// In some cases, a non-native parent class may not have an associated Blueprint asset - we consider that to be ok here since we're just preloading.
-			if (ClassIt->ClassGeneratedBy && ClassIt->ClassGeneratedBy->HasAnyFlags(RF_NeedLoad))
+			if (ClassIt->ClassGeneratedBy)
 			{
-				ClassIt->ClassGeneratedBy->GetLinker()->Preload(ClassIt->ClassGeneratedBy);
+				ClassIt->ClassGeneratedBy->ConditionalPreload();
 			}
 		}
 	}
@@ -1339,6 +1339,7 @@ bool UBlueprint::ForceLoad(UObject* Obj)
 	if (Linker && !Obj->HasAnyFlags(RF_LoadCompleted))
 	{
 		check(!GEventDrivenLoaderEnabled);
+		UE::TUniqueLock LinkerLock(Linker->Mutex);
 		Obj->SetFlags(RF_NeedLoad);
 		Linker->Preload(Obj);
 		return true;
diff --git a/Engine/Source/Runtime/Engine/Private/CompositeCurveTable.cpp b/Engine/Source/Runtime/Engine/Private/CompositeCurveTable.cpp
--- a/Engine/Source/Runtime/Engine/Private/CompositeCurveTable.cpp
+++ b/Engine/Source/Runtime/Engine/Private/CompositeCurveTable.cpp
@@ -60,13 +60,9 @@ void UCompositeCurveTable::Serialize(FArchive& Ar)
 	{
 		for (UCurveTable* ParentTable : ParentTables)
 		{
-			if (ParentTable && ParentTable->HasAnyFlags(RF_NeedLoad))
+			if (ParentTable)
 			{
-				FLinkerLoad* ParentTableLinker = ParentTable->GetLinker();
-				if (ParentTableLinker)
-				{
-					ParentTableLinker->Preload(ParentTable);
-				}
+				ParentTable->ConditionalPreload();
 			}
 		}
 
diff --git a/Engine/Source/Runtime/Engine/Private/CompositeDataTable.cpp b/Engine/Source/Runtime/Engine/Private/CompositeDataTable.cpp
--- a/Engine/Source/Runtime/Engine/Private/CompositeDataTable.cpp
+++ b/Engine/Source/Runtime/Engine/Private/CompositeDataTable.cpp
@@ -223,14 +223,7 @@ void UCompositeDataTable::Serialize(FArchive& Ar)
 	{
 		for (UDataTable* ParentTable : ParentTables)
 		{
-			if (ParentTable && ParentTable->HasAnyFlags(RF_NeedLoad))
-			{
-				FLinkerLoad* ParentTableLinker = ParentTable->GetLinker();
-				if (ParentTableLinker)
-				{
-					ParentTableLinker->Preload(ParentTable);
-				}
-			}
+			ParentTable->ConditionalPreload();
 		}
 
 		OnParentTablesUpdated();
diff --git a/Engine/Source/Runtime/Engine/Private/DataTable.cpp b/Engine/Source/Runtime/Engine/Private/DataTable.cpp
--- a/Engine/Source/Runtime/Engine/Private/DataTable.cpp
+++ b/Engine/Source/Runtime/Engine/Private/DataTable.cpp
@@ -280,13 +280,9 @@ void UDataTable::Serialize(FStructuredArchiveRecord Record)
 
 	Super::Serialize(Record); // When loading, this should load our RowStruct!	
 
-	if (RowStruct && RowStruct->HasAnyFlags(RF_NeedLoad))
+	if (RowStruct)
 	{
-		auto RowStructLinker = RowStruct->GetLinker();
-		if (RowStructLinker)
-		{
-			RowStructLinker->Preload(RowStruct);
-		}
+		RowStruct->ConditionalPreload();
 	}
 
 	if(BaseArchive.IsLoading())
diff --git a/Engine/Source/Runtime/Engine/Private/InheritableComponentHandler.cpp b/Engine/Source/Runtime/Engine/Private/InheritableComponentHandler.cpp
--- a/Engine/Source/Runtime/Engine/Private/InheritableComponentHandler.cpp
+++ b/Engine/Source/Runtime/Engine/Private/InheritableComponentHandler.cpp
@@ -476,27 +476,20 @@ FComponentKey UInheritableComponentHandler::FindKey(UActorComponent* ComponentTe
 
 void UInheritableComponentHandler::PreloadAllTemplates()
 {
+	TRACE_CPUPROFILER_EVENT_SCOPE(UInheritableComponentHandler::PreloadAllTemplates);
+
 	for (const FComponentOverrideRecord& Record : Records)
 	{
-		if (Record.ComponentTemplate && Record.ComponentTemplate->HasAllFlags(RF_NeedLoad))
+		if (Record.ComponentTemplate)
 		{
-			if (FLinkerLoad* Linker = Record.ComponentTemplate->GetLinker())
-			{
-				Linker->Preload(Record.ComponentTemplate);
+			Record.ComponentTemplate->ConditionalPreload();
 
-				TArray<UObject*> ComponentTemplateSubobjects;
-				// Can't use ForEachObjectWithOuter here as Preloading may modify UObject hash tables (it will most likely create new objects)
-				GetObjectsWithOuter(Record.ComponentTemplate, ComponentTemplateSubobjects);
-				for (UObject* SubObj : ComponentTemplateSubobjects)
-				{
-					if (SubObj->HasAllFlags(RF_NeedLoad))
-					{
-						if (FLinkerLoad* SubObjLinker = SubObj->GetLinker())
-						{
-							SubObjLinker->Preload(SubObj);
-						}
-					}
-				}
+			TArray<UObject*> ComponentTemplateSubobjects;
+			// Can't use ForEachObjectWithOuter here as Preloading may modify UObject hash tables (it will most likely create new objects)
+			GetObjectsWithOuter(Record.ComponentTemplate, ComponentTemplateSubobjects);
+			for (UObject* SubObj : ComponentTemplateSubobjects)
+			{
+				SubObj->ConditionalPreload();
 			}
 		}
 	}
@@ -504,13 +497,8 @@ void UInheritableComponentHandler::PreloadAllTemplates()
 
 void UInheritableComponentHandler::PreloadAll()
 {
-	if (HasAllFlags(RF_NeedLoad))
-	{
-		if (FLinkerLoad* Linker = GetLinker())
-		{
-			Linker->Preload(this);
-		}
-	}
+	ConditionalPreload();
+	
 	PreloadAllTemplates();
 	// This will get component names up to date - since these component 
 	// templates are used by GetArchetypeFromRequiredInfo we want to have real
diff --git a/Engine/Source/Runtime/Engine/Private/Materials/MaterialExpressions.cpp b/Engine/Source/Runtime/Engine/Private/Materials/MaterialExpressions.cpp
--- a/Engine/Source/Runtime/Engine/Private/Materials/MaterialExpressions.cpp
+++ b/Engine/Source/Runtime/Engine/Private/Materials/MaterialExpressions.cpp
@@ -17205,10 +17205,7 @@ void UMaterialExpressionMaterialFunctionCall::PostLoad()
 {
 	if (MaterialFunction)
 	{
-		if (MaterialFunction->HasAnyFlags(RF_NeedLoad))
-		{
-			MaterialFunction->GetLinker()->Preload(MaterialFunction);
-		}
+		MaterialFunction->ConditionalPreload();
 		MaterialFunction->ConditionalPostLoad();
 	}
 
diff --git a/Engine/Source/Runtime/Engine/Private/SCS_Node.cpp b/Engine/Source/Runtime/Engine/Private/SCS_Node.cpp
--- a/Engine/Source/Runtime/Engine/Private/SCS_Node.cpp
+++ b/Engine/Source/Runtime/Engine/Private/SCS_Node.cpp
@@ -341,31 +341,20 @@ bool USCS_Node::IsChildOf(USCS_Node* TestParent)
 
 void USCS_Node::PreloadChain()
 {
-	if( HasAnyFlags(RF_NeedLoad) )
-	{
-		GetLinker()->Preload(this);
-	}
+	TRACE_CPUPROFILER_EVENT_SCOPE(USCS_Node::PreloadChain);
+
+	ConditionalPreload();
 
-	if (ComponentTemplate && ComponentTemplate->HasAnyFlags(RF_NeedLoad))
+	if (ComponentTemplate)
 	{
-		if (ensureMsgf(ComponentTemplate->GetLinker(), TEXT("Failed to find linker for %s, likely a circular dependency"), *ComponentTemplate->GetPathName()))
-		{
-			ComponentTemplate->GetLinker()->Preload(ComponentTemplate);
+		ComponentTemplate->ConditionalPreload();
 
-			TArray<UObject*> Children;
-			GetObjectsWithOuter(ComponentTemplate, Children, true, RF_LoadCompleted);
-			for (UObject* Obj : Children)
-			{
-				if (!Obj->HasAnyFlags(RF_WasLoaded))
-				{
-					continue;
-				}
+		TArray<UObject*> Children;
+		GetObjectsWithOuter(ComponentTemplate, Children, true);
 
-				if (FLinkerLoad* Linker = Obj->GetLinker())
-				{
-					Linker->Preload(Obj);
-				}
-			}
+		for (UObject* Obj : Children)
+		{
+			Obj->ConditionalPreload();
 		}
 	}
 
diff --git a/Engine/Source/Runtime/Engine/Private/SimpleConstructionScript.cpp b/Engine/Source/Runtime/Engine/Private/SimpleConstructionScript.cpp
--- a/Engine/Source/Runtime/Engine/Private/SimpleConstructionScript.cpp
+++ b/Engine/Source/Runtime/Engine/Private/SimpleConstructionScript.cpp
@@ -55,10 +55,7 @@ void USimpleConstructionScript::Serialize(FArchive& Ar)
 			if(RootNode_DEPRECATED != NULL)
 			{
 				// Ensure it's been loaded so that its properties are valid
-				if(RootNode_DEPRECATED->HasAnyFlags(RF_NeedLoad))
-				{
-					RootNode_DEPRECATED->GetLinker()->Preload(RootNode_DEPRECATED);
-				}
+				RootNode_DEPRECATED->ConditionalPreload();
 
 				// If the root node was not native
 				if(!RootNode_DEPRECATED->bIsNative_DEPRECATED)
@@ -74,10 +71,7 @@ void USimpleConstructionScript::Serialize(FArchive& Ar)
 						if(Node != NULL)
 						{
 							// Ensure it's been loaded (may not have been yet if we're preloading the SCS)
-							if(Node->HasAnyFlags(RF_NeedLoad))
-							{
-								Node->GetLinker()->Preload(Node);
-							}
+							Node->ConditionalPreload();
 
 							// We only care about non-native child nodes (non-native nodes could only be attached to the root node in the previous version, so we don't need to examine native child nodes)
 							if(!Node->bIsNative_DEPRECATED)
@@ -103,11 +97,7 @@ void USimpleConstructionScript::Serialize(FArchive& Ar)
 				if(Node != NULL)
 				{
 					// Ensure it's been loaded (may not have been yet if we're preloading the SCS)
-					if(Node->HasAnyFlags(RF_NeedLoad))
-					{
-						Node->GetLinker()->Preload(Node);
-					}
-
+					Node->ConditionalPreload();
 					if(!Node->bIsNative_DEPRECATED)
 					{
 						RootNodes.Add(Node);
@@ -124,7 +114,7 @@ void USimpleConstructionScript::Serialize(FArchive& Ar)
 
 void USimpleConstructionScript::PreloadChain()
 {
-	GetLinker()->Preload(this);
+	ConditionalPreload();
 
 	for (USCS_Node* Node : RootNodes)
 	{
diff --git a/Engine/Source/Runtime/GameplayTags/Private/GameplayTagsManager.cpp b/Engine/Source/Runtime/GameplayTags/Private/GameplayTagsManager.cpp
--- a/Engine/Source/Runtime/GameplayTags/Private/GameplayTagsManager.cpp
+++ b/Engine/Source/Runtime/GameplayTags/Private/GameplayTagsManager.cpp
@@ -313,11 +313,7 @@ void UGameplayTagsManager::LoadGameplayTagTables(bool bAllowAsyncLoad)
 		// will not be initialized correctly.
 		if (TagTable)
 		{
-			FLinkerLoad* TagLinker = TagTable->GetLinker();
-			if (TagLinker)
-			{
-				TagTable->GetLinker()->Preload(TagTable);
-			}
+			TagTable->ConditionalPreload();
 		}
 		GameplayTagTables.Add(TagTable);
 	}
