Fixed selection when ShowFlag.HitProxies is manually set to 0

diff --git a/Engine/Source/Editor/UnrealEd/Private/EditorViewportClient.cpp b/Engine/Source/Editor/UnrealEd/Private/EditorViewportClient.cpp
--- a/Engine/Source/Editor/UnrealEd/Private/EditorViewportClient.cpp
+++ b/Engine/Source/Editor/UnrealEd/Private/EditorViewportClient.cpp
@@ -4372,12 +4372,6 @@ void FEditorViewportClient::Draw(FViewport* InViewport, FCanvas* Canvas)
 
 	ViewFamily.EngineShowFlags = UseEngineShowFlags;
 
-	// Set the view hit proxies flag as early as possible so that view extensions can check whether they render 
-	if (Canvas->IsHitTesting())
-	{
-		ViewFamily.EngineShowFlags.SetHitProxies(1);
-	}
-
 	ViewFamily.bIsHDR = Viewport->IsHDRViewport();
 	
 	// The view is in focus if it is currently in editing
@@ -4430,6 +4424,12 @@ void FEditorViewportClient::Draw(FViewport* InViewport, FCanvas* Canvas)
 		ViewFamily.EngineShowFlags.SetTemporalAA(false);
 	}
 
+	// Set the view hit proxies flag as early as possible (but after EngineShowFlagOverride) so that view extensions can check whether they render in the hit proxy pass in SetupView
+	if (Canvas->IsHitTesting())
+	{
+		ViewFamily.EngineShowFlags.SetHitProxies(1);
+	}
+
 	// Setup the screen percentage and upscaling method for the view family.
 	{
 		checkf(ViewFamily.GetScreenPercentageInterface() == nullptr,
