[Graph Editor] Fix a crash in graph editor deprecated delegates when the graph editor events struct falls out of scope

diff --git a/Engine/Source/Editor/UnrealEd/Private/GraphEditor.cpp b/Engine/Source/Editor/UnrealEd/Private/GraphEditor.cpp
--- a/Engine/Source/Editor/UnrealEd/Private/GraphEditor.cpp
+++ b/Engine/Source/Editor/UnrealEd/Private/GraphEditor.cpp
@@ -64,30 +64,30 @@ void SGraphEditor::UpgradeDeprecatedDelegates(FGraphEditorEvents& EventsToUpdate
 	PRAGMA_DISABLE_DEPRECATION_WARNINGS
 	if (EventsToUpdate.OnDropActor.IsBound() && !EventsToUpdate.OnDropActors.IsBound())
 	{
-		EventsToUpdate.OnDropActors = FOnDropActors::CreateLambda([&EventsToUpdate](const TArray< TWeakObjectPtr<AActor> >& Actors, UEdGraph* InGraph, const FVector2f& InDropLocation)
+		EventsToUpdate.OnDropActors = FOnDropActors::CreateLambda([OnDropActorCopy = EventsToUpdate.OnDropActor](const TArray< TWeakObjectPtr<AActor> >& Actors, UEdGraph* InGraph, const FVector2f& InDropLocation)
 		{
-			EventsToUpdate.OnDropActor.Execute(Actors, InGraph, FVector2D(InDropLocation));
+			OnDropActorCopy.Execute(Actors, InGraph, FVector2D(InDropLocation));
 		});
 	}
 	if (EventsToUpdate.OnDropStreamingLevel.IsBound() && !EventsToUpdate.OnDropStreamingLevels.IsBound())
 	{
-		EventsToUpdate.OnDropStreamingLevels = FOnDropStreamingLevels::CreateLambda([&EventsToUpdate](const TArray< TWeakObjectPtr<class ULevelStreaming> >& Levels, UEdGraph* InGraph, const FVector2f& InDropLocation)
+		EventsToUpdate.OnDropStreamingLevels = FOnDropStreamingLevels::CreateLambda([OnDropStreamingLevelCopy = EventsToUpdate.OnDropStreamingLevel](const TArray< TWeakObjectPtr<class ULevelStreaming> >& Levels, UEdGraph* InGraph, const FVector2f& InDropLocation)
 		{
-			EventsToUpdate.OnDropStreamingLevel.Execute(Levels, InGraph, FVector2D(InDropLocation));
+			OnDropStreamingLevelCopy.Execute(Levels, InGraph, FVector2D(InDropLocation));
 		});
 	}
 	if (EventsToUpdate.OnCreateActionMenu.IsBound() && !EventsToUpdate.OnCreateActionMenuAtLocation.IsBound())
 	{
-		EventsToUpdate.OnCreateActionMenuAtLocation = FOnCreateActionMenuAtLocation::CreateLambda([&EventsToUpdate](UEdGraph* InGraph, const FVector2f& InLocation, const TArray<UEdGraphPin*>& InPins, bool bAutoExpand, FActionMenuClosed InMenuClosed)
+		EventsToUpdate.OnCreateActionMenuAtLocation = FOnCreateActionMenuAtLocation::CreateLambda([OnCreateActionMenuCopy = EventsToUpdate.OnCreateActionMenu](UEdGraph* InGraph, const FVector2f& InLocation, const TArray<UEdGraphPin*>& InPins, bool bAutoExpand, FActionMenuClosed InMenuClosed)
 		{
-			return EventsToUpdate.OnCreateActionMenu.Execute(InGraph, FVector2D(InLocation), InPins, bAutoExpand, InMenuClosed);
+			return OnCreateActionMenuCopy.Execute(InGraph, FVector2D(InLocation), InPins, bAutoExpand, InMenuClosed);
 		});
 	}
 	if (EventsToUpdate.OnSpawnNodeByShortcut.IsBound() && !EventsToUpdate.OnSpawnNodeByShortcutAtLocation.IsBound())
 	{
-		EventsToUpdate.OnSpawnNodeByShortcutAtLocation = FOnSpawnNodeByShortcutAtLocation::CreateLambda([&EventsToUpdate](FInputChord InInputChord, const FVector2f& InLocation)
+		EventsToUpdate.OnSpawnNodeByShortcutAtLocation = FOnSpawnNodeByShortcutAtLocation::CreateLambda([OnSpawnNodeByShortcutCopy = EventsToUpdate.OnSpawnNodeByShortcut](FInputChord InInputChord, const FVector2f& InLocation)
 		{
-			return EventsToUpdate.OnSpawnNodeByShortcut.Execute(InInputChord, FVector2D(InLocation));
+			return OnSpawnNodeByShortcutCopy.Execute(InInputChord, FVector2D(InLocation));
 		});
 	}
 	PRAGMA_ENABLE_DEPRECATION_WARNINGS
