Restore the missing View Mode Options menu in the new viewport toolbar.

[RN] Minor Editor

diff --git a/Engine/Source/Editor/UnrealEd/Private/ViewportToolbar/UnrealEdViewportToolbar.cpp b/Engine/Source/Editor/UnrealEd/Private/ViewportToolbar/UnrealEdViewportToolbar.cpp
--- a/Engine/Source/Editor/UnrealEd/Private/ViewportToolbar/UnrealEdViewportToolbar.cpp
+++ b/Engine/Source/Editor/UnrealEd/Private/ViewportToolbar/UnrealEdViewportToolbar.cpp
@@ -1907,6 +1907,66 @@ void PopulateViewModesMenu(UToolMenu* InMenu)
 			);
 		}
 	}
+	
+	// View Mode options
+	{
+		FToolMenuSection& Section = InMenu->AddSection("ViewModeOptions", LOCTEXT("ViewModeOptionsHeader", "View Mode Options"), FToolMenuInsert("ViewMode", EToolMenuInsertType::After));
+		Section.Visibility = TAttribute<bool>::CreateLambda([WeakViewport = EditorViewport.ToWeakPtr()]
+		{
+			if (TSharedPtr<SEditorViewport> Viewport = WeakViewport.Pin())
+			{
+				if (TSharedPtr<FEditorViewportClient> Client = Viewport->GetViewportClient())
+				{
+					EViewModeIndex ViewMode = Client->GetViewMode();
+					return ViewMode == VMI_MeshUVDensityAccuracy || ViewMode == VMI_MaterialTextureScaleAccuracy || ViewMode == VMI_RequiredTextureResolution;	
+				}
+			}
+			return false;
+		});
+		
+		TAttribute<FText> ViewOptionsLabel = TAttribute<FText>::CreateLambda([WeakViewport = EditorViewport.ToWeakPtr()]
+		{
+			if (TSharedPtr<SEditorViewport> Viewport = WeakViewport.Pin())
+			{
+				if (TSharedPtr<FEditorViewportClient> Client = Viewport->GetViewportClient())
+				{
+					return ::GetViewModeOptionsMenuLabel(Client->GetViewMode());
+				}
+			}
+			return LOCTEXT("ViewModeOptionsLabel", "View Mode Options");
+		});
+		
+		FToolMenuEntry& OptionsEntry = Section.AddSubMenu(
+			"ViewModeOptions",
+			ViewOptionsLabel,
+			LOCTEXT("ViewModeOptionsToolTip", "Options for the current view mode."),
+			FNewMenuDelegate::CreateLambda([WeakViewport = EditorViewport.ToWeakPtr()](FMenuBuilder& MenuBuilder)
+			{
+				TSharedRef<SWidget> OptionsMenuContent = SNullWidget::NullWidget;
+			
+				if (TSharedPtr<SEditorViewport> Viewport = WeakViewport.Pin())
+				{
+					if (TSharedPtr<FEditorViewportClient> Client = Viewport->GetViewportClient())
+					{
+						const UWorld* World = Client->GetWorld();
+						
+						// TODO: This is a very awkward way of exposing this menu.
+						// The function should be modified to allow defining the menu builder.
+						// It was done this way to maintain ABI compatability and should be fixed at the earliest possible convenience.
+						OptionsMenuContent = ::BuildViewModeOptionsMenu(
+							Viewport->GetCommandList(),
+							Client->GetViewMode(),
+							World ? World->GetFeatureLevel() : GMaxRHIFeatureLevel,
+							Client->GetViewModeParamNameMap()
+						);
+					}	
+				}
+				MenuBuilder.AddWidget(OptionsMenuContent, FText::GetEmpty(), true, false);
+			})
+		);
+					
+		OptionsEntry.SetShowInToolbarTopLevel(true); 
+	}
 }
 
 UUnrealEdViewportToolbarContext* CreateViewportToolbarDefaultContext(const TWeakPtr<SEditorViewport>& InViewport)
