Guard against GetAssets failing due to a bad filter, and fallback to GetAllAssets

#rnx

diff --git a/Engine/Source/Editor/UnrealEd/Private/Commandlets/GatherTextFromAssetsCommandlet.cpp b/Engine/Source/Editor/UnrealEd/Private/Commandlets/GatherTextFromAssetsCommandlet.cpp
--- a/Engine/Source/Editor/UnrealEd/Private/Commandlets/GatherTextFromAssetsCommandlet.cpp
+++ b/Engine/Source/Editor/UnrealEd/Private/Commandlets/GatherTextFromAssetsCommandlet.cpp
@@ -953,12 +953,16 @@ void UGatherTextFromAssetsCommandlet::ApplyFirstPassFilter(const FARFilter& InFi
 		// Figure out a good way to still convey the information in a log or clog.
 		const double GetAllAssetsStartTime = FPlatformTime::Seconds();
 		IAssetRegistry::GetChecked().GetAllAssets(InOutAssetDataArray);
-		UE_LOG(LogGatherTextFromAssetsCommandlet, Display, TEXT("Loading all assets from asset registry took %.2f seconds."), FPlatformTime::Seconds() - GetAllAssetsStartTime);
+		UE_LOG(LogGatherTextFromAssetsCommandlet, Display, TEXT("Getting all assets from asset registry took %.2f seconds."), FPlatformTime::Seconds() - GetAllAssetsStartTime);
 	}
 	else
 	{
 		const double GetAllAssetsWithFirstPassFilterStartTime = FPlatformTime::Seconds();
-		IAssetRegistry::GetChecked().GetAssets(InFilter, InOutAssetDataArray);
+		if (!IAssetRegistry::GetChecked().GetAssets(InFilter, InOutAssetDataArray))
+		{
+			UE_LOG(LogGatherTextFromAssetsCommandlet, Warning, TEXT("The first pass asset registry filter was invalid. Falling back to getting all assets from the asset registry..."));
+			IAssetRegistry::GetChecked().GetAllAssets(InOutAssetDataArray);
+		}
 		UE_LOG(LogGatherTextFromAssetsCommandlet, Display, TEXT("Getting all assets with first pass filter from asset registry took %.2f seconds."), FPlatformTime::Seconds() - GetAllAssetsWithFirstPassFilterStartTime);
 	}
 }
