[Editor Viewports] Fix a crash on startup where the hit proxy can try to generate before the viewport client has drawn

diff --git a/Engine/Source/Editor/UnrealEd/Private/EditorViewportClient.cpp b/Engine/Source/Editor/UnrealEd/Private/EditorViewportClient.cpp
--- a/Engine/Source/Editor/UnrealEd/Private/EditorViewportClient.cpp
+++ b/Engine/Source/Editor/UnrealEd/Private/EditorViewportClient.cpp
@@ -5697,8 +5697,8 @@ void FEditorViewportClient::CheckHoveredHitProxy( HHitProxy* HoveredHitProxy )
 
 void FEditorViewportClient::ConditionalCheckHoveredHitProxy()
 {
-	// If it has been decided that there is more important things to do than check hit proxies, then don't check them.
-	if( !bShouldCheckHitProxy || bWidgetAxisControlledByDrag == true )
+	// If it has been decided that there is more important things to do than check hit proxies, or our mouse location is not valid, then don't check them.
+	if( !bShouldCheckHitProxy || bWidgetAxisControlledByDrag == true || CurrentMousePos == FIntPoint(-1, -1) )
 	{
 		return;
 	}
