Add new feature in the asset registry tracker : optionally create package dependencies for all the packages returned by a tracked query.

- FAddPackageDependenciesScope : new RAII class enabling the creation of package dependencies.
- FCookAssetRegistryAccessTracker : Each record now has a flag to tell whether or not the packages from the query should be added as dependencies.
- FCookPackageArtifact : When creating the AssetRegistryQuery dependencies, create also dependencies for the packages if necessary.
- ComputeBoundsFromLevel : create a FAddPackageDependenciesScope object so dependencies for the packages queried in GetLevelInstancePivotOffsetFromPackage are created.

#rnx
#incrementalcook

diff --git a/Engine/Source/Editor/UnrealEd/Private/Cooker/CookAssetRegistryAccessTracker.cpp b/Engine/Source/Editor/UnrealEd/Private/Cooker/CookAssetRegistryAccessTracker.cpp
--- a/Engine/Source/Editor/UnrealEd/Private/Cooker/CookAssetRegistryAccessTracker.cpp
+++ b/Engine/Source/Editor/UnrealEd/Private/Cooker/CookAssetRegistryAccessTracker.cpp
@@ -5,6 +5,7 @@
 #include "Algo/Unique.h"
 #include "AssetRegistry/AssetData.h"
 #include "AssetRegistry/AssetRegistryModule.h"
+#include "Cooker/CookAssetRegistryAccessTracker_AddPackageDep.h"
 #include "Interfaces/ITargetPlatform.h"
 #include "Misc/PackageAccessTracking.h"
 
@@ -15,12 +16,13 @@ namespace UE::CookAssetRegistryAccessTracker
 		uint32 Hash = 0;
 		Hash = HashCombineFast(Hash, GetTypeHash(Data.Platform));
 		Hash = HashCombineFast(Hash, UE::AssetRegistry::GetTypeHash(Data.Filter));
+		Hash = HashCombineFast(Hash, ::GetTypeHash(Data.bAddPackageDependencies));
 		return Hash;
 	}
 
 	bool FRecord::operator==(const FRecord& Other) const
 	{
-		return Filter == Other.Filter && Platform == Other.Platform;
+		return Filter == Other.Filter && Platform == Other.Platform && bAddPackageDependencies == Other.bAddPackageDependencies;
 	}
 
 	bool FRecord::operator<(const FRecord& Other) const
@@ -44,6 +46,16 @@ namespace UE::CookAssetRegistryAccessTracker
 			return Platform->PlatformName() < Other.Platform->PlatformName();
 		}
 
+		if (bAddPackageDependencies != Other.bAddPackageDependencies)
+		{
+			if (!bAddPackageDependencies)
+			{
+				return true;
+			}
+
+			return false;
+		}
+
 		return false;
 	}
 
@@ -149,6 +161,7 @@ namespace UE::CookAssetRegistryAccessTracker
 			TSet<FRecord>& Records = AccessRecords.FindOrAdd(Referencer);
 			FRecord NewRecord;
 			NewRecord.Platform = AccumulatedScopeData->TargetPlatform;
+			NewRecord.bAddPackageDependencies = FAddPackageDependenciesScope::GetAddPackageDependenciesEnabled();
 			NewRecord.Filter = UE::AssetRegistry::ConvertToNonCompiledFilter(CompiledFilter);
 			NewRecord.Filter.SortForSaving();
 
diff --git a/Engine/Source/Editor/UnrealEd/Private/Cooker/CookAssetRegistryAccessTracker.h b/Engine/Source/Editor/UnrealEd/Private/Cooker/CookAssetRegistryAccessTracker.h
--- a/Engine/Source/Editor/UnrealEd/Private/Cooker/CookAssetRegistryAccessTracker.h
+++ b/Engine/Source/Editor/UnrealEd/Private/Cooker/CookAssetRegistryAccessTracker.h
@@ -16,8 +16,14 @@ namespace UE::CookAssetRegistryAccessTracker
 {
 	struct FRecord
 	{
-		const ITargetPlatform* Platform;
 		FARFilter Filter;
+		const ITargetPlatform* Platform;
+
+		/** 
+		 * If true, when the dependency for the asset registry query is created, it will also create a package dependency for each package
+		 * returned by the query.
+		*/
+		bool bAddPackageDependencies = false;
 
 		bool operator==(const FRecord& Other) const;
 		bool operator<(const FRecord& Other) const;
diff --git a/Engine/Source/Editor/UnrealEd/Private/Cooker/CookPackageArtifacts.cpp b/Engine/Source/Editor/UnrealEd/Private/Cooker/CookPackageArtifacts.cpp
--- a/Engine/Source/Editor/UnrealEd/Private/Cooker/CookPackageArtifacts.cpp
+++ b/Engine/Source/Editor/UnrealEd/Private/Cooker/CookPackageArtifacts.cpp
@@ -478,6 +478,22 @@ bool FBuildDependencySet::TryCollectInternal(FBuildResultDependenciesMap& InOutR
 		for (const UE::CookAssetRegistryAccessTracker::FRecord& Record : Records)
 		{
 			DefaultResultDependencies.Add(FCookDependency::AssetRegistryQuery(Record.Filter));
+
+			if (Record.bAddPackageDependencies)
+			{
+				// Add as package dependency each package from the result of the query.
+				TSet<FName> ResultPackageNames;
+				AssetRegistry->EnumerateAssets(Record.Filter, [&ResultPackageNames](const FAssetData& AssetData)
+					{
+						ResultPackageNames.Add(AssetData.PackageName);
+						return true;
+					}, UE::AssetRegistry::EEnumerateAssetsFlags::None);
+
+				for (const FName& ResultPackageName : ResultPackageNames)
+				{
+					DefaultResultDependencies.Add(FCookDependency::Package(ResultPackageName));
+				}
+			}
 		}
 	}
 
diff --git a/Engine/Source/Runtime/CoreUObject/Private/Cooker/CookAssetRegistryAccessTracker_AddPackageDep.cpp b/Engine/Source/Runtime/CoreUObject/Private/Cooker/CookAssetRegistryAccessTracker_AddPackageDep.cpp
new file mode 100644
--- /dev/null
+++ b/Engine/Source/Runtime/CoreUObject/Private/Cooker/CookAssetRegistryAccessTracker_AddPackageDep.cpp
@@ -0,0 +1,24 @@
+// Copyright Epic Games, Inc. All Rights Reserved.
+
+#include "Cooker/CookAssetRegistryAccessTracker_AddPackageDep.h"
+
+namespace UE::CookAssetRegistryAccessTracker
+{
+	static thread_local bool bAddPackageDependenciesEnabled = false;
+
+	FAddPackageDependenciesScope::FAddPackageDependenciesScope()
+		: bPreviousValue(bAddPackageDependenciesEnabled)
+	{
+		bAddPackageDependenciesEnabled = true;
+	}
+
+	FAddPackageDependenciesScope::~FAddPackageDependenciesScope()
+	{
+		bAddPackageDependenciesEnabled = bPreviousValue;
+	}
+
+	bool FAddPackageDependenciesScope::GetAddPackageDependenciesEnabled()
+	{
+		return bAddPackageDependenciesEnabled;
+	}
+}
diff --git a/Engine/Source/Runtime/CoreUObject/Public/Cooker/CookAssetRegistryAccessTracker_AddPackageDep.h b/Engine/Source/Runtime/CoreUObject/Public/Cooker/CookAssetRegistryAccessTracker_AddPackageDep.h
new file mode 100644
--- /dev/null
+++ b/Engine/Source/Runtime/CoreUObject/Public/Cooker/CookAssetRegistryAccessTracker_AddPackageDep.h
@@ -0,0 +1,19 @@
+// Copyright Epic Games, Inc. All Rights Reserved.
+
+#include "CoreTypes.h"
+
+namespace UE::CookAssetRegistryAccessTracker
+{
+	/** RAII class used to enable the creation of dependencies for the packages returned by an asset registry request. */
+	class FAddPackageDependenciesScope
+	{
+	public:
+		COREUOBJECT_API FAddPackageDependenciesScope();
+		COREUOBJECT_API ~FAddPackageDependenciesScope();
+		 
+		COREUOBJECT_API static bool GetAddPackageDependenciesEnabled();
+
+	private:
+		bool bPreviousValue;
+	};
+}
