[Editor Viewport] Zoom to Cursor broken for Orthographic Top and Bottom views

diff --git a/Engine/Source/Editor/UnrealEd/Private/EditorViewportClient.cpp b/Engine/Source/Editor/UnrealEd/Private/EditorViewportClient.cpp
--- a/Engine/Source/Editor/UnrealEd/Private/EditorViewportClient.cpp
+++ b/Engine/Source/Editor/UnrealEd/Private/EditorViewportClient.cpp
@@ -3895,11 +3895,11 @@ void FEditorViewportClient::OnOrthoZoom( const struct FInputEventState& InputSta
 		}
 		int32 DeltaFromCenterX = -(CenterX - (ViewportWidth>>1));
 		int32 DeltaFromCenterY =  (CenterY - (ViewportHeight>>1));
-		const bool bIsUsingLUFCoordinates = AxisDisplayInfo::GetAxisDisplayCoordinateSystem() == EAxisList::LeftUpForward;
+
 		switch( GetViewportType() )
 		{
 		case LVT_OrthoTop:
-			OldOffsetFromCenter.Set(DeltaFromCenterX, bIsUsingLUFCoordinates ? DeltaFromCenterY : -DeltaFromCenterY, 0.0f);
+			OldOffsetFromCenter.Set(-DeltaFromCenterY, -DeltaFromCenterX , 0.0f);
 			break;
 		case LVT_OrthoLeft:
 			OldOffsetFromCenter.Set(DeltaFromCenterX, 0.0f, DeltaFromCenterY);
@@ -3908,7 +3908,7 @@ void FEditorViewportClient::OnOrthoZoom( const struct FInputEventState& InputSta
 			OldOffsetFromCenter.Set(0.0f, DeltaFromCenterX, DeltaFromCenterY);
 			break;
 		case LVT_OrthoBottom:
-			OldOffsetFromCenter.Set(-DeltaFromCenterX, bIsUsingLUFCoordinates ? DeltaFromCenterY : -DeltaFromCenterY, 0.0f);
+			OldOffsetFromCenter.Set(-DeltaFromCenterY, DeltaFromCenterX, 0.0f);
 			break;
 		case LVT_OrthoRight:
 			OldOffsetFromCenter.Set(-DeltaFromCenterX, 0.0f, DeltaFromCenterY);
