Fixed message log not displaying when failing to save packages (we weren't passing FSaveErrorOutputDevice to OBJ SAVEPACKAG - typo in 30503323).
Consolodate all save package log messages into a single page.
Added a tokenized log message for packages that fail to save so the user can easily navigate to them.
Consolodate all packages that faild to save or were made writable into a single dialog.
Don't show the packages failed to save dialog for UEditorLoadingAndSavingUtils::SavePackages and UEditorLoadingAndSavingUtils::SaveDirtyPackages (typo in 3790970).

diff --git a/Engine/Source/Editor/UnrealEd/Private/FileHelpers.cpp b/Engine/Source/Editor/UnrealEd/Private/FileHelpers.cpp
--- a/Engine/Source/Editor/UnrealEd/Private/FileHelpers.cpp
+++ b/Engine/Source/Editor/UnrealEd/Private/FileHelpers.cpp
@@ -7,6 +7,7 @@
 #include "Misc/MessageDialog.h"
 #include "HAL/FileManager.h"
 #include "Algo/Copy.h"
+#include "Algo/Unique.h"
 #include "Algo/Transform.h"
 #include "Containers/VersePath.h"
 #include "Misc/CommandLine.h"
@@ -60,6 +61,7 @@
 #include "DesktopPlatformModule.h"
 #include "Logging/TokenizedMessage.h"
 #include "Logging/MessageLog.h"
+#include "Logging/StructuredLog.h"
 
 #include "Dialogs/DlgPickPath.h"
 
@@ -117,10 +119,20 @@ static TAutoConsoleVariable<int32> CVarSkipSourceControlCheckForEditablePackages
 
 FEditorFileUtils::FOnPrepareWorldsForExplicitSave FEditorFileUtils::OnPrepareWorldsForExplicitSave;
 
-/** A special output device that puts save output in the message log when flushed */
+/** A special output device that associates output to the current package being saved and stores said messages to InMessages. */
 class FSaveErrorOutputDevice : public FOutputDevice
 {
 public:
+	explicit FSaveErrorOutputDevice(TArray<TSharedRef<FTokenizedMessage>>& InMessages)
+		: Messages(InMessages)
+	{
+	}
+
+	~FSaveErrorOutputDevice()
+	{
+		checkf(CurrentPackageBeingSaved == nullptr, TEXT("Expected SetCurrentPackageBeingSaved(nullptr) call to clear CurrentPackageBeingSaved after it is saved."));
+	}
+
 	virtual void Serialize( const TCHAR* InData, ELogVerbosity::Type Verbosity, const class FName& Category ) override
 	{
 		if ( Verbosity == ELogVerbosity::Error || Verbosity == ELogVerbosity::Warning )
@@ -134,25 +146,234 @@ public:
 			{
 				Severity = EMessageSeverity::Warning;
 			}
-			ErrorMessages.Add(FTokenizedMessage::Create(Severity, FText::FromName(InData)));
+
+			Messages.Add(FTokenizedMessage::Create(Severity, FText::FromString(InData)));
+		}
+		else
+		{
+			GLog->Serialize(InData, Verbosity, Category);
 		}
 	}
 
-	virtual void Flush() override
+	virtual void SerializeRecord(const UE::FLogRecord& Record) override
 	{
-		if ( ErrorMessages.Num() > 0 )
+		if (Record.GetVerbosity() == ELogVerbosity::Error || Record.GetVerbosity() == ELogVerbosity::Warning)
+		{
+			TStringBuilder<512> Text;
+			Record.FormatMessageTo(Text);
+			Serialize(*Text, Record.GetVerbosity(), Record.GetCategory());
+		}
+		else
+		{
+			GLog->SerializeRecord(Record);
+		}
+	}
+
+	void SetCurrentPackageBeingSaved(const UPackage* Package)
+	{
+		// FOutputDevice only allows for non-tokenized messages, so add an additional message with a token to link to the package.
+		if (CurrentPackageBeingSaved != nullptr && CurrentPackageFirstMessageIndex < Messages.Num())
+		{
+			int32 Errors = 0;
+			int32 Warnings = 0;
+			for (int32 Index = CurrentPackageFirstMessageIndex; Index < Messages.Num(); ++Index)
+			{
+				switch (Messages[Index]->GetSeverity())
+				{
+				case EMessageSeverity::Error:
+					++Errors;
+					break;
+				case EMessageSeverity::Warning:
+					++Warnings;
+					break;
+				}
+			}
+
+			if (Errors > 0 || Warnings > 0)
+			{
+				TSharedRef<FTokenizedMessage> Message = FTokenizedMessage::Create(EMessageSeverity::Info);
+
+				Message->AddToken(FAssetNameToken::Create(CurrentPackageBeingSaved->GetName(), FText::FromString(IAssetTools::Get().GetUserFacingLongPackageName(*CurrentPackageBeingSaved))));
+
+				if (Errors > 0 && Warnings > 0)
+				{
+					Message->AddText(FText::Format(
+						LOCTEXT("ErrorsAndWarningsSavingAsset", "{0} {0}|plural(one=error,other=errors), {1} {1}|plural(one=warning,other=warnings) saving asset."),
+						FText::AsNumber(Errors),
+						FText::AsNumber(Warnings)));
+				}
+				else if (Errors > 0)
+				{
+					Message->AddText(FText::Format(
+						LOCTEXT("ErrorsSavingAsset", "{0} {0}|plural(one=error,other=errors) saving asset."),
+						FText::AsNumber(Errors)));
+				}
+				else if (ensure(Warnings > 0))
+				{
+					Message->AddText(FText::Format(
+						LOCTEXT("WarningsSavingAsset", "{0} {0}|plural(one=warning,other=warnings) saving asset."),
+						FText::AsNumber(Warnings)));
+				}
+
+				Messages.Insert(MoveTemp(Message), CurrentPackageFirstMessageIndex);
+			}
+		}
+
+		CurrentPackageBeingSaved = Package;
+		CurrentPackageFirstMessageIndex = Messages.Num();
+	}
+
+private:
+	// Holds the errors for the message log.
+	TArray<TSharedRef<FTokenizedMessage>>& Messages;
+	const UPackage* CurrentPackageBeingSaved = nullptr;
+	int32 CurrentPackageFirstMessageIndex = 0;
+};
+
+class IScopedSaveOutput
+{
+public:
+	IScopedSaveOutput() = default;
+	IScopedSaveOutput(const IScopedSaveOutput&) = delete;
+	virtual ~IScopedSaveOutput() = default;
+
+	IScopedSaveOutput& operator=(const IScopedSaveOutput&) = delete;
+
+	virtual FOutputDevice& GetOutputDevice() = 0;
+
+	virtual void SetCurrentPackageBeingSaved(const UPackage* Package) = 0;
+
+	virtual void AddFailedPackage(const UPackage* InFailedPackage) = 0;
+
+	virtual void AddWritablePackage(const UPackage* InWritablePackage) = 0;
+};
+
+// An implementation of IScopedSaveOutput that just writes to GLog.
+class FScopedAutosaveOutput final : public IScopedSaveOutput
+{
+public:
+	FScopedAutosaveOutput()
+		: OutputDevice(*GLog)
+	{
+	}
+
+	FOutputDevice& GetOutputDevice() override
+	{
+		return OutputDevice;
+	}
+
+	void SetCurrentPackageBeingSaved(const UPackage* Package) override
+	{
+	}
+
+	void AddFailedPackage(const UPackage* InFailedPackage) override
+	{
+	}
+
+	void AddWritablePackage(const UPackage* InWritablePackage) override
+	{
+	}
+
+private:
+	FOutputDevice& OutputDevice;
+};
+
+// An implementation of IScopedSaveOutput that coalesces log messages to a single message log page,
+// and shows a dialog warning about and packages that weren't saved or were made writable (if bInUseDialog is true).
+class FScopedSaveOutput final : public IScopedSaveOutput
+{
+public:
+	explicit FScopedSaveOutput(bool bInUseDialog)
+		: bUseDialog(bInUseDialog)
+		, Messages()
+		, OutputDevice(Messages)
+	{
+	}
+
+	~FScopedSaveOutput()
+	{
+		// If any packages were saved that weren't actually in source control but instead forcibly made writable,
+		// then warn the user about those packages. We do not warn if the Uncontrolled Changelists are enabled since
+		// the file will be picked up.
+		if (!WritablePackages.IsEmpty() && !FUncontrolledChangelistsModule::Get().IsEnabled())
+		{
+			TArray<FString> LongPackageNames = IAssetTools::Get().GetUserFacingLongPackageNames(WritablePackages);
+			LongPackageNames.Sort();
+			LongPackageNames.SetNum(Algo::Unique(LongPackageNames));
+
+			const FText WritableFileWarning = FText::Format(NSLOCTEXT("UnrealEd", "Warning_WritablePackagesNotCheckedOut", "The following assets are writable on disk but not checked out from revision control:\n{0}"),
+				FText::FromString(FString::Join(LongPackageNames, TEXT("\n"))));
+
+			UE_LOG(LogFileHelpers, Warning, TEXT("%s"), *WritableFileWarning.ToString());
+			if (bUseDialog)
+			{
+				FSuppressableWarningDialog::FSetupInfo Info(WritableFileWarning, NSLOCTEXT("UnrealEd", "Warning_WritablePackagesNotCheckedOutTitle", "Writable Assets Not Checked Out"), "WritablePackagesNotCheckedOut");
+				Info.ConfirmText = NSLOCTEXT("ModalDialogs", "WritablePackagesNotCheckedOutConfirm", "Close");
+
+				FSuppressableWarningDialog PromptForWritableFiles(Info);
+				PromptForWritableFiles.ShowModal();
+			}
+		}
+
+		// Warn the user if any packages failed to save
+		if (!FailedPackages.IsEmpty() && bUseDialog)
+		{
+			TArray<FString> LongPackageNames = IAssetTools::Get().GetUserFacingLongPackageNames(FailedPackages);
+			LongPackageNames.Sort();
+			LongPackageNames.SetNum(Algo::Unique(LongPackageNames));
+
+			FFormatNamedArguments Arguments;
+			Arguments.Add(TEXT("Packages"), FText::FromString(FString::Join(LongPackageNames, TEXT("\n"))));
+			FText MessageFormatting = NSLOCTEXT("FileHelper", "FailedSavePromptMessageFormatting", "The following assets failed to save correctly:\n{Packages}");
+			FText Message = FText::Format(MessageFormatting, Arguments);
+
+			// Display warning
+			FText Title = NSLOCTEXT("FileHelper", "FailedSavePrompt_Title", "Packages Failed To Save");
+			FMessageDialog::Open(EAppMsgType::Ok, Message, Title);
+		}
+
+		if (!Messages.IsEmpty())
 		{
 			FMessageLog EditorErrors("EditorErrors");
 			EditorErrors.NewPage(LOCTEXT("SaveOutputPageLabel", "Save Output"));
-			EditorErrors.AddMessages(ErrorMessages);
-			EditorErrors.Open();
-			ErrorMessages.Empty();
+			EditorErrors.AddMessages(Messages);
+			if (bUseDialog)
+			{
+				EditorErrors.Open();
+			}
+			else
+			{
+				EditorErrors.Flush();
+			}
 		}
 	}
 
+	FOutputDevice& GetOutputDevice() override
+	{
+		return OutputDevice;
+	}
+
+	void SetCurrentPackageBeingSaved(const UPackage* Package) override
+	{
+		OutputDevice.SetCurrentPackageBeingSaved(Package);
+	}
+
+	void AddFailedPackage(const UPackage* InFailedPackage) override
+	{
+		FailedPackages.Add(InFailedPackage);
+	}
+
+	void AddWritablePackage(const UPackage* InWritablePackage) override
+	{
+		WritablePackages.Add(InWritablePackage);
+	}
+
 private:
-	// Holds the errors for the message log.
-	TArray< TSharedRef< FTokenizedMessage > > ErrorMessages;
+	const bool bUseDialog;
+	TArray<TSharedRef<FTokenizedMessage>> Messages;
+	FSaveErrorOutputDevice OutputDevice;
+	TArray<const UPackage*> FailedPackages;
+	TArray<const UPackage*> WritablePackages;
 };
 
 namespace FileDialogHelpers
@@ -597,13 +818,17 @@ FString FEditorFileUtils::GetFilterString(EFileInteraction Interaction)
 
 ///////////////////////////////////////////////////////////////////////////////
 
+static TArray<UPackage*> InternalGetValidPackages(const TArray<UPackage*>& PackagesToSave, bool bCheckDirty);
+static bool InternalCheckoutAndSavePackages(const TArray<UPackage*>& PackagesToSave, IScopedSaveOutput& SaveOutput);
+
 // Gather external packages to save for maps an other package that use external packages
 // @param	InPackage				The package that is being saved
 // @param	InOutPackagesToSave		Any packages to save. Array will be modified with additional packages if they are found, then those packages will all be saved
 // @param	bInNewlyCreated			Whether the package was newly created
 // @param	bInAutosaving			Should be set to true if autosaving
+// @param	SaveOutput				Object used to consolidate save output to present to the user
 // @returns true if all the save operations completed successfully
-static bool SaveExternalPackages(UPackage* InPackage, TArray<UPackage*>& InOutPackagesToSave, bool bInNewlyCreated, bool bInAutosaving)
+static bool SaveExternalPackages(UPackage* InPackage, TArray<UPackage*>& InOutPackagesToSave, bool bInNewlyCreated, bool bInAutosaving, IScopedSaveOutput& SaveOutput)
 {
 	bool bSuccess = true;
 
@@ -613,9 +838,9 @@ static bool SaveExternalPackages(UPackage* InPackage, TArray<UPackage*>& InOutPa
 
 		if (InOutPackagesToSave.Num())
 		{
-			if (!UEditorLoadingAndSavingUtils::SavePackages(InOutPackagesToSave, /*bCheckDirty=*/ !bInNewlyCreated))
+			TArray<UPackage*> Packages = InternalGetValidPackages(InOutPackagesToSave, /*bCheckDirty=*/ !bInNewlyCreated);
+			if (!InternalCheckoutAndSavePackages(Packages, SaveOutput))
 			{
-				FMessageDialog::Open(EAppMsgType::Ok, NSLOCTEXT("UnrealEd", "Error_FailedToSaveExternalPackages", "Failed to save external packages"));
 				bSuccess = false;
 			}
 		}
@@ -624,6 +849,30 @@ static bool SaveExternalPackages(UPackage* InPackage, TArray<UPackage*>& InOutPa
 	return bSuccess;
 }
 
+TArray<UPackage*> InternalGetMapDataPackages(UWorld& WorldToSave)
+{
+	TArray<UPackage*> MapDataPackages;
+
+	UPackage* WorldPackage = WorldToSave.GetOutermost();
+
+	if (!WorldPackage->HasAnyPackageFlags(PKG_PlayInEditor)
+		&& !WorldPackage->HasAnyFlags(RF_Transient))
+	{
+		ULevel* Level = WorldToSave.PersistentLevel;
+		if (Level->MapBuildData)
+		{
+			UPackage* BuiltDataPackage = Level->MapBuildData->GetOutermost();
+
+			if (BuiltDataPackage != WorldPackage)
+			{
+				MapDataPackages.Add(BuiltDataPackage);
+			}
+		}
+	}
+
+	return MapDataPackages;
+}
+
 /**
  * @param	World					The world to save.
  * @param	ForceFilename			If non-NULL, save the level package to this name (full path+filename).
@@ -634,6 +883,7 @@ static bool SaveExternalPackages(UPackage* InPackage, TArray<UPackage*>& InOutPa
  * @param	FinalFilename			[out] The full path+filename the level was saved to.
  * @param	bAutosaving				Should be set to true if autosaving; passed to UWorld::SaveWorld.
  * @param	bPIESaving				Should be set to true if saving for PIE; passed to UWorld::SaveWorld.
+ * @param	SaveOutput				Object used to consolidate save output to present to the user.
  * @return							true if the level was saved.
  */
 static bool SaveWorld(UWorld* World,
@@ -644,7 +894,8 @@ static bool SaveWorld(UWorld* World,
 					   bool bCheckDirty,
 					   FString& FinalFilename,
 					   bool bAutosaving,
-					   bool bPIESaving)
+					   bool bPIESaving,
+					   IScopedSaveOutput& SaveOutput)
 {
 	// SaveWorld not reentrant - check that we are not already in the process of saving here (for example, via autosave)
 	static bool bIsReentrant = false;
@@ -945,17 +1196,24 @@ static bool SaveWorld(UWorld* World,
 		{
 			const FString AutoSavingString = (bAutosaving || bPIESaving) ? TEXT("true") : TEXT("false");
 			const FString KeepDirtyString = bPIESaving ? TEXT("true") : TEXT("false");
-			FSaveErrorOutputDevice SaveErrors;
 
-			bSuccess = GEditor->Exec(NULL, *FString::Printf(TEXT("OBJ SAVEPACKAGE PACKAGE=\"%s\" FILE=\"%s\" SILENT=true AUTOSAVING=%s KEEPDIRTY=%s"), *Package->GetName(), *FinalFilename, *AutoSavingString, *KeepDirtyString), SaveErrors);
-			SaveErrors.Flush();
+			SaveOutput.SetCurrentPackageBeingSaved(Package);
+			ON_SCOPE_EXIT
+			{
+				SaveOutput.SetCurrentPackageBeingSaved(nullptr);
+			};
+
+			bSuccess = GEditor->Exec(
+				nullptr,
+				*FString::Printf(TEXT("OBJ SAVEPACKAGE PACKAGE=\"%s\" FILE=\"%s\" SILENT=true AUTOSAVING=%s KEEPDIRTY=%s"), *Package->GetName(), *FinalFilename, *AutoSavingString, *KeepDirtyString),
+				SaveOutput.GetOutputDevice());
 		}
 
 		SlowTask.EnterProgressFrame(50);
 
 		if (bSuccess)
 		{
-			bSuccess = SaveExternalPackages(Package, PackagesToSave, bNewlyCreated, bAutosaving);
+			bSuccess = SaveExternalPackages(Package, PackagesToSave, bNewlyCreated, bAutosaving, SaveOutput);
 		}
 
 		if (bSuccess)
@@ -998,9 +1256,10 @@ static bool SaveWorld(UWorld* World,
 		// @todo Autosaving should save build data as well
 		if (bSuccess && !bAutosaving)
 		{
-			if (!FEditorFileUtils::SaveMapDataPackages(SaveWorld, /*bCheckDirty*/true))
+			TArray<UPackage*> MapDataPackages = InternalGetMapDataPackages(*SaveWorld);
+			TArray<UPackage*> Packages = InternalGetValidPackages(MapDataPackages, bCheckDirty);
+			if (!InternalCheckoutAndSavePackages(Packages, SaveOutput))
 			{
-				FMessageDialog::Open(EAppMsgType::Ok, NSLOCTEXT("UnrealEd", "Error_FailedToSaveMapDataPackages", "Failed to save map data packages"));
 				bSuccess = false;
 			}
 		}
@@ -1050,13 +1309,23 @@ static bool SaveWorld(UWorld* World,
 // @param	InPackage				The package to save
 // @param	PackageName				The name of the package to save
 // @param	FinalPackageSavePath	The save path of the package
-// @param	SaveOutput				Output device for error reporting
+// @param	SaveOutput				Object used to consolidate save output to present to the user
 // @returns true if all the save operations completed successfully
-static bool SaveAsset(UPackage* InPackage, const FString& PackageName, const FString& FinalPackageSavePath, FOutputDevice& SaveOutput)
+static bool SaveAsset(UPackage* InPackage, const FString& PackageName, const FString& FinalPackageSavePath, IScopedSaveOutput& SaveOutput)
 {
 	TArray<UPackage*> PackagesToSave;
-	return	SaveExternalPackages(InPackage, PackagesToSave, InPackage->HasAnyPackageFlags(PKG_NewlyCreated), false) &&
-			GEngine->Exec(nullptr, *FString::Printf( TEXT("OBJ SAVEPACKAGE PACKAGE=\"%s\" FILE=\"%s\" SILENT=true"), *PackageName, *FinalPackageSavePath));
+	if (!SaveExternalPackages(InPackage, PackagesToSave, InPackage->HasAnyPackageFlags(PKG_NewlyCreated), false, SaveOutput))
+	{
+		return false;
+	}
+
+	SaveOutput.SetCurrentPackageBeingSaved(InPackage);
+	ON_SCOPE_EXIT
+	{
+		SaveOutput.SetCurrentPackageBeingSaved(nullptr);
+	};
+
+	return GEngine->Exec(nullptr, *FString::Printf( TEXT("OBJ SAVEPACKAGE PACKAGE=\"%s\" FILE=\"%s\" SILENT=true"), *PackageName, *FinalPackageSavePath), SaveOutput.GetOutputDevice());
 }
 
 FString FEditorFileUtils::GetAutoSaveFilename(UPackage* const Package, const FString& AbsoluteAutosaveDir, const int32 AutoSaveIndex, const FString& PackageExt)
@@ -3101,6 +3370,9 @@ bool FEditorFileUtils::SaveMap(UWorld* InWorld, const FString& Filename )
 {
 	bool bLevelWasSaved = false;
 
+	constexpr bool bUseDialog = true;
+	FScopedSaveOutput SaveOutput(bUseDialog);
+
 	const double SaveStartTime = FPlatformTime::Seconds();
 
 	FString FinalFilename;
@@ -3108,7 +3380,8 @@ bool FEditorFileUtils::SaveMap(UWorld* InWorld, const FString& Filename )
 								nullptr, nullptr,
 								true, false,
 								FinalFilename,
-								false, false );
+								false, false,
+								SaveOutput );
 
 	// Track time spent saving map.
 	UE_LOG(LogFileHelpers, Log, TEXT("Saving map '%s' took %.3f"), *FPaths::GetBaseFilename(Filename), FPlatformTime::Seconds() - SaveStartTime );
@@ -3156,6 +3429,8 @@ EAutosaveContentPackagesResult::Type FEditorFileUtils::AutosaveMapEx(const FStri
 	// Clean up any old worlds.
 	CollectGarbage( GARBAGE_COLLECTION_KEEPFLAGS );
 
+	FScopedAutosaveOutput SaveOutput;
+
 	FWorldContext& EditorContext = GEditor->GetEditorWorldContext();
 
 	// Get the set of all reference worlds.
@@ -3180,7 +3455,8 @@ EAutosaveContentPackagesResult::Type FEditorFileUtils::AutosaveMapEx(const FStri
 					NULL, NULL,
 					false, true,
 					FinalFilename,
-					true, false );
+					true, false,
+					SaveOutput );
 
 				// Remark the package as being dirty, as saving will have undiritied the package.
 				Package->MarkPackageDirty();
@@ -3418,22 +3694,33 @@ static void PrepareSavePackages(const TArray<UPackage*>& PackagesToSave)
  * @param bUseDialog					If true, use the normal behavior.
  *										If false, do not prompt message dialog. If it can't save the package, skip it. If the package is a map and the name is not valid, skip it.
  * @param OutPackageLocallyWritable		Set to true if the provided package was locally writable but not under source control (of if source control is disabled).
- * @param SaveOutput					The output from the save process.
+ * @param SaveOutput					Object used to consolidate save output to present to the user.
  * @return	InternalSavePackageResult::Success if package saving was a success
 			InternalSavePackageResult::Continue if the package saving failed and the user doesn't want to retry
 			InternalSavePackageResult::Cancel if the user wants to cancel everything
 			InternalSavePackageResult::Error if an error occured. Check OutFailureReason
  */
-static InternalSavePackageResult InternalSavePackage(UPackage* PackageToSave, bool bUseDialog, bool& bOutPackageLocallyWritable, FOutputDevice &SaveOutput)
+static InternalSavePackageResult InternalSavePackage(UPackage* PackageToSave, bool bUseDialog, IScopedSaveOutput& SaveOutput)
 {
 	TRACE_CPUPROFILER_EVENT_SCOPE(InternalSavePackage);
 
-	// What we will be returning. Assume for now that everything will go fine
 	InternalSavePackageResult ReturnCode = InternalSavePackageResult::Error;
 
 	// Assume the package is locally writable in case SCC is disabled; if SCC is enabled, it will
 	// correctly set this value later
-	bOutPackageLocallyWritable = true;
+	bool bPackageLocallyWritable = true;
+
+	ON_SCOPE_EXIT
+	{
+		if (ReturnCode == InternalSavePackageResult::Continue || ReturnCode == InternalSavePackageResult::Error)
+		{
+			SaveOutput.AddFailedPackage(PackageToSave);
+		}
+		if (bPackageLocallyWritable)
+		{
+			SaveOutput.AddWritablePackage(PackageToSave);
+		}
+	};
 
 	bool bShouldRetrySave = true;
 	UWorld*	AssociatedWorld	= UWorld::FindWorldInPackage(PackageToSave);
@@ -3497,7 +3784,7 @@ static InternalSavePackageResult InternalSavePackage(UPackage* PackageToSave, bo
 		// Right now, all of this code is specific to maps
 
 		// There wont be a "not checked out from SCC but writable on disk" conflict if the package is new.
-		bOutPackageLocallyWritable = false;
+		bPackageLocallyWritable = false;
 
 		// Make a list of file types
 		// We have to ask for save as.
@@ -3624,13 +3911,19 @@ static InternalSavePackageResult InternalSavePackage(UPackage* PackageToSave, bo
 		if (bIsMapPackage && !bSavingRedirectorToWorld)
 		{
 			// have a Helper attempt to save the map
-			SaveOutput.Log("LogFileHelpers", ELogVerbosity::Log, FString::Printf(TEXT("Saving Map: %s"), *PackageName));
-			bWasSuccessful = FEditorFileUtils::SaveMap( AssociatedWorld, FinalPackageSavePath );
+			SaveOutput.GetOutputDevice().Log("LogFileHelpers", ELogVerbosity::Log, FString::Printf(TEXT("Saving Map: %s"), *PackageName));
+			FString FinalFilename;
+			bWasSuccessful = SaveWorld(AssociatedWorld, &FinalPackageSavePath,
+				nullptr, nullptr,
+				true, false,
+				FinalFilename,
+				false, false,
+				SaveOutput);
 		}
 		else
 		{
 			// normally, we just save the package (and its external packages)
-			SaveOutput.Log("LogFileHelpers", ELogVerbosity::Log, FString::Printf(TEXT("Saving Package: %s"), *PackageName));
+			SaveOutput.GetOutputDevice().Log("LogFileHelpers", ELogVerbosity::Log, FString::Printf(TEXT("Saving Package: %s"), *PackageName));
 			bWasSuccessful = SaveAsset(PackageToSave, PackageName, FinalPackageSavePath, SaveOutput);
 		}
 
@@ -3642,7 +3935,7 @@ static InternalSavePackageResult InternalSavePackage(UPackage* PackageToSave, bo
 			if (ISourceControlModule::Get().IsEnabled())
 			{
 				// Assume the package was correctly checked out from SCC
-				bOutPackageLocallyWritable = false;
+				bPackageLocallyWritable = false;
 
 				// Trusting the SCC status in the package file cache to minimize network activity during save.
 				const FSourceControlStatePtr SourceControlState = SourceControlProvider.GetState(PackageToSave, EStateCacheUsage::Use);
@@ -3652,16 +3945,16 @@ static InternalSavePackageResult InternalSavePackage(UPackage* PackageToSave, bo
 				const bool bInDepot = SourceControlState.IsValid() && SourceControlState->IsSourceControlled();
 				if ( !bSCCCanEdit && bInDepot && !IFileManager::Get().IsReadOnly( *FinalPackageSavePath ) && SourceControlProvider.UsesLocalReadOnlyState() && !bSCCIsCheckedOut )
 				{
-					bOutPackageLocallyWritable = true;
+					bPackageLocallyWritable = true;
 				}
 			}
 			else
 			{
 				// If we are in offline mode, automatically add the modified package to an Uncontrolled Changelist
-				bOutPackageLocallyWritable = UncontrolledChangelistsModule.IsEnabled() && (!IFileManager::Get().IsReadOnly(*FinalPackageSavePath));
+				bPackageLocallyWritable = UncontrolledChangelistsModule.IsEnabled() && (!IFileManager::Get().IsReadOnly(*FinalPackageSavePath));
 			}
 
-			if (bWasSuccessful && bOutPackageLocallyWritable)
+			if (bWasSuccessful && bPackageLocallyWritable)
 			{
 				UncontrolledChangelistsModule.OnSaveWritable({ FinalPackageSavePath });
 			}
@@ -3715,33 +4008,6 @@ static InternalSavePackageResult InternalSavePackage(UPackage* PackageToSave, bo
 	return ReturnCode;
 }
 
-/**
- * Shows a dialog warning a user about packages which failed to save
- * 
- * @param Packages that should be displayed in the dialog
- */
-static void InternalWarnUserAboutFailedSave( const TArray<UPackage*>& InFailedPackages, bool bUseDialog )
-{
-	// Warn the user if any packages failed to save
-	if ( InFailedPackages.Num() > 0 )
-	{
-		FString FailedPackages;
-		for (const FString& UserFacingPath : IAssetTools::Get().GetUserFacingLongPackageNames(InFailedPackages))
-		{
-			FailedPackages += FString::Printf( TEXT("\n%s"), *UserFacingPath );
-		}
-
-		FFormatNamedArguments Arguments;
-		Arguments.Add(TEXT("Packages"), FText::FromString( FailedPackages ));
-		FText MessageFormatting = NSLOCTEXT("FileHelper", "FailedSavePromptMessageFormatting", "The following assets failed to save correctly:{Packages}");
-		FText Message = FText::Format( MessageFormatting, Arguments );
-
-		// Display warning
-		FText Title = NSLOCTEXT("FileHelper", "FailedSavePrompt_Title", "Packages Failed To Save");
-		FMessageDialog::Open(EAppMsgType::Ok, Message, Title);
-	}
-}
-
 static TArray<UPackage*> InternalGetDirtyPackages(const bool bSaveMapPackages, const bool bSaveContentPackages, const FEditorFileUtils::FShouldIgnorePackageFunctionRef& ShouldIgnorePackageFunction = FEditorFileUtils::FShouldIgnorePackage::Default)
 {
 	if (bSaveContentPackages)
@@ -3766,37 +4032,29 @@ static TArray<UPackage*> InternalGetDirtyPackages(const bool bSaveMapPackages, c
 	return PackagesToSave;
 }
 
-static void InternalNotifyNoPackagesSaved(const bool bUseDialog)
+static void InternalNotifyNoPackagesSaved()
 {
-	if (bUseDialog)
-	{
-		FNotificationInfo NotificationInfo(LOCTEXT("NoAssetsToSave", "All files are already saved."));
-		NotificationInfo.bFireAndForget = true;
-		NotificationInfo.ExpireDuration = 4.0f; // Need this message to last a little longer than normal since the user may have expected there to be modified files.
-		NotificationInfo.bUseThrobber = true;
-		FSlateNotificationManager::Get().AddNotification(NotificationInfo);
-	}
-	else
-	{
-		UE_LOG(LogFileHelpers, Log, TEXT("%s"), *LOCTEXT("NoAssetsToSave", "All files are already saved.").ToString());
-	}
+	FNotificationInfo NotificationInfo(LOCTEXT("NoAssetsToSave", "All files are already saved."));
+	NotificationInfo.bFireAndForget = true;
+	NotificationInfo.ExpireDuration = 4.0f; // Need this message to last a little longer than normal since the user may have expected there to be modified files.
+	NotificationInfo.bUseThrobber = true;
+	FSlateNotificationManager::Get().AddNotification(NotificationInfo);
 }
 
 
 /*
+ * @param PackageToSave					The package to save.
  * @param bUseDialog					If true, use the normal behavior.
  *										If false, do not prompt message dialog. If it can't save the package, skip it. If the package is a map and the name is not valid, skip it.
- * @param	bShowDialogIfError			If InternalSavePackage failed, tell the user with a Dialog
- * @param	OutFailedPackages			Packages that failed to save
+ * @param SaveOutput					Object used to consolidate save output to present to the user.
  */
-static bool InternalSavePackagesFast(const TArray<UPackage*>& PackagesToSave, bool bUseDialog, TArray<UPackage*>& OutFailedPackages)
+static bool InternalSavePackagesFast(const TArray<UPackage*>& PackagesToSave, bool bUseDialog, IScopedSaveOutput& SaveOutput)
 {
 	TRACE_CPUPROFILER_EVENT_SCOPE(InternalSavePackagesFast);
 	UE_SCOPED_TIMER(TEXT("InternalSavePackagesFast"), LogFileHelpers, Log);
 
 	bool bReturnCode = true;
 
-	FSaveErrorOutputDevice SaveErrors;
 	GWarn->BeginSlowTask(NSLOCTEXT("UnrealEd", "SavingPackagesE", "Saving packages..."), true);
 
 	TArray<UPackage*> PackagesToClean;
@@ -3869,8 +4127,7 @@ static bool InternalSavePackagesFast(const TArray<UPackage*>& PackagesToSave, bo
 
 	for (UPackage* Package : FinalPackagesToSave)
 	{
-		bool bPackageLocallyWritable;
-		const InternalSavePackageResult SaveStatus = InternalSavePackage(Package, bUseDialog, bPackageLocallyWritable, SaveErrors);
+		const InternalSavePackageResult SaveStatus = InternalSavePackage(Package, bUseDialog, SaveOutput);
 
 		if (SaveStatus == InternalSavePackageResult::Cancel)
 		{
@@ -3880,9 +4137,6 @@ static bool InternalSavePackagesFast(const TArray<UPackage*>& PackagesToSave, bo
 		}
 		else if (SaveStatus == InternalSavePackageResult::Continue || SaveStatus == InternalSavePackageResult::Error)
 		{
-			// The package could not be saved so add it to the failed array 
-			OutFailedPackages.Add(Package);
-
 			if (SaveStatus == InternalSavePackageResult::Error)
 			{
 				// exit gracefully.
@@ -3898,7 +4152,6 @@ static bool InternalSavePackagesFast(const TArray<UPackage*>& PackagesToSave, bo
 	}
 	
 	GWarn->EndSlowTask();
-	SaveErrors.Flush();
 
 	return bReturnCode;
 }
@@ -3934,11 +4187,10 @@ static bool InternalSavePackages(const TArray<UPackage*>& PackagesToSave, bool b
 	}
 	else
 	{
-		const bool bUseDialog = true;
-		TArray<UPackage*> FailedPackages;
-		bReturnCode = InternalSavePackagesFast(PackagesToSave, bUseDialog, FailedPackages);
-		// Warn the user about any packages which failed to save.
-		InternalWarnUserAboutFailedSave(FailedPackages, bUseDialog);
+		constexpr bool bUseDialog = true;
+		FScopedSaveOutput SaveOutput(bUseDialog);
+
+		bReturnCode = InternalSavePackagesFast(PackagesToSave, bUseDialog, SaveOutput);
 	}
 	return bReturnCode;
 }
@@ -3947,25 +4199,13 @@ bool FEditorFileUtils::SaveMapDataPackages(UWorld* WorldToSave, bool bCheckDirty
 {
 	TRACE_CPUPROFILER_EVENT_SCOPE(FEditorFileUtils_SaveMapDataPackages);
 
-	TArray<UPackage*> PackagesToSave;
-	UPackage* WorldPackage = WorldToSave->GetOutermost();
-
-	if (!WorldPackage->HasAnyPackageFlags(PKG_PlayInEditor)
-		&& !WorldPackage->HasAnyFlags(RF_Transient))
+	if (WorldToSave == nullptr)
 	{
-		ULevel* Level = WorldToSave->PersistentLevel;
-		if (Level->MapBuildData)
-		{
-			UPackage* BuiltDataPackage = Level->MapBuildData->GetOutermost();
-
-			if (BuiltDataPackage != WorldPackage)
-			{
-				return UEditorLoadingAndSavingUtils::SavePackages({ BuiltDataPackage }, bCheckDirty);
-			}
-		}
+		return false;
 	}
 
-	return true;
+	TArray<UPackage*> PackagesToSave = InternalGetMapDataPackages(*WorldToSave);
+	return UEditorLoadingAndSavingUtils::SavePackages(PackagesToSave, bCheckDirty);
 }
 
 /**
@@ -4019,13 +4259,17 @@ bool FEditorFileUtils::SaveLevel(ULevel* Level, const FString& DefaultFilename,
 		UWorld* WorldToSave = Cast<UWorld>( Level->GetOuter() );
 		if ( WorldToSave )
 		{
+			constexpr bool bUseDialog = true;
+			FScopedSaveOutput SaveOutput(bUseDialog);
+
 			FString FinalFilename;
 			bLevelWasSaved = SaveWorld( WorldToSave,
 										DefaultFilename.Len() > 0 ? &DefaultFilename : NULL,
 										NULL, NULL,
 										true, false,
 										FinalFilename,
-										false, false );
+										false, false,
+										SaveOutput );
 			if (bLevelWasSaved && OutSavedFilename)
 			{
 				*OutSavedFilename = FinalFilename;
@@ -4073,7 +4317,7 @@ bool FEditorFileUtils::SaveDirtyPackages(const bool bPromptUserToSave, const boo
 	}
 	else if (bNotifyNoPackagesSaved)
 	{
-		InternalNotifyNoPackagesSaved(true);
+		InternalNotifyNoPackagesSaved();
 	}
 	return bReturnCode;
 }
@@ -4141,7 +4385,7 @@ bool FEditorFileUtils::SaveDirtyContentPackages(TArray<UClass*>& SaveContentClas
 	}
 	else if (bNotifyNoPackagesSaved)
 	{
-		InternalNotifyNoPackagesSaved(true);
+		InternalNotifyNoPackagesSaved();
 		bResult = true;
 	}
 	return bResult;
@@ -4209,17 +4453,16 @@ bool FEditorFileUtils::SaveCurrentLevel()
  * Helper code for PromptForCheckoutAndSave
  * @param FinalSaveList					Package to save
  * @param bUseDialog					Use dialog with InternalSavePackage & if we show errors
+ * @param SaveOutput					Object used to consolidate save output to present to the user
  * @param OutFailedPackages				Packages that couldn't be save
  */
-FEditorFileUtils::EPromptReturnCode InternalPromptForCheckoutAndSave(const TArray<UPackage*>& FinalSaveList, bool bUseDialog, TArray<UPackage*>& OutFailedPackages)
+FEditorFileUtils::EPromptReturnCode InternalPromptForCheckoutAndSave(const TArray<UPackage*>& FinalSaveList, bool bUseDialog, IScopedSaveOutput& SaveOutput, TArray<UPackage*>& OutFailedPackages)
 {
 	UE_SCOPED_TIMER(TEXT("InternalPromptForCheckoutAndSave"), LogFileHelpers, Log);
 
 	FEditorFileUtils::EPromptReturnCode ReturnResponse = FEditorFileUtils::PR_Success;
 	const FScopedBusyCursor BusyCursor;
-	FSaveErrorOutputDevice SaveErrors;
 
-	TArray<UPackage*, TInlineAllocator<2>> WritablePackageFiles;
 	TArray<UPackage*> PackagesToClean;
 	TArray<UPackage*> PackagesToSave;
 	PackagesToSave.Reserve(FinalSaveList.Num());
@@ -4303,15 +4546,7 @@ FEditorFileUtils::EPromptReturnCode InternalPromptForCheckoutAndSave(const TArra
 			SlowTask.EnterProgressFrame(1, SavingPackageText);
 
 			// Save the package
-			bool bPackageLocallyWritable;
-			const InternalSavePackageResult SaveStatus = InternalSavePackage(Package, bUseDialog, bPackageLocallyWritable, SaveErrors);
-
-			// If InternalSavePackage reported that the provided package was locally writable, add it to the list of writable files
-			// to warn the user about
-			if (bPackageLocallyWritable)
-			{
-				WritablePackageFiles.Add(Package);
-			}
+			const InternalSavePackageResult SaveStatus = InternalSavePackage(Package, bUseDialog, SaveOutput);
 
 			if (SaveStatus == InternalSavePackageResult::Cancel)
 			{
@@ -4340,49 +4575,12 @@ FEditorFileUtils::EPromptReturnCode InternalPromptForCheckoutAndSave(const TArra
 		}
 	}
 
-	SaveErrors.Flush();
-
-	
-
 	// Add all files that needs to be marked for add in one command, if any
 	if (GEditor)
 	{
 		GEditor->RunDeferredMarkForAddFiles();
 	}
 
-	// If any packages were saved that weren't actually in source control but instead forcibly made writable,
-	// then warn the user about those packages. We do not warn if the Uncontrolled Changelists are enabled since
-	// the file will be picked up.
-	if (!FUncontrolledChangelistsModule::Get().IsEnabled() && (WritablePackageFiles.Num() > 0))
-	{
-		FString WritableFiles;
-		for (const FString& UserFacingPath : IAssetTools::Get().GetUserFacingLongPackageNames(WritablePackageFiles))
-		{
-			// A warning message was created.  Try and show it.
-			WritableFiles += FString::Printf(TEXT("\n%s"), *UserFacingPath);
-		}
-
-		const FText WritableFileWarning = FText::Format(NSLOCTEXT("UnrealEd", "Warning_WritablePackagesNotCheckedOut", "The following assets are writable on disk but not checked out from revision control:{0}"),
-			FText::FromString(WritableFiles));
-
-		UE_LOG(LogFileHelpers, Warning, TEXT("%s"), *WritableFileWarning.ToString());
-		if (bUseDialog)
-		{
-			FSuppressableWarningDialog::FSetupInfo Info(WritableFileWarning, NSLOCTEXT("UnrealEd", "Warning_WritablePackagesNotCheckedOutTitle", "Writable Assets Not Checked Out"), "WritablePackagesNotCheckedOut");
-			Info.ConfirmText = NSLOCTEXT("ModalDialogs", "WritablePackagesNotCheckedOutConfirm", "Close");
-
-			FSuppressableWarningDialog PromptForWritableFiles(Info);
-			PromptForWritableFiles.ShowModal();
-		}
-	}
-
-	// Warn the user if any packages failed to save
-	if (OutFailedPackages.Num() > 0)
-	{
-		// Show a dialog for the failed packages
-		InternalWarnUserAboutFailedSave(OutFailedPackages, bUseDialog);
-	}
-
 	return ReturnResponse;
 }
 
@@ -4667,8 +4865,6 @@ FEditorFileUtils::EPromptReturnCode FEditorFileUtils::PromptForCheckoutAndSave(c
 
 		PackagesToSave = MoveTemp(SortedPackagesToSave);
 
-		TArray<UPackage*> FailedPackages;
-
 		TArray<UPackage*> PackagesCheckedOutOrMadeWritable;
 		TArray<UPackage*> PackagesNotNeedingCheckout;
 		TArray<UPackage*> PackagesToRevert;
@@ -4705,9 +4901,12 @@ FEditorFileUtils::EPromptReturnCode FEditorFileUtils::PromptForCheckoutAndSave(c
 				FinalSaveList.Append(PackagesCheckedOutOrMadeWritable);
 			}
 
+			TArray<UPackage*> FailedPackages;
+
 			{
-				const bool bUseDialog = true;
-				ReturnResponse = InternalPromptForCheckoutAndSave(FinalSaveList, bUseDialog, FailedPackages);
+				constexpr bool bUseDialog = true;
+				FScopedSaveOutput SaveOutput(bUseDialog);
+				ReturnResponse = InternalPromptForCheckoutAndSave(FinalSaveList, bUseDialog, SaveOutput, FailedPackages);
 			}
 
 			// Set the failure array to have the same contents as the local one.
@@ -4752,8 +4951,10 @@ FEditorFileUtils::EPromptReturnCode FEditorFileUtils::PromptForCheckoutAndSave(c
 			}
 
 			// Save the packages that need to be reverted, so the SourceControl can act on them.
+			constexpr bool bUseDialog = false;
+			FScopedAutosaveOutput SaveOutput;
 			TArray<UPackage*> PackagesSaveFailed;
-			InternalPromptForCheckoutAndSave(PackagesToRevert, /*bUseDialog=*/false, PackagesSaveFailed);
+			InternalPromptForCheckoutAndSave(PackagesToRevert, bUseDialog, SaveOutput, PackagesSaveFailed);
 
 			// Revert packages that could not be checked out.
 			USourceControlHelpers::RevertAndReloadPackages(USourceControlHelpers::PackageFilenames(PackagesToRevert), /*bRevertAll=*/false, /*bReloadWorld=*/bReloadWorld);
@@ -5006,7 +5207,8 @@ bool FEditorFileUtils::SaveWorlds(UWorld* InWorld, const FString& RootPath, cons
 	TArray<UWorld*> WorldsArray;
 	EditorLevelUtils::GetWorlds( InWorld, WorldsArray, true );
 
-	TArray<FName> PackageNames; 
+	constexpr bool bUseDialog = true;
+	FScopedSaveOutput SaveOutput(bUseDialog);
 
 	// Save all packages containing levels that are currently "referenced" by the global world pointer.
 	bool bSavedAll = true;
@@ -5020,7 +5222,8 @@ bool FEditorFileUtils::SaveWorlds(UWorld* InWorld, const FString& RootPath, cons
 												*WorldPath, Prefix,
 												false, false,
 												FinalFilename,
-												false, true);
+												false, true,
+												SaveOutput );
 
 		if (bLevelWasSaved)
 		{
@@ -5575,63 +5778,85 @@ UWorld* UEditorLoadingAndSavingUtils::LoadMapWithDialog()
 	return GEditor->GetEditorWorldContext().World();
 }
 
-static bool InternalCheckoutAndSavePackages(const TArray<UPackage*>& PackagesToSave, bool bUseDialog)
+static bool InternalCheckoutAndSavePackages(const TArray<UPackage*>& PackagesToSave, IScopedSaveOutput& SaveOutput)
 {
 	bool bResult = true;
 	if (PackagesToSave.Num() > 0)
 	{
-		if (bUseDialog)
+		const FScopedBusyCursor BusyCursor;
+		// Prevent modal window if not requested.
+		TGuardValue<bool> UnattendedScriptGuard(GIsRunningUnattendedScript, true);
+
+		TArray<UPackage*> PackagesCheckedOut;
+		const bool bErrorIfAlreadyCheckedOut = false;
+		const bool bConfirmPackageBranchCheckOutStatus = false;
+
+		if (ISourceControlModule::Get().IsEnabled())
 		{
-			const bool bPromptUserToSave = true;
-			const bool bFastSave = false;
-			const bool bCanBeDeclined = true;
-			const bool bCheckDirty = true;
-			bResult = InternalSavePackages(PackagesToSave, bPromptUserToSave, bFastSave, bCanBeDeclined, bCheckDirty);
+			TArray<UPackage*> PackagesToCheckout;
+			Algo::CopyIf(PackagesToSave, PackagesToCheckout, [](UPackage* Package) { return !FPackageName::IsTempPackage(Package->GetName()); });
+			FEditorFileUtils::CheckoutPackages(PackagesToCheckout, &PackagesCheckedOut, bErrorIfAlreadyCheckedOut, bConfirmPackageBranchCheckOutStatus);
 		}
-		else
-		{
-			const FScopedBusyCursor BusyCursor;
-			// Prevent modal window if not requested.
-			TGuardValue<bool> UnattendedScriptGuard(GIsRunningUnattendedScript, true);
-
-			TArray<UPackage*> PackagesCheckedOut;
-			const bool bErrorIfAlreadyCheckedOut = false;
-			const bool bConfirmPackageBranchCheckOutStatus = false;
 
-			if (ISourceControlModule::Get().IsEnabled())
+		// Cannot mark files for add until after packages saved
+		TArray<UPackage*> PackagesToMarkForAdd;
+		for (UPackage* Package : PackagesToSave)
+		{
+			// List unsaved packages that were not checked out and are not going to be deleted
+			if (!PackagesCheckedOut.Contains(Package) && !UPackage::IsEmptyPackage(Package) && !FPackageName::IsTempPackage(Package->GetName()))
 			{
-				TArray<UPackage*> PackagesToCheckout;
-				Algo::CopyIf(PackagesToSave, PackagesToCheckout, [](UPackage* Package) { return !FPackageName::IsTempPackage(Package->GetName()); });
-				FEditorFileUtils::CheckoutPackages(PackagesToCheckout, &PackagesCheckedOut, bErrorIfAlreadyCheckedOut, bConfirmPackageBranchCheckOutStatus);
+				PackagesToMarkForAdd.Add(Package);
 			}
+		}
 
-			// Cannot mark files for add until after packages saved
-			TArray<UPackage*> PackagesToMarkForAdd;
-			for (UPackage* Package : PackagesToSave)
-			{
-				// List unsaved packages that were not checked out and are not going to be deleted
-				if (!PackagesCheckedOut.Contains(Package) && !UPackage::IsEmptyPackage(Package) && !FPackageName::IsTempPackage(Package->GetName()))
-				{
-					PackagesToMarkForAdd.Add(Package);
-				}
-			}
+		constexpr bool bUseDialog = false;
+		TArray<UPackage*> FailedPackages;
+		FEditorFileUtils::EPromptReturnCode ReturnResponse = InternalPromptForCheckoutAndSave(PackagesToSave, bUseDialog, SaveOutput, FailedPackages);
 
-			TArray<UPackage*> FailedPackages;
-			FEditorFileUtils::EPromptReturnCode ReturnResponse = InternalPromptForCheckoutAndSave(PackagesToSave, bUseDialog, FailedPackages);
+		// Mark files for add now that packages have saved
+		PackagesToMarkForAdd.RemoveAll([&FailedPackages](UPackage* Package) { return FailedPackages.Contains(Package); });
+		if (PackagesToMarkForAdd.Num() > 0 && ISourceControlModule::Get().IsEnabled())
+		{
+			FEditorFileUtils::CheckoutPackages(PackagesToMarkForAdd, nullptr, bErrorIfAlreadyCheckedOut, bConfirmPackageBranchCheckOutStatus);
+		}
 
-			// Mark files for add now that packages have saved
-			PackagesToMarkForAdd.RemoveAll([&FailedPackages](UPackage* Package) { return FailedPackages.Contains(Package); });
-			if (PackagesToMarkForAdd.Num() > 0 && ISourceControlModule::Get().IsEnabled())
-			{
-				FEditorFileUtils::CheckoutPackages(PackagesToMarkForAdd, nullptr, bErrorIfAlreadyCheckedOut, bConfirmPackageBranchCheckOutStatus);
-			}
+		bResult = (ReturnResponse == FEditorFileUtils::PR_Success);
+	}
 
-			bResult = (ReturnResponse == FEditorFileUtils::PR_Success);
-		}
+	return bResult;
+}
+
+static bool InternalCheckoutAndSavePackages(const TArray<UPackage*>& PackagesToSave)
+{
+	bool bResult = true;
+	if (PackagesToSave.Num() > 0)
+	{
+		constexpr bool bUseDialog = false;
+		FScopedSaveOutput SaveOutput(bUseDialog);
+		bResult = InternalCheckoutAndSavePackages(PackagesToSave, SaveOutput);
+	}
+	else
+	{
+		UE_LOG(LogFileHelpers, Log, TEXT("%s"), *LOCTEXT("NoAssetsToSave", "All files are already saved.").ToString());
+	}
+
+	return bResult;
+}
+
+static bool InternalCheckoutAndSavePackagesWithDialog(const TArray<UPackage*>& PackagesToSave)
+{
+	bool bResult = true;
+	if (PackagesToSave.Num() > 0)
+	{
+		const bool bPromptUserToSave = true;
+		const bool bFastSave = false;
+		const bool bCanBeDeclined = true;
+		const bool bCheckDirty = true;
+		bResult = InternalSavePackages(PackagesToSave, bPromptUserToSave, bFastSave, bCanBeDeclined, bCheckDirty);
 	}
 	else
 	{
-		InternalNotifyNoPackagesSaved(bUseDialog);
+		InternalNotifyNoPackagesSaved();
 	}
 
 	return bResult;
@@ -5663,25 +5888,25 @@ static TArray<UPackage*> InternalGetValidPackages(const TArray<UPackage*>& Packa
 bool UEditorLoadingAndSavingUtils::SavePackages(const TArray<UPackage*>& PackagesToSave, bool bOnlyDirty)
 {
 	TArray<UPackage*> Packages = InternalGetValidPackages(PackagesToSave, bOnlyDirty);
-	return InternalCheckoutAndSavePackages(Packages, false);
+	return InternalCheckoutAndSavePackages(Packages);
 }
 
 bool UEditorLoadingAndSavingUtils::SavePackagesWithDialog(const TArray<UPackage*>& PackagesToSave, bool bOnlyDirty)
 {
 	TArray<UPackage*> Packages = InternalGetValidPackages(PackagesToSave, bOnlyDirty);
-	return InternalCheckoutAndSavePackages(Packages, true);
+	return InternalCheckoutAndSavePackagesWithDialog(Packages);
 }
 
 bool UEditorLoadingAndSavingUtils::SaveDirtyPackages(const bool bSaveMapPackages, const bool bSaveContentPackages)
 {
 	TArray<UPackage*> Packages = InternalGetDirtyPackages(bSaveMapPackages, bSaveContentPackages);
-	return InternalCheckoutAndSavePackages(Packages, false);
+	return InternalCheckoutAndSavePackages(Packages);
 }
 
 bool UEditorLoadingAndSavingUtils::SaveDirtyPackagesWithDialog(const bool bSaveMapPackages, const bool bSaveContentPackages)
 {
 	TArray<UPackage*> Packages = InternalGetDirtyPackages(bSaveMapPackages, bSaveContentPackages);
-	return InternalCheckoutAndSavePackages(Packages, true);
+	return InternalCheckoutAndSavePackagesWithDialog(Packages);
 }
 
 bool UEditorLoadingAndSavingUtils::SaveCurrentLevel()
