EditorBulkData: Improve editor startup performance by not waiting for AssetRegistry scanning to give us a cache key for EditorBulkData packages that want to read cached data from DDC about their payload hashes. In the absence of the cached data, we recalculate them on demand when we need them for DDC builds, which is an insignificant cost for the packages that load during engine startup.
#rnx

diff --git a/Engine/Source/Editor/UnrealEd/Private/EditorDomain/BulkDataRegistryImpl.cpp b/Engine/Source/Editor/UnrealEd/Private/EditorDomain/BulkDataRegistryImpl.cpp
--- a/Engine/Source/Editor/UnrealEd/Private/EditorDomain/BulkDataRegistryImpl.cpp
+++ b/Engine/Source/Editor/UnrealEd/Private/EditorDomain/BulkDataRegistryImpl.cpp
@@ -2,6 +2,7 @@
 
 #include "BulkDataRegistryImpl.h"
 
+#include "AssetRegistry/IAssetRegistry.h"
 #include "Compression/CompressedBuffer.h"
 #include "Containers/ArrayView.h"
 #include "Containers/ContainersFwd.h"
@@ -290,8 +291,27 @@ FBulkDataRegistryImpl::TryRegister(UPackage* Owner, const UE::Serialization::FEd
 	}
 	if (bAllowedToReadWritePayloadIdFromCache)
 	{
-		AddPendingPackageBulkData(PackageName, MoveTemp(PendingPackageBulk));
+		AddPendingPackageBulkData(PackageName, MoveTemp(PendingPackageBulk), bAllowedToReadWritePayloadIdFromCache);
+
+		// If The PendingPackage found that it is not allowed to use cache operations, then go back into the BulkData
+		// in the registry and mark that it is not allowed to write its payload ids to the cache. PayloadIds are not
+		// allowed to be written unless we read them first and find they do not exist, and the only place we read them
+		// is from the PendingPackage.
+		if (!bAllowedToReadWritePayloadIdFromCache)
+		{
+			FWriteScopeLock RegistryScopeLock(RegistryLock);
+			if (bActive)
+			{
+				UE::BulkDataRegistry::Private::FRegisteredBulk* RegisteredBulk
+					= Registry.Find(BulkData.GetIdentifier());
+				if (RegisteredBulk)
+				{
+					RegisteredBulk->bAllowedToWritePayloadIdToCache = false;
+				}
+			}
+		}
 	}
+
 	ResaveSizeTracker.Register(Owner, BulkData);
 	return UE::BulkDataRegistry::ERegisterResult::Success;
 }
@@ -592,7 +612,8 @@ void FBulkDataRegistryImpl::Tick(float DeltaTime)
 	TickCook(DeltaTime, false /* bTickComplete */);
 }
 
-void FBulkDataRegistryImpl::AddPendingPackageBulkData(FName PackageName, UE::Serialization::FEditorBulkData&& BulkData)
+void FBulkDataRegistryImpl::AddPendingPackageBulkData(FName PackageName, UE::Serialization::FEditorBulkData&& BulkData,
+	bool& bOutCacheOperationsAllowed)
 {
 	FScopeLock PendingPackageScopeLock(&PendingPackageLock);
 	check(bActive); // Callers should early exit if !bActive
@@ -603,8 +624,10 @@ void FBulkDataRegistryImpl::AddPendingPackageBulkData(FName PackageName, UE::Ser
 	}
 	if (!PendingPackage->IsLoadInProgress())
 	{
+		bOutCacheOperationsAllowed = false;
 		return;
 	}
+	bOutCacheOperationsAllowed = PendingPackage->AreCacheOperationsAllowed();
 	PendingPackage->AddBulkData(MoveTemp(BulkData));
 }
 
@@ -794,8 +817,24 @@ FPendingPackage::FPendingPackage(FName InPackageName, FBulkDataRegistryImpl* InO
 {
 	PendingOperations = Flag_EndLoad | Flag_BulkDataListResults;
 
-	UE::EditorDomain::GetBulkDataList(PackageName, BulkDataListCacheRequest,
-		[this](FSharedBuffer Buffer) { OnBulkDataListResults(Buffer); });
+	// The BulkDataList for the package is keyed by the EditorDomain key for the package, which requires
+	// data from the AssetRegistry and will block on the AssetRegistry's gather of the package. That blocking makes editor
+	// startup too slow, and we want to skip it.
+	// Therefore, for packages that load while the AssetRegistry is still scanning, we skip reading their cache value,
+	// and never read/write their cache. Building DDC jobs with these packages, if they have any PlaceHolderPayloadIds,
+	// will therefore always need to recalculate the real PayloadId by loading the bulk data before the DDC build starts.
+	// TODO: Change this to be more nuanced, by waiting for the AssetRegistry to report the package is loaded and then
+	// sending off this GetBulkDataList request, and finish off the rest of the package.
+	if (!IAssetRegistry::GetChecked().IsGathering())
+	{
+		UE::EditorDomain::GetBulkDataList(PackageName, BulkDataListCacheRequest,
+			[this](FSharedBuffer Buffer) { OnBulkDataListResults(Buffer); });
+	}
+	else
+	{
+		bCacheOperationsAllowed = false;
+		OnBulkDataListResults(FSharedBuffer());
+	}
 }
 
 void FPendingPackage::Cancel()
@@ -817,7 +856,7 @@ void FPendingPackage::OnEndLoad(bool& bOutShouldRemove, bool& bOutShouldWriteCac
 	bLoadInProgress = false;
 	if (PendingOperations.fetch_and(~Flag_EndLoad) == Flag_EndLoad)
 	{
-		bOutShouldWriteCache = true;
+		bOutShouldWriteCache = bCacheOperationsAllowed;
 		bOutShouldRemove = true;
 	}
 	else
@@ -939,10 +978,17 @@ void FPendingPackage::ReadCache(bool& bOutAbort)
 
 void FPendingPackage::WriteCache()
 {
-	// If the BulkDataList cache read found some existing results, then exit; cache results are deterministic so there
-	// is no need to write the list to the cache again.
 	if (CachedBulkDatas.Num() > 0)
 	{
+		// If the BulkDataList cache read found some existing results, then exit; cache results are deterministic so there
+		// is no need to write the list to the cache again.
+		return;
+	}
+	if (!bCacheOperationsAllowed)
+	{
+		// bCacheOperationsAllowed is set to false for packages that are loaded to early to read from the cache.
+		// It might still be too early to write to the cache, but also, we're not allowed to write without
+		// first checking whether the value already exists by reading.
 		return;
 	}
 
diff --git a/Engine/Source/Editor/UnrealEd/Private/EditorDomain/BulkDataRegistryImpl.h b/Engine/Source/Editor/UnrealEd/Private/EditorDomain/BulkDataRegistryImpl.h
--- a/Engine/Source/Editor/UnrealEd/Private/EditorDomain/BulkDataRegistryImpl.h
+++ b/Engine/Source/Editor/UnrealEd/Private/EditorDomain/BulkDataRegistryImpl.h
@@ -102,6 +102,11 @@ public:
 	}
 	void WriteCache();
 
+	bool AreCacheOperationsAllowed() const
+	{
+		return bCacheOperationsAllowed;
+	}
+
 private:
 	void OnBulkDataListResults(FSharedBuffer Buffer);
 	void ReadCache(bool& bOutAbort);
@@ -128,6 +133,12 @@ private:
 	 * to the list when operations occur on the UPackage in the editor.
 	 */
 	bool bLoadInProgress = true;
+	/**
+	 * This value is false for FPendingPackages that were not able to query the cache (e.g. because
+	 * they were created too early in EngineStartup) and therefore are not allowed to write
+	 * the cache either.
+	 */
+	bool bCacheOperationsAllowed = true;
 };
 
 /** Data about a BulkData that has loaded its payload for TryGetMeta and should drop it after GetData or a timeout. */
@@ -229,7 +240,8 @@ public:
 	virtual TStatId GetStatId() const override { return TStatId(); }
 
 private:
-	void AddPendingPackageBulkData(FName PackageName, UE::Serialization::FEditorBulkData&& BulkData);
+	void AddPendingPackageBulkData(FName PackageName, UE::Serialization::FEditorBulkData&& BulkData,
+		bool& bOutCacheOperationsAllowed);
 	void PollPendingPackages(bool bWaitForCooldown);
 	void AddTempLoadedPayload(const FGuid& RegistryKey, uint64 PayloadSize);
 	void PruneTempLoadedPayloads();
