Make sure we relink properties after moving a class to a new super that is a skeleton class

[RN] Fixed a crash in garbage collection due to mislinked SKEL_ classes, most frequently seen during long running cook operations

diff --git a/Engine/Source/Editor/UnrealEd/Private/Kismet2/KismetReinstanceUtilities.cpp b/Engine/Source/Editor/UnrealEd/Private/Kismet2/KismetReinstanceUtilities.cpp
--- a/Engine/Source/Editor/UnrealEd/Private/Kismet2/KismetReinstanceUtilities.cpp
+++ b/Engine/Source/Editor/UnrealEd/Private/Kismet2/KismetReinstanceUtilities.cpp
@@ -2014,6 +2014,8 @@ void FBlueprintCompileReinstancer::MoveDependentSkelToReinst(UClass* const Owner
 			if (UClass* const* NewSuper = NewSkeletonToOldSkeleton.Find(CurClass->GetSuperClass()))
 			{
 				CurClass->SetSuperStruct(*NewSuper);
+				CurClass->Bind();
+				CurClass->StaticLink(true);
 			}
 			continue;
 		}
