Cooker: CookOnTheFly: Fix assertion that occurs when connecting a new client (e.g. WindowsServer) to the CookCommandlet in cookonthefly mode when it is processing requests for an existing client (e.g. WindowsClient).
#rnx

diff --git a/Engine/Source/Editor/UnrealEd/Classes/CookOnTheSide/CookOnTheFlyServer.h b/Engine/Source/Editor/UnrealEd/Classes/CookOnTheSide/CookOnTheFlyServer.h
--- a/Engine/Source/Editor/UnrealEd/Classes/CookOnTheSide/CookOnTheFlyServer.h
+++ b/Engine/Source/Editor/UnrealEd/Classes/CookOnTheSide/CookOnTheFlyServer.h
@@ -600,6 +600,7 @@ private:
 	 */
 	void OnRemoveSessionPlatform(const ITargetPlatform* TargetPlatform, int32 RemovedIndex);
 	/** Update structures that have data per platfrom when COTF adds a new SessionPlatform. */
+	void OnBeforePlatformAddedToSession(const ITargetPlatform* TargetPlatform);
 	void OnPlatformAddedToSession(const ITargetPlatform* TargetPlatform);
 
 	void InitializePollables();
diff --git a/Engine/Source/Editor/UnrealEd/Private/CookOnTheFlyServer.cpp b/Engine/Source/Editor/UnrealEd/Private/CookOnTheFlyServer.cpp
--- a/Engine/Source/Editor/UnrealEd/Private/CookOnTheFlyServer.cpp
+++ b/Engine/Source/Editor/UnrealEd/Private/CookOnTheFlyServer.cpp
@@ -3497,6 +3497,16 @@ void UCookOnTheFlyServer::OnRemoveSessionPlatform(const ITargetPlatform* TargetP
 	WorkerRequests->OnRemoveSessionPlatform(TargetPlatform);
 }
 
+void UCookOnTheFlyServer::OnBeforePlatformAddedToSession(const ITargetPlatform* TargetPlatform)
+{
+	using namespace UE::Cook;
+
+	for (TUniquePtr<FRequestCluster>& Cluster : PackageDatas->GetRequestQueue().GetRequestClusters())
+	{
+		Cluster->OnBeforePlatformAddedToSession(TargetPlatform);
+	}
+}
+
 void UCookOnTheFlyServer::OnPlatformAddedToSession(const ITargetPlatform* TargetPlatform)
 {
 	using namespace UE::Cook;
@@ -3505,10 +3515,6 @@ void UCookOnTheFlyServer::OnPlatformAddedToSession(const ITargetPlatform* Target
 	{
 		Element.ReachablePlatforms.OnPlatformAddedToSession(TargetPlatform);
 	}
-	for (TUniquePtr<FRequestCluster>& Cluster : PackageDatas->GetRequestQueue().GetRequestClusters())
-	{
-		Cluster->OnPlatformAddedToSession(TargetPlatform);
-	}
 }
 
 void UCookOnTheFlyServer::TickNetwork()
diff --git a/Engine/Source/Editor/UnrealEd/Private/Cooker/CookPlatformManager.cpp b/Engine/Source/Editor/UnrealEd/Private/Cooker/CookPlatformManager.cpp
--- a/Engine/Source/Editor/UnrealEd/Private/Cooker/CookPlatformManager.cpp
+++ b/Engine/Source/Editor/UnrealEd/Private/Cooker/CookPlatformManager.cpp
@@ -151,18 +151,34 @@ void FPlatformManager::AddSessionPlatform(UCookOnTheFlyServer& COTFS, FPlatformI
 	checkf(IsSchedulerThread(), TEXT("Writing to SessionPlatforms is only allowed from the scheduler thread."));
 	CreatePlatformData(TargetPlatform);
 
-	FRWScopeLock ScopeLock(SessionLock, FRWScopeLockType::SLT_Write);
-	if (SessionPlatforms.Contains(TargetPlatform))
 	{
-		return;
+		FRWScopeLock ScopeLock(SessionLock, FRWScopeLockType::SLT_ReadOnly);
+		if (SessionPlatforms.Contains(TargetPlatform))
+		{
+			return;
+		}
 	}
-	SessionPlatforms.Add(TargetPlatform);
+
+	if (COTFS.bSessionRunning)
+	{
+		COTFS.OnBeforePlatformAddedToSession(TargetPlatform);
+	}
+
+	{
+		FRWScopeLock ScopeLock(SessionLock, FRWScopeLockType::SLT_Write);
+		if (SessionPlatforms.Contains(TargetPlatform))
+		{
+			return;
+		}
+		SessionPlatforms.Add(TargetPlatform);
+		COTFS.bSessionRunning = true;
+		bHasSelectedSessionPlatforms = true;
+	}
+
 	if (COTFS.bSessionRunning)
 	{
 		COTFS.OnPlatformAddedToSession(TargetPlatform);
 	}
-	COTFS.bSessionRunning = true;
-	bHasSelectedSessionPlatforms = true;
 }
 
 FPlatformData* FPlatformManager::GetPlatformData(FPlatformId Platform)
diff --git a/Engine/Source/Editor/UnrealEd/Private/Cooker/CookRequestCluster.cpp b/Engine/Source/Editor/UnrealEd/Private/Cooker/CookRequestCluster.cpp
--- a/Engine/Source/Editor/UnrealEd/Private/Cooker/CookRequestCluster.cpp
+++ b/Engine/Source/Editor/UnrealEd/Private/Cooker/CookRequestCluster.cpp
@@ -637,32 +637,26 @@ void FRequestCluster::OnNewReachablePlatforms(FPackageData* PackageData)
 	}
 }
 
-void FRequestCluster::OnPlatformAddedToSession(const ITargetPlatform* TargetPlatform)
+void FRequestCluster::OnBeforePlatformAddedToSession(const ITargetPlatform* TargetPlatform)
 {
-	if (GraphSearch.IsStarted())
+	FCookerTimer CookerTimer(FCookerTimer::Forever);
+	bool bComplete;
+	while (Process(CookerTimer, bComplete), !bComplete)
 	{
-		FCookerTimer CookerTimer(FCookerTimer::Forever);
-		bool bComplete;
-		while (PumpExploration(CookerTimer, bComplete), !bComplete)
-		{
-			UE_LOG(LogCook, Display, TEXT("Waiting for RequestCluster to finish before adding platform to session."));
-			FPlatformProcess::Sleep(.001f);
-		}
+		UE_LOG(LogCook, Display, TEXT("Waiting for RequestCluster to finish before adding platform to session."));
+		FPlatformProcess::Sleep(.001f);
 	}
 }
 
 void FRequestCluster::OnRemoveSessionPlatform(const ITargetPlatform* TargetPlatform)
 {
-	if (GraphSearch.IsStarted())
+	FCookerTimer CookerTimer(FCookerTimer::Forever);
+	bool bComplete;
+	while (Process(CookerTimer, bComplete), !bComplete)
 	{
-		FCookerTimer CookerTimer(FCookerTimer::Forever);
-		bool bComplete;
-		while (PumpExploration(CookerTimer, bComplete), !bComplete)
-		{
-			UE_LOG(LogCook, Display,
-				TEXT("Waiting for RequestCluster to finish before removing platform from session."));
-			FPlatformProcess::Sleep(.001f);
-		}
+		UE_LOG(LogCook, Display,
+			TEXT("Waiting for RequestCluster to finish before removing platform from session."));
+		FPlatformProcess::Sleep(.001f);
 	}
 }
 
diff --git a/Engine/Source/Editor/UnrealEd/Private/Cooker/CookRequestCluster.h b/Engine/Source/Editor/UnrealEd/Private/Cooker/CookRequestCluster.h
--- a/Engine/Source/Editor/UnrealEd/Private/Cooker/CookRequestCluster.h
+++ b/Engine/Source/Editor/UnrealEd/Private/Cooker/CookRequestCluster.h
@@ -76,7 +76,7 @@ public:
 	/** PackageData container interface: remove the PackageData from this container. */
 	void RemovePackageData(FPackageData* PackageData);
 	void OnNewReachablePlatforms(FPackageData* PackageData);
-	void OnPlatformAddedToSession(const ITargetPlatform* TargetPlatform);
+	void OnBeforePlatformAddedToSession(const ITargetPlatform* TargetPlatform);
 	void OnRemoveSessionPlatform(const ITargetPlatform* TargetPlatform);
 	void RemapTargetPlatforms(TMap<ITargetPlatform*, ITargetPlatform*>& Remap);
 
