A few config variables are queried and then cached during a cook. Due to the config tracking, only the first package has those variables in its dependencies. All the other packages are using the cached value s so the config tracking mecanism doesn't detect it.
Compute a hash per platform representing the global dependencies and containing the config variables being cached plus any other value needed. This global hash is then added to the key hash of every cooked package.

- CalculateGlobalDependenciesHash : new function to calculate, per platform, a hash representing the global dependencies.
- UCookOnTheFlyServer : when starting the cooker, calculate the global hash dependencies for every platform.
- FBuildDependencySet::TryCalculateCurrentKey : Add the global dependencies hash to the key hash. No need to add CookIncrementalVersion, it is already inside the global hash.
- Add FIgnoreScope to disable the config tracking of a the config variables included in the global hash.
- Remove the old definition of FIgnoreScope from CookConfigAccessTracker.h
- Add empty constructor/destructor in FIgnoreScope when UE_WITH_CONFIG_TRACKING is 0. This prevents a compilation error "unreferenced local variable" in ShaderCompileWorker due to FIgnoreScope used in CoreUObject.

#rnx
#incrementalcook

diff --git a/Engine/Source/Editor/UnrealEd/Classes/CookOnTheSide/CookOnTheFlyServer.h b/Engine/Source/Editor/UnrealEd/Classes/CookOnTheSide/CookOnTheFlyServer.h
--- a/Engine/Source/Editor/UnrealEd/Classes/CookOnTheSide/CookOnTheFlyServer.h
+++ b/Engine/Source/Editor/UnrealEd/Classes/CookOnTheSide/CookOnTheFlyServer.h
@@ -49,6 +49,7 @@ struct FGenericMemoryStats;
 struct FPropertyChangedEvent;
 struct FResourceSizeEx;
 
+namespace UE::Cook { void CalculateGlobalDependenciesHash(const ITargetPlatform* Platform, const UCookOnTheFlyServer& COTFS); }
 namespace UE::LinkerLoad { enum class EImportBehavior : uint8; }
 
 enum class ECookInitializationFlags
@@ -1016,6 +1017,10 @@ private:
 	FBeginCookContext CreateBeginCookByTheBookContext(const FCookByTheBookStartupOptions& StartupOptions);
 	/** Set the PlatformManager's session platforms and finish filling out the BeginContext for StartCookByTheBook */
 	void SelectSessionPlatforms(FBeginCookContext& BeginContext);
+
+	/** Calculate the global dependencies hashes for every platform being cooked. This needs to be called after SelectSessionPlatforms.*/
+	void CalculateGlobalDependenciesHashes();
+
 	/** CollectFilesToCook and add them as external requests. */
 	void GenerateInitialRequests(FBeginCookContext& BeginContext);
 	/** If cooking DLC, initialize cooked PackageDatas for all of the packages cooked by the base release. */
@@ -1741,4 +1746,6 @@ private:
 	friend UE::Cook::FSoftGCHistory;
 	friend UE::Cook::FWorkerRequestsLocal;
 	friend UE::Cook::FWorkerRequestsRemote;
+
+	friend void UE::Cook::CalculateGlobalDependenciesHash(const ITargetPlatform* Platform, const UCookOnTheFlyServer& COTFS);
 };
diff --git a/Engine/Source/Editor/UnrealEd/Private/CookOnTheFlyServer.cpp b/Engine/Source/Editor/UnrealEd/Private/CookOnTheFlyServer.cpp
--- a/Engine/Source/Editor/UnrealEd/Private/CookOnTheFlyServer.cpp
+++ b/Engine/Source/Editor/UnrealEd/Private/CookOnTheFlyServer.cpp
@@ -27,6 +27,7 @@
 #include "Cooker/CookDirector.h"
 #include "Cooker/CookGarbageCollect.h"
 #include "Cooker/CookGenerationHelper.h"
+#include "Cooker/CookGlobalDependencies.h"
 #include "Cooker/CookImportsChecker.h"
 #include "Cooker/CookLogPrivate.h"
 #include "Cooker/CookOnTheFlyServerInterface.h"
@@ -829,6 +830,9 @@ void UCookOnTheFlyServer::StartCookOnTheFlySessionFromGameThread(ITargetPlatform
 		InitializeSession();
 	}
 	PlatformManager->AddSessionPlatform(*this, TargetPlatform);
+
+	CalculateGlobalDependenciesHashes();
+
 	TargetPlatform->InitializeForCook();
 	ResetCook({ TPair<const ITargetPlatform*,bool>{TargetPlatform, true /* bResetResults */} });
 
@@ -11145,6 +11149,7 @@ void UCookOnTheFlyServer::StartCookByTheBook( const FCookByTheBookStartupOptions
 	CreateSandboxFile(BeginContext);
 	LoadBeginCookConfigSettings(BeginContext);
 	SelectSessionPlatforms(BeginContext);
+	CalculateGlobalDependenciesHashes();
 	LoadBeginCookIncrementalFlags(BeginContext);
 
 	// Initialize systems referenced by later stages or that need to start early for async performance
@@ -11416,6 +11421,7 @@ void UCookOnTheFlyServer::StartCookAsCookWorker()
 	CreateSandboxFile(BeginContext);
 	LoadBeginCookConfigSettings(BeginContext);
 	SelectSessionPlatforms(BeginContext);
+	CalculateGlobalDependenciesHashes();
 	LoadBeginCookIncrementalFlags(BeginContext);
 	if (IsDirectorCookOnTheFly())
 	{
@@ -12085,6 +12091,15 @@ void UCookOnTheFlyServer::SelectSessionPlatforms(FBeginCookContext& BeginContext
 	}
 }
 
+void UCookOnTheFlyServer::CalculateGlobalDependenciesHashes()
+{
+	const TArray<const ITargetPlatform*>& Platforms = PlatformManager->GetSessionPlatforms();
+	for (const ITargetPlatform* Platform : Platforms)
+	{
+		UE::Cook::CalculateGlobalDependenciesHash(Platform, *this);
+	}
+}
+
 void UCookOnTheFlyServer::BeginCookEditorSystems()
 {
 	if (!IsCookingInEditor())
diff --git a/Engine/Source/Editor/UnrealEd/Private/Cooker/CookConfigAccessTracker.h b/Engine/Source/Editor/UnrealEd/Private/Cooker/CookConfigAccessTracker.h
--- a/Engine/Source/Editor/UnrealEd/Private/Cooker/CookConfigAccessTracker.h
+++ b/Engine/Source/Editor/UnrealEd/Private/Cooker/CookConfigAccessTracker.h
@@ -128,17 +128,4 @@ public:
 
 }
 
-#else // !UE_WITH_CONFIG_TRACKING
-
-namespace UE::ConfigAccessTracking
-{
-
-struct FIgnoreScope
-{
-	FIgnoreScope() {}
-	~FIgnoreScope() {}
-};
-
-} // namespace UE::ConfigAccessTracking
-
-#endif // else !UE_WITH_CONFIG_TRACKING
+#endif // #if UE_WITH_CONFIG_TRACKING
diff --git a/Engine/Source/Editor/UnrealEd/Private/Cooker/CookGlobalDependencies.cpp b/Engine/Source/Editor/UnrealEd/Private/Cooker/CookGlobalDependencies.cpp
new file mode 100644
--- /dev/null
+++ b/Engine/Source/Editor/UnrealEd/Private/Cooker/CookGlobalDependencies.cpp
@@ -0,0 +1,76 @@
+// Copyright Epic Games, Inc. All Rights Reserved.
+
+#include "Cooker/CookGlobalDependencies.h"
+
+#include "CookOnTheSide/CookLog.h"
+#include "CookOnTheSide/CookOnTheFlyServer.h"
+#include "CoreGlobals.h"
+#include "Hash/Blake3.h"
+#include "Interfaces/ITargetPlatform.h"
+#include "Misc/ConfigCacheIni.h"
+#include "Misc/Guid.h"
+
+namespace UE::Cook
+{
+	extern FGuid CookIncrementalVersion;
+	
+	// This hash is calculated per platform from different settings and guids
+	static TMap<const ITargetPlatform*, FBlake3Hash> GlobalHash;
+
+	void CalculateGlobalDependenciesHash(const ITargetPlatform* Platform, const UCookOnTheFlyServer& COTFS)
+	{
+		UE::ConfigAccessTracking::FIgnoreScope Ignore;
+
+		FBlake3 Hasher;
+
+		Hasher.Update(&CookIncrementalVersion, sizeof(CookIncrementalVersion));
+
+		{
+			bool bValue = false;
+			GConfig->GetBool(TEXT("Core.System"), TEXT("CanStripEditorOnlyExportsAndImports"), bValue, GEngineIni);
+			Hasher.Update(&bValue, sizeof(bValue));
+		}
+
+		{
+			bool bValue = false;
+			GConfig->GetBool(TEXT("Core.System"), TEXT("CanSkipEditorReferencedPackagesWhenCooking"), bValue, GEngineIni);
+			Hasher.Update(&bValue, sizeof(bValue));
+		}
+
+		{
+			bool bValue = false;
+			FConfigFile TargetIni;
+			FConfigCacheIni::LoadLocalIniFile(TargetIni, TEXT("Engine"), true, *Platform->IniPlatformName());
+			TargetIni.GetBool(TEXT("Core.System"), TEXT("CanUseUnversionedPropertySerialization"), bValue);
+			Hasher.Update(&bValue, sizeof(bValue));
+		}
+
+		{
+			TArray<FString> AssetList;
+			GConfig->GetArray(TEXT("CookSettings"), TEXT("AutomaticOptionalInclusionAssetType"), AssetList, GEditorIni);
+			for (const FString& AssetType : AssetList)
+			{
+				Hasher.Update(*AssetType, AssetType.NumBytesWithoutNull());
+			}
+		}
+
+		{
+			bool bIsUnversioned = COTFS.IsCookFlagSet(ECookInitializationFlags::Unversioned);
+			Hasher.Update(&bIsUnversioned, sizeof(bIsUnversioned));
+		}
+
+		GlobalHash.Add(Platform, Hasher.Finalize());
+	}
+
+	FBlake3Hash GetGlobalDependenciesHash(const ITargetPlatform* Platform)
+	{
+		FBlake3Hash* Hash = GlobalHash.Find(Platform);
+		if (!Hash)
+		{
+			UE_LOG(LogCook, Error, TEXT("CalculateGlobalDependenciesHash was not called for platform %s"), *Platform->PlatformName());
+			return FBlake3Hash();
+		}
+
+		return *Hash;
+	}
+}
diff --git a/Engine/Source/Editor/UnrealEd/Private/Cooker/CookGlobalDependencies.h b/Engine/Source/Editor/UnrealEd/Private/Cooker/CookGlobalDependencies.h
new file mode 100644
--- /dev/null
+++ b/Engine/Source/Editor/UnrealEd/Private/Cooker/CookGlobalDependencies.h
@@ -0,0 +1,16 @@
+// Copyright Epic Games, Inc. All Rights Reserved.
+
+struct FBlake3Hash;
+class ITargetPlatform;
+class UCookOnTheFlyServer;
+
+namespace UE::Cook
+{
+	/** 
+	 * Calculate a hash representing the dependencies applied to every package.
+	 * Some dependencies are per platform so this function needs to be called per platform cooked.
+	 */
+	void CalculateGlobalDependenciesHash(const ITargetPlatform* Platform, const UCookOnTheFlyServer& COTFS);
+
+	FBlake3Hash GetGlobalDependenciesHash(const ITargetPlatform* Platform);
+}
diff --git a/Engine/Source/Editor/UnrealEd/Private/Cooker/CookPackageArtifacts.cpp b/Engine/Source/Editor/UnrealEd/Private/Cooker/CookPackageArtifacts.cpp
--- a/Engine/Source/Editor/UnrealEd/Private/Cooker/CookPackageArtifacts.cpp
+++ b/Engine/Source/Editor/UnrealEd/Private/Cooker/CookPackageArtifacts.cpp
@@ -10,6 +10,7 @@
 #include "Cooker/CookDependency.h"
 #include "Cooker/CookEvents.h"
 #include "Cooker/CookGenerationHelper.h"
+#include "Cooker/CookGlobalDependencies.h"
 #include "Cooker/PackageBuildDependencyTracker.h"
 #include "CookOnTheSide/CookLog.h"
 #include "CookPackageSplitter.h"
@@ -85,8 +86,11 @@ FBuildDependencySet::ECurrentKeyResult FBuildDependencySet::TryCalculateCurrentK
 		if (OutMessages) OutMessages->Emplace(ELogVerbosity::Error, TEXT("EditorDomain is unavailable."));
 		return ECurrentKeyResult::Error;
 	}
+	
 	FBlake3 KeyBuilder;
-	KeyBuilder.Update(&CookIncrementalVersion, sizeof(CookIncrementalVersion));
+
+	FBlake3Hash GlobalDependencies = GetGlobalDependenciesHash(TargetPlatform);
+	KeyBuilder.Update(&GlobalDependencies, sizeof(GlobalDependencies));
 
 	ECurrentKeyResult Result = ECurrentKeyResult::Success;
 	FCookDependencyContext Context(&KeyBuilder,
diff --git a/Engine/Source/Runtime/Core/Public/Misc/ConfigAccessTracking.h b/Engine/Source/Runtime/Core/Public/Misc/ConfigAccessTracking.h
--- a/Engine/Source/Runtime/Core/Public/Misc/ConfigAccessTracking.h
+++ b/Engine/Source/Runtime/Core/Public/Misc/ConfigAccessTracking.h
@@ -160,6 +160,8 @@ struct FSection;
 // Define FIgnoreScope even when !UE_WITH_CONFIG_TRACKING so we don't need to wrap its use in a macro.
 struct FIgnoreScope
 {
+	FIgnoreScope(){}
+	~FIgnoreScope(){}
 };
 
 } // namespace UE::ConfigAccessTracking
diff --git a/Engine/Source/Runtime/CoreUObject/Private/Serialization/UnversionedPropertySerialization.cpp b/Engine/Source/Runtime/CoreUObject/Private/Serialization/UnversionedPropertySerialization.cpp
--- a/Engine/Source/Runtime/CoreUObject/Private/Serialization/UnversionedPropertySerialization.cpp
+++ b/Engine/Source/Runtime/CoreUObject/Private/Serialization/UnversionedPropertySerialization.cpp
@@ -5,6 +5,7 @@
 #include "Serialization/UnversionedPropertySerializationTest.h"
 #include "Hash/Blake3.h"
 #include "Interfaces/ITargetPlatform.h"
+#include "Misc/ConfigAccessTracking.h"
 #include "Misc/ScopeRWLock.h"
 #include "UObject/OverridableManager.h"
 #include "UObject/PropertyOptional.h"
@@ -910,6 +911,8 @@ GUPSIniValueCache;
 
 bool CanUseUnversionedPropertySerialization(const ITargetPlatform* Target)
 {
+	UE::ConfigAccessTracking::FIgnoreScope IgnoreScope;
+
 	if (!Target)
 	{
 		// Use current platform settings
diff --git a/Engine/Source/Runtime/CoreUObject/Private/UObject/SavePackage/SaveContext.cpp b/Engine/Source/Runtime/CoreUObject/Private/UObject/SavePackage/SaveContext.cpp
--- a/Engine/Source/Runtime/CoreUObject/Private/UObject/SavePackage/SaveContext.cpp
+++ b/Engine/Source/Runtime/CoreUObject/Private/UObject/SavePackage/SaveContext.cpp
@@ -317,6 +317,8 @@ void FSaveContext::SetupHarvestingRealms()
 	// Generate the automatic optional context inclusion asset list
 	static bool bAssetListGenerated = [](TArray<TSoftClassPtr<UObject>>& OutAssetList)
 		{
+			UE::ConfigAccessTracking::FIgnoreScope IgnoreScope;
+
 			TArray<FString> AssetList;
 			GConfig->GetArray(TEXT("CookSettings"), TEXT("AutomaticOptionalInclusionAssetType"), AssetList, GEditorIni);
 			for (const FString& AssetType : AssetList)
diff --git a/Engine/Source/Runtime/CoreUObject/Private/UObject/SavePackage/SavePackageUtilities.cpp b/Engine/Source/Runtime/CoreUObject/Private/UObject/SavePackage/SavePackageUtilities.cpp
--- a/Engine/Source/Runtime/CoreUObject/Private/UObject/SavePackage/SavePackageUtilities.cpp
+++ b/Engine/Source/Runtime/CoreUObject/Private/UObject/SavePackage/SavePackageUtilities.cpp
@@ -641,6 +641,7 @@ bool CanStripEditorOnlyImportsAndExports()
 		FCanStripEditorOnlyExportsAndImports()
 			: bCanStripEditorOnlyObjects(true)
 		{
+			UE::ConfigAccessTracking::FIgnoreScope IgnoreScope;
 			GConfig->GetBool(TEXT("Core.System"), TEXT("CanStripEditorOnlyExportsAndImports"), bCanStripEditorOnlyObjects, GEngineIni);
 		}
 		FORCEINLINE operator bool() const { return bCanStripEditorOnlyObjects; }
@@ -1299,6 +1300,9 @@ FCanSkipEditorReferencedPackagesWhenCooking::FCanSkipEditorReferencedPackagesWhe
 	{
 		bWarned = true;
 		bool bResult = true;
+
+		UE::ConfigAccessTracking::FIgnoreScope IgnoreScope;
+
 		GConfig->GetBool(TEXT("Core.System"), TEXT("CanSkipEditorReferencedPackagesWhenCooking"), bResult, GEngineIni);
 		if (bResult)
 		{
