Make the function CanUseUnversionedPropertySerialization internal so I can call it from CookGlobalDependencies.

- New file Internal/Serialization/UnversionedPropertySerializationInternal.h containing the declaration of CanUseUnversionedPropertySerialization.
- Remove the declaration of  CanUseUnversionedPropertySerialization from UnversionedPropertySerializationPrivate.h.

#rnx

diff --git a/Engine/Source/Editor/UnrealEd/Private/Cooker/CookGlobalDependencies.cpp b/Engine/Source/Editor/UnrealEd/Private/Cooker/CookGlobalDependencies.cpp
--- a/Engine/Source/Editor/UnrealEd/Private/Cooker/CookGlobalDependencies.cpp
+++ b/Engine/Source/Editor/UnrealEd/Private/Cooker/CookGlobalDependencies.cpp
@@ -9,6 +9,7 @@
 #include "Interfaces/ITargetPlatform.h"
 #include "Misc/ConfigCacheIni.h"
 #include "Misc/Guid.h"
+#include "Serialization/UnversionedPropertySerializationInternal.h"
 
 namespace UE::Cook
 {
@@ -38,11 +39,8 @@ namespace UE::Cook
 		}
 
 		{
-			bool bValue = false;
-			FConfigFile TargetIni;
-			FConfigCacheIni::LoadLocalIniFile(TargetIni, TEXT("Engine"), true, *Platform->IniPlatformName());
-			TargetIni.GetBool(TEXT("Core.System"), TEXT("CanUseUnversionedPropertySerialization"), bValue);
-			Hasher.Update(&bValue, sizeof(bValue));
+			bool bCanUseUnversionedPropertySerialization = CanUseUnversionedPropertySerialization(Platform);
+			Hasher.Update(&bCanUseUnversionedPropertySerialization, sizeof(bCanUseUnversionedPropertySerialization));
 		}
 
 		{
diff --git a/Engine/Source/Runtime/CoreUObject/Internal/Serialization/UnversionedPropertySerializationInternal.h b/Engine/Source/Runtime/CoreUObject/Internal/Serialization/UnversionedPropertySerializationInternal.h
new file mode 100644
--- /dev/null
+++ b/Engine/Source/Runtime/CoreUObject/Internal/Serialization/UnversionedPropertySerializationInternal.h
@@ -0,0 +1,10 @@
+// Copyright Epic Games, Inc. All Rights Reserved.
+
+#pragma once
+
+#include "HAL/Platform.h"
+
+class ITargetPlatform;
+
+// Check if unversioned property serialization is configured to be used on target platform
+COREUOBJECT_API bool CanUseUnversionedPropertySerialization(const ITargetPlatform* Target);
diff --git a/Engine/Source/Runtime/CoreUObject/Private/Serialization/UnversionedPropertySerialization.h b/Engine/Source/Runtime/CoreUObject/Private/Serialization/UnversionedPropertySerialization.h
--- a/Engine/Source/Runtime/CoreUObject/Private/Serialization/UnversionedPropertySerialization.h
+++ b/Engine/Source/Runtime/CoreUObject/Private/Serialization/UnversionedPropertySerialization.h
@@ -4,10 +4,9 @@
 
 #include "UObject/UnrealType.h"
 
-struct FBlake3Hash;
+#include "Serialization/UnversionedPropertySerializationInternal.h"
 
-// Check if unversioned property serialization is configured to be used on target platform
-bool CanUseUnversionedPropertySerialization(const ITargetPlatform* Target);
+struct FBlake3Hash;
 
 // Serialize sparse unversioned properties for a particular struct
 void SerializeUnversionedProperties(const UStruct* Struct, FStructuredArchive::FSlot Slot, uint8* Data, UStruct* DefaultsStruct, uint8* Defaults);
