Interchange:
- Set 'Collision Generated' to true when importing a static mesh and collision is generated

diff --git a/Engine/Source/Editor/UnrealEd/Private/GeomFitUtils.cpp b/Engine/Source/Editor/UnrealEd/Private/GeomFitUtils.cpp
--- a/Engine/Source/Editor/UnrealEd/Private/GeomFitUtils.cpp
+++ b/Engine/Source/Editor/UnrealEd/Private/GeomFitUtils.cpp
@@ -111,7 +111,7 @@ int32 GenerateKDopAsSimpleCollision(UStaticMesh* StaticMesh, const TArray<FVecto
 	ConvexElem.VertexData = HullVertices;
 	// Note: UpdateElemBox also computes the convex hull indices
 	ConvexElem.UpdateElemBox();
-
+	ConvexElem.bIsGenerated = true;
 	bs->AggGeom.ConvexElems.Add(ConvexElem);
 	
 	// create all body instances
@@ -164,6 +164,7 @@ int32 GenerateBoxAsSimpleCollision(UStaticMesh* StaticMesh, bool bUpdate)
 	BoxElem.X = Extents.X * 2.0f;
 	BoxElem.Y = Extents.Y * 2.0f;
 	BoxElem.Z = Extents.Z * 2.0f;
+	BoxElem.bIsGenerated = true;
 	bs->AggGeom.BoxElems.Add(BoxElem);
 
 	// refresh collision change back to staticmesh components
@@ -366,6 +367,7 @@ int32 GenerateSphereAsSimpleCollision(UStaticMesh* StaticMesh, bool bUpdate)
 	FKSphereElem SphereElem;
 	SphereElem.Center = bestSphere.Center;
 	SphereElem.Radius = bestSphere.W;
+	SphereElem.bIsGenerated = true;
 	bs->AggGeom.SphereElems.Add(SphereElem);
 
 	// refresh collision change back to staticmesh components
@@ -518,6 +520,7 @@ int32 GenerateSphylAsSimpleCollision(UStaticMesh* StaticMesh, bool bUpdate)
 	SphylElem.Rotation = rotation;
 	SphylElem.Radius = sphere.W;
 	SphylElem.Length = length;
+	SphylElem.bIsGenerated = true;
 	bs->AggGeom.SphylElems.Add(SphylElem);
 
 	// refresh collision change back to staticmesh components
