CIS fixes (Non-Unity build and deprecated member).

diff --git a/Engine/Source/Editor/UnrealEd/Public/CookPackageSplitter.h b/Engine/Source/Editor/UnrealEd/Public/CookPackageSplitter.h
--- a/Engine/Source/Editor/UnrealEd/Public/CookPackageSplitter.h
+++ b/Engine/Source/Editor/UnrealEd/Public/CookPackageSplitter.h
@@ -143,8 +143,6 @@ public:
 		FString GeneratedRootPath;
 		/** Generated package relative to <GeneratedRootPath>/_Generated_. */
 		FString RelativePath;
-		UE_DEPRECATED(5.3, "Write to PackageDependencies instead")
-		TArray<FName> Dependencies;
 		/**
 		 * Source packages outside of the generator package that will be incorporated into the generated
 		 * package (e.g. ExternalActor packages). These are used to construct the PackageSavedHash for the generated
