Remove the LOCTEXT_NAMESPACE Macro from DecoratedDragDrop since it's not used.

diff --git a/Engine/Source/Editor/UnrealEd/Public/DragAndDrop/DecoratedDragDropOp.h b/Engine/Source/Editor/UnrealEd/Public/DragAndDrop/DecoratedDragDropOp.h
--- a/Engine/Source/Editor/UnrealEd/Public/DragAndDrop/DecoratedDragDropOp.h
+++ b/Engine/Source/Editor/UnrealEd/Public/DragAndDrop/DecoratedDragDropOp.h
@@ -11,8 +11,6 @@
 #include "Widgets/Text/STextBlock.h"
 #include "Styling/AppStyle.h"
 
-#define LOCTEXT_NAMESPACE "DecoratedDragDrop"
-
 class FDecoratedDragDropOp : public FDragDropOperation
 {
 public:
@@ -125,6 +123,3 @@ protected:
 	/** Default color and opacity for the default icon to be displayed. */
 	FSlateColor							DefaultHoverIconColorAndOpacity;
 };
-
-
-#undef LOCTEXT_NAMESPACE
