[EditorEngine]
- Make some delegates thread-safe

diff --git a/Engine/Source/Editor/UnrealEd/Classes/Editor/EditorEngine.h b/Engine/Source/Editor/UnrealEd/Classes/Editor/EditorEngine.h
--- a/Engine/Source/Editor/UnrealEd/Classes/Editor/EditorEngine.h
+++ b/Engine/Source/Editor/UnrealEd/Classes/Editor/EditorEngine.h
@@ -717,21 +717,21 @@ public:
 	class FUObjectAnnotationSparseBool ObjectsThatExistInEditorWorld;
 
 	/** Called prior to a Blueprint compile */
-	DECLARE_EVENT_OneParam( UEditorEngine, FBlueprintPreCompileEvent, UBlueprint* );
+	DECLARE_TS_MULTICAST_DELEGATE_OneParam( FBlueprintPreCompileEvent, UBlueprint* );
 	FBlueprintPreCompileEvent& OnBlueprintPreCompile() { return BlueprintPreCompileEvent; }
 
 	/** Broadcasts that a Blueprint is about to be compiled */
 	void BroadcastBlueprintPreCompile(UBlueprint* BlueprintToCompile) { BlueprintPreCompileEvent.Broadcast(BlueprintToCompile); }
 
 	/** Called when a Blueprint compile is completed. */
-	DECLARE_EVENT( UEditorEngine, FBlueprintCompiledEvent );
+	DECLARE_TS_MULTICAST_DELEGATE( FBlueprintCompiledEvent );
 	FBlueprintCompiledEvent& OnBlueprintCompiled() { return BlueprintCompiledEvent; }
 
 	/**	Broadcasts that a blueprint just finished compiling. THIS SHOULD NOT BE PUBLIC */
 	void BroadcastBlueprintCompiled() { BlueprintCompiledEvent.Broadcast(); }
 
 	/** Called by the blueprint compiler after a blueprint has been compiled and all instances replaced, but prior to garbage collection. */
-	DECLARE_EVENT(UEditorEngine, FBlueprintReinstanced);
+	DECLARE_TS_MULTICAST_DELEGATE(FBlueprintReinstanced);
 	FBlueprintReinstanced& OnBlueprintReinstanced() { return BlueprintReinstanced; }
 
 	/**	Broadcasts that a blueprint just finished being reinstanced. THIS SHOULD NOT BE PUBLIC */
