IncrementalCook: Fix some bugs/missing features in -diffonly output.
* Fix spuriously reported diffs in import table due to <UnknownExport> text when the package being imported was not cooked in the previous session (e.g. because of -cooksinglepackagenorefs).
* Fix "DumpPackageHeaderDiffs does not yet implement describing the difference" when the only difference is in VersioningInfo.
#rnx
#incrementalcook

diff --git a/Engine/Source/Editor/UnrealEd/Private/Cooker/DiffWriterArchive.cpp b/Engine/Source/Editor/UnrealEd/Private/Cooker/DiffWriterArchive.cpp
--- a/Engine/Source/Editor/UnrealEd/Private/Cooker/DiffWriterArchive.cpp
+++ b/Engine/Source/Editor/UnrealEd/Private/Cooker/DiffWriterArchive.cpp
@@ -1495,6 +1495,123 @@ struct TTableItem
 	}
 };
 
+static bool Equal(const TOptional<FZenPackageVersioningInfo>& SourceInfo, const TOptional<FZenPackageVersioningInfo>& DestInfo)
+{
+	TStringBuilder<256> Message;
+	if (SourceInfo.IsSet() != DestInfo.IsSet())
+	{
+		return false;
+	}
+	else if (!SourceInfo.IsSet())
+	{
+		return true;
+	}
+	else
+	{
+		if (SourceInfo->ZenVersion != DestInfo->ZenVersion)
+		{
+			return false;
+		}
+		if (SourceInfo->PackageVersion != DestInfo->PackageVersion)
+		{
+			return false;
+		}
+		if (SourceInfo->LicenseeVersion != DestInfo->LicenseeVersion)
+		{
+			return false;
+		}
+		for (const FCustomVersion& SrcVersion : SourceInfo->CustomVersions.GetAllVersions())
+		{
+			const FCustomVersion* DestVersion = DestInfo->CustomVersions.GetVersion(SrcVersion.Key);
+			if (!DestVersion)
+			{
+				return false;
+			}
+			else if (SrcVersion.Version != DestVersion->Version)
+			{
+				return false;
+			}
+		}
+		for (const FCustomVersion& DestVersion : DestInfo->CustomVersions.GetAllVersions())
+		{
+			const FCustomVersion* SrcVersion = SourceInfo->CustomVersions.GetVersion(DestVersion.Key);
+			if (!SrcVersion)
+			{
+				return false;
+			}
+		}
+		return true;
+	}
+}
+
+template <typename ContextProvider>
+static void DumpDifferences(ContextProvider& SourceContext, const TCHAR* AssetFilename,
+	const TOptional<FZenPackageVersioningInfo>& SourceInfo, const TOptional<FZenPackageVersioningInfo>& DestInfo)
+{
+	TStringBuilder<256> Message;
+	if (SourceInfo.IsSet() != DestInfo.IsSet())
+	{
+		Message << NewLineToken << TEXT("VersioningInfo present in source package vs not present in dest package.");
+	}
+	else if (SourceInfo.IsSet())
+	{
+		if (SourceInfo->ZenVersion != DestInfo->ZenVersion)
+		{
+			Message << NewLineToken << IndentToken;
+			Message.Appendf(TEXT("-ZenVersion: %u"), (uint32)SourceInfo->ZenVersion);
+			Message << NewLineToken << IndentToken;
+			Message.Appendf(TEXT("+ZenVersion: %u"), (uint32)DestInfo->ZenVersion);
+		}
+		if (SourceInfo->PackageVersion != DestInfo->PackageVersion)
+		{
+			Message << NewLineToken << IndentToken;
+			Message.Appendf(TEXT("-PackageVersion: %d"), (uint32)SourceInfo->PackageVersion.ToValue());
+			Message << NewLineToken << IndentToken;
+			Message.Appendf(TEXT("+PackageVersion: %d"), (uint32)DestInfo->PackageVersion.ToValue());
+		}
+		if (SourceInfo->LicenseeVersion != DestInfo->LicenseeVersion)
+		{
+			Message << NewLineToken << IndentToken;
+			Message.Appendf(TEXT("-LicenseeVersion: %d"), (uint32)SourceInfo->LicenseeVersion);
+			Message << NewLineToken << IndentToken;
+			Message.Appendf(TEXT("+LicenseeVersion: %d"), (uint32)DestInfo->LicenseeVersion);
+		}
+		for (const FCustomVersion& SrcVersion : SourceInfo->CustomVersions.GetAllVersions())
+		{
+			const FCustomVersion* DestVersion = DestInfo->CustomVersions.GetVersion(SrcVersion.Key);
+			if (!DestVersion)
+			{
+				Message << NewLineToken << IndentToken;
+				Message.Appendf(TEXT("-CustomVersion: %s == %d"), *SrcVersion.Key.ToString(), SrcVersion.Version);
+			}
+			else if (SrcVersion.Version != DestVersion->Version)
+			{
+				Message << NewLineToken << IndentToken;
+				Message.Appendf(TEXT("-CustomVersion: %s == %d"), *SrcVersion.Key.ToString(), SrcVersion.Version);
+				Message << NewLineToken << IndentToken;
+				Message.Appendf(TEXT("+CustomVersion: %s == %d"), *DestVersion->Key.ToString(), DestVersion->Version);
+			}
+		}
+		for (const FCustomVersion& DestVersion : DestInfo->CustomVersions.GetAllVersions())
+		{
+			const FCustomVersion* SrcVersion = SourceInfo->CustomVersions.GetVersion(DestVersion.Key);
+			if (!SrcVersion)
+			{
+				Message << NewLineToken << IndentToken;
+				Message.Appendf(TEXT("+CustomVersion: %s == %d"), *DestVersion.Key.ToString(), DestVersion.Version);
+			}
+		}
+	}
+
+	if (Message.Len() == 0)
+	{
+		Message << NewLineToken << IndentToken << TEXT("<Unknown Difference>");
+	}
+	SourceContext.LogMessage(ELogVerbosity::Warning, FString::Printf(
+		TEXT("%s: VersioningInfo is different:%s"),
+		AssetFilename, *Message));
+}
+
 /** Dumps differences between Linker tables */
 template <typename T, typename ContextProvider>
 static void DumpTableDifferences(
@@ -1732,6 +1849,11 @@ void DumpPackageHeaderDiffs_ZenPackage(FPackageHeaderData& SourceHeaderData,
 	Algo::Sort(SourceNames, StringSortByNoCaseThenCase);
 	Algo::Sort(DestNames, StringSortByNoCaseThenCase);
 
+	if (!Equal(SourceHeader.GetPackageHeader().VersioningInfo, DestHeader.GetPackageHeader().VersioningInfo))
+	{
+		DumpDifferences(SourceHeader, *AssetFilename, SourceHeader.GetPackageHeader().VersioningInfo,
+			DestHeader.GetPackageHeader().VersioningInfo);
+	}
 	bool bFoundDifferenceInLinkerTableData = false;
 	if (!SourceHeader.IsNameMapIdentical(DestHeader, SourceNames, DestNames))
 	{
diff --git a/Engine/Source/Editor/UnrealEd/Private/Cooker/DiffWriterZenHeader.cpp b/Engine/Source/Editor/UnrealEd/Private/Cooker/DiffWriterZenHeader.cpp
--- a/Engine/Source/Editor/UnrealEd/Private/Cooker/DiffWriterZenHeader.cpp
+++ b/Engine/Source/Editor/UnrealEd/Private/Cooker/DiffWriterZenHeader.cpp
@@ -445,20 +445,21 @@ void FZenPackageExportsForDiff::Initialize(FAccumulatorGlobals& Globals, FName I
 	ExportPaths.Reset();
 	if (bUseZenStore)
 	{
-		InitializeFromZenStore(Globals);
-	}
-	else
-	{
-		InitializeFromMemory(Globals);
+		if (TryInitializeFromZenStore(Globals))
+		{
+			return;
+		}
 	}
+
+	// Silently ignore if TryInitializeFromMemory fails; GetExportName will report <UnknownExport.%u>
+	(void) TryInitializeFromMemory(Globals);
 }
 
-void FZenPackageExportsForDiff::InitializeFromZenStore(FAccumulatorGlobals& Globals)
+bool FZenPackageExportsForDiff::TryInitializeFromZenStore(FAccumulatorGlobals& Globals)
 {
-	// Silently ignore on errors; GetExportName will report <UnknownExport.%u>
 	if (!Globals.PackageWriter)
 	{
-		return;
+		return false;
 	}
 
 	FPackageId PackageId = FPackageId::FromName(PackageName);
@@ -468,14 +469,14 @@ void FZenPackageExportsForDiff::InitializeFromZenStore(FAccumulatorGlobals& Glob
 	ICookedPackageWriter::FPreviousCookedBytesData PackageBytes;
 	if (!Globals.PackageWriter->GetPreviousCookedBytes(PackageInfo, PackageBytes))
 	{
-		return;
+		return false;
 	}
 	FString Error;
 	FZenPackageHeader PackageHeader = FZenPackageHeader::MakeView(
 		FMemoryView(PackageBytes.Data.Get(), PackageBytes.Size), Error);
 	if (!Error.IsEmpty())
 	{
-		return;
+		return false;
 	}
 
 	TArray<FString> LocalPackageRelativePaths;
@@ -548,15 +549,16 @@ void FZenPackageExportsForDiff::InitializeFromZenStore(FAccumulatorGlobals& Glob
 		const FExportMapEntry& Export = PackageHeader.ExportMap[ExportIndex];
 		ExportPaths.Add(Export.PublicExportHash, LocalPackageRelativePaths[ExportIndex]);
 	}
+
+	return true;
 }
 
-void FZenPackageExportsForDiff::InitializeFromMemory(FAccumulatorGlobals& Globals)
+bool FZenPackageExportsForDiff::TryInitializeFromMemory(FAccumulatorGlobals& Globals)
 {
-	// Silently ignore on errors; GetExportName will report <UnknownExport.%u>
 	UPackage* Package = FindPackage(nullptr, *WriteToString<256>(PackageName));
 	if (!Package)
 	{
-		return;
+		return false;
 	}
 
 	// We do not add the UPackage object itself, because UPackage object is not stored in the list
@@ -571,6 +573,8 @@ void FZenPackageExportsForDiff::InitializeFromMemory(FAccumulatorGlobals& Global
 				ExportPaths.Add(PublicExportHash, FString(PackageRelativePath));
 			}
 		}, true /* bIncludeNestedObjects */);
+
+	return true;
 }
 
 FStringView FZenPackageExportsForDiff::GetExportPackageRelativePath(uint64 PublicExportHash)
diff --git a/Engine/Source/Editor/UnrealEd/Private/Cooker/DiffWriterZenHeader.h b/Engine/Source/Editor/UnrealEd/Private/Cooker/DiffWriterZenHeader.h
--- a/Engine/Source/Editor/UnrealEd/Private/Cooker/DiffWriterZenHeader.h
+++ b/Engine/Source/Editor/UnrealEd/Private/Cooker/DiffWriterZenHeader.h
@@ -20,8 +20,8 @@ public:
 	FStringView GetExportPackageRelativePath(uint64 PublicExportHash);
 
 private:
-	void InitializeFromZenStore(FAccumulatorGlobals& Globals);
-	void InitializeFromMemory(FAccumulatorGlobals& Globals);
+	bool TryInitializeFromZenStore(FAccumulatorGlobals& Globals);
+	bool TryInitializeFromMemory(FAccumulatorGlobals& Globals);
 
 private:
 	FName PackageName = NAME_None;
