Make sound previews audible when using PIE by adding the AudioComponent to the current main audio device.

diff --git a/Engine/Source/Editor/UnrealEd/Private/EditorEngine.cpp b/Engine/Source/Editor/UnrealEd/Private/EditorEngine.cpp
--- a/Engine/Source/Editor/UnrealEd/Private/EditorEngine.cpp
+++ b/Engine/Source/Editor/UnrealEd/Private/EditorEngine.cpp
@@ -2904,7 +2904,6 @@ UAudioComponent* UEditorEngine::PlayPreviewSound(USoundBase* Sound,  USoundNode*
 		AudioComponent->bCenterChannelOnly = false;
 		AudioComponent->bIsPreviewSound = true;
 
-#if WITH_EDITORONLY_DATA
 		if (!::IsValid(Sound) || Sound->EditorData.OverrideUseSubmixesOnPreview == EUseSubmixOnPreviewMode::UseEditorPreference)
 		{
 			const UAudioEditorSettings* AudioEditorSettings = GetDefault<UAudioEditorSettings>();
@@ -2914,7 +2913,11 @@ UAudioComponent* UEditorEngine::PlayPreviewSound(USoundBase* Sound,  USoundNode*
 		{
 			AudioComponent->bIsDryPreviewPlayback = (Sound->EditorData.OverrideUseSubmixesOnPreview == EUseSubmixOnPreviewMode::PlayPreviewsWithoutSubmixes) ? true : false;
 		}
-#endif //WITH_EDITORONLY_DATA
+
+		if (FAudioDeviceManager* DeviceManager = GEngine->GetAudioDeviceManager())
+		{
+			AudioComponent->AudioDeviceID = DeviceManager->GetActiveAudioDevice()->DeviceID;
+		}
 
 		AudioComponent->Play();
 
