Swap things to use FTextChar functions directly, rather than use the (soon to be deprecated) shims

* Swap FText::IsWhitespace for FTextChar::IsWhitespace
* Swap FUnicodeChar::CodepointToString for FTextChar::ConvertCodepointToString

#rnx

diff --git a/Engine/Source/Developer/Localization/Private/TextLocalizationResourceGenerator.cpp b/Engine/Source/Developer/Localization/Private/TextLocalizationResourceGenerator.cpp
--- a/Engine/Source/Developer/Localization/Private/TextLocalizationResourceGenerator.cpp
+++ b/Engine/Source/Developer/Localization/Private/TextLocalizationResourceGenerator.cpp
@@ -3,6 +3,7 @@
 #include "TextLocalizationResourceGenerator.h"
 
 #include "Internationalization/Culture.h"
+#include "Internationalization/TextChar.h"
 #include "Internationalization/TextLocalizationResource.h"
 #include "LocTextHelper.h"
 #include "Logging/StructuredLog.h"
@@ -183,7 +184,7 @@ bool FTextLocalizationResourceGenerator::GenerateLocRes(const FLocTextHelper& In
 					auto IsUnsafeWhitespace = [](const TCHAR InChar)
 						{
 							// Unsafe whitespace is any whitespace character, except new-lines
-							return FText::IsWhitespace(InChar) && !(InChar == TEXT('\r') || InChar == TEXT('\n'));
+							return FTextChar::IsWhitespace(InChar) && !(InChar == TEXT('\r') || InChar == TEXT('\n'));
 						};
 
 					if (IsUnsafeWhitespace(TranslationText.Text[0]) || IsUnsafeWhitespace(TranslationText.Text[TranslationText.Text.Len() - 1]))
diff --git a/Engine/Source/Editor/BlueprintGraph/Private/BasicTokenParser.cpp b/Engine/Source/Editor/BlueprintGraph/Private/BasicTokenParser.cpp
--- a/Engine/Source/Editor/BlueprintGraph/Private/BasicTokenParser.cpp
+++ b/Engine/Source/Editor/BlueprintGraph/Private/BasicTokenParser.cpp
@@ -4,6 +4,7 @@
 
 #include "HAL/UnrealMemory.h"
 #include "Internationalization/Internationalization.h"
+#include "Internationalization/TextChar.h"
 #include "Logging/LogCategory.h"
 #include "Logging/LogMacros.h"
 #include "Math/UnrealMathSSE.h"
@@ -233,7 +234,7 @@ namespace BasicTokenParserImpl
 	//------------------------------------------------------------------------------
 	static bool IsWhitespace(TCHAR c)
 	{
-		return FText::IsWhitespace(c);
+		return FTextChar::IsWhitespace(c);
 	}
 
 	//------------------------------------------------------------------------------
@@ -855,7 +856,7 @@ TCHAR FBasicTokenParser::GetLeadingChar()
 			{
 				MultipleNewlines = true;
 			}
-		} while (FText::IsWhitespace(c));
+		} while (FTextChar::IsWhitespace(c));
 
 		if (c != TEXT('/') || PeekChar() != TEXT('/'))
 		{
diff --git a/Engine/Source/Editor/UnrealEd/Private/Commandlets/GatherTextFromSourceCommandlet.cpp b/Engine/Source/Editor/UnrealEd/Private/Commandlets/GatherTextFromSourceCommandlet.cpp
--- a/Engine/Source/Editor/UnrealEd/Private/Commandlets/GatherTextFromSourceCommandlet.cpp
+++ b/Engine/Source/Editor/UnrealEd/Private/Commandlets/GatherTextFromSourceCommandlet.cpp
@@ -8,6 +8,7 @@
 #include "HAL/FileManager.h"
 #include "HAL/PlatformMath.h"
 #include "Internationalization/InternationalizationMetadata.h"
+#include "Internationalization/TextChar.h"
 #include "Internationalization/TextNamespaceUtil.h"
 #include "Logging/StructuredLog.h"
 #include "Misc/AsciiSet.h"
@@ -829,7 +830,7 @@ FString UGatherTextFromSourceCommandlet::UnescapeLiteralCharacterEscapeSequences
 						const uint32 UnicodeCodepoint = (uint32)FCString::Strtoi(*EscapedLiteralCharacter, nullptr, 16);
 
 						FString UnicodeString;
-						if (FUnicodeChar::CodepointToString(UnicodeCodepoint, UnicodeString))
+						if (FTextChar::ConvertCodepointToString(UnicodeCodepoint, UnicodeString))
 						{
 							RetString.Append(MoveTemp(UnicodeString));
 						}
@@ -845,7 +846,7 @@ FString UGatherTextFromSourceCommandlet::UnescapeLiteralCharacterEscapeSequences
 						const uint32 UnicodeCodepoint = (uint32)FCString::Strtoi(*EscapedLiteralCharacter, nullptr, 16);
 
 						FString UnicodeString;
-						if (FUnicodeChar::CodepointToString(UnicodeCodepoint, UnicodeString))
+						if (FTextChar::ConvertCodepointToString(UnicodeCodepoint, UnicodeString))
 						{
 							RetString.Append(MoveTemp(UnicodeString));
 						}
@@ -870,7 +871,7 @@ FString UGatherTextFromSourceCommandlet::UnescapeLiteralCharacterEscapeSequences
 						const uint32 UnicodeCodepoint = (uint32)FCString::Strtoui64(*EscapedLiteralCharacter, nullptr, 16);
 
 						FString UnicodeString;
-						if (FUnicodeChar::CodepointToString(UnicodeCodepoint, UnicodeString))
+						if (FTextChar::ConvertCodepointToString(UnicodeCodepoint, UnicodeString))
 						{
 							RetString.Append(MoveTemp(UnicodeString));
 						}
@@ -886,7 +887,7 @@ FString UGatherTextFromSourceCommandlet::UnescapeLiteralCharacterEscapeSequences
 						const uint32 UnicodeCodepoint = (uint32)FCString::Strtoui64(*EscapedLiteralCharacter, nullptr, 16);
 
 						FString UnicodeString;
-						if (FUnicodeChar::CodepointToString(UnicodeCodepoint, UnicodeString))
+						if (FTextChar::ConvertCodepointToString(UnicodeCodepoint, UnicodeString))
 						{
 							RetString.Append(MoveTemp(UnicodeString));
 						}
@@ -1743,7 +1744,7 @@ void UGatherTextFromSourceCommandlet::FSourceFileParseContext::SetDefine(const F
 	}
 	else if (!ExcludedRegion)
 	{
-		if (InDefineCtx.StartsWith(LocNamespaceString, ESearchCase::CaseSensitive) && InDefineCtx.IsValidIndex(LocNamespaceString.Len()) && (FText::IsWhitespace(InDefineCtx[LocNamespaceString.Len()]) || InDefineCtx[LocNamespaceString.Len()] == TEXT('"')))
+		if (InDefineCtx.StartsWith(LocNamespaceString, ESearchCase::CaseSensitive) && InDefineCtx.IsValidIndex(LocNamespaceString.Len()) && (FTextChar::IsWhitespace(InDefineCtx[LocNamespaceString.Len()]) || InDefineCtx[LocNamespaceString.Len()] == TEXT('"')))
 		{
 			// #define LOCTEXT_NAMESPACE <namespace>
 			if (WithinNamespaceDefineLineNumber != INDEX_NONE)
@@ -2679,7 +2680,7 @@ void UGatherTextFromSourceCommandlet::FNestedMacroDescriptor::TryParse(const FSt
 	if (Pos > 0)
 	{
 		TCHAR Char = Context.LineText[Pos - 1];
-		if (!FText::IsWhitespace(Char) && Char != TCHAR('(') && Char != TCHAR('{'))
+		if (!FTextChar::IsWhitespace(Char) && Char != TCHAR('(') && Char != TCHAR('{'))
 		{
 			return;
 		}
diff --git a/Engine/Source/Runtime/Core/Private/Containers/StringTest.cpp b/Engine/Source/Runtime/Core/Private/Containers/StringTest.cpp
--- a/Engine/Source/Runtime/Core/Private/Containers/StringTest.cpp
+++ b/Engine/Source/Runtime/Core/Private/Containers/StringTest.cpp
@@ -6,6 +6,7 @@
 #include "Misc/StringBuilder.h"
 #include "Containers/StringView.h"
 #include "Containers/UnrealString.h"
+#include "Internationalization/TextChar.h"
 #include "Serialization/MemoryReader.h"
 #include "Serialization/MemoryWriter.h"
 
@@ -103,7 +104,7 @@ bool FStringUnicodeTest::RunTest(const FString& Parameters)
 
 	// Test data used to verify basic processing of a Unicode character outside the BMP
 	FString TestStr;
-	if (FUnicodeChar::CodepointToString(128512, TestStr))
+	if (FTextChar::ConvertCodepointToString(128512, TestStr))
 	{
 		// Verify that the string can be serialized and deserialized without losing any data
 		{
@@ -133,7 +134,6 @@ bool FStringUnicodeTest::RunTest(const FString& Parameters)
 		}
 	}
 
-
 	return true;
 }
 
diff --git a/Engine/Source/Runtime/Core/Private/Internationalization/Text.cpp b/Engine/Source/Runtime/Core/Private/Internationalization/Text.cpp
--- a/Engine/Source/Runtime/Core/Private/Internationalization/Text.cpp
+++ b/Engine/Source/Runtime/Core/Private/Internationalization/Text.cpp
@@ -332,7 +332,7 @@ bool FText::IsEmptyOrWhitespace() const
 
 	for( const TCHAR Character : DisplayString )
 	{
-		if (!IsWhitespace(Character))
+		if (!FTextChar::IsWhitespace(Character))
 		{
 			return false;
 		}
@@ -372,7 +372,7 @@ FText FText::TrimPreceding( const FText& InText )
 	int32 StartPos = 0;
 	while (StartPos < CurrentString.Len())
 	{
-		if (!FText::IsWhitespace(CurrentString[StartPos]))
+		if (!FTextChar::IsWhitespace(CurrentString[StartPos]))
 		{
 			break;
 		}
@@ -398,7 +398,7 @@ FText FText::TrimTrailing( const FText& InText )
 	int32 EndPos = CurrentString.Len() - 1;
 	while (EndPos >= 0)
 	{
-		if (!FText::IsWhitespace(CurrentString[EndPos]))
+		if (!FTextChar::IsWhitespace(CurrentString[EndPos]))
 		{
 			break;
 		}
@@ -424,7 +424,7 @@ FText FText::TrimPrecedingAndTrailing( const FText& InText )
 	int32 StartPos = 0;
 	while (StartPos < CurrentString.Len())
 	{
-		if (!FText::IsWhitespace(CurrentString[StartPos]))
+		if (!FTextChar::IsWhitespace(CurrentString[StartPos]))
 		{
 			break;
 		}
@@ -435,7 +435,7 @@ FText FText::TrimPrecedingAndTrailing( const FText& InText )
 	int32 EndPos = CurrentString.Len() - 1;
 	while (EndPos > StartPos)
 	{
-		if (!FText::IsWhitespace(CurrentString[EndPos]))
+		if (!FTextChar::IsWhitespace(CurrentString[EndPos]))
 		{
 			break;
 		}
@@ -1332,7 +1332,7 @@ bool FText::ShouldGatherForLocalization() const
 	{
 		for(int32 i = 0; i < String.Len(); ++i)
 		{
-			if( !FText::IsWhitespace( String[i] ) )
+			if( !FTextChar::IsWhitespace( String[i] ) )
 			{
 				return false;
 			}
diff --git a/Engine/Source/Runtime/Core/Private/UObject/UnrealNames.cpp b/Engine/Source/Runtime/Core/Private/UObject/UnrealNames.cpp
--- a/Engine/Source/Runtime/Core/Private/UObject/UnrealNames.cpp
+++ b/Engine/Source/Runtime/Core/Private/UObject/UnrealNames.cpp
@@ -27,6 +27,7 @@
 #include "ProfilingDebugging/MetadataTrace.h"
 #include "Containers/Set.h"
 #include "Internationalization/Text.h"
+#include "Internationalization/TextChar.h"
 #include "Internationalization/Internationalization.h"
 #include "Misc/OutputDeviceRedirector.h"
 #include "HAL/FileManager.h"
@@ -3943,7 +3944,7 @@ bool FName::IsValidXName(const FStringView& InName, const FString& InInvalidChar
 				
 				for (TCHAR ch : MatchedInvalidChars)
 				{
-					if (FText::IsWhitespace(ch))
+					if (FTextChar::IsWhitespace(ch))
 					{
 						InvalidCharacters.Appendf(TEXT(" U+%04X"), ch);
 						continue;
@@ -3980,7 +3981,7 @@ FString FName::SanitizeWhitespace(const FString& FNameString)
 
 	for (const TCHAR CurrentChar : FNameString)
 	{
-		if (!FText::IsWhitespace(CurrentChar) || (CurrentChar == ' ' || CurrentChar =='\t'))
+		if (!FTextChar::IsWhitespace(CurrentChar) || (CurrentChar == ' ' || CurrentChar =='\t'))
 		{
 			ReturnStr.AppendChar(CurrentChar);
 		}	
diff --git a/Engine/Source/Runtime/Engine/Private/UserInterface/Canvas.cpp b/Engine/Source/Runtime/Engine/Private/UserInterface/Canvas.cpp
--- a/Engine/Source/Runtime/Engine/Private/UserInterface/Canvas.cpp
+++ b/Engine/Source/Runtime/Engine/Private/UserInterface/Canvas.cpp
@@ -19,6 +19,7 @@
 #include "Framework/Application/SlateApplication.h"
 #include "EngineFontServices.h"
 #include "Internationalization/BreakIterator.h"
+#include "Internationalization/TextChar.h"
 #include "Misc/CoreDelegates.h"
 #include "RHIStaticStates.h"
 #include "ClearQuad.h"
@@ -165,7 +166,7 @@ bool FCanvasWordWrapper::ProcessLine(FWrappingState& WrappingState)
 			--BreakIndex;			
 		}
 
-		while (BreakIndex > 0 && FText::IsWhitespace(WrappingState.String[BreakIndex - 1]))
+		while (BreakIndex > 0 && FTextChar::IsWhitespace(WrappingState.String[BreakIndex - 1]))
 		{
 			--BreakIndex;
 		}
@@ -177,7 +178,7 @@ bool FCanvasWordWrapper::ProcessLine(FWrappingState& WrappingState)
 			bHasAddedLine = true;
 		}
 
-		while (NextStartIndex < StringLength && FText::IsWhitespace(WrappingState.String[NextStartIndex]))
+		while (NextStartIndex < StringLength && FTextChar::IsWhitespace(WrappingState.String[NextStartIndex]))
 		{
 			++NextStartIndex;
 		}
diff --git a/Engine/Source/Runtime/Slate/Private/Framework/Text/TextLayout.cpp b/Engine/Source/Runtime/Slate/Private/Framework/Text/TextLayout.cpp
--- a/Engine/Source/Runtime/Slate/Private/Framework/Text/TextLayout.cpp
+++ b/Engine/Source/Runtime/Slate/Private/Framework/Text/TextLayout.cpp
@@ -9,6 +9,7 @@
 #include "Internationalization/BreakIterator.h"
 #include "Internationalization/Culture.h"
 #include "Internationalization/Internationalization.h"
+#include "Internationalization/TextChar.h"
 #include "Internationalization/TextTransformer.h"
 #include "Framework/Text/ShapedTextCache.h"
 
@@ -74,7 +75,7 @@ FTextLayout::FBreakCandidate FTextLayout::CreateBreakCandidate( int32& OutRunInd
 		WhitespaceStopIndex = StopIndex = FMath::Min( Range.EndIndex, CurrentBreak );
 		int32 BeginIndex = FMath::Max( PreviousBreak, Range.BeginIndex );
 
-		while( WhitespaceStopIndex > BeginIndex && FText::IsWhitespace( (*Line.Text)[ WhitespaceStopIndex - 1 ] ) )
+		while( WhitespaceStopIndex > BeginIndex && FTextChar::IsWhitespace( (*Line.Text)[ WhitespaceStopIndex - 1 ] ) )
 		{
 			--WhitespaceStopIndex;
 		}
@@ -2518,7 +2519,7 @@ FTextSelection FTextLayout::GetWordAt(const FTextLocation& Location) const
 		bool HasLetter = false;
 		for (int32 Index = CurrentBreak; Index < PreviousBreak; Index++)
 		{
-			if (!FText::IsWhitespace((*LineModel.Text)[Index]))
+			if (!FTextChar::IsWhitespace((*LineModel.Text)[Index]))
 			{
 				HasLetter = true;
 				break;
diff --git a/Engine/Source/Runtime/SlateCore/Private/Fonts/FontCache.cpp b/Engine/Source/Runtime/SlateCore/Private/Fonts/FontCache.cpp
--- a/Engine/Source/Runtime/SlateCore/Private/Fonts/FontCache.cpp
+++ b/Engine/Source/Runtime/SlateCore/Private/Fonts/FontCache.cpp
@@ -5,6 +5,7 @@
 #include "HAL/IConsoleManager.h"
 #include "Algo/Find.h"
 #include "Application/SlateApplicationBase.h"
+#include "Internationalization/TextChar.h"
 #include "Fonts/FontCacheFreeType.h"
 #include "Fonts/FontCacheHarfBuzz.h"
 #include "Fonts/FontCacheCompositeFont.h"
@@ -844,7 +845,7 @@ const FCharacterEntry* FCharacterList::CacheCharacter(TCHAR Character)
 
 			const bool bHasKerning = FT_HAS_KERNING(FaceGlyphData.FaceAndMemory->GetFace()) != 0;
 
-			const bool bIsWhitespace = FText::IsWhitespace(Character);
+			const bool bIsWhitespace = FTextChar::IsWhitespace(Character);
 			const uint32 GlyphIndex = FT_Get_Char_Index(FaceGlyphData.FaceAndMemory->GetFace(), Character);
 
 			int16 XAdvance = 0;
diff --git a/Engine/Source/Runtime/SlateCore/Private/Fonts/SlateTextShaper.cpp b/Engine/Source/Runtime/SlateCore/Private/Fonts/SlateTextShaper.cpp
--- a/Engine/Source/Runtime/SlateCore/Private/Fonts/SlateTextShaper.cpp
+++ b/Engine/Source/Runtime/SlateCore/Private/Fonts/SlateTextShaper.cpp
@@ -5,6 +5,7 @@
 #include "Fonts/SlateFontRenderer.h"
 #include "Fonts/FontProviderInterface.h"
 #include "Internationalization/BreakIterator.h"
+#include "Internationalization/TextChar.h"
 #include "SlateGlobals.h"
 
 #include <limits>
@@ -18,7 +19,7 @@ namespace
 
 bool RenderCodepointAsWhitespace(const TCHAR InCodepoint)
 {
-	return FText::IsWhitespace(InCodepoint)
+	return FTextChar::IsWhitespace(InCodepoint)
 		|| InCodepoint == TEXT('\u200B')	// Zero Width Space
 		|| InCodepoint == TEXT('\u2009')	// Thin Space
 		|| InCodepoint == TEXT('\u202F');	// Narrow No-Break Space
