[ITF Drag Tools] Add the option to modulate "strict" or "window/crossing" marquee selection depending on whether you drag to the left or to the right.
- Add & use "Dashed Box" canvas item to make crossing selection visually distinct.

diff --git a/Engine/Source/Editor/Experimental/EditorInteractiveToolsFramework/Private/EditorDragTools/EditorBoxSelect.cpp b/Engine/Source/Editor/Experimental/EditorInteractiveToolsFramework/Private/EditorDragTools/EditorBoxSelect.cpp
--- a/Engine/Source/Editor/Experimental/EditorInteractiveToolsFramework/Private/EditorDragTools/EditorBoxSelect.cpp
+++ b/Engine/Source/Editor/Experimental/EditorInteractiveToolsFramework/Private/EditorDragTools/EditorBoxSelect.cpp
@@ -66,16 +66,6 @@ void RemoveHoverEffect(const FViewportHoverTarget& InHoverTarget)
 
 } // namespace UE::EditorDragTools::Private
 
-void FEditorBoxSelect::Render(const FSceneView* View, FCanvas* Canvas)
-{
-	FCanvasBoxItem BoxItem(
-		FVector2D(Start.X, Start.Y) / Canvas->GetDPIScale(),
-		FVector2D(End.X - Start.X, End.Y - Start.Y) / Canvas->GetDPIScale()
-	);
-	BoxItem.SetColor(FLinearColor::White);
-	Canvas->DrawItem(BoxItem);
-}
-
 FInputRayHit FEditorBoxSelect::CanBeginClickDragSequence(const FInputDeviceRay& InPressPos)
 {
 	if (FEditorViewportClient* const EditorViewportClient = EditorViewportClientProxy->GetEditorViewportClient())
@@ -307,7 +297,7 @@ void FEditorBoxSelect::OnClickRelease(const FInputDeviceRay& InReleasePos)
 			// If the edit mode didn't handle the selection, try normal actor box selection.
 			if (!bEditorModeHandledBoxSelection && !bComponentVisHandledSelection)
 			{
-				const bool bStrictDragSelection = GetDefault<ULevelEditorViewportSettings>()->bStrictBoxSelection;
+				const bool bMustEncompassEntireElement = IsWindowSelection();
 
 				UTypedElementSelectionSet* SelectionSet =
 					GEditor->GetEditorSubsystem<ULevelEditorSubsystem>()->GetSelectionSet();
@@ -321,7 +311,7 @@ void FEditorBoxSelect::OnClickRelease(const FInputDeviceRay& InReleasePos)
 				FWorldSelectionElementArgs SelectionArgs{
 					SelectionSet,           ETypedElementSelectionMethod::Primary,
 					ElementSelectionOption, &(EditorViewportClient->EngineShowFlags),
-					bStrictDragSelection,   bGeometryMode
+					bMustEncompassEntireElement,   bGeometryMode
 				};
 
 				// Select all element that are within the selection box area.  Be aware that certain modes do special processing below.
@@ -372,7 +362,7 @@ void FEditorBoxSelect::OnClickRelease(const FInputDeviceRay& InReleasePos)
 					// Check every node in the model
 					for (int32 NodeIndex = 0; NodeIndex < Model.Nodes.Num(); NodeIndex++)
 					{
-						if (IntersectsBox(Model, NodeIndex, SelBBox, bStrictDragSelection))
+						if (IntersectsBox(Model, NodeIndex, SelBBox, bMustEncompassEntireElement))
 						{
 							// If the node intersected the frustum select the corresponding surface
 							GEditor->SelectBSPSurf(&Model, Model.Nodes[NodeIndex].iSurf, bShouldSelect, false);
diff --git a/Engine/Source/Editor/Experimental/EditorInteractiveToolsFramework/Private/EditorDragTools/EditorFrustumSelect.cpp b/Engine/Source/Editor/Experimental/EditorInteractiveToolsFramework/Private/EditorDragTools/EditorFrustumSelect.cpp
--- a/Engine/Source/Editor/Experimental/EditorInteractiveToolsFramework/Private/EditorDragTools/EditorFrustumSelect.cpp
+++ b/Engine/Source/Editor/Experimental/EditorInteractiveToolsFramework/Private/EditorDragTools/EditorFrustumSelect.cpp
@@ -137,22 +137,6 @@ void FViewportFrustum::Calculate(FConvexVolume& OutFrustum) const
 
 } // namespace UE::EditorDragTools
 
-void FEditorFrustumSelect::Render(const FSceneView* View, FCanvas* Canvas)
-{
-	if (!Canvas || !View)
-	{
-		return;
-	}
-
-	FCanvasBoxItem BoxItem(
-		FVector2D(Start.X, Start.Y) / Canvas->GetDPIScale(),
-		FVector2D(End.X - Start.X, End.Y - Start.Y) / Canvas->GetDPIScale()
-	);
-
-	BoxItem.SetColor(FLinearColor::White);
-	Canvas->DrawItem(BoxItem);
-}
-
 FInputRayHit FEditorFrustumSelect::CanBeginClickDragSequence(const FInputDeviceRay& InPressPos)
 {
 	// Todo: this could be retrieved from a command for customization
@@ -235,7 +219,7 @@ void FEditorFrustumSelect::OnClickRelease(const FInputDeviceRay& InReleasePos)
 	}
 
 	// Does an actor have to be fully contained in the box to be selected
-	const bool bStrictDragSelection = GetDefault<ULevelEditorViewportSettings>()->bStrictBoxSelection;
+	const bool bMustEncompassEntireElement = IsWindowSelection();
 
 	constexpr bool bShouldSelect = true; // bLeftMouseButtonDown
 	// Let the editor mode try to handle the selection.
@@ -255,7 +239,7 @@ void FEditorFrustumSelect::OnClickRelease(const FInputDeviceRay& InReleasePos)
 												 ETypedElementSelectionMethod::Primary,
 												 FTypedElementSelectionOptions(),
 												 &(EditorViewportClient->EngineShowFlags),
-												 bStrictDragSelection,
+												 bMustEncompassEntireElement,
 												 bGeometryMode };
 
 		const int32 ViewportSizeX = Viewport->GetSizeXY().X;
@@ -304,14 +288,14 @@ void FEditorFrustumSelect::OnClickRelease(const FInputDeviceRay& InReleasePos)
 				FTypedElementListRef ElementList = UTypedElementRegistry::GetInstance()->CreateElementList();
 				Viewport->GetElementHandlesInRect(BoxRect, ElementList);
 
-				if (bStrictDragSelection)
+				if (bMustEncompassEntireElement)
 				{
 					ElementList->ForEachElement<ITypedElementWorldInterface>(
-						[bStrictDragSelection, &Frustum, &SelectionSet, &ElementsToSelect](
+						[bMustEncompassEntireElement, &Frustum, &SelectionSet, &ElementsToSelect](
 							const TTypedElement<ITypedElementWorldInterface>& InElement
 						)
 						{
-							if (InElement.IsElementInConvexVolume(Frustum, bStrictDragSelection))
+							if (InElement.IsElementInConvexVolume(Frustum, bMustEncompassEntireElement))
 							{
 								ElementsToSelect.Add(
 									SelectionSet->GetSelectionElement(InElement, ETypedElementSelectionMethod::Primary)
@@ -352,7 +336,7 @@ void FEditorFrustumSelect::OnClickRelease(const FInputDeviceRay& InReleasePos)
 					// Check every node in the model
 					for (int32 NodeIndex = 0; NodeIndex < Model.Nodes.Num(); NodeIndex++)
 					{
-						if (IntersectsFrustum(Model, NodeIndex, Frustum, bStrictDragSelection))
+						if (IntersectsFrustum(Model, NodeIndex, Frustum, bMustEncompassEntireElement))
 						{
 							uint32 SurfaceIndex = Model.Nodes[NodeIndex].iSurf;
 							FBspSurf& Surf = Model.Surfs[SurfaceIndex];
@@ -367,7 +351,7 @@ void FEditorFrustumSelect::OnClickRelease(const FInputDeviceRay& InReleasePos)
 				for (auto It = BSPActors.CreateConstIterator(); It; ++It)
 				{
 					AActor* Actor = *It;
-					if (bStrictDragSelection)
+					if (bMustEncompassEntireElement)
 					{
 						ElementsToSelect.Append(
 							UE::EditorDragTools::Private::GetElementsIntersectingFrustum(
diff --git a/Engine/Source/Editor/Experimental/EditorInteractiveToolsFramework/Private/EditorDragTools/EditorMarqueeSelect.cpp b/Engine/Source/Editor/Experimental/EditorInteractiveToolsFramework/Private/EditorDragTools/EditorMarqueeSelect.cpp
new file mode 100644
--- /dev/null
+++ b/Engine/Source/Editor/Experimental/EditorInteractiveToolsFramework/Private/EditorDragTools/EditorMarqueeSelect.cpp
@@ -0,0 +1,70 @@
+﻿// Copyright Epic Games, Inc. All Rights Reserved.
+
+#include "EditorDragTools/EditorMarqueeSelect.h"
+
+#include "CanvasItem.h"
+#include "CanvasTypes.h"
+#include "Settings/LevelEditorViewportSettings.h"
+
+void FEditorMarqueeSelect::Render(const FSceneView* View, FCanvas* Canvas)
+{
+	if (!Canvas || !View)
+	{
+		return;
+	}
+	
+	if (IsWindowSelection())
+	{
+		FCanvasBoxItem BoxItem(
+			FVector2D(Start.X, Start.Y) / Canvas->GetDPIScale(),
+			FVector2D(End.X - Start.X, End.Y - Start.Y) / Canvas->GetDPIScale()
+		);
+	
+		BoxItem.SetColor(FLinearColor::White);
+		Canvas->DrawItem(BoxItem);
+	}
+	else
+	{
+		FVector LocalStart = Start;
+		FVector LocalEnd = End;
+		
+		// Place the starting corner in the upper-left
+		if (LocalStart.X > LocalEnd.X)
+		{
+			Swap(LocalStart.X, LocalEnd.X);
+		}
+
+		if (LocalStart.Y > LocalEnd.Y)
+		{
+			Swap(LocalStart.Y, LocalEnd.Y);
+		}
+
+		FCanvasDashedBoxItem BoxItem(
+			FVector2D(LocalStart.X, LocalStart.Y) / Canvas->GetDPIScale(),
+			FVector2D(LocalEnd.X - LocalStart.X, LocalEnd.Y - LocalStart.Y) / Canvas->GetDPIScale()
+		);
+		
+		BoxItem.DashLength = 12.0f;
+		BoxItem.DashGap = 6.0f;
+		BoxItem.SetColor(FLinearColor::White);
+	
+		Canvas->DrawItem(BoxItem);	
+	}
+}
+
+bool FEditorMarqueeSelect::IsWindowSelection() const
+{
+	switch (GetDefault<ULevelEditorViewportSettings>()->MarqueeSelectionMode)
+	{
+	case EMarqueeSelectionMode::Crossing:
+		return false;
+	case EMarqueeSelectionMode::Window:
+		return true;
+	case EMarqueeSelectionMode::CrossLeft:
+		return Start.X < End.X;
+	case EMarqueeSelectionMode::CrossRight:
+		return Start.X > End.X;
+	default:
+		return false;
+	}
+}
diff --git a/Engine/Source/Editor/Experimental/EditorInteractiveToolsFramework/Public/EditorDragTools/EditorBoxSelect.h b/Engine/Source/Editor/Experimental/EditorInteractiveToolsFramework/Public/EditorDragTools/EditorBoxSelect.h
--- a/Engine/Source/Editor/Experimental/EditorInteractiveToolsFramework/Public/EditorDragTools/EditorBoxSelect.h
+++ b/Engine/Source/Editor/Experimental/EditorInteractiveToolsFramework/Public/EditorDragTools/EditorBoxSelect.h
@@ -3,20 +3,19 @@
 #pragma once
 
 #include "EditorDragToolBehaviorTarget.h"
+#include "EditorMarqueeSelect.h"
 
 class FCanvas;
 class UModel;
 
-class FEditorBoxSelect : public FEditorDragToolBehaviorTarget
+class FEditorBoxSelect : public FEditorMarqueeSelect
 {
 public:
 	explicit FEditorBoxSelect(FEditorViewportClient* const InEditorViewportClient)
-		: FEditorDragToolBehaviorTarget(InEditorViewportClient)
+		: FEditorMarqueeSelect(InEditorViewportClient)
 	{
 	}
 
-	virtual void Render(const FSceneView* View, FCanvas* Canvas) override;
-
 	//~ Begin IClickDragBehaviorTarget
 	virtual FInputRayHit CanBeginClickDragSequence(const FInputDeviceRay& InPressPos) override;
 	virtual void OnClickPress(const FInputDeviceRay& InPressPos) override;
@@ -57,7 +56,7 @@ private:
 	 * @param OutBox	The created box.
 	 */
 	void CalculateBox(FBox& OutBox) const;
-
+	
 	/** List of BSP models to check for selection */
 	TArray<UModel*> ModelsToCheck;
 };
diff --git a/Engine/Source/Editor/Experimental/EditorInteractiveToolsFramework/Public/EditorDragTools/EditorFrustumSelect.h b/Engine/Source/Editor/Experimental/EditorInteractiveToolsFramework/Public/EditorDragTools/EditorFrustumSelect.h
--- a/Engine/Source/Editor/Experimental/EditorInteractiveToolsFramework/Public/EditorDragTools/EditorFrustumSelect.h
+++ b/Engine/Source/Editor/Experimental/EditorInteractiveToolsFramework/Public/EditorDragTools/EditorFrustumSelect.h
@@ -3,6 +3,7 @@
 #pragma once
 
 #include "EditorDragToolBehaviorTarget.h"
+#include "EditorMarqueeSelect.h"
 
 #define UE_API EDITORINTERACTIVETOOLSFRAMEWORK_API
 
@@ -17,16 +18,14 @@ struct FConvexVolume;
  * Draws a box in the current viewport and when the mouse button is released,
  * it selects/unselects everything inside of it.
  */
-class FEditorFrustumSelect : public FEditorDragToolBehaviorTarget
+class FEditorFrustumSelect : public FEditorMarqueeSelect
 {
 public:
 	explicit FEditorFrustumSelect(FEditorViewportClient* const InEditorViewportClient)
-		: FEditorDragToolBehaviorTarget(InEditorViewportClient)
+		: FEditorMarqueeSelect(InEditorViewportClient)
 	{
 	}
 
-	virtual void Render(const FSceneView* View, FCanvas* Canvas) override;
-
 	//~ Begin IClickDragBehaviorTarget
 	virtual FInputRayHit CanBeginClickDragSequence(const FInputDeviceRay& InPressPos) override;
 	virtual void OnClickPress(const FInputDeviceRay& InPressPos) override;
diff --git a/Engine/Source/Editor/Experimental/EditorInteractiveToolsFramework/Public/EditorDragTools/EditorMarqueeSelect.h b/Engine/Source/Editor/Experimental/EditorInteractiveToolsFramework/Public/EditorDragTools/EditorMarqueeSelect.h
new file mode 100644
--- /dev/null
+++ b/Engine/Source/Editor/Experimental/EditorInteractiveToolsFramework/Public/EditorDragTools/EditorMarqueeSelect.h
@@ -0,0 +1,30 @@
+﻿// Copyright Epic Games, Inc. All Rights Reserved.
+
+#pragma once
+#include "EditorDragToolBehaviorTarget.h"
+
+class FCanvas;
+class FEditorViewportClient;
+class FSceneView;
+class UModel;
+
+/**
+ * Base class for tools that use a marquee selection behavior
+ */
+class FEditorMarqueeSelect : public FEditorDragToolBehaviorTarget
+{
+public:
+	explicit FEditorMarqueeSelect(FEditorViewportClient* const InEditorViewportClient)
+			: FEditorDragToolBehaviorTarget(InEditorViewportClient)
+	{
+	}
+
+	virtual void Render(const FSceneView* View, FCanvas* Canvas);
+	
+protected:
+
+	/**
+	 * @return Whether selection should only include items completely contained by the box.  
+	 */
+	bool IsWindowSelection() const;
+};
diff --git a/Engine/Source/Editor/UnrealEd/Classes/Settings/LevelEditorViewportSettings.h b/Engine/Source/Editor/UnrealEd/Classes/Settings/LevelEditorViewportSettings.h
--- a/Engine/Source/Editor/UnrealEd/Classes/Settings/LevelEditorViewportSettings.h
+++ b/Engine/Source/Editor/UnrealEd/Classes/Settings/LevelEditorViewportSettings.h
@@ -82,6 +82,19 @@ enum class EMaterialKind : uint8
 	Emissive,
 };
 
+UENUM()
+enum class EMarqueeSelectionMode : uint8
+{
+	/** Selects anything the marquee touches. */
+	Crossing,
+	/** Selects only items the marquee surrounds. */
+	Window,
+	/** When dragging to the left, select anything the marquee touches. To the right, only select items the marquee surrounds. */
+	CrossLeft,
+	/** When dragging to the right, select anything the marquee touches. To the left, only select items the marquee surrounds. */
+	CrossRight
+};
+
 /**
  * Implements the Level Editor's per-instance view port settings.
  */
@@ -530,6 +543,10 @@ public:
 	/** True if viewport box selection requires objects to be fully encompassed by the selection box to be selected */
 	UPROPERTY(config)
 	uint32 bStrictBoxSelection:1;
+	
+	/** Sets how marquee selection should behave. */
+	UPROPERTY(EditAnywhere, Config, Category = Behavior)
+	EMarqueeSelectionMode MarqueeSelectionMode = EMarqueeSelectionMode::Crossing;
 
 	/** True if viewport box selection also selects occluded objects, false if only objects with visible pixels are selected */
 	UPROPERTY(config)
diff --git a/Engine/Source/Editor/UnrealEd/Private/Tools/EdModeInteractiveToolsContext.cpp b/Engine/Source/Editor/UnrealEd/Private/Tools/EdModeInteractiveToolsContext.cpp
--- a/Engine/Source/Editor/UnrealEd/Private/Tools/EdModeInteractiveToolsContext.cpp
+++ b/Engine/Source/Editor/UnrealEd/Private/Tools/EdModeInteractiveToolsContext.cpp
@@ -33,6 +33,8 @@
 #include "EditorModeManager.h"
 #include "EdMode.h"
 #include "EditorDragTools.h"
+#include "ToolMenuOwner.h"
+#include "ToolMenus.h"
 
 #include "BaseGizmos/GizmoRenderingUtil.h"
 #include "UnrealClient.h"
@@ -44,6 +46,8 @@
 
 DEFINE_LOG_CATEGORY_STATIC(LogInteractiveToolsContext, Log, All);
 
+#define LOCTEXT_NAMESPACE "EdModeInteractiveToolsContext"
+
 class FEdModeToolsContextQueriesImpl : public IToolsContextQueriesAPI
 {
 public:
@@ -1353,6 +1357,74 @@ void UModeManagerInteractiveToolsContext::RegisterDragTools()
 		DragToolsBehaviorSource = NewObject<UDragToolsBehaviorSource>();
 		DragToolsBehaviorSource->Initialize(this);
 		DragToolsBehaviorSource->RegisterSource();
+		
+		{
+			// Register menu extensions
+        	FToolMenuOwnerScoped Owner(this);
+        	
+        	UToolMenu* TransformMenu = UToolMenus::Get()->ExtendMenu("LevelEditor.ViewportToolbar.Transform");
+			FToolMenuSection& SelectionSection = TransformMenu->FindOrAddSection("Selection");
+			
+			FToolMenuEntry& MarqueeMode = SelectionSection.AddSubMenu(
+				"StrictBoxSelect", // Overwrite the other setting
+				LOCTEXT("MarqueeSelectionLabel", "Marquee Selection"),
+				LOCTEXT("MarqueeSelectionTooltip", "Define how objects are selected by marquee selection"),
+				FNewToolMenuDelegate::CreateLambda([](UToolMenu* Menu)
+				{
+					FProperty* Property = ULevelEditorViewportSettings::StaticClass()->FindPropertyByName(GET_MEMBER_NAME_CHECKED(ULevelEditorViewportSettings, MarqueeSelectionMode));
+					if (!Property || !Property->IsA<FEnumProperty>())
+					{
+						return;
+					}
+					
+					FEnumProperty* EnumProperty = StaticCast<FEnumProperty*>(Property);
+					UEnum* Enum = EnumProperty->GetEnum();
+					if (!Enum)
+					{
+						return;
+					}
+					
+					TWeakObjectPtr<UEnum> WeakEnum = Enum;
+					
+					FToolMenuSection& Section = Menu->AddSection("Mode", LOCTEXT("MarqueeModeSection", "Mode"));
+					
+					const int32 MaxIndex = Enum->NumEnums() - (Enum->ContainsExistingMax() ? 1 : 0);
+					for (int32 Index = 0; Index < MaxIndex; Index++)
+					{
+						Section.AddMenuEntry(
+							NAME_None,
+							Enum->GetDisplayNameTextByIndex(Index),
+							Enum->GetToolTipTextByIndex(Index),
+							FSlateIcon(),
+							FUIAction(
+								FExecuteAction::CreateLambda([WeakEnum, Index]
+								{
+									if (UEnum* Enum = WeakEnum.Get())
+									{
+										ULevelEditorViewportSettings* Settings = GetMutableDefault<ULevelEditorViewportSettings>();
+										Settings->MarqueeSelectionMode = static_cast<EMarqueeSelectionMode>(Enum->GetValueByIndex(Index));
+										Settings->SaveConfig();
+									}
+								}),
+								FCanExecuteAction(),
+								FIsActionChecked::CreateLambda([WeakEnum, Index]
+								{
+									if (UEnum* Enum = WeakEnum.Get())
+									{
+										const ULevelEditorViewportSettings* Settings = GetDefault<ULevelEditorViewportSettings>();
+										return Settings->MarqueeSelectionMode == static_cast<EMarqueeSelectionMode>(Enum->GetValueByIndex(Index));
+									}
+									return false;
+								})
+							),
+							EUserInterfaceActionType::RadioButton
+						);
+					}
+				})
+			);
+			MarqueeMode.Icon = FSlateIcon(FAppStyle::GetAppStyleSetName(), "LevelEditor.StrictBoxSelect");
+			MarqueeMode.InsertPosition = FToolMenuInsert("TransparentBoxSelect", EToolMenuInsertType::Before);
+		}
 	}
 }
 
@@ -1364,6 +1436,8 @@ void UModeManagerInteractiveToolsContext::UnregisterDragTools()
 		SetBypassAltModifier(true);
 		DragToolsBehaviorSource->DeregisterSource();
 		DragToolsBehaviorSource = nullptr;
+		
+		UToolMenus::UnregisterOwner(this);
 	}
 }
 
@@ -1408,3 +1482,5 @@ FRay UEdModeInteractiveToolsContext::GetLastWorldRay() const
 {
 	return ParentModeManagerToolsContext->GetLastWorldRay();
 }
+
+#undef LOCTEXT_NAMESPACE
diff --git a/Engine/Source/Runtime/Engine/Private/UserInterface/CanvasItem.cpp b/Engine/Source/Runtime/Engine/Private/UserInterface/CanvasItem.cpp
--- a/Engine/Source/Runtime/Engine/Private/UserInterface/CanvasItem.cpp
+++ b/Engine/Source/Runtime/Engine/Private/UserInterface/CanvasItem.cpp
@@ -30,6 +30,7 @@ DECLARE_CYCLE_STAT(TEXT("CanvasTileMaterialItem Time"),STAT_Canvas_TileMaterialI
 DECLARE_CYCLE_STAT(TEXT("CanvasTextItem Time"),STAT_Canvas_TextItemTime,STATGROUP_Canvas);
 DECLARE_CYCLE_STAT(TEXT("CanvasLineItem Time"),STAT_Canvas_LineItemTime,STATGROUP_Canvas);
 DECLARE_CYCLE_STAT(TEXT("CanvasBoxItem Time"),STAT_Canvas_BoxItemTime,STATGROUP_Canvas);
+DECLARE_CYCLE_STAT(TEXT("CanvasDashedBoxItem Time"),STAT_Canvas_DashedBoxItemTime,STATGROUP_Canvas);
 DECLARE_CYCLE_STAT(TEXT("CanvasTriTextureItem Time"),STAT_Canvas_TriTextureItemTime,STATGROUP_Canvas);
 DECLARE_CYCLE_STAT(TEXT("CanvasTriMaterialItem Time"), STAT_Canvas_TriMaterialItemTime, STATGROUP_Canvas);
 DECLARE_CYCLE_STAT(TEXT("CanvasBorderItem Time"),STAT_Canvas_BorderItemTime,STATGROUP_Canvas);
@@ -1614,6 +1615,72 @@ void FCanvasBoxItem::SetupBox()
 	Corners.AddUnique( FVector( Position.X, Position.Y + Size.Y, 0.0f ) );
 }
 
+void FCanvasDashedBoxItem::Draw(class FCanvas* InCanvas)
+{
+	SCOPE_CYCLE_COUNTER(STAT_Canvas_DashedBoxItemTime);
+	
+	FBatchedElements* BatchedElements = InCanvas->GetBatchedElements( FCanvas::ET_Line );
+	FHitProxyId HitProxyId = InCanvas->GetHitProxyId();
+
+	TArray<FVector, TInlineAllocator<4>> Corners;
+	// Top
+	Corners.AddUnique( FVector( Position.X , Position.Y, 0.0f ) );
+	// Right
+	Corners.AddUnique( FVector( Position.X + Size.X , Position.Y, 0.0f ) );
+	// Bottom
+	Corners.AddUnique( FVector( Position.X + Size.X , Position.Y + Size.Y, 0.0f ) );
+	// Left
+	Corners.AddUnique( FVector( Position.X, Position.Y + Size.Y, 0.0f ) );
+	
+	if (Corners.Num() < 2)
+	{
+		return;
+	}
+	
+	FVector CurrentPosition = Corners[0];
+	bool bIsDash = true;
+	float NextDistance = DashLength;
+	
+	// Draw the four dashed edges
+	for (int32 EdgeIndex = 0; EdgeIndex < Corners.Num() ; EdgeIndex++)
+	{
+		int32 NextCornerIndex = ( ( EdgeIndex + 1 ) % Corners.Num() );
+		
+		bool bMoveNext = false;
+		
+		while (!bMoveNext)
+		{
+			FVector ToCornerDelta = Corners[NextCornerIndex] - CurrentPosition;
+			bool bFlipDash = false;
+		
+			FVector NextPosition;
+			if (ToCornerDelta.SquaredLength() < NextDistance * NextDistance)
+			{
+				NextPosition = Corners[NextCornerIndex];
+				NextDistance = NextDistance - ToCornerDelta.Length();
+				bMoveNext = true;
+			}
+			else
+			{
+				NextPosition = CurrentPosition + ToCornerDelta.GetSafeNormal() * NextDistance;
+				bFlipDash = true;
+				NextDistance = bIsDash ? DashGap : DashLength;
+			}
+		
+			if (bIsDash)
+			{
+				BatchedElements->AddLine( CurrentPosition, NextPosition, Color, HitProxyId, DashThickness );
+			}
+		
+			CurrentPosition = NextPosition;
+			
+			if (bFlipDash)
+			{
+				bIsDash = !bIsDash;
+			}
+		}
+	}
+}
 
 void FCanvasTriangleItem::Draw( class FCanvas* InCanvas )
 {
diff --git a/Engine/Source/Runtime/Engine/Public/CanvasItem.h b/Engine/Source/Runtime/Engine/Public/CanvasItem.h
--- a/Engine/Source/Runtime/Engine/Public/CanvasItem.h
+++ b/Engine/Source/Runtime/Engine/Public/CanvasItem.h
@@ -789,6 +789,38 @@ private:
 	TArray< FVector > Corners;
 };
 
+/**
+ * A box drawn with dashed lines that wrap at the corners.
+ */
+class FCanvasDashedBoxItem : public FCanvasItem
+{
+public:
+	FCanvasDashedBoxItem( const FVector2D& InPosition, const FVector2D& InSize )
+		: FCanvasItem( InPosition )
+		, Size( InSize )
+		, DashThickness( 0.0f )
+		, DashLength( 1.0f )
+		, DashGap( 1.0f )
+		{};
+	
+	ENGINE_API virtual void Draw( FCanvas* InCanvas ) override;
+	
+	/* Expose the functions defined in the base class. */
+	using FCanvasItem::Draw;
+	
+	/* Size of the box. */	
+	FVector2d Size;
+	
+	/* The thickness of the dashes. */
+	float DashThickness;
+	
+	/* The length of the dashes. */
+	float DashLength;
+	
+	/* The length of the gap between dashes. */
+	float DashGap;
+};
+
 
 class FCanvasTriangleItem : public FCanvasItem
 {
