IncrementalCook: IncrementallyModifiedDiagnostics
* Write list of modified packages and the reason they were found to be modified to MetaData\ModifiedCookedPackages.txt.
* Write list of self-modified packages to MetaData\ModifiedEditorPackages.txt.
#rnx

diff --git a/Engine/Source/Editor/UnrealEd/Classes/CookOnTheSide/CookOnTheFlyServer.h b/Engine/Source/Editor/UnrealEd/Classes/CookOnTheSide/CookOnTheFlyServer.h
--- a/Engine/Source/Editor/UnrealEd/Classes/CookOnTheSide/CookOnTheFlyServer.h
+++ b/Engine/Source/Editor/UnrealEd/Classes/CookOnTheSide/CookOnTheFlyServer.h
@@ -184,6 +184,7 @@ namespace UE::Cook
 	class FCookWorkerClient;
 	class FCookWorkerServer;
 	class FDiagnostics;
+	class FIncrementallyModifiedDiagnostics;
 	class FLogHandler;
 	class FODSCClientData;
 	class FPackagePreloader;
@@ -1587,6 +1588,7 @@ private:
 	bool bCookIncremental = true;
 	bool bCookIncrementalAllowAllClasses = false;
 	bool bHiddenDependenciesDebug = false;
+	bool bIncrementallyModifiedDiagnostics = false;
 	bool bOnlyEditorOnlyDebug = false;
 	bool bHiddenDependenciesClassPathFilterListIsAllowList = true;
 	bool bFirstCookInThisProcessInitialized = false;
@@ -1667,6 +1669,8 @@ private:
 	TSharedPtr<FLooseFilesCookArtifactReader> SharedLooseFilesCookArtifactReader;
 	TUniquePtr<UE::Cook::FCookGCDiagnosticContext> GCDiagnosticContext;
 	TUniquePtr<UE::Cook::ILogHandler> LogHandler;
+	// These structs are TUniquePtr both for implementation hiding and because they are only sometimes needed.
+	TUniquePtr<UE::Cook::FIncrementallyModifiedDiagnostics> IncrementallyModifiedDiagnostics;
 
 	TArray<UE::Cook::FCookSavePackageContext*> SavePackageContexts;
 	/**
diff --git a/Engine/Source/Editor/UnrealEd/Private/CookOnTheFlyServer.cpp b/Engine/Source/Editor/UnrealEd/Private/CookOnTheFlyServer.cpp
--- a/Engine/Source/Editor/UnrealEd/Private/CookOnTheFlyServer.cpp
+++ b/Engine/Source/Editor/UnrealEd/Private/CookOnTheFlyServer.cpp
@@ -5781,6 +5781,21 @@ public:
 		return DiffMode != EDiffMode::None;
 	}
 
+	bool IsIncrementallyModifiedDiagnostics() const
+	{
+		switch (DiffMode)
+		{
+		case EDiffMode::None: return false;
+		case EDiffMode::DiffOnly: return false;
+		case EDiffMode::IncrementalValidate: return true;
+		case EDiffMode::IncrementalValidatePhase1: return true;
+		case EDiffMode::IncrementalValidatePhase2: return true;
+		default:
+			checkNoEntry();
+			return false;
+		}
+	}
+
 private:
 	bool bInitialized = false;
 	EDiffMode DiffMode = EDiffMode::None;
@@ -6239,16 +6254,18 @@ void UCookOnTheFlyServer::SetInitializeConfigSettings(UE::Cook::FInitializeConfi
 
 	UE_LOG(LogCook, Display, TEXT("CookSettings for Memory:%s"), *GetCookSettingsForMemoryLogText());
 
+	const TCHAR* CommandLine = FCommandLine::Get();
+
 	// Debugging hidden dependencies
-	bOnlyEditorOnlyDebug = FParse::Param(FCommandLine::Get(), TEXT("OnlyEditorOnlyDebug"));
+	bOnlyEditorOnlyDebug = FParse::Param(CommandLine, TEXT("OnlyEditorOnlyDebug"));
 	bSkipOnlyEditorOnly = false;
 	GConfig->GetBool(TEXT("CookSettings"), TEXT("SkipOnlyEditorOnly"), bSkipOnlyEditorOnly, GEditorIni);
 	FString ParamText;
-	if (FParse::Value(FCommandLine::Get(), TEXT("-SkipOnlyEditorOnly="), ParamText))
+	if (FParse::Value(CommandLine, TEXT("-SkipOnlyEditorOnly="), ParamText))
 	{
 		LexFromString(bSkipOnlyEditorOnly, *ParamText);
 	}
-	else if (FParse::Param(FCommandLine::Get(), TEXT("SkipOnlyEditorOnly")))
+	else if (FParse::Param(CommandLine, TEXT("SkipOnlyEditorOnly")))
 	{
 		bSkipOnlyEditorOnly = true;
 	}
@@ -6259,7 +6276,7 @@ void UCookOnTheFlyServer::SetInitializeConfigSettings(UE::Cook::FInitializeConfi
 			TEXT("SkipOnlyEditorOnly is enabled, unsolicited packages will not be cooked unless they are referenced from the cooked version of the instigator package."));
 	}
 
-	bHiddenDependenciesDebug = FParse::Param(FCommandLine::Get(), TEXT("HiddenDependenciesDebug"));
+	bHiddenDependenciesDebug = FParse::Param(CommandLine, TEXT("HiddenDependenciesDebug"));
 	if (bHiddenDependenciesDebug)
 	{
 		UE_LOG(LogCook, Display, TEXT("HiddenDependenciesDebug is enabled."));
@@ -6271,11 +6288,11 @@ void UCookOnTheFlyServer::SetInitializeConfigSettings(UE::Cook::FInitializeConfi
 
 		FString ClassPathListStr;
 		TOptional<bool> bAllowList;
-		if (FParse::Value(FCommandLine::Get(), TEXT("-HiddenDependenciesIgnore="), ClassPathListStr))
+		if (FParse::Value(CommandLine, TEXT("-HiddenDependenciesIgnore="), ClassPathListStr))
 		{
 			bAllowList = false;
 		}
-		if (FParse::Value(FCommandLine::Get(), TEXT("-HiddenDependenciesReport="), ClassPathListStr))
+		if (FParse::Value(CommandLine, TEXT("-HiddenDependenciesReport="), ClassPathListStr))
 		{
 			if (bAllowList.IsSet() && !*bAllowList)
 			{
@@ -6320,20 +6337,20 @@ void UCookOnTheFlyServer::SetInitializeConfigSettings(UE::Cook::FInitializeConfi
 
 	ParseCookFilters();
 
-	bCallIsCachedOnSaveCreatedObjects = FParse::Param(FCommandLine::Get(), TEXT("CallIsCachedOnSaveCreatedObjects"));
+	bCallIsCachedOnSaveCreatedObjects = FParse::Param(CommandLine, TEXT("CallIsCachedOnSaveCreatedObjects"));
 
 	bLegacyIterativeIgnoreIni = false;
 	GConfig->GetBool(TEXT("CookSettings"), TEXT("IterativeIgnoreIni"), bLegacyIterativeIgnoreIni, GEditorIni);
 	GConfig->GetBool(TEXT("CookSettings"), TEXT("LegacyIterativeIgnoreIni"), bLegacyIterativeIgnoreIni, GEditorIni);
 	bLegacyIterativeIgnoreIni =
-		!FParse::Param(FCommandLine::Get(), TEXT("iteraterequireini")) &&
-		!FParse::Param(FCommandLine::Get(), TEXT("iterativerequireini")) &&
-		!FParse::Param(FCommandLine::Get(), TEXT("legacyiterativerequireini")) &&
+		!FParse::Param(CommandLine, TEXT("iteraterequireini")) &&
+		!FParse::Param(CommandLine, TEXT("iterativerequireini")) &&
+		!FParse::Param(CommandLine, TEXT("legacyiterativerequireini")) &&
 		(bLegacyIterativeIgnoreIni ||
 			IsCookFlagSet(ECookInitializationFlags::IgnoreIniSettingsOutOfDate) ||
-			FParse::Param(FCommandLine::Get(), TEXT("iterateignoreini")) ||
-			FParse::Param(FCommandLine::Get(), TEXT("iterativeignoreini")) ||
-			FParse::Param(FCommandLine::Get(), TEXT("legacyiterativeignoreini")));
+			FParse::Param(CommandLine, TEXT("iterateignoreini")) ||
+			FParse::Param(CommandLine, TEXT("iterativeignoreini")) ||
+			FParse::Param(CommandLine, TEXT("legacyiterativeignoreini")));
 	bLegacyIterativeCalculateExe = true;
 	bool bConfigSettingSetLegacyIterativeIgnoreExe = false;
 	GConfig->GetBool(TEXT("CookSettings"), TEXT("IterativeIgnoreExe"),
@@ -6341,31 +6358,31 @@ void UCookOnTheFlyServer::SetInitializeConfigSettings(UE::Cook::FInitializeConfi
 	GConfig->GetBool(TEXT("CookSettings"), TEXT("LegacyIterativeIgnoreExe"),
 		bConfigSettingSetLegacyIterativeIgnoreExe, GEditorIni);
 	bLegacyIterativeIgnoreExe =
-		!FParse::Param(FCommandLine::Get(), TEXT("iteraterequireexe")) &&
-		!FParse::Param(FCommandLine::Get(), TEXT("iterativerequireexe")) &&
-		!FParse::Param(FCommandLine::Get(), TEXT("legacyiterativerequireexe")) &&
+		!FParse::Param(CommandLine, TEXT("iteraterequireexe")) &&
+		!FParse::Param(CommandLine, TEXT("iterativerequireexe")) &&
+		!FParse::Param(CommandLine, TEXT("legacyiterativerequireexe")) &&
 		(bConfigSettingSetLegacyIterativeIgnoreExe ||
-			FParse::Param(FCommandLine::Get(), TEXT("iterateignoreexe")) ||
-			FParse::Param(FCommandLine::Get(), TEXT("iterativeignoreexe")) ||
-			FParse::Param(FCommandLine::Get(), TEXT("legacyiterativeignoreexe")));
+			FParse::Param(CommandLine, TEXT("iterateignoreexe")) ||
+			FParse::Param(CommandLine, TEXT("iterativeignoreexe")) ||
+			FParse::Param(CommandLine, TEXT("legacyiterativeignoreexe")));
 	// Calculate the exe hash if LegacyIterativeExeInvalidation is required by ini OR required by commandline
 	// It would be better to always calculate it, but we want to avoid the performance cost until it becomes more widely used
 	bLegacyIterativeCalculateExe = !bLegacyIterativeIgnoreIni || !bConfigSettingSetLegacyIterativeIgnoreExe;
 
-	bRunningAsShaderServer = FParse::Param(FCommandLine::Get(), TEXT("odsc"));
+	bRunningAsShaderServer = FParse::Param(CommandLine, TEXT("odsc"));
     ODSCClientData = nullptr;
 	if (bRunningAsShaderServer)
 	{
 		ODSCClientData = MakeUnique<UE::Cook::FODSCClientData>();
 	}
 
-	bSkipSave = FParse::Param(FCommandLine::Get(), TEXT("CookSkipSave"));
+	bSkipSave = FParse::Param(CommandLine, TEXT("CookSkipSave"));
 
 	FString Severity;
 	GConfig->GetString(TEXT("CookSettings"), TEXT("CookerIdleWarningSeverity"), Severity, GEditorIni);
 	CookerIdleWarningSeverity = ParseLogVerbosityFromString(Severity);
 
-	bCookFastStartup = FParse::Param(FCommandLine::Get(), TEXT("cookfaststartup"));
+	bCookFastStartup = FParse::Param(CommandLine, TEXT("cookfaststartup"));
 
 	const UProjectPackagingSettings* PackagingSettings = Cast<UProjectPackagingSettings>(UProjectPackagingSettings::StaticClass()->GetDefaultObject());
 	if (PackagingSettings->bTreatWarningsAsErrorsOnCook)
@@ -6375,14 +6392,16 @@ void UCookOnTheFlyServer::SetInitializeConfigSettings(UE::Cook::FInitializeConfi
 
 	FString GeneratorSplit;
 	GConfig->GetString(TEXT("CookSettings"), TEXT("MPCookGeneratorSplit"), GeneratorSplit, GEditorIni);
-	FParse::Value(FCommandLine::Get(), TEXT("-MPCookGeneratorSplit="), GeneratorSplit);
+	FParse::Value(CommandLine, TEXT("-MPCookGeneratorSplit="), GeneratorSplit);
 	MPCookGeneratorSplit = UE::Cook::ParseMPCookGeneratorSplitFromString(GeneratorSplit);
 
-	bDeterminismDebug = FParse::Param(FCommandLine::Get(), TEXT("cookdeterminism")) ||
-		FParse::Param(FCommandLine::Get(), TEXT("diffonlybase"));
+	bDeterminismDebug = FParse::Param(CommandLine, TEXT("cookdeterminism")) ||
+		FParse::Param(CommandLine, TEXT("diffonlybase"));
 
-	FParse::Value(FCommandLine::Get(), TEXT("-CookCVarControl="), OverrideCookCVarControl);
-	FParse::Value(FCommandLine::Get(), TEXT("-DeviceProfile="), OverrideDeviceProfileName);
+	FParse::Value(CommandLine, TEXT("-CookCVarControl="), OverrideCookCVarControl);
+	FParse::Value(CommandLine, TEXT("-DeviceProfile="), OverrideDeviceProfileName);
+
+	bIncrementallyModifiedDiagnostics = FParse::Param(CommandLine, TEXT("CookIncrementallyModifiedDiagnostics"));
 }
 
 void UCookOnTheFlyServer::ParseCookFilters()
@@ -9993,7 +10012,18 @@ void UCookOnTheFlyServer::CookByTheBookFinishedInternal()
 			UE_SCOPED_HIERARCHICAL_COOKTIMER(WriteReferencedSet);
 			WriteReferencedSet(TargetPlatform, CookedPackageNames.Array());
 		}
+	}
+
+	if (!PlatformManager->GetSessionPlatforms().IsEmpty())
+	{
+		const ITargetPlatform* TargetPlatform = PlatformManager->GetSessionPlatforms()[0];
+		FString MetaDataDirectory = GetMetadataDirectory();
+		MetaDataDirectory = ConvertToFullSandboxPath(MetaDataDirectory, true, TargetPlatform->PlatformName());
 
+		if (IncrementallyModifiedDiagnostics)
+		{
+			IncrementallyModifiedDiagnostics->WriteToFile(MetaDataDirectory);
+		}
 	}
 
 	FString ActualLibraryName = GenerateShaderCodeLibraryName(LibraryName, IsCookFlagSet(ECookInitializationFlags::LegacyIterativeSharedBuild));
@@ -10898,6 +10928,13 @@ UE::Cook::FCookSavePackageContext* UCookOnTheFlyServer::CreateSaveContext(const
 	{
 		DeterminismManager = MakeUnique<FDeterminismManager>();
 	}
+	if ((bIncrementallyModifiedDiagnostics || DiffModeHelper->IsIncrementallyModifiedDiagnostics())
+		&& !IncrementallyModifiedDiagnostics
+		&& !IsCookWorkerMode())
+	{
+		IncrementallyModifiedDiagnostics = MakeUnique<UE::Cook::FIncrementallyModifiedDiagnostics>();
+	}
+
 	if (IsUsingZenStore())
 	{
 		TSharedRef<FLayeredCookArtifactReader> LayeredReader = MakeShared<FLayeredCookArtifactReader>();
diff --git a/Engine/Source/Editor/UnrealEd/Private/Cooker/CookDiagnostics.cpp b/Engine/Source/Editor/UnrealEd/Private/Cooker/CookDiagnostics.cpp
--- a/Engine/Source/Editor/UnrealEd/Private/Cooker/CookDiagnostics.cpp
+++ b/Engine/Source/Editor/UnrealEd/Private/Cooker/CookDiagnostics.cpp
@@ -1,18 +1,22 @@
 // Copyright Epic Games, Inc. All Rights Reserved.
 
-#include "CookDiagnostics.h"
+#include "Cooker/CookDiagnostics.h"
 
+#include "Algo/Sort.h"
 #include "AssetRegistry/IAssetRegistry.h"
 #include "Containers/Array.h"
 #include "Containers/RingBuffer.h"
 #include "Cooker/CookGenerationHelper.h"
+#include "Cooker/CookPackageArtifacts.h"
 #include "Cooker/CookPackageData.h"
 #include "Cooker/CookRequestCluster.h"
 #include "Cooker/CookTypes.h"
 #include "CookOnTheSide/CookLog.h"
 #include "CookOnTheSide/CookOnTheFlyServer.h"
 #include "Logging/LogMacros.h"
+#include "Misc/FileHelper.h"
 #include "Misc/Optional.h"
+#include "Misc/Paths.h"
 #include "Misc/StringBuilder.h"
 #include "Serialization/ArchiveSerializedPropertyChain.h"
 #include "UObject/ICookInfo.h"
@@ -437,4 +441,167 @@ void FDiagnostics::AnalyzeHiddenDependencies(UCookOnTheFlyServer& COTFS, FPackag
 	}
 }
 
-}
\ No newline at end of file
+const TCHAR* LexToString(UE::Cook::EIncrementallyModifiedReason Reason)
+{
+	static_assert(static_cast<int>(EIncrementallyModifiedReason::Count) == 7);
+	switch (Reason)
+	{
+	case EIncrementallyModifiedReason::None: return TEXT("None");
+	case EIncrementallyModifiedReason::NotPreviouslyCooked: return TEXT("NotPreviouslyCooked");
+	case EIncrementallyModifiedReason::TargetDomainKey:  return TEXT("TargetDomainKey");
+	case EIncrementallyModifiedReason::TransitiveBuildDependency:  return TEXT("TransitiveBuildDependency");
+	case EIncrementallyModifiedReason::IncrementalCookDisabled:  return TEXT("IncrementalCookDisabled");
+	case EIncrementallyModifiedReason::NoGenerator:  return TEXT("NoGenerator");
+	case EIncrementallyModifiedReason::GeneratorOtherPlatform: return TEXT("GeneratorOtherPlatform");
+	default: return TEXT("Invalid");
+	}
+}
+
+void FIncrementallyModifiedDiagnostics::AddPackage(FIncrementallyModifiedContext& Context)
+{
+	check(Context.PackageData);
+	check(Context.Artifacts);
+
+	FPackageData& PackageData = PackageDatas.Emplace_GetRef();
+	PackageData.AssetClass = Context.AssetClass;
+	PackageData.PackageName = Context.PackageData->GetPackageName();
+	PackageData.bSelfModified = Context.Artifacts->GetIsPackageModified().Get(false /* DefaultValue */);
+	PackageData.Reason = Context.Reason;
+
+	switch (Context.Reason)
+	{
+	case EIncrementallyModifiedReason::TargetDomainKey:
+	{
+		TArray<FString>& Lines = PackageData.ReasonLines;
+		TArray<FString> ModifiedDependencies;
+		Context.Artifacts->TryGetModifiedDependencies(Context.TargetPlatform, Context.GenerationHelper,
+			ModifiedDependencies);
+
+		const int32 MaxDependencyCount = 3;
+		for (int32 Index = 0; Index < ModifiedDependencies.Num(); ++Index)
+		{
+			if (Index == MaxDependencyCount)
+			{
+				Lines.Add(TEXT("..."));
+				break;
+			}
+			Lines.Add(ModifiedDependencies[Index]);
+		}
+		if (Lines.Num() == 0)
+		{
+			Lines.Add(TEXT("<UnknownModifiedDependency>"));
+		}
+		break;
+	}
+	case EIncrementallyModifiedReason::TransitiveBuildDependency:
+		PackageData.ReasonLines.Add(Context.TransitiveBuildDependencyName.ToString());
+		break;
+	default:
+		break;
+	}
+}
+
+void FIncrementallyModifiedDiagnostics::WriteToFile(FStringView OutputDirectory)
+{
+	TMap<FTopLevelAssetPath, TArray<const FPackageData*>> PackagesPerClass;
+	for (const FPackageData& PackageData : PackageDatas)
+	{
+		PackagesPerClass.FindOrAdd(PackageData.AssetClass).Add(&PackageData);
+	}
+
+	PackagesPerClass.KeySort([](const FTopLevelAssetPath& A, const FTopLevelAssetPath& B)
+		{
+			return A.Compare(B) < 0;
+		});
+	for (TPair<FTopLevelAssetPath, TArray<const FPackageData*>>& Pair : PackagesPerClass)
+	{
+		Algo::Sort(Pair.Value, [](const FPackageData* A, const FPackageData* B)
+			{
+				return A->PackageName.LexicalLess(B->PackageName);
+			});
+	}
+
+	{
+		FString FilePath = FPaths::Combine(OutputDirectory, TEXT("ModifiedCookedPackages.txt"));
+
+		TArray<FString> Lines;
+		Lines.Reserve(PackageDatas.Num() + PackagesPerClass.Num());
+		const TCHAR* Indent = TEXT("    ");
+
+		// Reserve leading line for the total
+		Lines.Emplace();
+		int32 NumTotalModified = 0;
+
+		for (TPair<FTopLevelAssetPath, TArray<const FPackageData*>>& Pair : PackagesPerClass)
+		{
+			Lines.Add(FString::Printf(TEXT("%s, %d"), *Pair.Key.ToString(), Pair.Value.Num()));
+			NumTotalModified += Pair.Value.Num();
+			for (const FPackageData* PackageData : Pair.Value)
+			{
+				if (PackageData->ReasonLines.Num() == 0)
+				{
+					Lines.Add(FString::Printf(TEXT("%s%s, %s"), Indent, *PackageData->PackageName.ToString(),
+						LexToString(PackageData->Reason)));
+				}
+				else if (PackageData->ReasonLines.Num() == 1)
+				{
+					Lines.Add(FString::Printf(TEXT("%s%s, %s: %s"), Indent, *PackageData->PackageName.ToString(),
+						LexToString(PackageData->Reason), *PackageData->ReasonLines[0]));
+				}
+				else
+				{
+					Lines.Add(FString::Printf(TEXT("%s%s, %s:"), Indent, *PackageData->PackageName.ToString(),
+						LexToString(PackageData->Reason)));
+					for (const FString& ReasonLine : PackageData->ReasonLines)
+					{
+						Lines.Add(FString::Printf(TEXT("%s%s%s"), Indent, Indent, *ReasonLine));
+					}
+				}
+			}
+		}
+
+		Lines[0] = FString::Printf(TEXT("Total, %d"), NumTotalModified);
+		FFileHelper::SaveStringArrayToFile(Lines, *FilePath);
+	}
+
+	{
+		int32 NumTotalModified = 0;
+		FString FilePath = FPaths::Combine(OutputDirectory, TEXT("ModifiedEditorPackages.txt"));
+
+		TArray<FString> Lines;
+		Lines.Reserve(PackageDatas.Num() + PackagesPerClass.Num() + 1);
+		const TCHAR* Indent = TEXT("    ");
+
+		// Reserve leading line for the total
+		Lines.Emplace();
+
+		for (TPair<FTopLevelAssetPath, TArray<const FPackageData*>>& Pair : PackagesPerClass)
+		{
+			int NumClassModified = 0;
+			int32 ClassModifiedIndex = -1;
+			for (const FPackageData* PackageData : Pair.Value)
+			{
+				if (PackageData->bSelfModified)
+				{
+					if (ClassModifiedIndex == -1)
+					{
+						ClassModifiedIndex = Lines.Num();
+						Lines.Emplace();
+					}
+					Lines.Add(FString::Printf(TEXT("%s%s"), Indent, *PackageData->PackageName.ToString()));
+					++NumClassModified;
+				}
+			}
+			if (ClassModifiedIndex != -1)
+			{
+				Lines[ClassModifiedIndex] = FString::Printf(TEXT("%s, %d"), *Pair.Key.ToString(), NumClassModified);
+				NumTotalModified += NumClassModified;
+			}
+		}
+
+		Lines[0] = FString::Printf(TEXT("Total, %d"), NumTotalModified);
+		FFileHelper::SaveStringArrayToFile(Lines, *FilePath);
+	}
+}
+
+} // namespace UE::Cook
\ No newline at end of file
diff --git a/Engine/Source/Editor/UnrealEd/Private/Cooker/CookDiagnostics.h b/Engine/Source/Editor/UnrealEd/Private/Cooker/CookDiagnostics.h
--- a/Engine/Source/Editor/UnrealEd/Private/Cooker/CookDiagnostics.h
+++ b/Engine/Source/Editor/UnrealEd/Private/Cooker/CookDiagnostics.h
@@ -2,14 +2,19 @@
 
 #pragma once
 
+#include "Containers/Array.h"
 #include "Containers/ArrayView.h"
 #include "Containers/Map.h"
 #include "Containers/Set.h"
+#include "Containers/StringView.h"
+#include "UObject/TopLevelAssetPath.h"
 
 class ITargetPlatform;
 class UCookOnTheFlyServer;
+namespace UE::Cook { class FPackageArtifacts; }
 namespace UE::Cook { enum class EInstigator : uint8; }
 namespace UE::Cook { struct FPackageData; }
+namespace UE::Cook { struct FGenerationHelper; }
 
 namespace UE::Cook
 {
@@ -29,4 +34,54 @@ public:
 
 };
 
-}
\ No newline at end of file
+/** Reasons that CookRequestCluster can find to mark a package as IncrementallyModified. */
+enum class EIncrementallyModifiedReason : uint8
+{
+	None,
+	NotPreviouslyCooked,
+	TargetDomainKey,
+	TransitiveBuildDependency,
+	IncrementalCookDisabled,
+	NoGenerator,
+	GeneratorOtherPlatform,
+	Count,
+};
+const TCHAR* LexToString(UE::Cook::EIncrementallyModifiedReason Reason);
+
+/** Diagnostic data passed in from CookRequestCluster for why a package was found IncrementallyModified. */
+struct FIncrementallyModifiedContext
+{
+	FPackageData* PackageData = nullptr;
+	const ITargetPlatform* TargetPlatform = nullptr;
+	FGenerationHelper* GenerationHelper = nullptr;
+	const FPackageArtifacts* Artifacts = nullptr;
+	FTopLevelAssetPath AssetClass;
+	FName TransitiveBuildDependencyName = NAME_None;
+	EIncrementallyModifiedReason Reason = EIncrementallyModifiedReason::None;
+};
+
+/**
+ * Record diagnostics from CookRequestCluster for which packages were found modified and why, and
+ * write them out to files in MetaData.
+ */
+class FIncrementallyModifiedDiagnostics
+{
+public:
+	void AddPackage(FIncrementallyModifiedContext& Context);
+	void WriteToFile(FStringView OutputDirectory);
+
+private:
+	/** The saved diagnostic data about each reported modified Package. */
+	struct FPackageData
+	{
+		FName PackageName;
+		FTopLevelAssetPath AssetClass;
+		TArray<FString> ReasonLines;
+		bool bSelfModified = false;
+		EIncrementallyModifiedReason Reason = EIncrementallyModifiedReason::None;
+	};
+
+	TArray<FPackageData> PackageDatas;
+};
+
+} // namespace UE::Cook
\ No newline at end of file
diff --git a/Engine/Source/Editor/UnrealEd/Private/Cooker/CookIncrementalVersion.cpp b/Engine/Source/Editor/UnrealEd/Private/Cooker/CookIncrementalVersion.cpp
--- a/Engine/Source/Editor/UnrealEd/Private/Cooker/CookIncrementalVersion.cpp
+++ b/Engine/Source/Editor/UnrealEd/Private/Cooker/CookIncrementalVersion.cpp
@@ -8,5 +8,6 @@ namespace UE::Cook
 {
 
 // Change CookIncrementalVersion to a new guid when all packages in an incremental cook need to be invalidated
-FGuid CookIncrementalVersion( 0xD6B3741E, 0x72CD49EF, 0xBEDEFBA9, 0xDB7C61CF );
+FGuid CookIncrementalVersion(0x8EB02510, 0xA8414F57, 0x9D68A4A8, 0x2F8BC58E);
+
 }
diff --git a/Engine/Source/Editor/UnrealEd/Private/Cooker/CookPackageArtifacts.cpp b/Engine/Source/Editor/UnrealEd/Private/Cooker/CookPackageArtifacts.cpp
--- a/Engine/Source/Editor/UnrealEd/Private/Cooker/CookPackageArtifacts.cpp
+++ b/Engine/Source/Editor/UnrealEd/Private/Cooker/CookPackageArtifacts.cpp
@@ -55,7 +55,8 @@ bool FBuildDependencySet::HasKeyMatch(FName PackageName, const ITargetPlatform*
 	}
 	if (CurrentKey.IsZero())
 	{
-		ECurrentKeyResult Result = TryCalculateCurrentKey(PackageName, TargetPlatform, GenerationHelper);
+		ECurrentKeyResult Result = TryCalculateCurrentKey(PackageName, TargetPlatform, GenerationHelper,
+			false /* bUpdateValues */);
 		if (Result != ECurrentKeyResult::Success)
 		{
 			CurrentKey.Reset();
@@ -67,7 +68,7 @@ bool FBuildDependencySet::HasKeyMatch(FName PackageName, const ITargetPlatform*
 
 FBuildDependencySet::ECurrentKeyResult FBuildDependencySet::TryCalculateCurrentKey(FName PackageName,
 	const ITargetPlatform* TargetPlatform, FGenerationHelper* GenerationHelper,
-	TArray<TPair<ELogVerbosity::Type, FString>>* OutMessages)
+	bool bUpdateValues, TArray<TPair<ELogVerbosity::Type, FString>>* OutMessages)
 {
 	IAssetRegistry* AssetRegistry = IAssetRegistry::Get();
 	if (PackageName.IsNone())
@@ -141,7 +142,10 @@ FBuildDependencySet::ECurrentKeyResult FBuildDependencySet::TryCalculateCurrentK
 				}
 				if (PackageDigest.IsSuccessful())
 				{
-					PackageDependency.SetValue(PackageDigest.Hash);
+					if (bUpdateValues)
+					{
+						PackageDependency.SetValue(PackageDigest.Hash);
+					}
 					KeyBuilder.Update(PackageDependency.GetRawValue(), FCookDependency::ValueSizeInBytes);
 					
 					continue;
@@ -170,7 +174,10 @@ FBuildDependencySet::ECurrentKeyResult FBuildDependencySet::TryCalculateCurrentK
 				KeyBuilder.Update(&Marker, sizeof(Marker));
 				if (!Value.IsEmpty())
 				{
-					ConfigDependency.SetValue(FUtf8String(Value));
+					if (bUpdateValues)
+					{
+						ConfigDependency.SetValue(FUtf8String(Value));
+					}
 					KeyBuilder.Update(ConfigDependency.GetRawValue(), FCookDependency::ValueSizeInBytes);
 				}
 			}
@@ -179,7 +186,7 @@ FBuildDependencySet::ECurrentKeyResult FBuildDependencySet::TryCalculateCurrentK
 
 		case ECookDependency::NativeClass:
 		{
-			UE::EditorDomain::TryAppendClassDigests(Batch, KeyBuilder, Context);
+			UE::EditorDomain::TryAppendClassDigests(Batch, KeyBuilder, Context, bUpdateValues);
 			break;
 		}
 
@@ -196,13 +203,16 @@ FBuildDependencySet::ECurrentKeyResult FBuildDependencySet::TryCalculateCurrentK
 			Hashes.SetNum(PackageNames.Num());
 			FCoreRedirects::GetHashOfRedirectsAffectingPackages(PackageNames, Hashes);
 
-			for(int Index = 0; Index < Hashes.Num(); ++Index)
+			for (int Index = 0; Index < Hashes.Num(); ++Index)
 			{
 				const FBlake3Hash& Hash = Hashes[Index];
 				FCookDependency& RedirectionDependency = Batch[Index];
 
-				RedirectionDependency.SetValue(Hash);
-				KeyBuilder.Update(RedirectionDependency.GetRawValue(), FCookDependency::ValueSizeInBytes);
+				if (bUpdateValues)
+				{
+					RedirectionDependency.SetValue(Hash);
+				}
+				KeyBuilder.Update(&Hash, sizeof(Hash));
 			}
 
 			FCoreRedirects::AppendHashOfGlobalRedirects(KeyBuilder);
@@ -212,7 +222,12 @@ FBuildDependencySet::ECurrentKeyResult FBuildDependencySet::TryCalculateCurrentK
 		default:
 			for (FCookDependency& BatchDependency : Batch)
 			{
-				BatchDependency.UpdateHash(Context);
+				FBlake3Hash Hash;
+				BatchDependency.UpdateHash(Context, &Hash);
+				if (bUpdateValues)
+				{
+					BatchDependency.SetValue(Hash);
+				}
 			}
 			break;
 		}
@@ -225,6 +240,152 @@ FBuildDependencySet::ECurrentKeyResult FBuildDependencySet::TryCalculateCurrentK
 	return Result;
 }
 
+bool FBuildDependencySet::TryGetModifiedDependencies(FName PackageName, const ITargetPlatform* TargetPlatform,
+	FGenerationHelper* GenerationHelper, TArray<FString>& OutModifiedDependencies) const
+{
+	if (!bValid)
+	{
+		return false;
+	}
+	if (StoredKey.IsZero())
+	{
+		return false;
+	}
+	if (PackageName.IsNone())
+	{
+		return false;
+	}
+	FEditorDomain* EditorDomain = FEditorDomain::Get();
+	if (!EditorDomain)
+	{
+		return false;
+	}
+
+	bool bInvalidated = false;
+	FBlake3 KeyBuilder;
+	FCookDependencyContext Context(&KeyBuilder,
+		[](ELogVerbosity::Type Verbosity, FString&& Message)
+		{
+		},
+		[&bInvalidated](ELogVerbosity::Type Verbosity)
+		{
+			bInvalidated = true;
+		},
+		PackageName);
+
+	for (const FCookDependency& Dependency : Dependencies)
+	{
+		// Some CookDependency types can not handle Update being called, because their Update relies on functions
+		// only available outside of the CoreUObject module. Handle those types.
+		switch (Dependency.GetType())
+		{
+		case ECookDependency::Package:
+		{
+			FName DependencyPackageName = Dependency.GetPackageName();
+			UE::EditorDomain::FPackageDigest PackageDigest = EditorDomain->GetPackageDigest(DependencyPackageName);
+			if (!PackageDigest.IsSuccessful() && GenerationHelper)
+			{
+				PackageDigest = GenerationHelper->GetPackageDigest(DependencyPackageName, TargetPlatform);
+			}
+			if (!PackageDigest.IsSuccessful())
+			{
+				OutModifiedDependencies.Add(FString::Printf(TEXT("%s: %s: Could not get current PackageDigest"),
+					LexToString(Dependency.GetType()), *Dependency.GetDebugIdentifier()));
+			}
+			else
+			{
+				FBlake3Hash CurrentValue = FCookDependency::ConvertToHash(PackageDigest.Hash);
+				if (CurrentValue != Dependency.GetHashValue())
+				{
+					OutModifiedDependencies.Add(FString::Printf(TEXT("%s: %s: %s -> %s"),
+						LexToString(Dependency.GetType()), *Dependency.GetDebugIdentifier(),
+						*LexToString(Dependency.GetHashValue()), *LexToString(CurrentValue)));
+				}
+			}
+			break;
+		}
+
+		case ECookDependency::Config:
+		{
+#if UE_WITH_CONFIG_TRACKING
+			using namespace UE::ConfigAccessTracking;
+			FCookConfigAccessTracker& ConfigTracker = FCookConfigAccessTracker::Get();
+#endif
+			FString Value;
+#if UE_WITH_CONFIG_TRACKING
+			Value = ConfigTracker.GetValue(Dependency.GetConfigAccessData());
+#endif
+			FBlake3Hash CurrentValue = FCookDependency::ConvertToHash(FUtf8String(Value));
+			if (CurrentValue != Dependency.GetHashValue())
+			{
+				OutModifiedDependencies.Add(FString::Printf(TEXT("%s: %s: %s -> %s"),
+					LexToString(Dependency.GetType()), *Dependency.GetDebugIdentifier(),
+					*LexToString(Dependency.GetHashValue()), *LexToString(CurrentValue)));
+			}
+			break;
+		}
+
+		case ECookDependency::NativeClass:
+		{
+			FStringView ClassPath = Dependency.GetClassPath();
+			TOptional<FBlake3Hash> CurrentValue = UE::EditorDomain::TryGetClassDigest(FTopLevelAssetPath(ClassPath));
+			if (!CurrentValue)
+			{
+				OutModifiedDependencies.Add(FString::Printf(TEXT("%s: %s: Could not get current ClassDigest"),
+					LexToString(Dependency.GetType()), *Dependency.GetDebugIdentifier()));
+			}
+			if (*CurrentValue != Dependency.GetHashValue())
+			{
+				OutModifiedDependencies.Add(FString::Printf(TEXT("%s: %s: %s -> %s"),
+					LexToString(Dependency.GetType()), *Dependency.GetDebugIdentifier(),
+					*LexToString(Dependency.GetHashValue()), *LexToString(*CurrentValue)));
+			}
+			break;
+		}
+
+		case ECookDependency::RedirectionTarget:
+		{
+			FName DependencyPackageName = Dependency.GetPackageName();
+
+			TArray<FName, TInlineAllocator<1>> PackageNames;
+			PackageNames.Add(DependencyPackageName);
+			TArray<FBlake3Hash> Hashes;
+			Hashes.SetNum(PackageNames.Num());
+			FCoreRedirects::GetHashOfRedirectsAffectingPackages(PackageNames, Hashes);
+
+			FBlake3Hash& CurrentValue = Hashes[0];
+			if (CurrentValue != Dependency.GetHashValue())
+			{
+				OutModifiedDependencies.Add(FString::Printf(TEXT("%s: %s: %s -> %s"),
+					LexToString(Dependency.GetType()), *Dependency.GetDebugIdentifier(),
+					*LexToString(Dependency.GetHashValue()), *LexToString(CurrentValue)));
+			}
+			break;
+		}
+
+		default:
+		{
+			bInvalidated = false;
+			FBlake3Hash CurrentValue;
+			Dependency.UpdateHash(Context, &CurrentValue);
+			if (CurrentValue != Dependency.GetHashValue())
+			{
+				OutModifiedDependencies.Add(FString::Printf(TEXT("%s: %s: %s -> %s"),
+					LexToString(Dependency.GetType()), *Dependency.GetDebugIdentifier(),
+					*LexToString(Dependency.GetHashValue()), *LexToString(CurrentValue)));
+			}
+			else if (bInvalidated)
+			{
+				OutModifiedDependencies.Add(FString::Printf(TEXT("%s: %s: ReportInvalidated was called"),
+					LexToString(Dependency.GetType()), *Dependency.GetDebugIdentifier()));
+			}
+			break;
+		}
+		}
+	}
+	return true;
+}
+
 void FBuildDependencySet::Empty()
 {
 	Dependencies.Empty();
@@ -812,10 +973,42 @@ bool FPackageArtifacts::HasKeyMatch(const ITargetPlatform* TargetPlatform, FGene
 	return bValid && SaveBuildDependencies.HasKeyMatch(PackageName, TargetPlatform, GenerationHelper);
 }
 
+bool FPackageArtifacts::TryGetModifiedDependencies(const ITargetPlatform* TargetPlatform,
+	FGenerationHelper* GenerationHelper, TArray<FString>& OutModifiedDependencies) const
+{
+	if (!bValid)
+	{
+		return false;
+	}
+	return SaveBuildDependencies.TryGetModifiedDependencies(PackageName, TargetPlatform, GenerationHelper,
+		OutModifiedDependencies);
+}
+
+TOptional<bool> FPackageArtifacts::GetIsPackageModified() const
+{
+	if (!bValid)
+	{
+		return TOptional<bool>();
+	}
+	IAssetRegistry* AssetRegistry = IAssetRegistry::Get();
+	if (!AssetRegistry)
+	{
+		return TOptional<bool>();
+	}
+
+	TOptional<FAssetPackageData> CurrentValue = AssetRegistry->GetAssetPackageDataCopy(PackageName);
+	if (!CurrentValue)
+	{
+		return TOptional<bool>();
+	}
+	return TOptional<bool>(CurrentValue->GetPackageSavedHash() != StoredPackageSavedHash);
+}
+
 FBuildDependencySet::ECurrentKeyResult FPackageArtifacts::TryCalculateCurrentKey(const ITargetPlatform* TargetPlatform,
-	FGenerationHelper* GenerationHelper, TArray<TPair<ELogVerbosity::Type, FString>>* OutMessages)
+	FGenerationHelper* GenerationHelper, bool bUpdateValues, TArray<TPair<ELogVerbosity::Type, FString>>* OutMessages)
 {
-	return SaveBuildDependencies.TryCalculateCurrentKey(PackageName, TargetPlatform, GenerationHelper, OutMessages);
+	return SaveBuildDependencies.TryCalculateCurrentKey(PackageName, TargetPlatform, GenerationHelper, bUpdateValues,
+		OutMessages);
 }
 
 void FPackageArtifacts::Empty()
@@ -876,6 +1069,13 @@ bool LoadFromCompactBinary(FCbObjectView ObjectView, UE::Cook::FPackageArtifacts
 				return false;
 			}
 		}
+		if (FieldView.GetName().Equals(UTF8TEXTVIEW("PackageSavedHash")))
+		{
+			if (!LoadFromCompactBinary(FieldView++, Artifacts.StoredPackageSavedHash))
+			{
+				return false;
+			}
+		}
 		if (FieldView == Last)
 		{
 			++FieldView;
@@ -908,6 +1108,7 @@ FCbWriter& operator<<(FCbWriter& Writer, const UE::Cook::FPackageArtifacts& Arti
 	{
 		Writer << "RuntimeDependencies" << Artifacts.RuntimeDependencies;
 	}
+	Writer << "PackageSavedHash" << Artifacts.StoredPackageSavedHash;
 
 	Writer.EndObject();
 	return Writer;
@@ -1225,6 +1426,8 @@ FPackageArtifacts FPackageArtifacts::Collect(const UPackage* Package, const ITar
 	FPackageArtifacts Result;
 	Result.PackageName = PackageName;
 	Result.bHasSaveResults = bHasSaveResult;
+	TOptional<FAssetPackageData> AssetPackageData = AssetRegistry->GetAssetPackageDataCopy(PackageName);
+	Result.StoredPackageSavedHash = AssetPackageData ? AssetPackageData->GetPackageSavedHash() : FIoHash();
 
 	// Store input+collected RuntimeDependencies on the result
 	Algo::Sort(InRuntimeDependencies, FNameLexicalLess());
@@ -1236,7 +1439,7 @@ FPackageArtifacts FPackageArtifacts::Collect(const UPackage* Package, const ITar
 	Result.LoadBuildDependencies.SetName(BuildResult::NAME_Load);
 	Result.LoadBuildDependencies.SetNormalizedDependencies(MoveTemp(LoadDependencies));
 		Result.LoadBuildDependencies.TryCalculateCurrentKey(PackageName, TargetPlatform,
-			nullptr /* GenerationHelper */, nullptr /* OutMessages */);
+			nullptr /* GenerationHelper */, true /* bUpdateValues */, nullptr /* OutMessages */);
 		Result.LoadBuildDependencies.StoreCurrentKey();
 		Result.LoadBuildDependencies.SetValid(true);
 
@@ -1251,7 +1454,7 @@ FPackageArtifacts FPackageArtifacts::Collect(const UPackage* Package, const ITar
 	Result.SaveBuildDependencies.SetName(BuildResult::NAME_Save);
 	Result.SaveBuildDependencies.SetNormalizedDependencies(MoveTemp(SaveDependencies));
 	FBuildDependencySet::ECurrentKeyResult CurrentKeyResult = Result.SaveBuildDependencies.TryCalculateCurrentKey(
-		PackageName, TargetPlatform, GenerationHelper, OutMessages);
+		PackageName, TargetPlatform, GenerationHelper, true /* bUpdateValues */, OutMessages);
 	if (CurrentKeyResult == FBuildDependencySet::ECurrentKeyResult::Error)
 	{
 		return FPackageArtifacts();
diff --git a/Engine/Source/Editor/UnrealEd/Private/Cooker/CookPackageArtifacts.h b/Engine/Source/Editor/UnrealEd/Private/Cooker/CookPackageArtifacts.h
--- a/Engine/Source/Editor/UnrealEd/Private/Cooker/CookPackageArtifacts.h
+++ b/Engine/Source/Editor/UnrealEd/Private/Cooker/CookPackageArtifacts.h
@@ -14,6 +14,7 @@
 #include "DerivedDataBuildDefinition.h"
 #include "IO/IoHash.h"
 #include "Logging/LogVerbosity.h"
+#include "Misc/Optional.h"
 #include "Serialization/PackageWriter.h"
 #include "Templates/Tuple.h"
 #include "UObject/NameTypes.h"
@@ -102,7 +103,15 @@ public:
 	const FIoHash& GetCurrentKey() const;
 	/** Call TryCalculateCurrentKey if not yet called, and return whether StoredKey == CurrentKey. */
 	bool HasKeyMatch(FName PackageName, const ITargetPlatform* TargetPlatform, FGenerationHelper* GenerationHelper);
-
+	/**
+	 * Report a StringDescriptor for each dependency that was found to be modified, for writing log output
+	 * about why the BuildDependencySet's current key is different than its stored key.
+	 * 
+	 * @return True if data is valid and modified dependencies can be calculated, else false. Will return
+	 *         true but with an empty list if not modified.
+	 */
+	bool TryGetModifiedDependencies(FName PackageName, const ITargetPlatform* TargetPlatform,
+		FGenerationHelper* GenerationHelper, TArray<FString>& OutModifiedDependencies) const;
 	enum class ECurrentKeyResult
 	{
 		Success,
@@ -114,9 +123,11 @@ public:
 	 *
 	 * @param GenerationHelper - If non-null, provides the lookup for the AssetPackageData of generated packages.
 	 *                           Must be provided if any generated packages are in the dependencies.
+	 * @param bUpdateValues - If true, also updates the stored value on each CookDependency.
 	 */
 	ECurrentKeyResult TryCalculateCurrentKey(FName PackageName, const ITargetPlatform* TargetPlatform,
-		FGenerationHelper* GenerationHelper, TArray<TPair<ELogVerbosity::Type, FString>>* OutMessages = nullptr);
+		FGenerationHelper* GenerationHelper, bool bUpdateValues,
+		TArray<TPair<ELogVerbosity::Type, FString>>* OutMessages = nullptr);
 
 	/** Clear data (except Name) and free memory. */
 	void Empty();
@@ -226,9 +237,11 @@ public:
 	 * @param TargetPlatform - Which targetplatform we are collecting for.
 	 * @param GenerationHelper - If non-null, provides the lookup for the AssetPackageData of generated packages.
 	 *                           Must be provided if any generated packages are in the dependencies.
+	 * @param bUpdateValues - If true, also updates the stored value on each CookDependency.
 	 */
 	FBuildDependencySet::ECurrentKeyResult TryCalculateCurrentKey(const ITargetPlatform* TargetPlatform,
-		FGenerationHelper* GenerationHelper, TArray<TPair<ELogVerbosity::Type, FString>>* OutMessages = nullptr);
+		FGenerationHelper* GenerationHelper, bool bUpdateValues,
+		TArray<TPair<ELogVerbosity::Type, FString>>* OutMessages = nullptr);
 	void Empty();
 
 	// Construction functions
@@ -257,10 +270,30 @@ public:
 	const FIoHash& GetCurrentKey() const;
 	bool HasKeyMatch(const ITargetPlatform* TargetPlatform, FGenerationHelper* GenerationHelper);
 
+	/**
+	 * Report a StringDescriptor for each dependency that was found to be modified, for writing log output
+	 * about why the HasKeyMatch returned false.
+	 * 
+	 * @return True if data is valid and modified dependencies can be calculated, else false. Will return
+	 *         true but with an empty list if not modified.
+	 */
+	bool TryGetModifiedDependencies(const ITargetPlatform* TargetPlatform, FGenerationHelper* GenerationHelper,
+		TArray<FString>& OutModifiedDependencies) const;
+
+	/**
+	 * Return whether the current Package's PackageSavedHash (from its package header on disk) is different than
+	 * it was during the previous cook that recorded the PackageArtifacts.
+	 *
+	 * @return An optional with true if modified, false if unmodified, and unset if data is invalid and whether
+	 *         it is modified is therefore unknown.
+	 */
+	TOptional<bool> GetIsPackageModified() const;
+
 private:
 	FBuildDependencySet LoadBuildDependencies;
 	FBuildDependencySet SaveBuildDependencies;
 	TArray<FName> RuntimeDependencies;
+	FIoHash StoredPackageSavedHash;
 	FName PackageName;
 	bool bHasSaveResults = false;
 	bool bValid = false;
diff --git a/Engine/Source/Editor/UnrealEd/Private/Cooker/CookRequestCluster.cpp b/Engine/Source/Editor/UnrealEd/Private/Cooker/CookRequestCluster.cpp
--- a/Engine/Source/Editor/UnrealEd/Private/Cooker/CookRequestCluster.cpp
+++ b/Engine/Source/Editor/UnrealEd/Private/Cooker/CookRequestCluster.cpp
@@ -10,6 +10,7 @@
 #include "AssetRegistry/IAssetRegistry.h"
 #include "Async/Async.h"
 #include "Cooker/CookDependency.h"
+#include "Cooker/CookDiagnostics.h"
 #include "Cooker/CookGenerationHelper.h"
 #include "Cooker/CookImportsChecker.h"
 #include "Cooker/CookLogPrivate.h"
@@ -1580,13 +1581,22 @@ bool FRequestCluster::FGraphSearch::FExploreEdgesContext::TryCalculateIncrementa
 		FPackagePlatformData& PackagePlatformData = PackageData->FindOrAddPlatformData(TargetPlatform);
 
 		UE::Cook::FPackageArtifacts& Artifacts = QueryPlatformData.CookAttachments.Artifacts;
+		bool bPreviouslyCooked =
+			(QueryPlatformData.CookAttachments.CommitStatus == IPackageWriter::ECommitStatus::Success)
+			| (QueryPlatformData.CookAttachments.CommitStatus == IPackageWriter::ECommitStatus::Error);
+		FIncrementallyModifiedContext SetModifiedArgs;
+		SetModifiedArgs.Artifacts = &Artifacts;
+		SetModifiedArgs.GenerationHelper = GenerationHelper.GetReference();
+		SetModifiedArgs.TargetPlatform = TargetPlatform;
+
 		if (PackageData->IsGenerated())
 		{
 			if (!GenerationHelper)
 			{
 				// Generated packages for which we do not have the GenerationHelper available are not incrementally
 				// skippable.
-				SetIncrementallyUnmodified(PlatformIndex, false, PackagePlatformData);
+				SetIncrementallyModified(PlatformIndex, PackagePlatformData, bPreviouslyCooked,
+					EIncrementallyModifiedReason::NoGenerator, SetModifiedArgs);
 				continue;
 			}
 			check(ParentPackageData); // Set above in the IsGenerated && GenerationHelper case
@@ -1598,14 +1608,17 @@ bool FRequestCluster::FGraphSearch::FExploreEdgesContext::TryCalculateIncrementa
 			{
 				if (ParentPlatformData->IsIncrementallyUnmodified())
 				{
-					SetIncrementallyUnmodified(PlatformIndex, true, PackagePlatformData);
+					SetIncrementallyUnmodified(PlatformIndex, PackagePlatformData);
 					continue;
 				}
 			}
 		}
 		if (!Artifacts.HasKeyMatch(TargetPlatform, GenerationHelper.GetReference()))
 		{
-			SetIncrementallyUnmodified(PlatformIndex, false, PackagePlatformData);
+			SetIncrementallyModified(PlatformIndex, PackagePlatformData, bPreviouslyCooked,
+				Artifacts.IsValid() ? EIncrementallyModifiedReason::TargetDomainKey
+				: EIncrementallyModifiedReason::NotPreviouslyCooked,
+				 SetModifiedArgs);
 			continue;
 		}
 
@@ -1613,17 +1626,19 @@ bool FRequestCluster::FGraphSearch::FExploreEdgesContext::TryCalculateIncrementa
 		// if they would otherwise qualify for incremental on their own. e.g. if worlds are incrementally disallowed,
 		// then streamingobject generated packages of the world are also disallowed.
 		FName PackageNameForIncrementalTest = ParentPackageData ? ParentPackageData->GetPackageName() : PackageName;
-		bool bIncrementalCookEnabled = IsIncrementalCookEnabled(PackageNameForIncrementalTest, Cluster.COTFS.bCookIncrementalAllowAllClasses);
+		bool bIncrementalCookEnabled = IsIncrementalCookEnabled(PackageNameForIncrementalTest,
+			Cluster.COTFS.bCookIncrementalAllowAllClasses);
 		
 		if (!bIncrementalCookEnabled)
 		{
-			SetIncrementallyUnmodified(PlatformIndex, false, PackagePlatformData);
+			SetIncrementallyModified(PlatformIndex, PackagePlatformData, bPreviouslyCooked,
+				EIncrementallyModifiedReason::IncrementalCookDisabled, SetModifiedArgs);
 			continue;
 		}
 
 		if (!QueryPlatformData.bTransitiveBuildDependenciesResolvedAsNotModified)
 		{
-			bool bAnyTransitiveBuildDependencyIsModified = false;
+			FName ModifiedTransitiveBuildDependency;
 			UnreadyTransitiveBuildVertices.Reset();
 			TArray<FName, TInlineAllocator<10>> TransitiveBuildDependencies;
 			Artifacts.GetTransitiveBuildDependencies(TransitiveBuildDependencies);
@@ -1643,7 +1658,7 @@ bool FRequestCluster::FGraphSearch::FExploreEdgesContext::TryCalculateIncrementa
 						TEXT(" Package %s will be marked as not incrementally skippable and will be recooked."),
 						*Vertex->GetPackageName().ToString(), *TransitiveBuildPackageName.ToString(),
 						*Vertex->GetPackageName().ToString());
-					bAnyTransitiveBuildDependencyIsModified = true;
+					ModifiedTransitiveBuildDependency = TransitiveBuildPackageName;
 					break;
 				}
 
@@ -1655,14 +1670,16 @@ bool FRequestCluster::FGraphSearch::FExploreEdgesContext::TryCalculateIncrementa
 				}
 				if (!TransitivePlatformData.bIncrementallyUnmodified.GetValue())
 				{
-					bAnyTransitiveBuildDependencyIsModified = true;
+					ModifiedTransitiveBuildDependency = TransitiveBuildPackageName;
 					break;
 				}
 			}
 
-			if (bAnyTransitiveBuildDependencyIsModified)
+			if (!ModifiedTransitiveBuildDependency.IsNone())
 			{
-				SetIncrementallyUnmodified(PlatformIndex, false, PackagePlatformData);
+				SetModifiedArgs.TransitiveBuildDependencyName = ModifiedTransitiveBuildDependency;
+				SetIncrementallyModified(PlatformIndex, PackagePlatformData, bPreviouslyCooked,
+					EIncrementallyModifiedReason::TransitiveBuildDependency, SetModifiedArgs);
 				continue;
 			}
 			if (!UnreadyTransitiveBuildVertices.IsEmpty())
@@ -1689,7 +1706,7 @@ bool FRequestCluster::FGraphSearch::FExploreEdgesContext::TryCalculateIncrementa
 			}
 		}
 
-		SetIncrementallyUnmodified(PlatformIndex, true, PackagePlatformData);
+		SetIncrementallyUnmodified(PlatformIndex, PackagePlatformData);
 	}
 
 	if (!bAllPlatformsAreReady)
@@ -1725,8 +1742,19 @@ bool FRequestCluster::FGraphSearch::FExploreEdgesContext::TryCalculateIncrementa
 				}
 				FFetchPlatformData& FetchPlatformData = GraphSearch.FetchPlatforms[PlatformIndex];
 				const ITargetPlatform* TargetPlatform = FetchPlatformData.Platform;
+				FQueryPlatformData& QueryPlatformData = Vertex->GetPlatformData()[PlatformIndex];
+				bool bPreviouslyCooked =
+					(QueryPlatformData.CookAttachments.CommitStatus == IPackageWriter::ECommitStatus::Success)
+					| (QueryPlatformData.CookAttachments.CommitStatus == IPackageWriter::ECommitStatus::Error);
+
+				FIncrementallyModifiedContext SetModifiedArgs;
+				SetModifiedArgs.Artifacts = &QueryPlatformData.CookAttachments.Artifacts;
+				SetModifiedArgs.GenerationHelper = GenerationHelper.GetReference();
+				SetModifiedArgs.TargetPlatform = TargetPlatform;
+
 				FPackagePlatformData& PackagePlatformData = PackageData->FindOrAddPlatformData(TargetPlatform);
-				SetIncrementallyUnmodified(PlatformIndex, false, PackagePlatformData);
+				SetIncrementallyModified(PlatformIndex, PackagePlatformData, bPreviouslyCooked,
+					EIncrementallyModifiedReason::GeneratorOtherPlatform, SetModifiedArgs);
 				PackagePlatformData.SetIncrementallyUnmodified(false);
 			}
 		}
@@ -2196,24 +2224,7 @@ void FRequestCluster::FGraphSearch::FExploreEdgesContext::ProcessPlatformAttachm
 			Context.bPreviouslyCooked = bPreviouslyCooked;
 			Context.bInOutShouldIncrementallySkip = bShouldIncrementallySkip;
 			PackageWriter->UpdatePackageModifiedStatus(Context);
-			bShouldIncrementallySkip = Context.bInOutShouldIncrementallySkip;
-		
-#if ENABLE_COOK_STATS
-			if (PlatformIndex == FirstSessionPlatformIndex && QueryPlatformData.bIncrementallyModifiedInstigated &&
-				bPreviouslyCooked)
-			{
-				FTopLevelAssetPath AssetClassPath;
-				FARFilter PackageNameFilter;
-				PackageNameFilter.PackageNames.Add(PackageData->GetPackageName());
-				Cluster.AssetRegistry.EnumerateAssets(PackageNameFilter, [&AssetClassPath](const FAssetData& AssetData)
-					{
-						AssetClassPath = AssetData.AssetClassPath;
-						return false;
-					});
-
-				++DetailedCookStats::NumPackagesIncrementallyModifiedByClass.FindOrAdd(AssetClassPath, 0);
-			}
-#endif
+			bShouldIncrementallySkip = Context.bInOutShouldIncrementallySkip;		
 		}
 
 		TRefCountPtr<FGenerationHelper> ParentGenerationHelper;
@@ -2325,18 +2336,47 @@ void FRequestCluster::FGraphSearch::FExploreEdgesContext::ProcessPlatformAttachm
 }
 
 void FRequestCluster::FGraphSearch::FExploreEdgesContext::SetIncrementallyUnmodified(int32 PlatformIndex,
-	bool bIncrementallyUnmodified, FPackagePlatformData& PackagePlatformData)
+	FPackagePlatformData& PackagePlatformData)
 {
 	FQueryPlatformData& ClusterPlatformData = Vertex->GetPlatformData()[PlatformIndex];
-	ClusterPlatformData.bIncrementallyUnmodified.Emplace(bIncrementallyUnmodified);
-	if (bIncrementallyUnmodified)
-	{
-		PackagePlatformData.SetIncrementallyUnmodified(true);
-	}
-	else if (!PackagePlatformData.IsIncrementallyUnmodifiedSet())
+	ClusterPlatformData.bIncrementallyUnmodified.Emplace(true);
+	PackagePlatformData.SetIncrementallyUnmodified(true);
+}
+
+void FRequestCluster::FGraphSearch::FExploreEdgesContext::SetIncrementallyModified(int32 PlatformIndex,
+	FPackagePlatformData& PackagePlatformData, bool bPreviouslyCooked, EIncrementallyModifiedReason Reason,
+	FIncrementallyModifiedContext& Context)
+{
+	FQueryPlatformData& ClusterPlatformData = Vertex->GetPlatformData()[PlatformIndex];
+	ClusterPlatformData.bIncrementallyUnmodified.Emplace(false);
+	if (!PackagePlatformData.IsIncrementallyUnmodifiedSet())
 	{
 		ClusterPlatformData.bIncrementallyModifiedInstigated = true;
 		PackagePlatformData.SetIncrementallyUnmodified(false);
+
+#if ENABLE_COOK_STATS
+		if (PlatformIndex == FirstSessionPlatformIndex && bPreviouslyCooked && Cluster.TraversalMarkCookable())
+		{
+			FTopLevelAssetPath AssetClassPath;
+			FARFilter PackageNameFilter;
+			PackageNameFilter.PackageNames.Add(PackageData->GetPackageName());
+			Cluster.AssetRegistry.EnumerateAssets(PackageNameFilter, [&AssetClassPath](const FAssetData& AssetData)
+				{
+					AssetClassPath = AssetData.AssetClassPath;
+					return false;
+				});
+
+			++DetailedCookStats::NumPackagesIncrementallyModifiedByClass.FindOrAdd(AssetClassPath, 0);
+
+			if (Cluster.COTFS.IncrementallyModifiedDiagnostics)
+			{
+				Context.PackageData = PackageData;
+				Context.Reason = Reason;
+				Context.AssetClass = AssetClassPath;
+				Cluster.COTFS.IncrementallyModifiedDiagnostics->AddPackage(Context);
+			}
+		}
+#endif
 	}
 }
 
diff --git a/Engine/Source/Editor/UnrealEd/Private/Cooker/CookRequestCluster.h b/Engine/Source/Editor/UnrealEd/Private/Cooker/CookRequestCluster.h
--- a/Engine/Source/Editor/UnrealEd/Private/Cooker/CookRequestCluster.h
+++ b/Engine/Source/Editor/UnrealEd/Private/Cooker/CookRequestCluster.h
@@ -30,10 +30,13 @@ class IAssetRegistry;
 class ICookedPackageWriter;
 class ITargetPlatform;
 class UCookOnTheFlyServer;
+namespace UE::Cook { class FPackageArtifacts; }
 namespace UE::Cook { class FRequestQueue; }
+namespace UE::Cook { enum class EIncrementallyModifiedReason : uint8; }
 namespace UE::Cook { enum class EReachability : uint8; }
 namespace UE::Cook { struct FDiscoveryQueueElement; }
 namespace UE::Cook { struct FFilePlatformRequest; }
+namespace UE::Cook { struct FIncrementallyModifiedContext; }
 namespace UE::Cook { struct FPackageData; }
 namespace UE::Cook { struct FPackageDatas; }
 namespace UE::Cook { struct FPackagePlatformData; }
@@ -396,8 +399,10 @@ private:
 				FIncrementalCookAttachments& PlatformAttachments, bool bExploreDependencies);
 			void ProcessPlatformDiscoveredDependencies(int32 PlatformIndex, const ITargetPlatform* TargetPlatform);
 
-			void SetIncrementallyUnmodified(int32 PlatformIndex, bool bIncrementallyUnmodified,
-				FPackagePlatformData& PackagePlatformData);
+			void SetIncrementallyUnmodified(int32 PlatformIndex, FPackagePlatformData& PackagePlatformData);
+			void SetIncrementallyModified(int32 PlatformIndex, FPackagePlatformData& PackagePlatformData,
+				bool bPreviouslyCooked, EIncrementallyModifiedReason Reason,
+				FIncrementallyModifiedContext& Context);
 
 		private:
 			FRequestCluster& Cluster;
diff --git a/Engine/Source/Editor/UnrealEd/Private/EditorDomain/EditorDomainUtils.cpp b/Engine/Source/Editor/UnrealEd/Private/EditorDomain/EditorDomainUtils.cpp
--- a/Engine/Source/Editor/UnrealEd/Private/EditorDomain/EditorDomainUtils.cpp
+++ b/Engine/Source/Editor/UnrealEd/Private/EditorDomain/EditorDomainUtils.cpp
@@ -459,62 +459,93 @@ FPackageDigest CalculatePackageDigest(const FAssetPackageData& PackageData, FNam
 	return Result;
 }
 
-bool TryAppendClassDigests(const TArrayView<UE::Cook::FCookDependency>& Dependencies, FBlake3& Writer, UE::Cook::FCookDependencyContext& Context)
+static const FBlake3Hash* GetClassDigestInternal(FTopLevelAssetPath ClassPath, FClassDigestMap& ClassDigests)
 {
-	FClassDigestMap& ClassDigests = GetClassDigests();
+	// Caller has read-locked ClassDigests.Lock.
+
+	const FClassDigestData* ExistingData = ClassDigests.Map.Find(ClassPath);
+	if (!ExistingData)
+	{
+		if (!IsInGameThread() || GIsSavingPackage)
+		{
+			return nullptr;
+		}
+
+		// PrecacheClassDigest will try to acquire a write lock. It can't upgrade a read lock to a write lock. That's why
+		// I release the read lock here.
+		ClassDigests.Lock.ReadUnlock();
+
+		PrecacheClassDigest(ClassPath);
+
+		ClassDigests.Lock.ReadLock();
+
+		ExistingData = ClassDigests.Map.Find(ClassPath);
+		if (!ExistingData)
+		{
+			return nullptr;
+		}
+	}
+
+	// We only support hashes for the ClosestNative class
+	const FClassDigestData* NativeData = ExistingData;
+	if (ExistingData->ClosestNative != ClassPath)
+	{
+		NativeData = ClassDigests.Map.Find(ExistingData->ClosestNative);
+		checkf(NativeData, TEXT("Classes are only stored in a ClosestNative field if they exist"));
+	}
+	check(NativeData->bNative);
+
+	return &NativeData->InclusiveSchemaHash;
+}
 
+bool TryAppendClassDigests(const TArrayView<UE::Cook::FCookDependency>& Dependencies, FBlake3& Writer,
+	UE::Cook::FCookDependencyContext& Context, bool bUpdateValues)
+{
+	FClassDigestMap& ClassDigests = GetClassDigests();
 	FReadScopeLock ClassDigestsScopeLock(ClassDigests.Lock);
 
 	for (UE::Cook::FCookDependency& ClassDependency : Dependencies)
 	{
-		FStringView ClassPath = ClassDependency.GetClassPath();
-		FTopLevelAssetPath AssetPath(ClassPath);
-		if (!AssetPath.IsValid())
+		FStringView ClassPathStr = ClassDependency.GetClassPath();
+		FTopLevelAssetPath ClassPath(ClassPathStr);
+		if (!ClassPath.IsValid())
 		{
-			Context.LogError(FString::Printf(TEXT("ClassDependency failed: %s is not a valid TopLevelAssetPath."), *AssetPath.ToString()));
+			Context.LogError(FString::Printf(TEXT("ClassDependency failed: %.*s is not a valid TopLevelAssetPath."),
+				ClassPathStr.Len(), ClassPathStr.GetData()));
 			continue;
 		}
 
-		FClassDigestData* ExistingData = ClassDigests.Map.Find(AssetPath);
-		if (!ExistingData && (!IsInGameThread() || GIsSavingPackage))
+		const FBlake3Hash* Hash = GetClassDigestInternal(ClassPath, ClassDigests);
+		if (!Hash)
 		{
-			Context.LogError(FString::Printf(TEXT("%s is not a valid class."), *AssetPath.ToString()));
+			Context.LogError(FString::Printf(TEXT("%s is not a valid class."), *ClassPath.ToString()));
 			return false;
 		}
-		else if (!ExistingData)
-		{
-			// PrecacheClassDigest will try to acquire a write lock. It can't upgrade a read lock to a write lock. That's why
-			// I release the read lock here.
-			ClassDigests.Lock.ReadUnlock();
-
-			PrecacheClassDigest(AssetPath);
-
-			ClassDigests.Lock.ReadLock();
-
-			ExistingData = ClassDigests.Map.Find(AssetPath);
-			if (!ExistingData)
-			{
-				Context.LogError(FString::Printf(TEXT("%s is not a valid class."), *AssetPath.ToString()));
-				return false;
-			}
-		}	
 
-		// We only support hashes for the ClosestNative class
-		FClassDigestData* NativeData = ExistingData;
-		if (ExistingData->ClosestNative != AssetPath)
+		if (bUpdateValues)
 		{
-			NativeData = ClassDigests.Map.Find(ExistingData->ClosestNative);
-			checkf(NativeData, TEXT("Classes are only stored in a ClosestNative field if they exist"));
+			ClassDependency.SetValue(*Hash);
 		}
-		check(NativeData->bNative);
-
-		ClassDependency.SetValue(NativeData->InclusiveSchemaHash);
-		Writer.Update(ClassDependency.GetRawValue(), UE::Cook::FCookDependency::ValueSizeInBytes);
+		Writer.Update(Hash->GetBytes(), sizeof(Hash->GetBytes()));
 	}
 
 	return true;
 }
 
+TOptional<FBlake3Hash> TryGetClassDigest(FTopLevelAssetPath ClassPath)
+{
+	if (!ClassPath.IsValid())
+	{
+		return TOptional<FBlake3Hash>();
+	}
+
+	FClassDigestMap& ClassDigests = GetClassDigests();
+	FReadScopeLock ClassDigestsScopeLock(ClassDigests.Lock);
+
+	const FBlake3Hash* Hash = GetClassDigestInternal(ClassPath, ClassDigests);
+	return Hash ? TOptional<FBlake3Hash>(*Hash) : TOptional<FBlake3Hash>();
+}
+
 /**
  * Holds context data for a call to PrecacheClassDigests, which needs to recursively
  * traverse a graph of of class parents and construction classes
diff --git a/Engine/Source/Editor/UnrealEd/Private/EditorDomain/EditorDomainUtils.h b/Engine/Source/Editor/UnrealEd/Private/EditorDomain/EditorDomainUtils.h
--- a/Engine/Source/Editor/UnrealEd/Private/EditorDomain/EditorDomainUtils.h
+++ b/Engine/Source/Editor/UnrealEd/Private/EditorDomain/EditorDomainUtils.h
@@ -14,6 +14,7 @@
 #include "Memory/MemoryFwd.h"
 #include "Misc/ConfigCacheIni.h"
 #include "Misc/EnumClassFlags.h"
+#include "Misc/Optional.h"
 #include "ProfilingDebugging/CookStats.h"
 #include "Serialization/CustomVersion.h"
 #include "Templates/Function.h"
@@ -94,7 +95,11 @@ FPackageDigest CalculatePackageDigest(IAssetRegistry& AssetRegistry, FName Packa
 FPackageDigest CalculatePackageDigest(const FAssetPackageData& PackageData, FName PackageName);
 
 /** FindOrAdd class digests for the given class dependencies, and append hashes onto the writer. Returns false if not found. */
-bool TryAppendClassDigests(const TArrayView<UE::Cook::FCookDependency>& Dependencies, FBlake3& Writer, UE::Cook::FCookDependencyContext& Context);
+bool TryAppendClassDigests(const TArrayView<UE::Cook::FCookDependency>& Dependencies, FBlake3& Writer,
+	UE::Cook::FCookDependencyContext& Context, bool bUpdateValues);
+
+/** FindOrAdd class digests for the given class dependencies, and return it. Returns empty optional if not found. */
+TOptional<FBlake3Hash> TryGetClassDigest(FTopLevelAssetPath ClassPath);
 
 /** For any ClassNames not already in ClassDigests, look up their UStruct and add them. */
 void PrecacheClassDigests(TConstArrayView<FTopLevelAssetPath> ClassNames);
diff --git a/Engine/Source/Runtime/CoreUObject/Private/Cooker/CookDependency.cpp b/Engine/Source/Runtime/CoreUObject/Private/Cooker/CookDependency.cpp
--- a/Engine/Source/Runtime/CoreUObject/Private/Cooker/CookDependency.cpp
+++ b/Engine/Source/Runtime/CoreUObject/Private/Cooker/CookDependency.cpp
@@ -42,6 +42,25 @@ const TMap<FName, FCookDependencyFunction>& GetCookDependencyFunctions();
 namespace UE::Cook
 {
 
+const TCHAR* LexToString(ECookDependency CookDependency)
+{
+	switch (CookDependency)
+	{
+	case ECookDependency::None: return TEXT("None");
+	case ECookDependency::File: return TEXT("File");
+	case ECookDependency::ConsoleVariable: return TEXT("ConsoleVariable");
+	case ECookDependency::NativeClass: return TEXT("NativeClass");
+	case ECookDependency::Function: return TEXT("Function");
+	case ECookDependency::TransitiveBuild: return TEXT("TransitiveBuild");
+	case ECookDependency::Package: return TEXT("Package");
+	case ECookDependency::Config: return TEXT("Config");
+	case ECookDependency::SettingsObject: return TEXT("SettingsObject");
+	case ECookDependency::AssetRegistryQuery: return TEXT("AssetRegistryQuery");
+	case ECookDependency::RedirectionTarget: return TEXT("RedirectionTarget");
+	default: return TEXT("Invalid");
+	}
+}
+
 namespace ResultProjection
 {
 
@@ -330,6 +349,36 @@ FCookDependency& FCookDependency::operator=(FCookDependency&& Other)
 	return *this;
 }
 
+FString FCookDependency::GetDebugIdentifier() const
+{
+	switch (Type)
+	{
+	case ECookDependency::None:
+		return FString();
+	case ECookDependency::File:
+	case ECookDependency::ConsoleVariable:
+	case ECookDependency::NativeClass:
+		return StringData;
+	case ECookDependency::Function:
+		return FunctionData.Name.ToString();
+	case ECookDependency::TransitiveBuild:
+		return TransitiveBuildData.PackageName.ToString();
+	case ECookDependency::Package:
+		return NameData.ToString();
+	case ECookDependency::Config:
+		return ConfigAccessData.IsValid() ? ConfigAccessData->FullPathToString() : FString();
+	case ECookDependency::SettingsObject:
+		return ObjectPtr ? ObjectPtr->GetPathName() : FString();
+	case ECookDependency::AssetRegistryQuery:
+		return FString();
+	case ECookDependency::RedirectionTarget:
+		return NameData.ToString();
+	default:
+		checkNoEntry();
+		return FString();
+	}
+}
+
 UE::ConfigAccessTracking::FConfigAccessData FCookDependency::GetConfigAccessData() const
 {
 	return Type == ECookDependency::Config && ConfigAccessData.IsValid() ?
@@ -342,17 +391,16 @@ FString FCookDependency::GetConfigPath() const
 		ConfigAccessData->FullPathToString() : FString();
 }
 
-const uint8* FCookDependency::GetRawValue() const
-{
-	return RawValue;
-}
-
-void FCookDependency::UpdateHash(FCookDependencyContext& Context)
+void FCookDependency::UpdateHash(FCookDependencyContext& Context, FBlake3Hash* OutCurrentValue) const
 {
 	switch (GetType())
 	{
 	case ECookDependency::None:
 		// Nothing to add, Nones are never invalidated
+		if (OutCurrentValue)
+		{
+			*OutCurrentValue = FBlake3Hash();
+		}
 		return;
 	case ECookDependency::File:
 	{
@@ -364,6 +412,10 @@ void FCookDependency::UpdateHash(FCookDependencyContext& Context)
 			Context.LogError(FString::Printf(
 				TEXT("FCookDependency::File('%s') failed to UpdateHash: could not read file."),
 				LocalFilename));
+			if (OutCurrentValue)
+			{
+				*OutCurrentValue = FBlake3Hash();
+			}
 			return;
 		}
 		TArray64<uint8> Buffer;
@@ -382,8 +434,11 @@ void FCookDependency::UpdateHash(FCookDependencyContext& Context)
 		}
 
 		FBlake3Hash Hash = Hasher.Finalize();
-		SetValue(Hash);
-		Context.Update(GetRawValue(), ValueSizeInBytes);
+		if (OutCurrentValue)
+		{
+			*OutCurrentValue = Hash;
+		}
+		Context.Update(Hash.GetBytes(), sizeof(Hash.GetBytes()));
 
 		break;
 	}
@@ -397,6 +452,10 @@ void FCookDependency::UpdateHash(FCookDependencyContext& Context)
 			Context.LogError(FString::Printf(
 				TEXT("FCookDependency::Function('%s') failed to UpdateHash: Function not found."),
 				*GetFunctionName().ToString()));
+			if (OutCurrentValue)
+			{
+				*OutCurrentValue = FBlake3Hash();
+			}
 			return;
 		}
 
@@ -413,19 +472,30 @@ void FCookDependency::UpdateHash(FCookDependencyContext& Context)
 		Context.SetHasher(OldHasher);
 
 		FBlake3Hash Hash = NewHasher.Finalize();
-		SetValue(Hash);
-		Context.Update(GetRawValue(), ValueSizeInBytes);
+		if (OutCurrentValue)
+		{
+			*OutCurrentValue = Hash;
+		}
+		Context.Update(Hash.GetBytes(), sizeof(Hash.GetBytes()));
 
 		return;
 	}
 	case ECookDependency::TransitiveBuild:
 		// Build dependencies do not impact the hash; they instead operate by marking the package
 		// as invalidated based on the invalidation of other packages, in a separate pass after its hash is compared
+		if (OutCurrentValue)
+		{
+			*OutCurrentValue = FBlake3Hash();
+		}
 		return;
 	case ECookDependency::Package:
 		Context.LogError(FString::Printf(
 			TEXT("FCookDependency::Package('%s') failed to UpdateHash: Package dependencies do not implement UpdateHash and it should not be called on them."),
 			*NameData.ToString()));
+		if (OutCurrentValue)
+		{
+			*OutCurrentValue = FBlake3Hash();
+		}
 		return;
 	case ECookDependency::ConsoleVariable:
 	{
@@ -435,12 +505,20 @@ void FCookDependency::UpdateHash(FCookDependencyContext& Context)
 			Context.LogError(FString::Printf(
 				TEXT("FCookDependency::ConsoleVariable('%s') failed to UpdateHash: could not find console variable."),
 				*StringData));
+			if (OutCurrentValue)
+			{
+				*OutCurrentValue = FBlake3Hash();
+			}
 			return;
 		}
 		FString VariableAsString = VariableInstance->GetString();
 
-		SetValue(FUtf8String(VariableAsString));
-		Context.Update(GetRawValue(), ValueSizeInBytes);
+		FBlake3Hash Hash = ConvertToHash(FUtf8String(VariableAsString));
+		if (OutCurrentValue)
+		{
+			*OutCurrentValue = Hash;
+		}
+		Context.Update(Hash.GetBytes(), sizeof(Hash.GetBytes()));
 
 		return;
 	}
@@ -448,16 +526,28 @@ void FCookDependency::UpdateHash(FCookDependencyContext& Context)
 		Context.LogError(FString::Printf(
 			TEXT("FCookDependency::Config('%s') failed to UpdateHash: Config dependencies do not implement UpdateHash and it should not be called on them."),
 			*GetConfigPath()));
+		if (OutCurrentValue)
+		{
+			*OutCurrentValue = FBlake3Hash();
+		}
 		return;
 	case ECookDependency::SettingsObject:
 		Context.LogError(FString::Printf(
 			TEXT("FCookDependency::SettingsObject('%s') failed to UpdateHash: SettingsObject dependencies do not implement UpdateHash and it should not be called on them."),
 			ObjectPtr ? *ObjectPtr->GetPathName() : TEXT("<null>")));
+		if (OutCurrentValue)
+		{
+			*OutCurrentValue = FBlake3Hash();
+		}
 		return;
 	case ECookDependency::NativeClass:
 		Context.LogError(FString::Printf(
 			TEXT("FCookDependency::NativeClass('%s') failed to UpdateHash: NativeClass dependencies do not implement UpdateHash and it should not be called on them."),
 			*StringData));
+		if (OutCurrentValue)
+		{
+			*OutCurrentValue = FBlake3Hash();
+		}
 		return;
 	case ECookDependency::AssetRegistryQuery:
 		if (ARFilter.IsValid())
@@ -484,8 +574,11 @@ void FCookDependency::UpdateHash(FCookDependencyContext& Context)
 				}
 
 				FBlake3Hash Hash = Hasher.Finalize();
-				SetValue(Hash);
-				Context.Update(GetRawValue(), ValueSizeInBytes);
+				if (OutCurrentValue)
+				{
+					*OutCurrentValue = Hash;
+				}
+				Context.Update(Hash.GetBytes(), sizeof(Hash.GetBytes()));
 			}
 		}
 		return;
@@ -493,6 +586,10 @@ void FCookDependency::UpdateHash(FCookDependencyContext& Context)
 		Context.LogError(FString::Printf(
 			TEXT("FCookDependency::RedirectionTarget('%s') failed to UpdateHash: RedirectionTarget dependencies do not implement UpdateHash and it should not be called on them."),
 			*NameData.ToString()));
+		if (OutCurrentValue)
+		{
+			*OutCurrentValue = FBlake3Hash();
+		}
 		return;
 	default:
 		checkNoEntry();
@@ -500,42 +597,49 @@ void FCookDependency::UpdateHash(FCookDependencyContext& Context)
 	}
 }
 
-void FCookDependency::SetValue(const FIoHash& Hash)
+void FCookDependency::SetValue(const FBlake3Hash& Hash)
+{
+	static_assert(sizeof(Hash) == ValueSizeInBytes);
+	HashedValue = Hash;
+}
+
+FBlake3Hash FCookDependency::ConvertToHash(const FIoHash& Hash)
 {
-	static_assert(sizeof(Hash) <= ValueSizeInBytes);
-	FMemory::Memcpy(RawValue, &Hash, sizeof(Hash));
+	FBlake3Hash OutHash;
+	constexpr size_t OutSize = sizeof(OutHash.GetBytes());
+	static_assert(sizeof(Hash) <= OutSize);
+	FMemory::Memcpy(OutHash.GetBytes(), &Hash, OutSize);
 
-	if constexpr (sizeof(Hash) < ValueSizeInBytes)
+	if constexpr (sizeof(Hash) < OutSize)
 	{
-		FMemory::Memzero(RawValue + sizeof(Hash), ValueSizeInBytes - sizeof(Hash));
+		FMemory::Memzero(OutHash.GetBytes() + sizeof(Hash), OutSize - sizeof(Hash));
 	}
+	return OutHash;
 }
 
-void FCookDependency::SetValue(const FBlake3Hash& Hash)
+FBlake3Hash FCookDependency::ConvertToHash(const FUtf8String& String)
 {
-	static_assert(sizeof(Hash) == ValueSizeInBytes);
-	HashedValue = Hash;
-}
+	FBlake3Hash OutHash;
+	constexpr size_t OutSize = sizeof(OutHash.GetBytes());
 
-void FCookDependency::SetValue(const FUtf8String& String)
-{
 	SIZE_T StringNumBytes = String.NumBytesWithoutNull();
-	if (StringNumBytes > ValueSizeInBytes)
+	if (StringNumBytes > OutSize)
 	{
 		FUtf8StringView View(String);
 
 		FBlake3 Hasher;
-		HashedValue = Hasher.HashBuffer(View.GetData(), View.NumBytes());
+		OutHash = Hasher.HashBuffer(View.GetData(), View.NumBytes());
 	}
 	else
 	{
-		FMemory::Memcpy(RawValue, String.GetCharArray().GetData(), StringNumBytes);
+		FMemory::Memcpy(OutHash.GetBytes(), String.GetCharArray().GetData(), StringNumBytes);
 
-		if (StringNumBytes < ValueSizeInBytes)
+		if (StringNumBytes < OutSize)
 		{
-			FMemory::Memzero(RawValue + StringNumBytes, ValueSizeInBytes - StringNumBytes);
+			FMemory::Memzero(OutHash.GetBytes() + StringNumBytes, OutSize - StringNumBytes);
 		}
 	}
+	return OutHash;
 }
 
 void FCookDependency::Construct()
diff --git a/Engine/Source/Runtime/CoreUObject/Public/Cooker/CookDependency.h b/Engine/Source/Runtime/CoreUObject/Public/Cooker/CookDependency.h
--- a/Engine/Source/Runtime/CoreUObject/Public/Cooker/CookDependency.h
+++ b/Engine/Source/Runtime/CoreUObject/Public/Cooker/CookDependency.h
@@ -6,6 +6,7 @@
 #include "Containers/StringView.h"
 #include "Containers/UnrealString.h"
 #include "HAL/PreprocessorHelpers.h"
+#include "Hash/Blake3.h"
 #include "Logging/LogVerbosity.h"
 #include "Serialization/CompactBinary.h"
 #include "Templates/Function.h"
@@ -127,6 +128,7 @@ enum class ECookDependency : uint8
 
 	Count,
 };
+COREUOBJECT_API const TCHAR* LexToString(ECookDependency CookDependency);
 
 namespace ResultProjection
 {
@@ -269,6 +271,9 @@ public:
 	/** FCookDependency is a vararg type. Return the type of this instance. */
 	ECookDependency GetType() const;
 
+	/** Return a string that describes which data is looked up by this CookDependency, within the context of its type. */
+	COREUOBJECT_API FString GetDebugIdentifier() const;
+
 	/** FileName if GetType() == File, else empty. StringView points to null or a null-terminated string. */
 	FStringView GetFileName() const;
 
@@ -307,6 +312,8 @@ public:
 
 	/** Returns the value of the dependency. It's an array of size FCookDependency::VALUE_MAX_BYTES_SIZE */
 	COREUOBJECT_API const uint8* GetRawValue() const;
+	/** Returns the value of the dependency, interpreted as a Blake3Hash. */
+	COREUOBJECT_API const FBlake3Hash& GetHashValue() const;
 
 	/**
 	 * Comparison operator for e.g. deterministic ordering of dependencies.
@@ -317,12 +324,15 @@ public:
 	bool operator==(const FCookDependency& Other) const;
 	bool operator!=(const FCookDependency& Other) const;
 
-	/** Calculate the current hash of this CookDependency, and add it into Context. */
-	COREUOBJECT_API void UpdateHash(FCookDependencyContext& Context);
+	/** Calculate the current hash of this CookDependency, add it into Context, and optionally return it. */
+	COREUOBJECT_API void UpdateHash(FCookDependencyContext& Context, FBlake3Hash* OutCurrentValue = nullptr) const;
 
-	COREUOBJECT_API void SetValue(const FIoHash& Hash);
 	COREUOBJECT_API void SetValue(const FBlake3Hash& Hash);
-	COREUOBJECT_API void SetValue(const FUtf8String& String);
+	void SetValue(const FIoHash& Hash);
+	void SetValue(const FUtf8String& String);
+
+	COREUOBJECT_API static FBlake3Hash ConvertToHash(const FIoHash& Hash);
+	COREUOBJECT_API static FBlake3Hash ConvertToHash(const FUtf8String& String);
 
 	static const uint8 ValueSizeInBytes = 32;
 
@@ -531,6 +541,16 @@ inline const FARFilter* FCookDependency::GetARFilter() const
 	return Type == ECookDependency::AssetRegistryQuery ? ARFilter.Get() : nullptr;
 }
 
+inline const uint8* FCookDependency::GetRawValue() const
+{
+	return RawValue;
+}
+
+inline const FBlake3Hash& FCookDependency::GetHashValue() const
+{
+	return HashedValue;
+}
+
 inline bool FCookDependency::operator<(const FCookDependency& Other) const
 {
 	if (static_cast<uint8>(Type) != static_cast<uint8>(Other.Type))
@@ -690,6 +710,16 @@ inline bool FCookDependency::operator!=(const FCookDependency& Other) const
 	return !(*this == Other);
 }
 
+inline void FCookDependency::SetValue(const FIoHash& Hash)
+{
+	SetValue(ConvertToHash(Hash));
+}
+
+inline void FCookDependency::SetValue(const FUtf8String& String)
+{
+	SetValue(ConvertToHash(String));
+}
+
 inline FCookDependencyContext::FCookDependencyContext(void* InHasher,
 	TUniqueFunction<void(ELogVerbosity::Type, FString&&)>&& InOnLog,
 	TUniqueFunction<void(ELogVerbosity::Type)>&& InOnInvalidated,
