The SDeleteAssetsDialog will now be called via SCustomDialog with Cancel and Confirm hotkeys enabled.

[RN]

diff --git a/Engine/Source/Editor/UnrealEd/Private/Dialogs/SDeleteAssetsDialog.cpp b/Engine/Source/Editor/UnrealEd/Private/Dialogs/SDeleteAssetsDialog.cpp
--- a/Engine/Source/Editor/UnrealEd/Private/Dialogs/SDeleteAssetsDialog.cpp
+++ b/Engine/Source/Editor/UnrealEd/Private/Dialogs/SDeleteAssetsDialog.cpp
@@ -379,7 +379,7 @@ TSharedRef<SWidget> SDeleteAssetsDialog::BuildDeleteDialog()
 				.FillWidth( 1.0f )
 				.Padding( 6, 0 )
 				[
-					SNew( SBorder )
+					SAssignNew( ForceDeleteButton, SBorder )
 					.BorderImage( FAppStyle::GetBrush( "NoBorder" ) )
 					.Visibility( this, &SDeleteAssetsDialog::GetForceDeleteVisibility )
 					[
@@ -395,7 +395,7 @@ TSharedRef<SWidget> SDeleteAssetsDialog::BuildDeleteDialog()
 					.BorderImage( FAppStyle::GetBrush( "NoBorder" ) )
 					.Visibility( this, &SDeleteAssetsDialog::GetDeleteVisibility )
 					[
-						SNew( SButton )
+						SAssignNew(DeleteButton, SButton )
 						.HAlign( HAlign_Center )
 						.Text( LOCTEXT( "Delete", "Delete" ) )
 						.ToolTipText( LOCTEXT( "DeleteTooltipText", "Perform the delete" ) )
@@ -583,20 +583,41 @@ TSharedRef<SWidget> SDeleteAssetsDialog::BuildForceDeleteWidget()
 
 FReply SDeleteAssetsDialog::OnKeyDown( const FGeometry& MyGeometry, const FKeyEvent& InKeyEvent )
 {
-	if( InKeyEvent.GetKey() == EKeys::Escape )
+	if ( ReferencerCommands->ProcessCommandBindings(InKeyEvent) )
 	{
-		if (TSharedPtr<SWindow> OwningWindow = ParentWindow.Pin())
-		{
-			OwningWindow->RequestDestroyWindow();
-		}
 		return FReply::Handled();
 	}
+	return FReply::Unhandled();
+}
 
-	if ( ReferencerCommands->ProcessCommandBindings(InKeyEvent) )
+void SDeleteAssetsDialog::OnCancelHotkeyPressed()
+{
+	Cancel();
+}
+
+void SDeleteAssetsDialog::OnConfirmHotkeyPressed()
+{
+	// If we're already in the process of deleting, we don't want to do anything
+	if (DeleteModel->GetState() != FAssetDeleteModel::Finished)
 	{
-		return FReply::Handled();
+		return;
 	}
-	return FReply::Unhandled();
+
+	// If Force Delete is available, trigger that
+	if (CanForceDelete())
+	{
+		ForceDelete();
+	}
+	// Otherwise trigger the normal Delete
+	else
+	{
+		Delete();
+	}
+}
+
+TSharedPtr<SWidget> SDeleteAssetsDialog::GetWidgetToFocusOnActivate()
+{
+	return CanForceDelete() ? ForceDeleteButton : DeleteButton;
 }
 
 EActiveTimerReturnType SDeleteAssetsDialog::TickDeleteModel( double InCurrentTime, float InDeltaTime )
diff --git a/Engine/Source/Editor/UnrealEd/Private/ObjectTools.cpp b/Engine/Source/Editor/UnrealEd/Private/ObjectTools.cpp
--- a/Engine/Source/Editor/UnrealEd/Private/ObjectTools.cpp
+++ b/Engine/Source/Editor/UnrealEd/Private/ObjectTools.cpp
@@ -3134,9 +3134,10 @@ namespace ObjectTools
 			const FVector2D DEFAULT_WINDOW_SIZE = FVector2D( 600, 700 );
 
 			/** Create the window to host our package dialog widget */
-			TSharedRef< SWindow > DeleteAssetsWindow = SNew( SWindow )
-				.Title( FText::FromString( "Delete Assets" ) )
-				.ClientSize( DEFAULT_WINDOW_SIZE );
+			TSharedRef<SCustomDialog> DeleteAssetsWindow = SNew(SCustomDialog)
+				.Title(FText::FromString("Delete Assets"))
+				.ClientSize(DEFAULT_WINDOW_SIZE)
+				.AutoCloseOnButtonPress(false);
 
 			/** Set the content of the window to our package dialog widget */
 			TSharedRef< SDeleteAssetsDialog > DeleteDialog =
@@ -3144,6 +3145,9 @@ namespace ObjectTools
 				.ParentWindow( DeleteAssetsWindow );
 
 			DeleteAssetsWindow->SetContent( DeleteDialog );
+			DeleteAssetsWindow->SetOnCancelHotkeyPressed(FExecuteAction::CreateSP(DeleteDialog, &SDeleteAssetsDialog::OnCancelHotkeyPressed));
+			DeleteAssetsWindow->SetOnConfirmHotkeyPressed(FExecuteAction::CreateSP(DeleteDialog, &SDeleteAssetsDialog::OnConfirmHotkeyPressed));
+			DeleteAssetsWindow->SetGetWidgetToFocusOnActivate(FGetFocusWidget::CreateSP(DeleteDialog, &SDeleteAssetsDialog::GetWidgetToFocusOnActivate));
 
 			/** Show the package dialog window as a modal window */
 			GEditor->EditorAddModalWindow( DeleteAssetsWindow );
diff --git a/Engine/Source/Editor/UnrealEd/Public/Dialogs/SDeleteAssetsDialog.h b/Engine/Source/Editor/UnrealEd/Public/Dialogs/SDeleteAssetsDialog.h
--- a/Engine/Source/Editor/UnrealEd/Public/Dialogs/SDeleteAssetsDialog.h
+++ b/Engine/Source/Editor/UnrealEd/Public/Dialogs/SDeleteAssetsDialog.h
@@ -54,6 +54,15 @@ public:
 
 	virtual FReply OnKeyDown( const FGeometry& MyGeometry, const FKeyEvent& InKeyEvent ) override;
 
+	// Function to be hooked up to the Cancel hotkey (Escape) for the dialog window
+	void OnCancelHotkeyPressed();
+
+	// Function to be hooked up to the Confirm hotkey (Enter) for the dialog window
+	void OnConfirmHotkeyPressed();
+
+	// Function to get the appropriate widget to focus when the parent window is activated
+	TSharedPtr<SWidget> GetWidgetToFocusOnActivate();
+
 private:
 	TSharedRef<SWidget> CreateThumbnailWidget();
 	TSharedRef<SWidget> MakeAssetViewForReferencerAssets();
@@ -155,6 +164,8 @@ private:
 	TSharedPtr< SBorder > AssetReferenceNameBorderWidget;
 	TSharedPtr< SComboButton > ConsolidationPickerComboButton;
 	TSharedPtr< SCheckBox > DeleteSourceFilesCheckbox;
+	TSharedPtr<SWidget> ForceDeleteButton;
+	TSharedPtr<SWidget> DeleteButton;
 
 	/** The selected asset we're going to consolidate the would be deleted assets into. */
 	FAssetData ConsolidationAsset;
