PR #13424: Implement fix for null terminating string in FAssetRegistryGenerator

diff --git a/Engine/Source/Editor/UnrealEd/Private/Commandlets/AssetRegistryGenerator.cpp b/Engine/Source/Editor/UnrealEd/Private/Commandlets/AssetRegistryGenerator.cpp
--- a/Engine/Source/Editor/UnrealEd/Private/Commandlets/AssetRegistryGenerator.cpp
+++ b/Engine/Source/Editor/UnrealEd/Private/Commandlets/AssetRegistryGenerator.cpp
@@ -1175,8 +1175,12 @@ TMap<int32, FString> FAssetRegistryGenerator::GetPakChunkIdToStringOverrideMap()
 	for (const TPair<int32, FString>& PakChunkToName : ChunkIdStringOverride)
 	{
 		const FString& PakChunkName = PakChunkToName.Value;
-		for (TCHAR Char : PakChunkName.GetCharArray())
+
+		// Iterate over the real characters. FString::Len() excludes the null terminator.
+		for (int32 CharIndex = 0; CharIndex < PakChunkName.Len(); ++CharIndex)
 		{
+			const TCHAR& Char = PakChunkName[CharIndex];
+
 			if (!ValidCharacters.Contains(Char))
 			{
 				UE_LOG(LogCook, Error, TEXT("GetPakChunkIdToStringMapping contains invalid string for chunk mapping. Character '%c' within %s. Only the following character are allowed:\n%s"), Char, *PakChunkName, VALID_SAVEDDIRSUFFIX_CHARACTERS);
