MPCook: Restore the cookdirector's memory and cpu limits after the CookWorkers have shut down and before CookByTheBookFinished, so we can use the full number of worker thread to go wide on operations in CookByTheBookFinished.
#rnx

diff --git a/Engine/Source/Editor/UnrealEd/Private/Cooker/CookDirector.cpp b/Engine/Source/Editor/UnrealEd/Private/Cooker/CookDirector.cpp
--- a/Engine/Source/Editor/UnrealEd/Private/Cooker/CookDirector.cpp
+++ b/Engine/Source/Editor/UnrealEd/Private/Cooker/CookDirector.cpp
@@ -322,6 +322,7 @@ void FCookDirector::StartCook(const FBeginCookContext& InBeginContext)
 	// work the Director has to do during the first tick, for PumpRequests and AssignRequests.
 	FScopeLock CommunicationScopeLock(&CommunicationLock);
 	InitializeWorkers();
+	ActivateMachineResourceReduction();
 }
 
 void FCookDirector::AssignRequests(TArrayView<FPackageData*> Requests, TArray<FWorkerId>& OutAssignments,
@@ -1016,6 +1017,11 @@ void FCookDirector::PumpCookComplete(bool& bCompleted)
 			}
 		}
 		bCompleted = !bWorkersActive;
+
+		if (bCompleted)
+		{
+			DeactivateMachineResourceReduction();
+		}
 	}
 	TickFromSchedulerThread();
 }
@@ -1416,12 +1422,17 @@ void FCookDirector::ActivateMachineResourceReduction()
 	// When running a multiprocess cook, we remove the Memory triggers and trigger GC based solely on PressureLevel.
 	// But keep the Soft GC settings
 	PRAGMA_DISABLE_DEPRECATION_WARNINGS
+	SavedResources.MemoryMaxUsedPhysical = COTFS.MemoryMaxUsedPhysical;
+	SavedResources.MemoryMaxUsedVirtual = COTFS.MemoryMaxUsedVirtual;
 	COTFS.MemoryMaxUsedPhysical = 0;
 	COTFS.MemoryMaxUsedVirtual = 0;
 	PRAGMA_ENABLE_DEPRECATION_WARNINGS
 
+	SavedResources.MemoryMinFreeVirtual = COTFS.MemoryMinFreeVirtual;
+	SavedResources.MemoryMinFreePhysical = COTFS.MemoryMinFreePhysical;
 	COTFS.MemoryMinFreeVirtual = 0;
 	COTFS.MemoryMinFreePhysical = 0;
+	SavedResources.MemoryTriggerGCAtPressureLevel = COTFS.MemoryTriggerGCAtPressureLevel;
 	COTFS.MemoryTriggerGCAtPressureLevel = FGenericPlatformMemoryStats::EMemoryPressureStatus::Critical;
 
 	UE_LOG(LogCook, Display, TEXT("CookMultiprocess changed CookSettings for Memory:%s"),
@@ -1433,6 +1444,8 @@ void FCookDirector::ActivateMachineResourceReduction()
 	int32 HyperThreadCount = FPlatformMisc::NumberOfCoresIncludingHyperthreads();
 	int32 NumberOfHyperThreadsPerCore = HyperThreadCount / NumberOfCores;
 	CoreLimit = FMath::Max(NumberOfCores / NumProcesses, 1);
+	SavedResources.NumberOfCores = NumberOfCores;
+	SavedResources.HyperThreadCount = HyperThreadCount;
 
 	const TCHAR* CommandLine = FCommandLine::Get();
 	float CoreOversubscription = 1.0f;
@@ -1446,9 +1459,9 @@ void FCookDirector::ActivateMachineResourceReduction()
 
 	// Update the number of Cores and WorkerThreads for this process
 	check(IsInGameThread());
-	int32 NumBackgroundWorkers = FMath::Max(1,
-		NumberOfWorkers - FMath::Min<int32>(GNumForegroundWorkers, NumberOfWorkers));
-	int32 NumForegroundWorkers = FMath::Max(1, NumberOfWorkers - NumBackgroundWorkers);
+	int32 NumForegroundWorkers = FMath::Min<uint32>(GNumForegroundWorkers, NumberOfWorkers);
+	int32 NumBackgroundWorkers = FMath::Max(1, NumberOfWorkers - NumForegroundWorkers);
+	NumForegroundWorkers = FMath::Max(1, NumberOfWorkers - NumBackgroundWorkers);
 	LowLevelTasks::FScheduler::Get().RestartWorkers(NumForegroundWorkers, NumBackgroundWorkers);
 
 	// Update the number of ShaderCompilerWorkers that can be launched
@@ -1460,6 +1473,38 @@ void FCookDirector::ActivateMachineResourceReduction()
 		HyperThreadCount, CoreIncludingHyperthreadsLimit);
 }
 
+void FCookDirector::DeactivateMachineResourceReduction()
+{
+	if (!bHasReducedMachineResources)
+	{
+		return;
+	}
+	bHasReducedMachineResources = false;
+
+	PRAGMA_DISABLE_DEPRECATION_WARNINGS
+	COTFS.MemoryMaxUsedPhysical = SavedResources.MemoryMaxUsedPhysical;
+	COTFS.MemoryMaxUsedVirtual = SavedResources.MemoryMaxUsedVirtual;
+	PRAGMA_ENABLE_DEPRECATION_WARNINGS
+
+	COTFS.MemoryMinFreeVirtual = SavedResources.MemoryMinFreeVirtual;
+	COTFS.MemoryMinFreePhysical = SavedResources.MemoryMinFreePhysical;
+	COTFS.MemoryTriggerGCAtPressureLevel = SavedResources.MemoryTriggerGCAtPressureLevel;
+
+	UE_LOG(LogCook, Display, TEXT("CookMultiprocess restored CookSettings for Memory:%s"),
+		*COTFS.GetCookSettingsForMemoryLogText());
+
+	check(IsInGameThread());
+	// Calling RestartWorkers with no arguments calculates NumForegroundWorkers and NumBackgroundWorkers based on the
+	// number of cores reported by FPlatformMisc; this is what we want for the restored-to-single-process worker state.
+	LowLevelTasks::FScheduler::Get().RestartWorkers();
+
+	GShaderCompilingManager->OnMachineResourcesChanged(SavedResources.NumberOfCores, SavedResources.HyperThreadCount);
+
+	UE_LOG(LogCook, Display, TEXT("CookMultiprocess restored number of cores to %d."), SavedResources.NumberOfCores);
+	UE_LOG(LogCook, Display, TEXT("CookMultiprocess restored number of hyperthreads to %d."),
+		SavedResources.HyperThreadCount);
+}
+
 void FCookDirector::TickWorkerConnects(ECookDirectorThread TickThread)
 {
 	using namespace UE::CompactBinaryTCP;
diff --git a/Engine/Source/Editor/UnrealEd/Private/Cooker/CookDirector.h b/Engine/Source/Editor/UnrealEd/Private/Cooker/CookDirector.h
--- a/Engine/Source/Editor/UnrealEd/Private/Cooker/CookDirector.h
+++ b/Engine/Source/Editor/UnrealEd/Private/Cooker/CookDirector.h
@@ -190,6 +190,18 @@ private:
 	};
 	class FRetractionHandler;
 
+	struct FMachineResourceValues
+	{
+		uint64 MemoryMaxUsedPhysical = 0;
+		uint64 MemoryMaxUsedVirtual = 0;
+		uint64 MemoryMinFreeVirtual = 0;
+		uint64 MemoryMinFreePhysical = 0;
+		FGenericPlatformMemoryStats::EMemoryPressureStatus MemoryTriggerGCAtPressureLevel
+			= FGenericPlatformMemoryStats::EMemoryPressureStatus::Unknown;
+		int32 NumberOfCores = 0;
+		int32 HyperThreadCount = 0;
+	};
+
 private:
 	/** Helper for constructor parsing. */
 	void ParseConfig(int32 CookProcessCount, bool& bOutValid);
@@ -209,6 +221,8 @@ private:
 	void RecreateWorkers();
 	/** Reduce memory settings, cpusettings, and anything else that needs to be shared with CookWorkers. */
 	void ActivateMachineResourceReduction();
+	/** Restore the settings that were reduced in ActivateMachineResourceReduction. */
+	void DeactivateMachineResourceReduction();
 	/** Start the communication thread if not already started. */
 	void LaunchCommunicationThread();
 	/** Signal the communication thread to stop, wait for it to finish, and deallocate it. */
@@ -283,6 +297,7 @@ private:
 	TUniquePtr<FCookWorkerProfileData> LocalWorkerProfileData;
 	TArray<TPair<UE::CompactBinaryTCP::FMarshalledMessage, ECookBroadcastTiming>> QueuedBroadcasts;
 	UCookOnTheFlyServer& COTFS;
+	FMachineResourceValues SavedResources;
 	double WorkersStalledStartTimeSeconds = 0.;
 	double WorkersStalledWarnTimeSeconds = 0.;
 	double LastTickTimeSeconds = 0.;
