Unify packagename source to make sure Actors&Cells use the same trimmed name all the time.

diff --git a/Engine/Source/Editor/UnrealEd/Private/WorldPartition/WorldPartitionStaticLightingBuilder.cpp b/Engine/Source/Editor/UnrealEd/Private/WorldPartition/WorldPartitionStaticLightingBuilder.cpp
--- a/Engine/Source/Editor/UnrealEd/Private/WorldPartition/WorldPartitionStaticLightingBuilder.cpp
+++ b/Engine/Source/Editor/UnrealEd/Private/WorldPartition/WorldPartitionStaticLightingBuilder.cpp
@@ -654,7 +654,7 @@ void FStaticLightingDescriptors::InitializeFromWorld(UWorld* InWorld)
 				FLightingActorDesc& LightingActorDesc = ActorGuidsToDesc.Add(Actor.ActorGuid);
 				LightingActorDesc.ActorGuid = Actor.ActorGuid;
 				LightingActorDesc.ActorPackage = Actor.Package;
-				LightingActorDesc.CellLevelPackage = Cell.CellPackage;
+				LightingActorDesc.CellLevelPackage = CellPackageName;
 
 				CellDesc.ActorInstanceGuids.Add(Actor.ActorGuid);
 			}
