Generate the AutogenShaderHeaders.usd files for every shader platforms we are about to cook. Those files can be hashed and used in dependencies so they need to be present. Only do this for the director.

#incrementalcook
#rnx

diff --git a/Engine/Source/Editor/UnrealEd/Private/CookOnTheFlyServer.cpp b/Engine/Source/Editor/UnrealEd/Private/CookOnTheFlyServer.cpp
--- a/Engine/Source/Editor/UnrealEd/Private/CookOnTheFlyServer.cpp
+++ b/Engine/Source/Editor/UnrealEd/Private/CookOnTheFlyServer.cpp
@@ -280,6 +280,20 @@ static FString GetReleaseVersionAssetRegistryPath(const FString& ReleaseVersion,
 	return (*ReleasesRoot) / ReleaseVersion / PlatformName;
 }
 
+static void GenerateAutogenShaderHeadersForAllShaderFormats(const TArray<ITargetPlatform*>& Platforms)
+{
+	for (const ITargetPlatform* Platform : Platforms)
+	{
+		TArray<FName> ShaderPlatformNameArray;
+		Platform->GetAllTargetedShaderFormats(ShaderPlatformNameArray);
+
+		for (const FName& ShaderPlatformName : ShaderPlatformNameArray)
+		{
+			FShaderCompileUtilities::GenerateBrdfHeaders(ShaderPlatformName);
+		}
+	}
+}
+
 template<typename T>
 struct FOneTimeCommandlineReader
 {
@@ -9051,8 +9065,17 @@ void UCookOnTheFlyServer::BeginCookStartShaderCodeLibrary(FBeginCookContext& Beg
 void UCookOnTheFlyServer::BeginCookFinishShaderCodeLibrary(FBeginCookContext& BeginContext)
 {
 	check(IsDirectorCookByTheBook()); // CookByTheBook only for now
+
+	// Generate the AutogenShaderHeaders.usd files for every shader platforms we are about to cook. Those files can be hashed and used
+	// in dependencies so they need to be present. Only do this for the director.
+	bool bIsCookWorker = IsCookWorkerMode();
+	if (!bIsCookWorker)
+	{
+		GenerateAutogenShaderHeadersForAllShaderFormats(BeginContext.TargetPlatforms);
+	}
+
 	// don't resave the global shader map files in dlc
-	if (!IsCookWorkerMode() && !IsCookingDLC() && !EnumHasAnyFlags(CookByTheBookOptions->StartupOptions, ECookByTheBookOptions::ForceDisableSaveGlobalShaders))
+	if (!bIsCookWorker && !IsCookingDLC() && !EnumHasAnyFlags(CookByTheBookOptions->StartupOptions, ECookByTheBookOptions::ForceDisableSaveGlobalShaders))
 	{
 		OpenGlobalShaderLibrary();
 
