Cooker: DiffOnly: Add optional -CookDiffJson parameter to output the results to a json file.

- refactor the message logging to route everything through a helper class.
- determinism diagnostic no longer reporting the ddc key by default.

#rnx

diff --git a/Engine/Source/Editor/UnrealEd/Private/Cooker/CookDeterminismManager.cpp b/Engine/Source/Editor/UnrealEd/Private/Cooker/CookDeterminismManager.cpp
--- a/Engine/Source/Editor/UnrealEd/Private/Cooker/CookDeterminismManager.cpp
+++ b/Engine/Source/Editor/UnrealEd/Private/Cooker/CookDeterminismManager.cpp
@@ -258,6 +258,11 @@ void FDeterminismConstructDiagnosticsContext::AddDiagnostic(FUtf8StringView Diag
 	ExportData.AddNewDiagnostic(DiagnosticName, Value);
 }
 
+bool FDeterminismConstructDiagnosticsContext::FullDDCKeysRequested() const
+{
+	return false;
+}
+
 FExportDeterminismData::FExportDeterminismData(FPackageDeterminismData& InPackageData, UObject* InExport)
 	: PackageData(InPackageData)
 	, Export(InExport)
diff --git a/Engine/Source/Editor/UnrealEd/Private/Cooker/CookDeterminismManager.h b/Engine/Source/Editor/UnrealEd/Private/Cooker/CookDeterminismManager.h
--- a/Engine/Source/Editor/UnrealEd/Private/Cooker/CookDeterminismManager.h
+++ b/Engine/Source/Editor/UnrealEd/Private/Cooker/CookDeterminismManager.h
@@ -32,6 +32,8 @@ public:
 	virtual const ITargetPlatform* GetTargetPlatform() override;
 	virtual void AddDiagnostic(FUtf8StringView DiagnosticName, const FCbField& Value) override;
 
+	virtual bool FullDDCKeysRequested() const override;
+
 public:
 	FExportDeterminismData& ExportData;
 };
diff --git a/Engine/Source/Editor/UnrealEd/Private/Cooker/DiffPackageWriter.cpp b/Engine/Source/Editor/UnrealEd/Private/Cooker/DiffPackageWriter.cpp
--- a/Engine/Source/Editor/UnrealEd/Private/Cooker/DiffPackageWriter.cpp
+++ b/Engine/Source/Editor/UnrealEd/Private/Cooker/DiffPackageWriter.cpp
@@ -17,8 +17,10 @@
 #include "Misc/Parse.h"
 #include "Misc/OutputDevice.h"
 #include "Misc/OutputDeviceHelper.h"
+#include "Misc/App.h"
 #include "ProfilingDebugging/CookStats.h"
 #include "Serialization/LargeMemoryWriter.h"
+#include "Serialization/JsonWriter.h"
 #include "Templates/UniquePtr.h"
 #include "Templates/UnrealTemplate.h"
 #include "UObject/LinkerSave.h"
@@ -28,6 +30,319 @@
 
 DEFINE_LOG_CATEGORY_STATIC(LogDiff, Log, All);
 
+
+static const TCHAR* const CallstackCutoffString = TEXT("UEditorEngine::Save()");
+
+namespace UE::DiffWriter
+{
+extern const TCHAR* const IndentToken;
+class FJsonDetailRecorder : public IDetailRecorder
+{
+public:
+	FJsonDetailRecorder( const FString& InFilename )
+	{
+		FString JsonFilename = InFilename;
+		if (UE::GetMultiprocessId() > 0)
+		{
+			FString JsonUniqueFilename = FString::Printf(TEXT("%s-%d.%s"), *FPaths::GetBaseFilename(InFilename), 
+				UE::GetMultiprocessId(), *FPaths::GetExtension(InFilename));
+			JsonFilename = FPaths::Combine( FPaths::GetPath(InFilename), JsonUniqueFilename);
+		}
+		UE_LOG(LogDiff, Display, TEXT("Using Json details file %s"), *JsonFilename);
+
+		JsonFile.Reset(IFileManager::Get().CreateFileWriter(*JsonFilename));
+		JsonWriter = TJsonWriterFactory<ANSICHAR>::Create(JsonFile.Get());
+
+		JsonWriter->WriteObjectStart(); // global
+		JsonWriter->WriteValue("ProjectName", FApp::GetProjectName());
+		JsonWriter->WriteValue("BuildVersion", FApp::GetBuildVersion());
+
+		JsonWriter->WriteArrayStart("Packages");
+	}
+
+	virtual ~FJsonDetailRecorder()
+	{
+		if (JsonFile.IsValid())
+		{
+			JsonWriter->WriteArrayEnd(); // "Packages"
+			JsonWriter->WriteObjectEnd(); // global
+			JsonWriter->Close();
+			JsonWriter.Reset();
+
+			JsonFile.Reset();
+		}
+	}
+
+	virtual void BeginPackage() override
+	{
+		check(State == EState::Idle);
+		State = EState::InPackage;
+	}
+
+	virtual void BeginSection(const TCHAR* SectionFilename, EDiffWriterSectionType SectionType, 
+		const FPackageData& SourceSection, const FPackageData& DestSection) override
+	{
+		check(State == EState::InPackage);
+
+		CurrentPackage.CurrentSectionIndex = (int32)SectionType;
+		State = EState::InSection;
+
+		FSection& Section = GetSection();
+		Section.SectionFilename = SectionFilename;
+		Section.SourceSize = SourceSection.Size;
+		Section.DestSize = DestSection.Size;
+
+		GetSection().SectionFilename = SectionFilename;
+	}
+		 
+	virtual void RecordDiff( int64 LocalOffset, 
+		const FCallstacks::FCallstackData* DifferenceCallstackData ) override
+	{
+		check(State == EState::InSection);
+		FDiff& Diff = GetSection().Diffs.AddDefaulted_GetRef();
+
+		Diff.LocalOffset = LocalOffset;
+		Diff.Size = 1;
+		if (DifferenceCallstackData)
+		{
+			Diff.ObjectName = DifferenceCallstackData->SerializedObjectName;
+			Diff.PropertyName = DifferenceCallstackData->SerializedPropertyName;
+			Diff.Callstack = DifferenceCallstackData->ToString(CallstackCutoffString);
+		}
+		CurrentPackage.bHasDiffs = true;
+	}
+
+	virtual void ExtendPreviousDiff( int64 LocalOffset ) override
+	{
+		check(State == EState::InSection);
+		check(CurrentPackage.bHasDiffs);
+		FDiff& Diff = GetSection().Diffs.Last();
+
+		check(Diff.LocalOffset <= LocalOffset);
+		Diff.Size = (LocalOffset - Diff.LocalOffset) + 1;
+	}
+
+	virtual void IncrementPreviousDiff() override
+	{
+		check(State == EState::InSection);
+		check(CurrentPackage.bHasDiffs);
+		FDiff& Diff = GetSection().Diffs.Last();
+		Diff.Count++;
+	}
+
+	virtual void RecordUndiagnosedDiff() override
+	{
+		check(State == EState::InSection);
+
+		CurrentPackage.bHasDiffs = true;
+		GetSection().bUndiagnosedDiff = true;
+	}
+
+	virtual void RecordUnreportedDiffs(int32 NumUnreportedDiffs) override
+	{
+		check(State == EState::InSection);
+
+		CurrentPackage.bHasDiffs = true;
+		GetSection().NumUnreportedDiffs = NumUnreportedDiffs;
+	}
+
+	virtual void RecordDetermismDiagnostics(const TCHAR* DeterminismLines) override
+	{
+		check(State == EState::InPackage);			
+		CurrentPackage.DeterminismDiagnostics = DeterminismLines;
+	}
+
+	virtual void RecordTableDifferences(const TCHAR* ItemName, const TCHAR* HumanReadableString) override
+	{
+		check(State == EState::InSection);
+		FTableDiff& TableDiff = GetSection().TableDiffs.AddDefaulted_GetRef();
+		TableDiff.ItemName = ItemName;
+		TableDiff.HumanReadableString = HumanReadableString;
+		CurrentPackage.bHasDiffs = true;
+	}
+
+
+	virtual void EndSection() override
+	{
+		check(State == EState::InSection);
+		CurrentPackage.CurrentSectionIndex = -1;
+		State = EState::InPackage;
+	}
+
+	virtual void EndPackage(const TCHAR* Filename, const FPackageData& SourcePackage, 
+		const FPackageData& DestPackage, const TCHAR* ClassName) override
+	{
+		static const TCHAR* SectionNames[(int32)EDiffWriterSectionType::MAX] =
+		{
+			TEXT("Header"),
+			TEXT("Exports"),
+		};
+
+		check(State == EState::InPackage);
+
+		if (CurrentPackage.bHasDiffs)
+		{
+			JsonWriter->WriteObjectStart();
+			JsonWriter->WriteValue("Filename",  GetSanitizedPaths(Filename));
+			JsonWriter->WriteValue("ClassName", ClassName);
+
+			for (int32 SectionIndex = 0; SectionIndex < (int32)EDiffWriterSectionType::MAX; SectionIndex++)
+			{
+				FSection& Section = CurrentPackage.Sections[SectionIndex];
+				if (Section.Diffs.Num() > 0 || Section.TableDiffs.Num() > 0 || Section.bUndiagnosedDiff || Section.NumUnreportedDiffs)
+				{
+					JsonWriter->WriteObjectStart( SectionNames[SectionIndex] );
+					if (SectionIndex != (int32)EDiffWriterSectionType::Header || Section.SectionFilename != Filename)
+					{
+						JsonWriter->WriteValue("Filename", GetSanitizedPaths(Section.SectionFilename));
+					}
+					JsonWriter->WriteValue("SourceSize",      Section.SourceSize);
+					JsonWriter->WriteValue("DestSize",        Section.DestSize);
+
+					if (Section.bUndiagnosedDiff)
+					{
+						JsonWriter->WriteValue("UndiagnosedDiff", true);
+					}
+					if (Section.NumUnreportedDiffs > 0)
+					{
+						JsonWriter->WriteValue("UnreportedDiffs", Section.NumUnreportedDiffs);
+					}
+					if (Section.Diffs.Num() > 0)
+					{
+						JsonWriter->WriteArrayStart("Diffs");
+						for (const FDiff& Diff : Section.Diffs)
+						{
+							JsonWriter->WriteObjectStart();
+							JsonWriter->WriteValue("LocalOffset", Diff.LocalOffset);
+							JsonWriter->WriteValue("Size",        Diff.Size);
+							if (Diff.Count > 1)
+							{
+								JsonWriter->WriteValue("Count", Diff.Count);
+							}
+							if (!Diff.Callstack.IsEmpty())
+							{
+								JsonWriter->WriteValue("ObjectName",   Diff.ObjectName);
+								JsonWriter->WriteValue("PropertyName", Diff.PropertyName);
+								JsonWriter->WriteValue("Callstack",    GetSanitizedCallstack(Diff));
+							}
+							JsonWriter->WriteObjectEnd();
+						}
+						JsonWriter->WriteArrayEnd();
+					}
+					if (Section.TableDiffs.Num() > 0)
+					{
+						JsonWriter->WriteArrayStart("TableDiffs");
+						for (const FTableDiff& TableDiff : Section.TableDiffs)
+						{
+							JsonWriter->WriteObjectStart();
+							JsonWriter->WriteValue("ItemName", TEXT("Map") + TableDiff.ItemName);
+							JsonWriter->WriteValue("Detail", 
+								TableDiff.HumanReadableString.Replace(IndentToken,TEXT("")));
+							JsonWriter->WriteObjectEnd();
+						}
+						JsonWriter->WriteArrayEnd();
+
+					}
+
+					JsonWriter->WriteObjectEnd();
+				}
+			}
+
+			if (!CurrentPackage.DeterminismDiagnostics.IsEmpty())
+			{
+				JsonWriter->WriteValue("Diagnostics", CurrentPackage.DeterminismDiagnostics);
+			}
+			JsonWriter->WriteObjectEnd();
+		}
+
+		CurrentPackage = FCurrentPackage();
+		State = EState::Idle;
+	}
+
+protected:
+	TUniquePtr<FArchive> JsonFile;
+	TSharedPtr<TJsonWriter<ANSICHAR>> JsonWriter;
+
+	enum class EState : uint8
+	{
+		Idle,
+		InPackage,
+		InSection,
+	};
+	EState State = EState::Idle;
+
+	struct FDiff
+	{
+		int64 LocalOffset = 0;
+		int64 Size = 0;
+		FString ObjectName;
+		FString PropertyName;
+		FString Callstack;
+		int32 Count = 1;
+	};
+
+	struct FTableDiff
+	{
+		FString ItemName;
+		FString HumanReadableString;
+	};
+
+	struct FSection
+	{
+		FString SectionFilename;
+		uint64 SourceSize = 0;
+		uint64 DestSize = 0;
+		bool bUndiagnosedDiff = false;
+		int32 NumUnreportedDiffs = 0;
+		TArray<FDiff> Diffs;
+		TArray<FTableDiff> TableDiffs;
+	};
+
+	struct FCurrentPackage
+	{
+		FString DeterminismDiagnostics;
+		FSection Sections[(int32)EDiffWriterSectionType::MAX];
+
+		bool bHasDiffs = false;
+		int32 CurrentSectionIndex = -1;
+	};
+	FCurrentPackage CurrentPackage;
+
+	FSection& GetSection()
+	{
+		check(CurrentPackage.CurrentSectionIndex != -1);
+		check(State == EState::InSection);
+		return CurrentPackage.Sections[CurrentPackage.CurrentSectionIndex];
+	}
+
+	FString GetSanitizedPaths( const FString& Path ) const
+	{
+		FStringView RelativePath;
+		if (FPathViews::TryMakeChildPathRelativeTo(Path, FPaths::RootDir(), RelativePath))
+		{
+			return FString(RelativePath);
+		}
+		else
+		{
+			// this is used for the absolute file paths embedded in the callstack strings
+			FString Result = Path;
+
+			Result.ReplaceCharInline('\\', '/');
+			Result.ReplaceInline(*FPaths::RootDir(), TEXT(""));		
+			return Result;
+		}
+	}
+
+
+	FString GetSanitizedCallstack(const FDiff& Diff) const
+	{
+		FString Result = Diff.Callstack.Replace(IndentToken, TEXT("")).Replace(TEXT("%DWA%"), TEXT(""));
+		return GetSanitizedPaths(Result);
+	}
+};
+} // namespace UE::DiffWriter
+
+
 FDiffPackageWriter::FDiffPackageWriter(TUniquePtr<ICookedPackageWriter>&& InInner,
 	UE::Cook::FDeterminismManager* InDeterminismManager)
 	: Inner(MoveTemp(InInner))
@@ -52,6 +367,13 @@ FDiffPackageWriter::FDiffPackageWriter(TUniquePtr<ICookedPackageWriter>&& InInne
 		bIgnoreHeaderDiffs = FParse::Param(FCommandLine::Get(), TEXT("IgnoreHeaderDiffs"));
 	}
 
+	FString JsonFile;
+	if (FParse::Value(FCommandLine::Get(), TEXT("CookDiffJson="), JsonFile))
+	{
+		AccumulatorGlobals->DetailRecorder = MakeUnique<UE::DiffWriter::FJsonDetailRecorder>(JsonFile);
+	}
+
+
 	ParseCmds();
 
 	Indent = FCString::Spc(FOutputDeviceHelper::FormatLogLine(ELogVerbosity::Warning,
@@ -212,9 +534,8 @@ void FDiffPackageWriter::WritePackageData(const FPackageInfo& Info, FLargeMemory
 		Accumulator.OnSecondSaveComplete(LocalInfo.HeaderSize);
 
 		TMap<FName, FArchiveDiffStats> PackageDiffStats;
-		const TCHAR* CutoffString = TEXT("UEditorEngine::Save()");
 		Accumulator.SetDeterminismManager(*DeterminismManager);
-		Accumulator.CompareWithPrevious(CutoffString, PackageDiffStats);
+		Accumulator.CompareWithPrevious(CallstackCutoffString, PackageDiffStats);
 
 		//COOK_STAT(FSavePackageStats::NumberOfDifferentPackages++);
 		//COOK_STAT(FSavePackageStats::MergeStats(PackageDiffStats));
@@ -287,9 +608,13 @@ UE::DiffWriter::FAccumulator& FDiffPackageWriter::ConstructAccumulator(FName Pac
 	TRefCountPtr<UE::DiffWriter::FAccumulator>& Accumulator = Accumulators[MultiOutputIndex];
 	if (!Accumulator.IsValid())
 	{
+		TSharedPtr<UE::DiffWriter::FDiffOutputRecorder> DiffOutputRecorder = 
+			MakeShared<UE::DiffWriter::FDiffOutputRecorder>(GetDiffWriterMessageCallback(), 
+				AccumulatorGlobals->DetailRecorder.Get());
+
 		check(!bHasStartedSecondSave); // Accumulator should already exist from CreateLinkerArchive in the first save
 		Accumulator = new UE::DiffWriter::FAccumulator(*AccumulatorGlobals, Asset, *PackageName.ToString(), MaxDiffsToLog,
-			bIgnoreHeaderDiffs, GetDiffWriterMessageCallback(), Inner->GetCookCapabilities().HeaderFormat);
+			bIgnoreHeaderDiffs, DiffOutputRecorder, Inner->GetCookCapabilities().HeaderFormat);
 	}
 	return *Accumulator;
 }
diff --git a/Engine/Source/Editor/UnrealEd/Private/Cooker/DiffWriterArchive.cpp b/Engine/Source/Editor/UnrealEd/Private/Cooker/DiffWriterArchive.cpp
--- a/Engine/Source/Editor/UnrealEd/Private/Cooker/DiffWriterArchive.cpp
+++ b/Engine/Source/Editor/UnrealEd/Private/Cooker/DiffWriterArchive.cpp
@@ -46,7 +46,7 @@ class FPackageHeaderData
 public:
 	FPackageHeaderData(const TCHAR* InName, bool bInReadFromPackageStore,
 		const FString& InAssetFilename, const FPackageData& InPackageData, EPackageHeaderFormat InFormat,
-		FAccumulatorGlobals& InGlobals, FMessageCallback& InMessageCallback);
+		FAccumulatorGlobals& InGlobals, FDiffOutputRecorder& InDiffOutputRecorder);
 	~FPackageHeaderData();
 	void Initialize();
 
@@ -55,7 +55,7 @@ public:
 	const FPackageData& GetPackageData() const;
 	EPackageHeaderFormat GetFormat() const;
 	FAccumulatorGlobals& GetGlobals() const;
-	FMessageCallback& GetMessageCallback() const;
+	FDiffOutputRecorder& GetDiffOutputRecorder() const;
 
 	FDiffWriterZenHeader& GetZenHeader() const;
 	FLinkerLoad* GetLinker() const;
@@ -68,7 +68,7 @@ private:
 	const FString& AssetFilename;
 	const FPackageData& PackageData;
 	FAccumulatorGlobals& Globals;
-	FMessageCallback& MessageCallback;
+	FDiffOutputRecorder& DiffOutputRecorder;
 	EPackageHeaderFormat Format = EPackageHeaderFormat::PackageFileSummary;
 	bool bReadFromPackageStore = false;
 	bool bInitialized = false;
@@ -94,7 +94,7 @@ private:
 
 /** Logs any mismatching header data. */
 void DumpPackageHeaderDiffs(FPackageHeaderData& SourceHeaderData, FPackageHeaderData& DestHeaderData,
-	const int32 MaxDiffsToLog);
+	const int32 MaxDiffsToLog, TSharedPtr<UE::DiffWriter::FDiffOutputRecorder> DiffOutputRecorder);
 
 /** Returns a new linker for loading the specified package. */
 FLinkerLoad* CreateLinkerForPackage(
@@ -396,10 +396,14 @@ void FCallstacks::Add(
 	}
 }
 
-int32 FCallstacks::GetCallstackIndexAtOffset(int64 Offset, int32 MinOffsetIndex) const
+int32 FCallstacks::GetCallstackIndexAtOffset(int64 Offset, int32 MinOffsetIndex, int64* OutOffsetEnd) const
 {
 	if (Offset < 0 || Offset >= EndOffset || MinOffsetIndex >= CallstackAtOffsetMap.Num())
 	{
+		if (OutOffsetEnd)
+		{
+			(*OutOffsetEnd) = -1;
+		}
 		return -1;
 	}
 
@@ -446,6 +450,18 @@ int32 FCallstacks::GetCallstackIndexAtOffset(int64 Offset, int32 MinOffsetIndex)
 		}
 	}
 
+	if (OutOffsetEnd)
+	{
+		if (OffsetForCallstackIndex != -1)
+		{
+			(*OutOffsetEnd) = CallstackAtOffsetMap[OffsetForCallstackIndex].Offset + CallstackAtOffsetMap[OffsetForCallstackIndex].Length - 1;
+		}
+		else
+		{
+			(*OutOffsetEnd) = -1;
+		}
+	}
+
 	return OffsetForCallstackIndex;
 }
 
@@ -637,13 +653,275 @@ void FAccumulatorGlobals::Initialize(EPackageHeaderFormat InFormat)
 	}
 }
 
+FDiffOutputRecorder::FDiffOutputRecorder( FMessageCallback&& InMessageCallback, IDetailRecorder* InDetailRecorder )
+	: MessageCallback(MoveTemp(InMessageCallback))
+	, DetailRecorder(InDetailRecorder)
+{
+}
+
+void FDiffOutputRecorder::RecordHeaderSizeMismatch( const TCHAR* Filename, int64 FirstHeaderSize, int64 SecondHeaderSize )
+{
+	MessageCallback(ELogVerbosity::Error, FString::Printf(
+		TEXT("%s: Indeterministic header size. When saving the package twice into memory, first header size %" INT64_FMT " != second header size %" INT64_FMT ". Callstacks for indeterminism in the exports will be incorrect.")
+		TEXT("\n\tDumping differences from first and second memory saves."),
+		Filename, FirstHeaderSize, SecondHeaderSize));
+}
+
+void FDiffOutputRecorder::RecordUndiagnosedHeaderDifference( const TCHAR* Filename )
+{
+	MessageCallback(ELogVerbosity::Warning, FString::Printf(
+		TEXT("%s: headers are different, but DumpPackageHeaderDiffs does not yet implement describing the difference."),
+		Filename
+		));
+
+	if (DetailRecorder != nullptr)
+	{
+		DetailRecorder->RecordUndiagnosedDiff();
+	}
+
+}
+
+void FDiffOutputRecorder::RecordTableDifferences(const TCHAR* Filename, const TCHAR* ItemName, int32 SourceTableNum, 
+	int32 DestTableNum, const TCHAR* HumanReadableString)
+{
+	//SourceContext.LogMessage(ELogVerbosity::Warning, FString::Printf(
+	MessageCallback(ELogVerbosity::Warning, FString::Printf(
+		TEXT("%s: %sMap is different (%d %ss in source package vs %d %ss in dest package):%s%s"),		
+		Filename,
+		ItemName,
+		SourceTableNum,
+		ItemName,
+		DestTableNum,
+		ItemName,
+		NewLineToken,
+		HumanReadableString));
+
+	if (DetailRecorder != nullptr)
+	{
+		DetailRecorder->RecordTableDifferences(ItemName, HumanReadableString);
+	}
+}
+
+void FDiffOutputRecorder::RecordTableDifferences(const TCHAR* Filename, const TCHAR* ItemName, 
+	const TCHAR* HumanReadableString)
+{
+	//LogMessage(ELogVerbosity::Warning, FString::Printf(
+	MessageCallback(ELogVerbosity::Warning, FString::Printf(
+		TEXT("%s: %sMap is different:%s%s"),
+		Filename,
+		ItemName,
+		NewLineToken,
+		HumanReadableString));
+
+	if (DetailRecorder != nullptr)
+	{
+		DetailRecorder->RecordTableDifferences(ItemName, HumanReadableString);
+	}
+}
+
+void FDiffOutputRecorder::RecordSectionSizeMismatch(const TCHAR* SectionFilename, int64 SourceSize, int64 DestSize)
+{
+	MessageCallback(ELogVerbosity::Warning, FString::Printf(
+		TEXT("%s: Size mismatch: on disk: %lld vs memory: %lld"), SectionFilename, SourceSize, DestSize));
+}
+
+void FDiffOutputRecorder::RecordDiff(const TCHAR* SectionFilename, int64 LocalOffset, int64 DestAbsoluteOffset, 
+	uint8 SourceByte, uint8 DestByte, bool bHasOptimizedHeader)
+{
+	if (bHasOptimizedHeader)
+	{
+		MessageCallback(ELogVerbosity::Warning, FString::Printf(
+			TEXT("%s: Difference at offset %lld (Combined/DiffBreak Offset: %" INT64_FMT "): OnDisk %d != %d InMemory.%s")
+			TEXT("Callstack is unknown because the offset is in the header and the header has been optimized. See the output of DumpPackageHeaderDiffs to debug this difference."),
+			SectionFilename, LocalOffset, DestAbsoluteOffset, SourceByte, DestByte, NewLineToken
+		));
+	}
+	else
+	{
+		MessageCallback(ELogVerbosity::Warning, FString::Printf(
+			TEXT("%s: Difference at offset %lld (Combined/DiffBreak Offset: %" INT64_FMT "): OnDisk %d != %d InMemory.%s")
+			TEXT("Callstack is unknown."),
+			SectionFilename, LocalOffset, DestAbsoluteOffset, SourceByte, DestByte, NewLineToken
+		));
+	}
+
+	if (DetailRecorder != nullptr)
+	{
+		DetailRecorder->RecordDiff(LocalOffset, nullptr);
+	}
+}
+
+void FDiffOutputRecorder::RecordDiff(const TCHAR* SectionFilename, int64 LocalOffset, int64 DestAbsoluteOffset, 
+	uint8 SourceByte, uint8 DestByte, int64 DifferenceOffset, const TCHAR* LastDifferenceCallstackDataText, 
+	const FString& BeforePropertyVal, const FString& AfterPropertyVal, 
+	const FCallstacks::FCallstackData& DifferenceCallstackData)
+{
+	FString DiffValues;
+	if (BeforePropertyVal != AfterPropertyVal)
+	{
+		DiffValues = FString::Printf(TEXT("\r\n%sBefore: %s\r\n%sAfter:  %s"),
+			IndentToken, *BeforePropertyVal,
+			IndentToken, *AfterPropertyVal);
+	}
+
+
+	FString DebugDataStackText;
+	//check for a debug data stack as part of the unique stack entry, and log it out if we find it.
+	FString FullStackText = DifferenceCallstackData.Callstack.Get();
+	int32 DebugDataIndex = FullStackText.Find(ANSI_TO_TCHAR(DebugDataStackMarker), ESearchCase::CaseSensitive);
+	if (DebugDataIndex > 0)
+	{
+		DebugDataStackText = FString::Printf(TEXT("\r\n%s"), IndentToken)
+			+ FullStackText.RightChop(DebugDataIndex + 2);
+	}
+
+	MessageCallback(ELogVerbosity::Warning, FString::Printf(
+		TEXT("%s: Difference at offset %lld (Combined/DiffBreak Offset: %" INT64_FMT "): OnDisk %d != %d InMemory.%s")
+			TEXT("Difference occurs at index %lld within Serialize call at callstack:%s%s%s%s"),
+		SectionFilename, LocalOffset, DestAbsoluteOffset, SourceByte, DestByte, NewLineToken,
+		DifferenceOffset, NewLineToken,
+		LastDifferenceCallstackDataText, *DiffValues, *DebugDataStackText
+	));
+
+	if (DetailRecorder != nullptr)
+	{
+		DetailRecorder->RecordDiff(LocalOffset, &DifferenceCallstackData);
+	}
+}
+
+void FDiffOutputRecorder::ExtendPreviousDiff(int64 LocalOffset)
+{
+	if (DetailRecorder != nullptr)
+	{
+		DetailRecorder->ExtendPreviousDiff(LocalOffset);
+	}
+}
+
+void FDiffOutputRecorder::IncrementPreviousDiff()
+{
+	if (DetailRecorder != nullptr)
+	{
+		DetailRecorder->IncrementPreviousDiff();
+	}
+}
+
+void FDiffOutputRecorder::RecordDiffBytes(const TCHAR* SectionFilename, int32 BytesToLog, int64 LocalOffset, 
+	const FPackageData& SourcePackage, const FPackageData& DestPackage)
+{
+	MessageCallback(ELogVerbosity::Display, FString::Printf(
+		TEXT("%s: Logging %d bytes around offset: %" INT64_FMT " (%016" INT64_X_FMT ") in the OnDisk package:"),
+		SectionFilename,
+		BytesToLog, LocalOffset, LocalOffset
+	));
+	TArray<FString> HexDumpLines = FCompressionUtil::HexDumpLines(SourcePackage.Data + SourcePackage.StartOffset,
+		SourcePackage.Size - SourcePackage.StartOffset,
+		LocalOffset - BytesToLog / 2, LocalOffset + BytesToLog / 2);
+	for (FString& Line : HexDumpLines)
+	{
+		MessageCallback(ELogVerbosity::Display, Line);
+	}
+
+	MessageCallback(ELogVerbosity::Display, FString::Printf(
+		TEXT("%s: Logging %d bytes around offset: %" INT64_FMT " (%016" INT64_X_FMT ") in the InMemory package:"),
+		SectionFilename, BytesToLog, LocalOffset, LocalOffset
+	));
+	HexDumpLines = FCompressionUtil::HexDumpLines(DestPackage.Data + DestPackage.StartOffset,
+		DestPackage.Size - DestPackage.StartOffset,
+		LocalOffset - BytesToLog / 2, LocalOffset + BytesToLog / 2);
+	for (FString& Line : HexDumpLines)
+	{
+		MessageCallback(ELogVerbosity::Display, Line);
+	}
+
+}
+
+void FDiffOutputRecorder::RecordUnreportedDiffs(const TCHAR* SectionFilename, int32 NumUnreportedDiffs, 
+	int32 FirstUnreportedDiffIndex)
+{
+	MessageCallback(ELogVerbosity::Warning, FString::Printf(
+		TEXT("%s: %lld difference(s) not logged (first at offset: %lld)."),
+		SectionFilename,NumUnreportedDiffs, FirstUnreportedDiffIndex));
+
+	if (DetailRecorder != nullptr)
+	{
+		DetailRecorder->RecordUnreportedDiffs(NumUnreportedDiffs);
+	}
+}
+
+void FDiffOutputRecorder::RecordDeterminismDiagnostics(const FString& DeterminismLines)
+{
+	if (!DeterminismLines.IsEmpty())
+	{
+		if (DetailRecorder != nullptr)
+		{
+			DetailRecorder->RecordDetermismDiagnostics(*DeterminismLines);
+		}
+
+		MessageCallback(ELogVerbosity::Display, FString(TEXT("DeterminismHelper Diagnostics:\n") + DeterminismLines));
+	}
+}
+
+void FDiffOutputRecorder::RecordNewPackage(const TCHAR* Filename)
+{
+	MessageCallback(ELogVerbosity::Warning, FString::Printf(TEXT("New package: %s"), Filename));
+}
+
+void FDiffOutputRecorder::BeginPackage(const TCHAR* Filename, const TCHAR* Classname)
+{
+	MessageCallback(ELogVerbosity::Display, FString::Printf(TEXT("Comparing: %s"), Filename));
+	MessageCallback(ELogVerbosity::Warning, FString::Printf(TEXT("Asset class: %s"), Classname));
+
+	if (DetailRecorder != nullptr)
+	{
+		DetailRecorder->BeginPackage();
+	}
+}
+void FDiffOutputRecorder::BeginSection(const TCHAR* SectionFilename, EDiffWriterSectionType SectionType, 
+	const FPackageData& SourceSection, const FPackageData& DestSection)
+{
+	if (DetailRecorder != nullptr)
+	{
+		DetailRecorder->BeginSection(SectionFilename, SectionType, SourceSection, DestSection);
+	}
+}
+void FDiffOutputRecorder::EndSection()
+{
+	if (DetailRecorder != nullptr)
+	{
+		DetailRecorder->EndSection();
+	}
+}
+void FDiffOutputRecorder::EndPackage(const TCHAR* Filename, const FPackageData& SourcePackage, const FPackageData& DestPackage, 
+	const TCHAR* ClassName)
+{
+	if (DetailRecorder != nullptr)
+	{
+		DetailRecorder->EndPackage(Filename, SourcePackage, DestPackage, ClassName);
+	}
+}
+
+
+
+const FMessageCallback& FDiffOutputRecorder::GetMessageCallback() const
+{
+	return MessageCallback;
+}
+
+IDetailRecorder* FDiffOutputRecorder::GetDetailRecorder() const
+{
+	return DetailRecorder;
+}
+
+
+
+
 FAccumulator::FAccumulator(FAccumulatorGlobals& InGlobals, UObject* InAsset, FName InPackageName,
-	int32 InMaxDiffsToLog, bool bInIgnoreHeaderDiffs, FMessageCallback&& InMessageCallback,
+	int32 InMaxDiffsToLog, bool bInIgnoreHeaderDiffs, TSharedPtr<FDiffOutputRecorder> InDiffOutputRecorder,
 	EPackageHeaderFormat InPackageHeaderFormat)
 	: LinkerCallstacks()
 	, ExportsCallstacks()
 	, Globals(InGlobals)
-	, MessageCallback(MoveTemp(InMessageCallback))
+	, DiffOutputRecorder(InDiffOutputRecorder)
 	, PackageName(InPackageName)
 	, Asset(InAsset)
 	, MaxDiffsToLog(InMaxDiffsToLog)
@@ -727,11 +1005,6 @@ void FAccumulator::OnSecondSaveComplete(int64 InHeaderSize)
 	check(bFirstSaveComplete); // Should have been set in OnFirstSaveComplete
 	if (HeaderSize != InHeaderSize)
 	{
-		MessageCallback(ELogVerbosity::Error, FString::Printf(
-			TEXT("%s: Indeterministic header size. When saving the package twice into memory, first header size %" INT64_FMT " != second header size %" INT64_FMT ". Callstacks for indeterminism in the exports will be incorrect.")
-			TEXT("\n\tDumping differences from first and second memory saves."),
-			*this->Filename, HeaderSize, InHeaderSize));
-
 		check(bFirstSaveComplete);
 		check(FirstSaveLinkerData.Num() >= HeaderSize);
 		check(LinkerArchive && LinkerArchive->TotalSize() >= InHeaderSize);
@@ -741,21 +1014,25 @@ void FAccumulator::OnSecondSaveComplete(int64 InHeaderSize)
 		int32 NumHeaderDiffMessages = 0;
 		FMessageCallback HeaderMessageCallback = [&NumHeaderDiffMessages, this](ELogVerbosity::Type Verbosity, FStringView Message)
 			{
-				MessageCallback(Verbosity, Message);
+				DiffOutputRecorder->GetMessageCallback()(Verbosity, Message);
 				++NumHeaderDiffMessages;
 			};
+        // Note: we do not record to the DetailRecorder when reporting the header differences between the first and second save.
+        // The DetailRecorder is reporting the differences between the first save and the previous cook's save.
+		FDiffOutputRecorder HeaderDiffOutputRecorder(MoveTemp(HeaderMessageCallback), nullptr);
+
+		HeaderDiffOutputRecorder.RecordHeaderSizeMismatch(*this->Filename, HeaderSize, InHeaderSize);
+
 		FPackageHeaderData FirstSaveHeader(TEXT("source"), false /* bReadFromPackageStore */, Filename,
-			FirstSaveHeaderSegment, PackageHeaderFormat, Globals, HeaderMessageCallback);
+			FirstSaveHeaderSegment, PackageHeaderFormat, Globals, HeaderDiffOutputRecorder);
 		FPackageHeaderData SecondSaveHeader(TEXT("dest"), false /* bReadFromPackageStore */, Filename,
-			SecondSaveHeaderSegment, PackageHeaderFormat, Globals, HeaderMessageCallback);
+			SecondSaveHeaderSegment, PackageHeaderFormat, Globals, HeaderDiffOutputRecorder);
 
-		DumpPackageHeaderDiffs(FirstSaveHeader, SecondSaveHeader, MaxDiffsToLog);
+		DumpPackageHeaderDiffs(FirstSaveHeader, SecondSaveHeader, MaxDiffsToLog, DiffOutputRecorder);
 		// Suppress static analysis warning V547: Expression 'NumHeaderDiffMessages == 0' is always true
 		if (NumHeaderDiffMessages == 0) // -V547
 		{
-			MessageCallback(ELogVerbosity::Warning, FString::Printf(
-				TEXT("%s: headers are different, but DumpPackageHeaderDiffs does not yet implement describing the difference."),
-				*Filename));
+			HeaderDiffOutputRecorder.RecordUndiagnosedHeaderDifference(*Filename);
 		}
 	}
 
@@ -902,8 +1179,7 @@ void FAccumulator::CompareWithPreviousForSection(const FPackageData& SourcePacka
 	
 	if (SourceSize != DestSize)
 	{
-		MessageCallback(ELogVerbosity::Warning, FString::Printf(
-			TEXT("%s: Size mismatch: on disk: %lld vs memory: %lld"), *SectionFilename, SourceSize, DestSize));
+		DiffOutputRecorder->RecordSectionSizeMismatch(*SectionFilename, SourceSize, DestSize);
 		int64 SizeDiff = DestPackage.Size - SourcePackage.Size;
 		OutStats.FindOrAdd(AssetClass).DiffSize += SizeDiff;
 	}
@@ -915,10 +1191,46 @@ void FAccumulator::CompareWithPreviousForSection(const FPackageData& SourcePacka
 	int64 NumDiffsLoggedLocal = 0;
 	int64 FirstUnreportedDiffIndex = -1;
 
+    // When we discover the first byte difference within a serialize call, we record the offset of the
+    // end of the serialize call so we can record what range of bytes within the serialize call was
+    // different.
+    // We record some variables about the offset in the Local reference frame (from the beginning of the
+    // current Section), and some in the Absolute reference frame (from the beginning of the Package that
+    // contains all of the Sections).
+
+    /**
+     * If non-negative, we have found a diff and this is the offset to the end of the current serialize call. */
+	int64 DestAbsoluteEndOffsetForActiveDiff = -1;
+
+   /**
+     * If non-negative, we have found a diff. The diff spans one or more bytes. This is the offset to the last byte
+     * we have reached that was different, so we can record the range of bytes that differed. We reset this and start
+     * a new diff report when we reach the end of the Serialize call in which we found the diff.
+     */
+	int64 ActiveDiffLocalOffset = -1;
+
+    /**
+     * If non-negative, we have found a diff and the diff has the same callstack repeated from the previous diff,
+     * but the callstack was called repeatededly and we recorded a diff in the repeated call. This is the offset
+     * to the end of the current serialize call for the repeated callstack.
+     */
+	int64 DestAbsoluteEndOffsetForDuplicateDiff = -1;
+
 	for (int64 LocalOffset = 0; LocalOffset < SizeToCompare; ++LocalOffset)
 	{
 		const int64 SourceAbsoluteOffset = LocalOffset + SourcePackage.StartOffset;
 		const int64 DestAbsoluteOffset = LocalOffset + DestPackage.StartOffset;
+		if (DestAbsoluteEndOffsetForActiveDiff != -1 && DestAbsoluteOffset > DestAbsoluteEndOffsetForActiveDiff)
+		{
+			DiffOutputRecorder->ExtendPreviousDiff(ActiveDiffLocalOffset);
+			DestAbsoluteEndOffsetForActiveDiff = -1;
+			ActiveDiffLocalOffset = -1;
+		}
+		if (DestAbsoluteEndOffsetForDuplicateDiff != -1 && DestAbsoluteOffset > DestAbsoluteEndOffsetForDuplicateDiff)
+		{
+			DiffOutputRecorder->IncrementPreviousDiff();
+			DestAbsoluteEndOffsetForDuplicateDiff = -1;
+		}
 
 		const uint8 SourceByte = SourcePackage.Data[SourceAbsoluteOffset];
 		const uint8 DestByte   = DestPackage  .Data[DestAbsoluteOffset];
@@ -926,10 +1238,12 @@ void FAccumulator::CompareWithPreviousForSection(const FPackageData& SourcePacka
 		{
 			continue;
 		}
+		ActiveDiffLocalOffset = LocalOffset;
 
 		constexpr int BytesToLog = 128;
+		int64 DestAbsoluteEndOffset = -1;
 		int32 DifferenceCallstackOffsetIndex = Callstacks.GetCallstackIndexAtOffset(DestAbsoluteOffset,
-			LastDifferenceCallstackOffsetIndex < 0 ? 0 : LastDifferenceCallstackOffsetIndex);
+			LastDifferenceCallstackOffsetIndex < 0 ? 0 : LastDifferenceCallstackOffsetIndex, &DestAbsoluteEndOffset);
 
 		// Skip reporting the difference if we are still within the same Serialize call, or for any bytes after
 		// the first different byte if we don't know the callstack.
@@ -954,6 +1268,7 @@ void FAccumulator::CompareWithPreviousForSection(const FPackageData& SourcePacka
 			if (NumDiffsLocal > 0 && 
 				LastDifferenceCallstackDataText.Compare(DifferenceCallstackDataText, ESearchCase::CaseSensitive) == 0)
 			{
+				DestAbsoluteEndOffsetForDuplicateDiff = DestAbsoluteEndOffset;
 				continue;
 			}
 		}
@@ -997,22 +1312,8 @@ void FAccumulator::CompareWithPreviousForSection(const FPackageData& SourcePacka
 
 		if (DifferenceCallstackOffsetIndex < 0)
 		{
-			if (IsWriterUsingPostSaveTransforms() && DestAbsoluteOffset < HeaderSize)
-			{
-				MessageCallback(ELogVerbosity::Warning, FString::Printf(
-					TEXT("%s: Difference at offset %lld (Combined/DiffBreak Offset: %" INT64_FMT "): OnDisk %d != %d InMemory.%s")
-					TEXT("Callstack is unknown because the offset is in the header and the header has been optimized. See the output of DumpPackageHeaderDiffs to debug this difference."),
-					*SectionFilename, LocalOffset, DestAbsoluteOffset, SourceByte, DestByte, NewLineToken
-				));
-			}
-			else
-			{
-				MessageCallback(ELogVerbosity::Warning, FString::Printf(
-					TEXT("%s: Difference at offset %lld (Combined/DiffBreak Offset: %" INT64_FMT "): OnDisk %d != %d InMemory.%s")
-					TEXT("Callstack is unknown."),
-					*SectionFilename, LocalOffset, DestAbsoluteOffset, SourceByte, DestByte, NewLineToken
-				));
-			}
+			const bool bHasOptimizedHeader = (IsWriterUsingPostSaveTransforms() && DestAbsoluteOffset < HeaderSize);
+			DiffOutputRecorder->RecordDiff(*SectionFilename, LocalOffset, DestAbsoluteOffset, SourceByte, DestByte, bHasOptimizedHeader);
 		}
 		else
 		{
@@ -1083,64 +1384,28 @@ void FAccumulator::CompareWithPreviousForSection(const FPackageData& SourcePacka
 				}
 			}
 
-			FString DiffValues;
-			if (BeforePropertyVal != AfterPropertyVal)
-			{
-				DiffValues = FString::Printf(TEXT("\r\n%sBefore: %s\r\n%sAfter:  %s"),
-					IndentToken, *BeforePropertyVal,
-					IndentToken, *AfterPropertyVal);
-			}
-
-			FString DebugDataStackText;
-			//check for a debug data stack as part of the unique stack entry, and log it out if we find it.
-			FString FullStackText = DifferenceCallstackData.Callstack.Get();
-			int32 DebugDataIndex = FullStackText.Find(ANSI_TO_TCHAR(DebugDataStackMarker), ESearchCase::CaseSensitive);
-			if (DebugDataIndex > 0)
-			{
-				DebugDataStackText = FString::Printf(TEXT("\r\n%s"), IndentToken)
-					+ FullStackText.RightChop(DebugDataIndex + 2);
-			}
-
-			MessageCallback(ELogVerbosity::Warning, FString::Printf(
-				TEXT("%s: Difference at offset %lld (Combined/DiffBreak Offset: %" INT64_FMT "): OnDisk %d != %d InMemory.%s")
-					TEXT("Difference occurs at index %lld within Serialize call at callstack:%s%s%s%s"),
-				*SectionFilename, LocalOffset, DestAbsoluteOffset, SourceByte, DestByte, NewLineToken,
-				DestAbsoluteOffset - CallstackAtOffset.SerializeCallOffset, NewLineToken,
-				*LastDifferenceCallstackDataText, *DiffValues, *DebugDataStackText
-			));
+			DiffOutputRecorder->RecordDiff(*SectionFilename, LocalOffset, DestAbsoluteOffset, SourceByte, DestByte,
+				DestAbsoluteOffset - CallstackAtOffset.SerializeCallOffset,
+				*LastDifferenceCallstackDataText, BeforePropertyVal, AfterPropertyVal, DifferenceCallstackData);
+			DestAbsoluteEndOffsetForActiveDiff = DestAbsoluteEndOffset;
 		}
 
-		MessageCallback(ELogVerbosity::Display, FString::Printf(
-			TEXT("%s: Logging %d bytes around offset: %" INT64_FMT " (%016" INT64_X_FMT ") in the OnDisk package:"),
-			*SectionFilename,
-			BytesToLog, LocalOffset, LocalOffset
-		));
-		TArray<FString> HexDumpLines = FCompressionUtil::HexDumpLines(SourcePackage.Data + SourcePackage.StartOffset,
-			SourcePackage.Size - SourcePackage.StartOffset,
-			LocalOffset - BytesToLog / 2, LocalOffset + BytesToLog / 2);
-		for (FString& Line : HexDumpLines)
-		{
-			MessageCallback(ELogVerbosity::Display, Line);
-		}
+		DiffOutputRecorder->RecordDiffBytes(*SectionFilename, BytesToLog, LocalOffset, SourcePackage, DestPackage);
+	}
 
-		MessageCallback(ELogVerbosity::Display, FString::Printf(
-			TEXT("%s: Logging %d bytes around offset: %" INT64_FMT " (%016" INT64_X_FMT ") in the InMemory package:"),
-			*SectionFilename, BytesToLog, LocalOffset, LocalOffset
-		));
-		HexDumpLines = FCompressionUtil::HexDumpLines(DestPackage.Data + DestPackage.StartOffset,
-			DestPackage.Size - DestPackage.StartOffset,
-			LocalOffset - BytesToLog / 2, LocalOffset + BytesToLog / 2);
-		for (FString& Line : HexDumpLines)
-		{
-			MessageCallback(ELogVerbosity::Display, Line);
-		}
+	if (DestAbsoluteEndOffsetForActiveDiff != -1)
+	{
+		DiffOutputRecorder->ExtendPreviousDiff(ActiveDiffLocalOffset);
+	}
+	if (DestAbsoluteEndOffsetForDuplicateDiff != -1)
+	{
+		DiffOutputRecorder->IncrementPreviousDiff();
 	}
 
+
 	if (MaxDiffsToLog >= 0 && NumDiffsForLogStatLocal > NumDiffsLoggedLocal)
 	{
-		MessageCallback(ELogVerbosity::Warning, FString::Printf(
-			TEXT("%s: %lld difference(s) not logged (first at offset: %lld)."),
-			*SectionFilename, NumDiffsForLogStatLocal - NumDiffsLoggedLocal, FirstUnreportedDiffIndex));
+		DiffOutputRecorder->RecordUnreportedDiffs(*SectionFilename, NumDiffsForLogStatLocal - NumDiffsLoggedLocal, FirstUnreportedDiffIndex);
 	}
 }
 
@@ -1155,7 +1420,7 @@ void FAccumulator::CompareWithPrevious(const TCHAR* CallstackCutoffText, TMap<FN
 	OutStats.FindOrAdd(AssetClass).NewFileTotalSize = LinkerArchive->TotalSize();
 	if (PreviousPackageData.Size == 0)
 	{
-		MessageCallback(ELogVerbosity::Warning, FString::Printf(TEXT("New package: %s"), *Filename));
+		DiffOutputRecorder->RecordNewPackage(*Filename);
 		OutStats.FindOrAdd(AssetClass).DiffSize = OutStats.FindOrAdd(AssetClass).NewFileTotalSize;
 		return;
 	}
@@ -1172,8 +1437,7 @@ void FAccumulator::CompareWithPrevious(const TCHAR* CallstackCutoffText, TMap<FN
 	DestPackage.HeaderSize = HeaderSize;
 	DestPackage.StartOffset = 0;
 
-	MessageCallback(ELogVerbosity::Display, FString::Printf(TEXT("Comparing: %s"), *Filename));
-	MessageCallback(ELogVerbosity::Warning, FString::Printf(TEXT("Asset class: %s"), *AssetClass.ToString()));
+	DiffOutputRecorder->BeginPackage(*Filename, *AssetClass.ToString());
 	if (DeterminismManager)
 	{
 		DeterminismManager->RecordPackageModified(Asset);
@@ -1194,27 +1458,30 @@ void FAccumulator::CompareWithPrevious(const TCHAR* CallstackCutoffText, TMap<FN
 	int32 NumHeaderDiffMessages = 0;
 	FMessageCallback HeaderMessageCallback = [&NumHeaderDiffMessages, this](ELogVerbosity::Type Verbosity, FStringView Message)
 		{
-			MessageCallback(Verbosity, Message);
+			DiffOutputRecorder->GetMessageCallback()(Verbosity, Message);
 			++NumHeaderDiffMessages;
 		};
+	FDiffOutputRecorder HeaderDiffOutputRecorder(MoveTemp(HeaderMessageCallback), DiffOutputRecorder->GetDetailRecorder());
+
 	FPackageHeaderData SourceHeader(TEXT("source"), true /* bReadFromPackageStore */, Filename, SourcePackage,
-		PackageHeaderFormat, Globals, HeaderMessageCallback);
+		PackageHeaderFormat, Globals, HeaderDiffOutputRecorder);
 	FPackageHeaderData DestHeader(TEXT("dest"), false /* bReadFromPackageStore */, Filename, DestPackage,
-		PackageHeaderFormat, Globals, HeaderMessageCallback);
+		PackageHeaderFormat, Globals, HeaderDiffOutputRecorder);
 
+	DiffOutputRecorder->BeginSection(*Filename, EDiffWriterSectionType::Header, SourceHeaderSegment, DestHeaderSegment);
 	CompareWithPreviousForSection(SourceHeaderSegment, DestHeaderSegment, SourceHeader, DestHeader,
 		CallstackCutoffText, NumLoggedDiffs, OutStats, Filename);
 	if (HeaderSize > 0 && OutStats.FindOrAdd(AssetClass).NumDiffs > 0)
 	{
-		DumpPackageHeaderDiffs(SourceHeader, DestHeader, MaxDiffsToLog);
+		DumpPackageHeaderDiffs(SourceHeader, DestHeader, MaxDiffsToLog, DiffOutputRecorder);
 		// Suppress static analysis warning V547: Expression 'NumHeaderDiffMessages == 0' is always true
 		if (NumHeaderDiffMessages == 0) // -V547
 		{
-			MessageCallback(ELogVerbosity::Warning, FString::Printf(
-				TEXT("%s: headers are different, but DumpPackageHeaderDiffs does not yet implement describing the difference."),
-				*Filename));
+			DiffOutputRecorder->RecordUndiagnosedHeaderDifference(*Filename);
 		}
 	}
+	DiffOutputRecorder->EndSection();
+
 
 	FPackageData SourcePackageExports = SourcePackage;
 	SourcePackageExports.HeaderSize = 0;
@@ -1234,18 +1501,18 @@ void FAccumulator::CompareWithPrevious(const TCHAR* CallstackCutoffText, TMap<FN
 		ExportsFilename = Filename;
 	}
 
+	
+	DiffOutputRecorder->BeginSection(*ExportsFilename, EDiffWriterSectionType::Exports, SourcePackageExports, DestPackageExports);
 	CompareWithPreviousForSection(SourcePackageExports, DestPackageExports, SourceHeader, DestHeader,
 		CallstackCutoffText, NumLoggedDiffs, OutStats, ExportsFilename);
+	DiffOutputRecorder->EndSection();
+
 
 	// Log comparison of the DeterminismHelper diagnostics, if we have any
 	if (DeterminismManager)
 	{
 		FString DeterminismLines = DeterminismManager->GetCurrentPackageDiagnosticsAsText();
-		if (!DeterminismLines.IsEmpty())
-		{
-			DeterminismLines = FString(TEXT("DeterminismHelper Diagnostics:\n") + DeterminismLines);
-			MessageCallback(ELogVerbosity::Display, MoveTemp(DeterminismLines));
-		}
+		DiffOutputRecorder->RecordDeterminismDiagnostics(DeterminismLines);
 	}
 
 	// Optionally save out any differences we detected.
@@ -1310,11 +1577,12 @@ void FAccumulator::CompareWithPrevious(const TCHAR* CallstackCutoffText, TMap<FN
 			}
 			else
 			{
-				MessageCallback(ELogVerbosity::Warning,
+				DiffOutputRecorder->GetMessageCallback()(ELogVerbosity::Warning,
 					FString::Printf(TEXT("Package '%s' doesn't seem to be writing to the Saved directory - skipping writing diff"), *OutputFilename));
 			}
 		}
 	}
+	DiffOutputRecorder->EndPackage(*Filename, SourcePackage, DestPackage, *AssetClass.ToString());
 }
 
 void FAccumulator::GenerateDiffMapForSection(const FPackageData& SourcePackage, const FPackageData& DestPackage,
@@ -1621,7 +1889,8 @@ static void DumpTableDifferences(
 	TConstArrayView<T> DestTable,
 	const TCHAR* AssetFilename,
 	const TCHAR* ItemName,
-	const int32 MaxDiffsToLog
+	const int32 MaxDiffsToLog,
+	TSharedPtr<UE::DiffWriter::FDiffOutputRecorder> DiffOutputRecorder
 )
 {
 	FString HumanReadableString;
@@ -1692,16 +1961,8 @@ static void DumpTableDifferences(
 		HumanReadableString += NewLineToken;
 	}
 
-	SourceContext.LogMessage(ELogVerbosity::Warning, FString::Printf(
-		TEXT("%s: %sMap is different (%d %ss in source package vs %d %ss in dest package):%s%s"),		
-		AssetFilename,
-		ItemName,
-		SourceTable.Num(),
-		ItemName,
-		DestTable.Num(),
-		ItemName,
-		NewLineToken,
-		*HumanReadableString));
+
+	DiffOutputRecorder->RecordTableDifferences(AssetFilename, ItemName, SourceTable.Num(), DestTable.Num(), *HumanReadableString);
 }
 
 static void DumpOrderedArrayDifferences(
@@ -1713,7 +1974,8 @@ static void DumpOrderedArrayDifferences(
 	TFunctionRef<void(ELogVerbosity::Type, FString)> LogMessage,
 	const TCHAR* AssetFilename,
 	const TCHAR* ItemName,
-	const int32 MaxDiffsToLog
+	const int32 MaxDiffsToLog,
+	TSharedPtr<UE::DiffWriter::FDiffOutputRecorder> DiffOutputRecorder
 )
 {
 	FString HumanReadableString;
@@ -1768,16 +2030,11 @@ static void DumpOrderedArrayDifferences(
 		HumanReadableString += NewLineToken;
 	}
 
-	LogMessage(ELogVerbosity::Warning, FString::Printf(
-		TEXT("%s: %sMap is different:%s%s"),
-		AssetFilename,
-		ItemName,
-		NewLineToken,
-		*HumanReadableString));
+	DiffOutputRecorder->RecordTableDifferences(AssetFilename, ItemName, *HumanReadableString);
 }
 
 void DumpPackageHeaderDiffs_LinkerLoad(FPackageHeaderData& SourceHeaderData, FPackageHeaderData& DestHeaderData,
-	const int32 MaxDiffsToLog)
+	const int32 MaxDiffsToLog, TSharedPtr<UE::DiffWriter::FDiffOutputRecorder> DiffOutputRecorder)
 {
 	SourceHeaderData.Initialize();
 	DestHeaderData.Initialize();
@@ -1789,30 +2046,30 @@ void DumpPackageHeaderDiffs_LinkerLoad(FPackageHeaderData& SourceHeaderData, FPa
 		return;
 	}
 
-	FDiffWriterLinkerLoadHeader SourceContext(SourceLinker, SourceHeaderData.GetMessageCallback());
-	FDiffWriterLinkerLoadHeader DestContext(DestLinker, SourceHeaderData.GetMessageCallback());
+	FDiffWriterLinkerLoadHeader SourceContext(SourceLinker, SourceHeaderData.GetDiffOutputRecorder());
+	FDiffWriterLinkerLoadHeader DestContext(DestLinker, SourceHeaderData.GetDiffOutputRecorder());
 
 	if (SourceLinker->NameMap != DestLinker->NameMap)
 	{
 		DumpTableDifferences<FNameEntryId>(SourceContext, DestContext, SourceLinker->NameMap, DestLinker->NameMap,
-			*AssetFilename, TEXT("Name"), MaxDiffsToLog);
+			*AssetFilename, TEXT("Name"), MaxDiffsToLog, DiffOutputRecorder);
 	}
 
 	if (!SourceContext.IsImportMapIdentical(DestContext))
 	{
 		DumpTableDifferences<FObjectImport>(SourceContext, DestContext, SourceLinker->ImportMap, DestLinker->ImportMap,
-			*AssetFilename, TEXT("Import"), MaxDiffsToLog);
+			*AssetFilename, TEXT("Import"), MaxDiffsToLog, DiffOutputRecorder);
 	}
 
 	if (!SourceContext.IsExportMapIdentical(DestContext))
 	{
 		DumpTableDifferences<FObjectExport>(SourceContext, DestContext, SourceLinker->ExportMap, DestLinker->ExportMap,
-			*AssetFilename, TEXT("Export"), MaxDiffsToLog);
+			*AssetFilename, TEXT("Export"), MaxDiffsToLog, DiffOutputRecorder);
 	}
 }
 
 void DumpPackageHeaderDiffs_ZenPackage(FPackageHeaderData& SourceHeaderData,
-	FPackageHeaderData& DestHeaderData, const int32 MaxDiffsToLog)
+	FPackageHeaderData& DestHeaderData, const int32 MaxDiffsToLog, TSharedPtr<UE::DiffWriter::FDiffOutputRecorder> DiffOutputRecorder)
 {
 	SourceHeaderData.Initialize();
 	DestHeaderData.Initialize();
@@ -1859,14 +2116,14 @@ void DumpPackageHeaderDiffs_ZenPackage(FPackageHeaderData& SourceHeaderData,
 	{
 		bFoundDifferenceInLinkerTableData = true;
 		DumpTableDifferences<FString>(SourceHeader, DestHeader, SourceNames, DestNames,
-			*AssetFilename, TEXT("Name"), MaxDiffsToLog);
+			*AssetFilename, TEXT("Name"), MaxDiffsToLog, DiffOutputRecorder);
 	}
 
 	if (!SourceHeader.IsImportMapIdentical(DestHeader))
 	{
 		bFoundDifferenceInLinkerTableData = true;
 		DumpTableDifferences<FPackageObjectIndex>(SourceHeader, DestHeader, SourceHeader.GetPackageHeader().ImportMap,
-			DestHeader.GetPackageHeader().ImportMap, *AssetFilename, TEXT("Import"), MaxDiffsToLog);
+			DestHeader.GetPackageHeader().ImportMap, *AssetFilename, TEXT("Import"), MaxDiffsToLog, DiffOutputRecorder);
 	}
 
 	if (!SourceHeader.IsExportMapIdentical(DestHeader))
@@ -1887,7 +2144,7 @@ void DumpPackageHeaderDiffs_ZenPackage(FPackageHeaderData& SourceHeaderData,
 			DestIndices.Add(FZenHeaderIndexIntoExportMap{ N });
 		}
 		DumpTableDifferences<FZenHeaderIndexIntoExportMap>(SourceHeader, DestHeader, SourceIndices, DestIndices,
-			*AssetFilename, TEXT("Export"), MaxDiffsToLog);
+			*AssetFilename, TEXT("Export"), MaxDiffsToLog, DiffOutputRecorder);
 	}
 
 	// Don't show differences in the dependency data if we found differences in the names/imports/exports data,
@@ -1915,7 +2172,7 @@ void DumpPackageHeaderDiffs_ZenPackage(FPackageHeaderData& SourceHeaderData,
 				{
 					SourceHeader.LogMessage(Verbosity, Message);
 				},
-				*AssetFilename, TEXT("ExportBundles"), MaxDiffsToLog);
+				*AssetFilename, TEXT("ExportBundles"), MaxDiffsToLog, DiffOutputRecorder);
 		}
 		if (!SourceHeader.IsDependencyBundlesIdentical(DestHeader))
 		{
@@ -1937,21 +2194,21 @@ void DumpPackageHeaderDiffs_ZenPackage(FPackageHeaderData& SourceHeaderData,
 				{
 					SourceHeader.LogMessage(Verbosity, Message);
 				},
-				*AssetFilename, TEXT("DependencyBundles"), MaxDiffsToLog);
+				*AssetFilename, TEXT("DependencyBundles"), MaxDiffsToLog, DiffOutputRecorder);
 		}
 	}
 }
 
 void DumpPackageHeaderDiffs(FPackageHeaderData& SourceHeaderData, FPackageHeaderData& DestHeaderData,
-	const int32 MaxDiffsToLog)
+	const int32 MaxDiffsToLog, TSharedPtr<UE::DiffWriter::FDiffOutputRecorder> DiffOutputRecorder)
 {
 	switch (SourceHeaderData.GetFormat())
 	{
 	case EPackageHeaderFormat::PackageFileSummary:
-		DumpPackageHeaderDiffs_LinkerLoad(SourceHeaderData, DestHeaderData, MaxDiffsToLog);
+		DumpPackageHeaderDiffs_LinkerLoad(SourceHeaderData, DestHeaderData, MaxDiffsToLog, DiffOutputRecorder);
 		break;
 	case EPackageHeaderFormat::ZenPackageSummary:
-		DumpPackageHeaderDiffs_ZenPackage(SourceHeaderData, DestHeaderData, MaxDiffsToLog);
+		DumpPackageHeaderDiffs_ZenPackage(SourceHeaderData, DestHeaderData, MaxDiffsToLog, DiffOutputRecorder);
 		break;
 	default:
 		unimplemented();
@@ -1960,12 +2217,12 @@ void DumpPackageHeaderDiffs(FPackageHeaderData& SourceHeaderData, FPackageHeader
 
 FPackageHeaderData::FPackageHeaderData(const TCHAR* InName, bool bInReadFromPackageStore,
 	const FString& InAssetFilename, const FPackageData& InPackageData, EPackageHeaderFormat InFormat,
-	FAccumulatorGlobals& InGlobals, FMessageCallback& InMessageCallback)
+	FAccumulatorGlobals& InGlobals, FDiffOutputRecorder& InDiffOutputRecorder)
 	: Name(InName)
 	, AssetFilename(InAssetFilename)
 	, PackageData(InPackageData)
 	, Globals(InGlobals)
-	, MessageCallback(InMessageCallback)
+	, DiffOutputRecorder(InDiffOutputRecorder)
 	, Format(InFormat)
 	, bReadFromPackageStore(bInReadFromPackageStore)
 {
@@ -2019,7 +2276,7 @@ void FPackageHeaderData::Initialize()
 		break;
 	}
 	case EPackageHeaderFormat::ZenPackageSummary:
-		ZenHeader = MakeUnique<FDiffWriterZenHeader>(Globals, MessageCallback, bReadFromPackageStore, PackageData,
+		ZenHeader = MakeUnique<FDiffWriterZenHeader>(Globals, DiffOutputRecorder, bReadFromPackageStore, PackageData,
 			AssetFilename, Name);
 		break;
 	default:
@@ -2053,9 +2310,9 @@ FAccumulatorGlobals& FPackageHeaderData::GetGlobals() const
 	return Globals;
 }
 
-FMessageCallback& FPackageHeaderData::GetMessageCallback() const
+FDiffOutputRecorder& FPackageHeaderData::GetDiffOutputRecorder() const
 {
-	return MessageCallback;
+	return DiffOutputRecorder;
 }
 
 FDiffWriterZenHeader& FPackageHeaderData::GetZenHeader() const
diff --git a/Engine/Source/Editor/UnrealEd/Private/Cooker/DiffWriterArchive.h b/Engine/Source/Editor/UnrealEd/Private/Cooker/DiffWriterArchive.h
--- a/Engine/Source/Editor/UnrealEd/Private/Cooker/DiffWriterArchive.h
+++ b/Engine/Source/Editor/UnrealEd/Private/Cooker/DiffWriterArchive.h
@@ -183,7 +183,7 @@ public:
 	void Append(const FCallstacks& Other, int64 OtherStartOffset);
 
 	/** Finds a callstack associated with data at the specified offset */
-	int32 GetCallstackIndexAtOffset(int64 Offset, int32 MinOffsetIndex = 0) const;
+	int32 GetCallstackIndexAtOffset(int64 Offset, int32 MinOffsetIndex = 0, int64* OutOffsetEnd = nullptr) const;
 
 	/** Finds a callstack associated with data at the specified offset */
 	const FCallstackAtOffset& GetCallstack(int32 CallstackIndex) const
@@ -221,6 +221,84 @@ private:
 	int64 EndOffset;
 };
 
+
+enum class EDiffWriterSectionType : uint8
+{
+	Header,
+	Exports,
+
+	MAX
+};
+
+
+
+/** Diff detail serialization interface. Used for generating additional report artifacts */
+class IDetailRecorder
+{
+public:
+	virtual void BeginPackage() = 0;
+	virtual void BeginSection(const TCHAR* SectionFilename, EDiffWriterSectionType SectionType, const FPackageData& SourceSection, 
+		const FPackageData& DestSection) = 0;
+	virtual void RecordDiff(int64 LocalOffset, const FCallstacks::FCallstackData* DifferenceCallstackData) = 0;
+	virtual void ExtendPreviousDiff(int64 LocalOffset) = 0;
+	virtual void IncrementPreviousDiff() = 0;
+	virtual void RecordUndiagnosedDiff() = 0;
+	virtual void RecordUnreportedDiffs(int32 NumUnreportedDiffs) = 0;
+	virtual void RecordTableDifferences(const TCHAR* ItemName, const TCHAR* HumanReadableString) = 0;
+	virtual void RecordDetermismDiagnostics(const TCHAR* DeterminismLines) = 0;
+	virtual void EndSection() = 0;
+	virtual void EndPackage(const TCHAR* Filename, const FPackageData& SourcePackage, const FPackageData& DestPackage, 
+		const TCHAR* ClassName) = 0;
+
+	virtual ~IDetailRecorder() = default;
+};
+
+/** Manages all output from the diff process */
+class FDiffOutputRecorder
+{
+public:
+	FDiffOutputRecorder(FMessageCallback&& InMessageCallback, IDetailRecorder* InDetailRecorder);
+
+	void RecordNewPackage(const TCHAR* Filename);
+
+	void BeginPackage(const TCHAR* Filename, const TCHAR* ClassName);
+	void BeginSection(const TCHAR* SectionFilename, EDiffWriterSectionType SectionType, const FPackageData& SourceSection, 
+		const FPackageData& DestSection);
+	void EndSection();
+	void EndPackage(const TCHAR* Filename, const FPackageData& SourcePackage, const FPackageData& DestPackage, const TCHAR* ClassName);
+
+	void RecordHeaderSizeMismatch(const TCHAR* Filename, int64 FirstHeaderSize, int64 SecondHeaderSize);
+	void RecordUndiagnosedHeaderDifference(const TCHAR* Filename);
+	void RecordTableDifferences(const TCHAR* Filename, const TCHAR* ItemName, int32 SourceTableNum, int32 DestTableNum, 
+		const TCHAR* HumanReadableString);
+	void RecordTableDifferences(const TCHAR* Filename, const TCHAR* ItemName, const TCHAR* HumanReadableString);
+	void RecordSectionSizeMismatch(const TCHAR* SectionFilename, int64 SourceSize, int64 DestSize);
+	void RecordDiff(const TCHAR* SectionFilename, int64 LocalOffset, int64 DestAbsoluteOffset, uint8 SourceByte, uint8 DestByte, 
+		bool bHasOptimizedHeader);
+	void RecordDiff(const TCHAR* SectionFilename, int64 LocalOffset, int64 DestAbsoluteOffset, uint8 SourceByte, uint8 DestByte, 
+		int64 DifferenceOffset, const TCHAR* LastDifferenceCallstackDataText, const FString& BeforePropertyVal, const FString& AfterPropertyVal, 
+		const FCallstacks::FCallstackData& DifferenceCallstackData);
+	void ExtendPreviousDiff(int64 LocalOffset);
+	void IncrementPreviousDiff();
+
+	void RecordDiffBytes(const TCHAR* SectionFilename, int32 BytesToLog, int64 LocalOffset, const FPackageData& SourcePackage, 
+		const FPackageData& DestPackage);
+	void RecordUnreportedDiffs(const TCHAR* SectionFilename, int32 NumUnreportedDiffs, int32 FirstUnreportedDiffIndex);
+
+	void RecordDeterminismDiagnostics(const FString& DeterminismLines);
+
+	const FMessageCallback& GetMessageCallback() const;
+	IDetailRecorder* GetDetailRecorder() const;
+
+protected:
+	FMessageCallback MessageCallback;
+	IDetailRecorder* DetailRecorder;
+};
+
+
+
+
+
 /** Global data (e.g. the FPackageId of every object in /Script) used during diffing */
 struct FAccumulatorGlobals
 {
@@ -233,6 +311,8 @@ public:
 	EPackageHeaderFormat Format = EPackageHeaderFormat::PackageFileSummary;
 	bool bInitialized = false;
 
+	TUniquePtr<IDetailRecorder> DetailRecorder;
+
 public:
 	FAccumulatorGlobals(ICookedPackageWriter* InnerPackageWriter = nullptr);
 	void Initialize(EPackageHeaderFormat Format);
@@ -251,7 +331,7 @@ class FAccumulator : public FRefCountBase
 {
 public:
 	FAccumulator(FAccumulatorGlobals& InGlobals, UObject* InAsset, FName InPackageName, int32 InMaxDiffsToLog,
-		bool bInIgnoreHeaderDiffs, FMessageCallback&& InMessageCallback, EPackageHeaderFormat InPackageHeaderFormat);
+		bool bInIgnoreHeaderDiffs, TSharedPtr<FDiffOutputRecorder> InDiffOutputRecorder, EPackageHeaderFormat InPackageHeaderFormat);
 	virtual ~FAccumulator();
 
 	void OnFirstSaveComplete(FStringView InLooseFilePath, int64 InHeaderSize, int64 InPreTransformHeaderSize,
@@ -289,7 +369,7 @@ private:
 	UE::Cook::FDeterminismManager* DeterminismManager = nullptr;
 
 	FDiffMap DiffMap;
-	FMessageCallback MessageCallback;
+	TSharedPtr<FDiffOutputRecorder> DiffOutputRecorder;
 	FName PackageName;
 	FString Filename;
 	UObject* Asset = nullptr;
diff --git a/Engine/Source/Editor/UnrealEd/Private/Cooker/DiffWriterLinkerLoadHeader.cpp b/Engine/Source/Editor/UnrealEd/Private/Cooker/DiffWriterLinkerLoadHeader.cpp
--- a/Engine/Source/Editor/UnrealEd/Private/Cooker/DiffWriterLinkerLoadHeader.cpp
+++ b/Engine/Source/Editor/UnrealEd/Private/Cooker/DiffWriterLinkerLoadHeader.cpp
@@ -7,9 +7,9 @@
 namespace UE::DiffWriter
 {
 
-FDiffWriterLinkerLoadHeader::FDiffWriterLinkerLoadHeader(FLinkerLoad* InLinker, const FMessageCallback& InMessageCallback)
+FDiffWriterLinkerLoadHeader::FDiffWriterLinkerLoadHeader(FLinkerLoad* InLinker, const FDiffOutputRecorder& InDiffOutputRecorder)
 	: Linker(InLinker)
-	, MessageCallback(InMessageCallback)
+	, DiffOutputRecorder(InDiffOutputRecorder)
 {
 	Linker = InLinker;
 }
@@ -251,7 +251,7 @@ bool FDiffWriterLinkerLoadHeader::IsExportMapIdentical(FDiffWriterLinkerLoadHead
 
 void FDiffWriterLinkerLoadHeader::LogMessage(ELogVerbosity::Type Verbosity, FStringView Message)
 {
-	MessageCallback(Verbosity, Message);
+	DiffOutputRecorder.GetMessageCallback()(Verbosity, Message);
 }
 
 }
\ No newline at end of file
diff --git a/Engine/Source/Editor/UnrealEd/Private/Cooker/DiffWriterLinkerLoadHeader.h b/Engine/Source/Editor/UnrealEd/Private/Cooker/DiffWriterLinkerLoadHeader.h
--- a/Engine/Source/Editor/UnrealEd/Private/Cooker/DiffWriterLinkerLoadHeader.h
+++ b/Engine/Source/Editor/UnrealEd/Private/Cooker/DiffWriterLinkerLoadHeader.h
@@ -15,11 +15,11 @@ struct FDiffWriterLinkerLoadHeader
 {
 public:
 	FLinkerLoad* Linker = nullptr;
-	const FMessageCallback& MessageCallback;
+	const FDiffOutputRecorder& DiffOutputRecorder;
 
 public:
 	// Constructor
-	FDiffWriterLinkerLoadHeader(FLinkerLoad* InLinker, const FMessageCallback& InMessageCallback);
+	FDiffWriterLinkerLoadHeader(FLinkerLoad* InLinker, const FDiffOutputRecorder& InDiffOutputRecorder);
 
 	// API for DumpTableDifferences
 	FString GetTableKey(const FObjectExport& Export);
diff --git a/Engine/Source/Editor/UnrealEd/Private/Cooker/DiffWriterZenHeader.cpp b/Engine/Source/Editor/UnrealEd/Private/Cooker/DiffWriterZenHeader.cpp
--- a/Engine/Source/Editor/UnrealEd/Private/Cooker/DiffWriterZenHeader.cpp
+++ b/Engine/Source/Editor/UnrealEd/Private/Cooker/DiffWriterZenHeader.cpp
@@ -10,17 +10,17 @@ namespace UE::DiffWriter
 {
 
 FDiffWriterZenHeader::FDiffWriterZenHeader(FAccumulatorGlobals& InGlobals,
-	const FMessageCallback& InMessageCallback, bool bInUseZenStoreForImports,
+	const FDiffOutputRecorder& InDiffOutputRecorder, bool bInUseZenStoreForImports,
 	const FPackageData& Package, const FString& AssetFilename, const TCHAR* WhichComparisonPackage)
 	: Globals(InGlobals)
-	, MessageCallback(InMessageCallback)
+	, DiffOutputRecorder(InDiffOutputRecorder)
 	, bUseZenStoreForImports(bInUseZenStoreForImports)
 {
 	FString ErrorMessage;
 	PackageHeader = FZenPackageHeader::MakeView(FMemoryView(Package.Data, Package.Size), ErrorMessage);
 	if (!ErrorMessage.IsEmpty())
 	{
-		InMessageCallback(ELogVerbosity::Warning, FString::Printf(
+		DiffOutputRecorder.GetMessageCallback()(ELogVerbosity::Warning, FString::Printf(
 			TEXT("%s: header is different, but cannot display differences. Package data for %s is invalid: %s"),
 			*AssetFilename, WhichComparisonPackage, *ErrorMessage));
 		PackageHeader = FZenPackageHeader();
@@ -432,7 +432,7 @@ FString FDiffWriterZenHeader::ConvertDependencyBundleToText(int32 Index)
 
 void FDiffWriterZenHeader::LogMessage(ELogVerbosity::Type Verbosity, FStringView Message)
 {
-	MessageCallback(Verbosity, Message);
+	DiffOutputRecorder.GetMessageCallback()(Verbosity, Message);
 }
 
 void FZenPackageExportsForDiff::Initialize(FAccumulatorGlobals& Globals, FName InPackageName, bool bUseZenStore)
diff --git a/Engine/Source/Editor/UnrealEd/Private/Cooker/DiffWriterZenHeader.h b/Engine/Source/Editor/UnrealEd/Private/Cooker/DiffWriterZenHeader.h
--- a/Engine/Source/Editor/UnrealEd/Private/Cooker/DiffWriterZenHeader.h
+++ b/Engine/Source/Editor/UnrealEd/Private/Cooker/DiffWriterZenHeader.h
@@ -43,7 +43,7 @@ class FDiffWriterZenHeader
 {
 public:
 	// Constructor and simple data readers
-	FDiffWriterZenHeader(FAccumulatorGlobals& InGlobals, const FMessageCallback& InMessageCallback,
+	FDiffWriterZenHeader(FAccumulatorGlobals& InGlobals, const FDiffOutputRecorder& InDiffOutputRecorder,
 		bool bInUseZenStoreForImports, const FPackageData& Package, const FString& AssetFilename,
 		const TCHAR* WhichComparisonPackage);
 	bool IsValid() const;
@@ -86,7 +86,7 @@ private:
 	TMap<FName, FZenPackageExportsForDiff> ImportedPackageExports;
 	TMap<FPackageObjectIndex, FString> ObjectIndexPathNames;
 	FAccumulatorGlobals& Globals;
-	const FMessageCallback& MessageCallback;
+	const FDiffOutputRecorder& DiffOutputRecorder;
 	bool bUseZenStoreForImports = false;
 	bool bValid = true;
 };
diff --git a/Engine/Source/Editor/UnrealEd/Private/Tests/DiffWriterArchiveTests.cpp b/Engine/Source/Editor/UnrealEd/Private/Tests/DiffWriterArchiveTests.cpp
--- a/Engine/Source/Editor/UnrealEd/Private/Tests/DiffWriterArchiveTests.cpp
+++ b/Engine/Source/Editor/UnrealEd/Private/Tests/DiffWriterArchiveTests.cpp
@@ -37,6 +37,13 @@ ICookedPackageWriter::FPreviousCookedBytesData ToPackageData(FLargeMemoryWriter&
 	};
 }
 
+TSharedPtr<FDiffOutputRecorder> MakeOutputRecorder()
+{
+	auto EmptyCallback = [](ELogVerbosity::Type Verbosity, FStringView Message) {};
+
+	return MakeShared<FDiffOutputRecorder>(EmptyCallback, nullptr);
+}
+
 } // namespace UE::DiffWriter::Tests
 
 IMPLEMENT_SIMPLE_AUTOMATION_TEST(FDiffWriterArchiveTestsCallstacks, "System.Core.Cooker.DiffWriterArchive.Callstacks",
@@ -54,7 +61,7 @@ bool FDiffWriterArchiveTestsCallstacks::RunTest(const FString& Parameters)
 	UE::DiffWriter::FAccumulatorGlobals AccumulatorGlobals;
 
 	TRefCountPtr<FAccumulator> Accumulator = new FAccumulator(AccumulatorGlobals, nullptr, PackageName, MaxDiffsToLog,
-		bIgnoreHeaderDiffs, [](ELogVerbosity::Type Verbosity, FStringView Message) {}, HeaderFormat);
+		bIgnoreHeaderDiffs, UE::DiffWriter::Tests::MakeOutputRecorder(), HeaderFormat);
 	FBasicDiffState State = FBasicDiffState{ 1,2,3 };
 	FDiffArchiveForLinker Ar(*Accumulator);
 
@@ -94,7 +101,7 @@ bool FDiffWriterArchiveTestsBasic::RunTest(const FString& Parameters)
 	{
 		TRefCountPtr<FAccumulator> Accumulator = new FAccumulator(AccumulatorGlobals, nullptr, PackageName,
 			MaxDiffsToLog, bIgnoreHeaderDiffs,
-			[](ELogVerbosity::Type Verbosity, FStringView Message) {}, HeaderFormat);
+			UE::DiffWriter::Tests::MakeOutputRecorder(), HeaderFormat);
 
 		FLargeMemoryWriter InitialState;
 		{
@@ -118,7 +125,7 @@ bool FDiffWriterArchiveTestsBasic::RunTest(const FString& Parameters)
 	{
 		TRefCountPtr<FAccumulator> Accumulator = new FAccumulator(AccumulatorGlobals, nullptr, PackageName,
 			MaxDiffsToLog, bIgnoreHeaderDiffs,
-			[](ELogVerbosity::Type Verbosity, FStringView Message) {}, HeaderFormat);
+			UE::DiffWriter::Tests::MakeOutputRecorder(), HeaderFormat);
 
 		FLargeMemoryWriter InitialState;
 		{
@@ -142,7 +149,7 @@ bool FDiffWriterArchiveTestsBasic::RunTest(const FString& Parameters)
 	{
 		TRefCountPtr<FAccumulator> Accumulator = new FAccumulator(AccumulatorGlobals, nullptr, PackageName,
 			MaxDiffsToLog, bIgnoreHeaderDiffs,
-			[](ELogVerbosity::Type Verbosity, FStringView Message) {}, HeaderFormat);
+			UE::DiffWriter::Tests::MakeOutputRecorder(), HeaderFormat);
 
 		FBasicDiffState InitialState = FBasicDiffState {1,2,3};
 		FLargeMemoryWriter InitialMemory;
diff --git a/Engine/Source/Runtime/CoreUObject/Public/Cooker/CookDeterminismHelper.h b/Engine/Source/Runtime/CoreUObject/Public/Cooker/CookDeterminismHelper.h
--- a/Engine/Source/Runtime/CoreUObject/Public/Cooker/CookDeterminismHelper.h
+++ b/Engine/Source/Runtime/CoreUObject/Public/Cooker/CookDeterminismHelper.h
@@ -38,6 +38,9 @@ struct IDeterminismConstructDiagnosticsContext
 
 	/** Output function; call this to add a diagnostic with the given value. */
 	virtual void AddDiagnostic(FUtf8StringView DiagnosticName, const FCbField& Value) = 0;
+
+	/** Whether full DDC keys are required in the diagnostic */
+	virtual bool FullDDCKeysRequested() const = 0;
 };
 
 /**
diff --git a/Engine/Source/Runtime/Engine/Private/Animation/AnimSequence.cpp b/Engine/Source/Runtime/Engine/Private/Animation/AnimSequence.cpp
--- a/Engine/Source/Runtime/Engine/Private/Animation/AnimSequence.cpp
+++ b/Engine/Source/Runtime/Engine/Private/Animation/AnimSequence.cpp
@@ -1072,7 +1072,10 @@ public:
 		FCbWriter Writer;
 		Writer.BeginObject();
 		Writer << "DDCKeyHash" << WriteToString<64>(DDCKeyHash);
-		Writer << "DDCKey" << DDCKey;
+		if (Context.FullDDCKeysRequested())
+		{
+			Writer << "DDCKey" << DDCKey;
+		}
 		//Limiting the diagnostic size to the first 1000 chars
 		Writer << "Model" << FStringView(ModelStateString).Left(1000);
 		Writer << "RawDataGuidEqual" << bRawDataGuidEqual;
