FAssetReferenceFilterContext: add creation flags for the asset
#rnx

diff --git a/Engine/Source/Editor/UnrealEd/Classes/Editor/EditorEngine.h b/Engine/Source/Editor/UnrealEd/Classes/Editor/EditorEngine.h
--- a/Engine/Source/Editor/UnrealEd/Classes/Editor/EditorEngine.h
+++ b/Engine/Source/Editor/UnrealEd/Classes/Editor/EditorEngine.h
@@ -357,7 +357,7 @@ struct FAssetReferenceFilterContext
 {
 	UNREALED_API void AddReferencingAsset(const FAssetData& InReferencingAsset, EAssetReferenceFilterProperties InProperties = EAssetReferenceFilterProperties::None);
 	UNREALED_API void AddReferencingAssets(const TArray<FAssetData>& InReferencingAssets, EAssetReferenceFilterProperties InProperties = EAssetReferenceFilterProperties::None);
-	UNREALED_API void AddReferencingAssetsFromPropertyHandle(const TSharedPtr<class IPropertyHandle>& PropertyHandle);
+	UNREALED_API void AddReferencingAssetsFromPropertyHandle(const TSharedPtr<class IPropertyHandle>& PropertyHandle, FAssetData::ECreationFlags InCreationFlags = FAssetData::ECreationFlags::None);
 
 	const TArray<FAssetReferenceFilterReferencerInfo>& GetReferencingAssets() const { return ReferencingAssetInfo; }
 
diff --git a/Engine/Source/Editor/UnrealEd/Private/EditorEngine.cpp b/Engine/Source/Editor/UnrealEd/Private/EditorEngine.cpp
--- a/Engine/Source/Editor/UnrealEd/Private/EditorEngine.cpp
+++ b/Engine/Source/Editor/UnrealEd/Private/EditorEngine.cpp
@@ -495,7 +495,7 @@ void FAssetReferenceFilterContext::AddReferencingAssets(const TArray<FAssetData>
 	}
 }
 
-void FAssetReferenceFilterContext::AddReferencingAssetsFromPropertyHandle(const TSharedPtr<IPropertyHandle>& PropertyHandle)
+void FAssetReferenceFilterContext::AddReferencingAssetsFromPropertyHandle(const TSharedPtr<IPropertyHandle>& PropertyHandle, FAssetData::ECreationFlags InCreationFlags)
 {
 	if (!PropertyHandle.IsValid())
 	{
@@ -527,7 +527,7 @@ void FAssetReferenceFilterContext::AddReferencingAssetsFromPropertyHandle(const
 	{
 		if (ReferencingObject)
 		{
-			FAssetData ReferencingObjectData(ReferencingObject);
+			FAssetData ReferencingObjectData(ReferencingObject, InCreationFlags);
 			ReferencingAssetInfo.Emplace(ReferencingObjectData, ReferencerProperties);
 			
 			PRAGMA_DISABLE_DEPRECATION_WARNINGS
