Use skeleton class for interface conformance

[RN] Fixed compiler failing to find newly added functions on blueprint interfaces, leading to implemented functions being renamed

diff --git a/Engine/Source/Editor/UnrealEd/Private/Kismet2/BlueprintEditorUtils.cpp b/Engine/Source/Editor/UnrealEd/Private/Kismet2/BlueprintEditorUtils.cpp
--- a/Engine/Source/Editor/UnrealEd/Private/Kismet2/BlueprintEditorUtils.cpp
+++ b/Engine/Source/Editor/UnrealEd/Private/Kismet2/BlueprintEditorUtils.cpp
@@ -7203,7 +7203,7 @@ static void ConformInterfaceByName(UBlueprint* Blueprint, FBPInterfaceDescriptio
 				continue;
 			}
 
-			UFunction* InterfaceFunction = EventNode->EventReference.ResolveMember<UFunction>(CurrentInterfaceDesc.Interface);
+			UFunction* InterfaceFunction = EventNode->EventReference.ResolveMember<UFunction>(FBlueprintEditorUtils::GetMostUpToDateClass(CurrentInterfaceDesc.Interface));
 			// if the function is still ok as an event, no need to try and fix it up
 			if (UEdGraphSchema_K2::FunctionCanBePlacedAsEvent(InterfaceFunction))
 			{
