IncrementalCook: Fix bug in previous IncrementallyModifiedDiagnostics change that was causing PackageDependencies and ConfigDependencies to be ignored; TryCalculateCurrentKey was using the no-longer-updated values of the stored PackageDependency to calculate the current key, rather than using the current value of the PackageDependency.
#rnx
#incrementalcook

diff --git a/Engine/Source/Editor/UnrealEd/Private/Cooker/CookPackageArtifacts.cpp b/Engine/Source/Editor/UnrealEd/Private/Cooker/CookPackageArtifacts.cpp
--- a/Engine/Source/Editor/UnrealEd/Private/Cooker/CookPackageArtifacts.cpp
+++ b/Engine/Source/Editor/UnrealEd/Private/Cooker/CookPackageArtifacts.cpp
@@ -142,11 +142,12 @@ FBuildDependencySet::ECurrentKeyResult FBuildDependencySet::TryCalculateCurrentK
 				}
 				if (PackageDigest.IsSuccessful())
 				{
+					FBlake3Hash ConvertedHash = FCookDependency::ConvertToHash(PackageDigest.Hash);
 					if (bUpdateValues)
 					{
-						PackageDependency.SetValue(PackageDigest.Hash);
+						PackageDependency.SetValue(ConvertedHash);
 					}
-					KeyBuilder.Update(PackageDependency.GetRawValue(), FCookDependency::ValueSizeInBytes);
+					KeyBuilder.Update(&ConvertedHash.GetBytes(), sizeof(ConvertedHash.GetBytes()));
 					
 					continue;
 				}
@@ -174,11 +175,12 @@ FBuildDependencySet::ECurrentKeyResult FBuildDependencySet::TryCalculateCurrentK
 				KeyBuilder.Update(&Marker, sizeof(Marker));
 				if (!Value.IsEmpty())
 				{
+					FBlake3Hash ConvertedHash = FCookDependency::ConvertToHash(FUtf8String(Value));
 					if (bUpdateValues)
 					{
-						ConfigDependency.SetValue(FUtf8String(Value));
+						ConfigDependency.SetValue(ConvertedHash);
 					}
-					KeyBuilder.Update(ConfigDependency.GetRawValue(), FCookDependency::ValueSizeInBytes);
+					KeyBuilder.Update(&ConvertedHash.GetBytes(), sizeof(ConvertedHash.GetBytes()));
 				}
 			}
 			break;
