FFeedbackContextEditor::StartSlowTask will now attempt to parent the window to the active top level regular window before resorting to using the main frame.

[RN]

diff --git a/Engine/Source/Editor/UnrealEd/Private/FeedbackContextEditor.cpp b/Engine/Source/Editor/UnrealEd/Private/FeedbackContextEditor.cpp
--- a/Engine/Source/Editor/UnrealEd/Private/FeedbackContextEditor.cpp
+++ b/Engine/Source/Editor/UnrealEd/Private/FeedbackContextEditor.cpp
@@ -468,18 +468,18 @@ void FFeedbackContextEditor::StartSlowTask( const FText& Task, bool bShowCancelB
 			}
 		}
 
-		// Attempt to parent the slow task window to the slate main frame
-		if (!ParentWindow.IsValid() && FModuleManager::Get().IsModuleLoaded("MainFrame"))
+		// Attempt to parent the slow task window to the current active window
+		// This is more likely to show the slow task window where it's relevant and the user is already looking, then return focus appropriately when destroyed
+		if (!ParentWindow.IsValid())
 		{
-			IMainFrameModule& MainFrame = FModuleManager::LoadModuleChecked<IMainFrameModule>("MainFrame");
-			ParentWindow = MainFrame.GetParentWindow();
+			ParentWindow = FSlateApplicationBase::Get().GetActiveTopLevelRegularWindow();
 		}
 
-		// If we still don't have a valid window, attempt to parent the slow task window to the top level window of the application
-		// This might be a program compiling against (and acting as) editor, but that does not use the Mainframe module
-		if (!ParentWindow.IsValid() && FPlatformProperties::IsProgram())
+		// If we still don't have a valid window, attempt to parent the slow task window to the slate main frame
+		if (!ParentWindow.IsValid() && FModuleManager::Get().IsModuleLoaded("MainFrame"))
 		{
-			ParentWindow = FSlateApplication::Get().GetActiveTopLevelWindow();
+			IMainFrameModule& MainFrame = FModuleManager::LoadModuleChecked<IMainFrameModule>("MainFrame");
+			ParentWindow = MainFrame.GetParentWindow();
 		}
 
 		if (ParentWindow.IsValid())
