Fix a regression where the 2D Snap Layer menu was not available in the new viewport toolbar.

diff --git a/Engine/Source/Editor/EditorStyle/Private/StarshipStyle.cpp b/Engine/Source/Editor/EditorStyle/Private/StarshipStyle.cpp
--- a/Engine/Source/Editor/EditorStyle/Private/StarshipStyle.cpp
+++ b/Engine/Source/Editor/EditorStyle/Private/StarshipStyle.cpp
@@ -4793,7 +4793,7 @@ void FStarshipEditorStyle::FStyle::SetupLevelEditorStyle()
 		Set( "EditorViewport.LocationGridSnap", new IMAGE_BRUSH_SVG("Starship/EditorViewport/grid", Icon16x16));
 		Set( "EditorViewport.RotationGridSnap", new IMAGE_BRUSH_SVG("Starship/EditorViewport/angle", Icon16x16));
 
-		Set( "EditorViewport.Layer2DSnap", new IMAGE_BRUSH("Old/LevelEditor/Layer2DSnap", Icon14x14));
+		Set( "EditorViewport.Layer2DSnap", new IMAGE_BRUSH("Starship/EditorViewport/Layer2DSnap", Icon14x14));
 
 		Set("EditorViewport.ScaleGridSnap", new IMAGE_BRUSH_SVG( "Starship/EditorViewport/scale-grid-snap", Icon16x16 ) );
 		Set("EditorViewport.ToggleSurfaceSnapping", new IMAGE_BRUSH_SVG( "Starship/EditorViewport/surface-snap", Icon16x16 ) );
diff --git a/Engine/Source/Editor/UnrealEd/Private/ViewportToolbar/UnrealEdViewportToolbar.cpp b/Engine/Source/Editor/UnrealEd/Private/ViewportToolbar/UnrealEdViewportToolbar.cpp
--- a/Engine/Source/Editor/UnrealEd/Private/ViewportToolbar/UnrealEdViewportToolbar.cpp
+++ b/Engine/Source/Editor/UnrealEd/Private/ViewportToolbar/UnrealEdViewportToolbar.cpp
@@ -15,6 +15,7 @@
 #include "GPUSkinCacheVisualizationMenuCommands.h"
 #include "IPreviewProfileController.h"
 #include "IPreviewLODController.h"
+#include "ISettingsModule.h"
 #include "LevelEditor.h"
 #include "LevelEditorActions.h"
 #include "RayTracingDebugVisualizationMenuCommands.h"
@@ -582,6 +583,145 @@ FToolMenuEntry CreateScaleSnapCheckboxMenu(const FToolMenuContext& InContext)
 	return Entry;
 }
 
+FToolMenuEntry CreateLayer2DSnapCheckboxMenu(const FToolMenuContext& InContext)
+{
+	TSharedPtr<FUICommandInfo> Command = FEditorViewportCommands::Get().Layer2DSnap;
+
+	FUIAction Action;
+	
+	if (const FUIAction* FoundAction = InContext.GetActionForCommand(Command))
+	{
+		Action = *FoundAction;
+	}
+	else
+	{
+		// Provide a default implementation
+		// TODO: This would be better handled by ensuring each editor has a bound command.
+		Action.ExecuteAction.BindLambda([]
+		{
+			ULevelEditorViewportSettings* ViewportSettings = GetMutableDefault<ULevelEditorViewportSettings>();
+			const ULevelEditor2DSettings* Settings2D = GetDefault<ULevelEditor2DSettings>();
+			if (!ViewportSettings->bEnableLayerSnap && (Settings2D->SnapLayers.Num() > 0))
+			{
+				ViewportSettings->bEnableLayerSnap = true;
+				ViewportSettings->ActiveSnapLayerIndex = FMath::Clamp(ViewportSettings->ActiveSnapLayerIndex, 0, Settings2D->SnapLayers.Num() - 1);
+			}
+			else
+			{
+				ViewportSettings->bEnableLayerSnap = false;
+			}
+			ViewportSettings->PostEditChange();
+		});
+		Action.GetActionCheckState.BindLambda([]
+		{
+			const ULevelEditorViewportSettings* ViewportSettings = GetDefault<ULevelEditorViewportSettings>();
+			const ULevelEditor2DSettings* Settings2D = GetDefault<ULevelEditor2DSettings>();
+			const bool bChecked = ViewportSettings->bEnableLayerSnap && Settings2D->SnapLayers.IsValidIndex(ViewportSettings->ActiveSnapLayerIndex);
+			return bChecked ? ECheckBoxState::Checked : ECheckBoxState::Unchecked;
+		});
+	}
+	
+	TAttribute<FText> Label = TAttribute<FText>::CreateLambda([]
+	{
+		const ULevelEditorViewportSettings* ViewportSettings = GetDefault<ULevelEditorViewportSettings>();
+		const ULevelEditor2DSettings* Settings2D = GetDefault<ULevelEditor2DSettings>();
+		if (Settings2D->SnapLayers.IsValidIndex(ViewportSettings->ActiveSnapLayerIndex))
+		{
+			return FText::FromString(Settings2D->SnapLayers[ViewportSettings->ActiveSnapLayerIndex].Name);
+		}
+
+		return FText();
+	});
+
+	FToolMenuEntry Entry = FToolMenuEntry::InitSubMenu(
+		"Layer2DSnap",
+		Label,
+		Command->GetDescription(),
+		FNewToolMenuDelegate::CreateLambda([](UToolMenu* Menu)
+		{
+			const ULevelEditor2DSettings* Settings2D = GetDefault<ULevelEditor2DSettings>();
+		
+			{
+				FToolMenuSection& LayerSection = Menu->AddSection("2DLayers", LOCTEXT("Layer2DLayerHeader", "Layers")); 
+				const int32 LayerCount = Settings2D->SnapLayers.Num();
+			
+				for (int32 LayerIndex = 0; LayerIndex < LayerCount; ++LayerIndex)
+				{
+					FName LayerName(*Settings2D->SnapLayers[LayerIndex].Name);
+
+					FUIAction Action(
+						FExecuteAction::CreateLambda([LayerIndex]
+						{
+							ULevelEditorViewportSettings* ViewportSettings = GetMutableDefault<ULevelEditorViewportSettings>();
+							ViewportSettings->bEnableLayerSnap = true;
+							ViewportSettings->ActiveSnapLayerIndex = LayerIndex;
+							ViewportSettings->PostEditChange();
+						}),
+						FCanExecuteAction(),
+						FIsActionChecked::CreateLambda([LayerIndex]
+						{
+							const ULevelEditorViewportSettings* ViewportSettings = GetDefault<ULevelEditorViewportSettings>();
+							return (ViewportSettings->ActiveSnapLayerIndex == LayerIndex);
+						})
+					);
+				
+					LayerSection.AddMenuEntry(
+						LayerName,
+						FText::FromName(LayerName),
+						FText::GetEmpty(),
+						FSlateIcon(),
+						Action,
+						EUserInterfaceActionType::RadioButton
+					);
+				}
+				
+				LayerSection.AddMenuEntry(
+					"GoToSettings",
+					LOCTEXT("2DSnap_EditLayer", "Edit Layers..."),
+					FText::GetEmpty(),
+					FSlateIcon(),
+					FExecuteAction::CreateLambda([]
+					{
+						if (ISettingsModule* SettingsModule = FModuleManager::GetModulePtr<ISettingsModule>("Settings"))
+						{
+							SettingsModule->ShowViewer("Project", "Editor", "LevelEditor2DSettings");
+						}
+					})
+				);
+			}
+			
+			{
+				FToolMenuSection& ActionSection = Menu->AddSection("Actions", LOCTEXT("Layer2DActionHeader", "Actions"));
+				
+				ActionSection.AddMenuEntry(FLevelEditorCommands::Get().SnapTo2DLayer);
+
+				ActionSection.AddSeparator(NAME_None);
+				ActionSection.AddMenuEntry(FLevelEditorCommands::Get().MoveSelectionToTop2DLayer);
+				ActionSection.AddMenuEntry(FLevelEditorCommands::Get().MoveSelectionUpIn2DLayers);
+				ActionSection.AddMenuEntry(FLevelEditorCommands::Get().MoveSelectionDownIn2DLayers);
+				ActionSection.AddMenuEntry(FLevelEditorCommands::Get().MoveSelectionToBottom2DLayer);
+				
+				ActionSection.AddSeparator(NAME_None);
+				ActionSection.AddMenuEntry(FLevelEditorCommands::Get().Select2DLayerAbove);
+				ActionSection.AddMenuEntry(FLevelEditorCommands::Get().Select2DLayerBelow);
+			}
+			
+		}),
+		Action,
+		EUserInterfaceActionType::ToggleButton
+	);
+	
+	Entry.Icon = Command->GetIcon();
+	Entry.ToolBarData.LabelOverride = Label;
+	
+	Entry.Visibility = TAttribute<bool>::CreateLambda([]
+	{
+		return GetDefault<ULevelEditor2DSettings>()->bEnableSnapLayers;
+	});
+	
+	return Entry;
+}
+
 bool IsViewModeSupported(EViewModeIndex InViewModeIndex)
 {
 	switch (InViewModeIndex)
@@ -1372,6 +1512,7 @@ FToolMenuEntry CreateSnappingSubmenu()
 				SnappingSection.AddEntry(Private::CreateLocationSnapCheckboxMenu(Submenu->Context)).SetShowInToolbarTopLevel(true);
 				SnappingSection.AddEntry(Private::CreateRotationSnapCheckboxMenu(Submenu->Context)).SetShowInToolbarTopLevel(true);
 				SnappingSection.AddEntry(Private::CreateScaleSnapCheckboxMenu(Submenu->Context)).SetShowInToolbarTopLevel(true);
+				SnappingSection.AddEntry(Private::CreateLayer2DSnapCheckboxMenu(Submenu->Context)).SetShowInToolbarTopLevel(true);
 			}
 		)
 	);
