Speculative fix for FScriptArrayHelper assertion when using -diffonly

#rnx

diff --git a/Engine/Source/Editor/UnrealEd/Private/Cooker/DiffWriterArchive.cpp b/Engine/Source/Editor/UnrealEd/Private/Cooker/DiffWriterArchive.cpp
--- a/Engine/Source/Editor/UnrealEd/Private/Cooker/DiffWriterArchive.cpp
+++ b/Engine/Source/Editor/UnrealEd/Private/Cooker/DiffWriterArchive.cpp
@@ -1177,13 +1177,6 @@ void FAccumulator::CompareWithPreviousForSection(const FPackageData& SourcePacka
 	const int64 SizeToCompare = FMath::Min(SourceSize, DestSize);
 	const FName AssetClass = GetAssetClass();
 	
-	if (SourceSize != DestSize)
-	{
-		DiffOutputRecorder->RecordSectionSizeMismatch(*SectionFilename, SourceSize, DestSize);
-		int64 SizeDiff = DestPackage.Size - SourcePackage.Size;
-		OutStats.FindOrAdd(AssetClass).DiffSize += SizeDiff;
-	}
-
 	FString LastDifferenceCallstackDataText;
 	int32 LastDifferenceCallstackOffsetIndex = -1;
 	int64 NumDiffsLocal = 0;
@@ -1402,6 +1395,14 @@ void FAccumulator::CompareWithPreviousForSection(const FPackageData& SourcePacka
 		DiffOutputRecorder->IncrementPreviousDiff();
 	}
 
+	if (SourceSize != DestSize)
+	{
+		DiffOutputRecorder->RecordSectionSizeMismatch(*SectionFilename, SourceSize, DestSize);
+		int64 SizeDiff = DestPackage.Size - SourcePackage.Size;
+		OutStats.FindOrAdd(AssetClass).DiffSize += SizeDiff;
+		OutStats.FindOrAdd(AssetClass).NumDiffs++;
+		InOutDiffsLogged++;
+	}
 
 	if (MaxDiffsToLog >= 0 && NumDiffsForLogStatLocal > NumDiffsLoggedLocal)
 	{
