Speculative fix for second save package header diff json capture assert

#rnx

diff --git a/Engine/Source/Editor/UnrealEd/Private/Cooker/DiffWriterArchive.cpp b/Engine/Source/Editor/UnrealEd/Private/Cooker/DiffWriterArchive.cpp
--- a/Engine/Source/Editor/UnrealEd/Private/Cooker/DiffWriterArchive.cpp
+++ b/Engine/Source/Editor/UnrealEd/Private/Cooker/DiffWriterArchive.cpp
@@ -94,7 +94,7 @@ private:
 
 /** Logs any mismatching header data. */
 void DumpPackageHeaderDiffs(FPackageHeaderData& SourceHeaderData, FPackageHeaderData& DestHeaderData,
-	const int32 MaxDiffsToLog, TSharedPtr<UE::DiffWriter::FDiffOutputRecorder> DiffOutputRecorder);
+	const int32 MaxDiffsToLog, UE::DiffWriter::FDiffOutputRecorder& DiffOutputRecorder);
 
 /** Returns a new linker for loading the specified package. */
 FLinkerLoad* CreateLinkerForPackage(
@@ -1028,7 +1028,7 @@ void FAccumulator::OnSecondSaveComplete(int64 InHeaderSize)
 		FPackageHeaderData SecondSaveHeader(TEXT("dest"), false /* bReadFromPackageStore */, Filename,
 			SecondSaveHeaderSegment, PackageHeaderFormat, Globals, HeaderDiffOutputRecorder);
 
-		DumpPackageHeaderDiffs(FirstSaveHeader, SecondSaveHeader, MaxDiffsToLog, DiffOutputRecorder);
+		DumpPackageHeaderDiffs(FirstSaveHeader, SecondSaveHeader, MaxDiffsToLog, HeaderDiffOutputRecorder);
 		// Suppress static analysis warning V547: Expression 'NumHeaderDiffMessages == 0' is always true
 		if (NumHeaderDiffMessages == 0) // -V547
 		{
@@ -1474,7 +1474,7 @@ void FAccumulator::CompareWithPrevious(const TCHAR* CallstackCutoffText, TMap<FN
 		CallstackCutoffText, NumLoggedDiffs, OutStats, Filename);
 	if (HeaderSize > 0 && OutStats.FindOrAdd(AssetClass).NumDiffs > 0)
 	{
-		DumpPackageHeaderDiffs(SourceHeader, DestHeader, MaxDiffsToLog, DiffOutputRecorder);
+		DumpPackageHeaderDiffs(SourceHeader, DestHeader, MaxDiffsToLog, *DiffOutputRecorder.Get());
 		// Suppress static analysis warning V547: Expression 'NumHeaderDiffMessages == 0' is always true
 		if (NumHeaderDiffMessages == 0) // -V547
 		{
@@ -1891,7 +1891,7 @@ static void DumpTableDifferences(
 	const TCHAR* AssetFilename,
 	const TCHAR* ItemName,
 	const int32 MaxDiffsToLog,
-	TSharedPtr<UE::DiffWriter::FDiffOutputRecorder> DiffOutputRecorder
+	UE::DiffWriter::FDiffOutputRecorder& DiffOutputRecorder
 )
 {
 	FString HumanReadableString;
@@ -1963,7 +1963,7 @@ static void DumpTableDifferences(
 	}
 
 
-	DiffOutputRecorder->RecordTableDifferences(AssetFilename, ItemName, SourceTable.Num(), DestTable.Num(), *HumanReadableString);
+	DiffOutputRecorder.RecordTableDifferences(AssetFilename, ItemName, SourceTable.Num(), DestTable.Num(), *HumanReadableString);
 }
 
 static void DumpOrderedArrayDifferences(
@@ -1976,7 +1976,7 @@ static void DumpOrderedArrayDifferences(
 	const TCHAR* AssetFilename,
 	const TCHAR* ItemName,
 	const int32 MaxDiffsToLog,
-	TSharedPtr<UE::DiffWriter::FDiffOutputRecorder> DiffOutputRecorder
+	UE::DiffWriter::FDiffOutputRecorder& DiffOutputRecorder
 )
 {
 	FString HumanReadableString;
@@ -2031,11 +2031,11 @@ static void DumpOrderedArrayDifferences(
 		HumanReadableString += NewLineToken;
 	}
 
-	DiffOutputRecorder->RecordTableDifferences(AssetFilename, ItemName, *HumanReadableString);
+	DiffOutputRecorder.RecordTableDifferences(AssetFilename, ItemName, *HumanReadableString);
 }
 
 void DumpPackageHeaderDiffs_LinkerLoad(FPackageHeaderData& SourceHeaderData, FPackageHeaderData& DestHeaderData,
-	const int32 MaxDiffsToLog, TSharedPtr<UE::DiffWriter::FDiffOutputRecorder> DiffOutputRecorder)
+	const int32 MaxDiffsToLog, UE::DiffWriter::FDiffOutputRecorder& DiffOutputRecorder)
 {
 	SourceHeaderData.Initialize();
 	DestHeaderData.Initialize();
@@ -2070,7 +2070,7 @@ void DumpPackageHeaderDiffs_LinkerLoad(FPackageHeaderData& SourceHeaderData, FPa
 }
 
 void DumpPackageHeaderDiffs_ZenPackage(FPackageHeaderData& SourceHeaderData,
-	FPackageHeaderData& DestHeaderData, const int32 MaxDiffsToLog, TSharedPtr<UE::DiffWriter::FDiffOutputRecorder> DiffOutputRecorder)
+	FPackageHeaderData& DestHeaderData, const int32 MaxDiffsToLog, UE::DiffWriter::FDiffOutputRecorder& DiffOutputRecorder)
 {
 	SourceHeaderData.Initialize();
 	DestHeaderData.Initialize();
@@ -2201,7 +2201,7 @@ void DumpPackageHeaderDiffs_ZenPackage(FPackageHeaderData& SourceHeaderData,
 }
 
 void DumpPackageHeaderDiffs(FPackageHeaderData& SourceHeaderData, FPackageHeaderData& DestHeaderData,
-	const int32 MaxDiffsToLog, TSharedPtr<UE::DiffWriter::FDiffOutputRecorder> DiffOutputRecorder)
+	const int32 MaxDiffsToLog, UE::DiffWriter::FDiffOutputRecorder& DiffOutputRecorder)
 {
 	switch (SourceHeaderData.GetFormat())
 	{
