[LUF] Fix Shift-Gizmo Transform Drag where an actor roated in local space wouldn't move in teh expected direction

- Inverse Y of OutDrag was being applied too late and wouldn't have effect on the transform applied

diff --git a/Engine/Source/Editor/UnrealEd/Private/MouseDeltaTracker.cpp b/Engine/Source/Editor/UnrealEd/Private/MouseDeltaTracker.cpp
--- a/Engine/Source/Editor/UnrealEd/Private/MouseDeltaTracker.cpp
+++ b/Engine/Source/Editor/UnrealEd/Private/MouseDeltaTracker.cpp
@@ -635,11 +635,6 @@ void FMouseDeltaTracker::ConvertMovementDeltaToDragRot(FSceneView* InView, FEdit
 	if( InViewportClient->GetCurrentWidgetAxis() != EAxisList::None )
 	{
 		InViewportClient->Widget->ConvertMouseMovementToAxisMovement( InView, InViewportClient, bUsedDragModifier, InOutDragDelta, OutDrag, OutRotation, OutScale );
-
-		if (AxisDisplayInfo::GetAxisDisplayCoordinateSystem() == EAxisList::LeftUpForward)
-		{
-			OutDrag.Y = -OutDrag.Y;
-		}
 	}
 	else
 	{
diff --git a/Engine/Source/Editor/UnrealEd/Private/UnrealWidget.cpp b/Engine/Source/Editor/UnrealEd/Private/UnrealWidget.cpp
--- a/Engine/Source/Editor/UnrealEd/Private/UnrealWidget.cpp
+++ b/Engine/Source/Editor/UnrealEd/Private/UnrealWidget.cpp
@@ -117,6 +117,11 @@ void FWidget::ConvertMouseToAxis_Translate(FVector2D DragDir, FVector& InOutDelt
 	const FVector2D EffectiveDelta = OutDrag.X * XAxisDir + OutDrag.Y * YAxisDir + OutDrag.Z * ZAxisDir;
 	InOutDelta                     = FVector(EffectiveDelta.X, -EffectiveDelta.Y, 0.0f);
 
+	if (AxisDisplayInfo::GetAxisDisplayCoordinateSystem() == EAxisList::LeftUpForward)
+	{
+		OutDrag.Y = -OutDrag.Y;
+	}
+
 	// Transform drag delta into world space
 	OutDrag = CustomCoordSystem.TransformPosition(OutDrag);
 }
@@ -464,7 +469,7 @@ void FWidget::ConvertMouseMovementToAxisMovement(FSceneView* InView, FEditorView
 	OutScale    = FVector::ZeroVector;
 
 	const int32 WidgetMode = InViewportClient->GetWidgetMode();
-
+	
 	// Get input delta as 2D vector, adjusted for inverted screen space Y axis
 	const FVector2D DragDir = FVector2D(InOutDelta.X, -InOutDelta.Y);
 
