WorldPartitionConvertCommandlet: add support for converting a world with external scene graph entities.

diff --git a/Engine/Source/Editor/UnrealEd/Private/Commandlets/WorldPartitionConvertCommandlet.cpp b/Engine/Source/Editor/UnrealEd/Private/Commandlets/WorldPartitionConvertCommandlet.cpp
--- a/Engine/Source/Editor/UnrealEd/Private/Commandlets/WorldPartitionConvertCommandlet.cpp
+++ b/Engine/Source/Editor/UnrealEd/Private/Commandlets/WorldPartitionConvertCommandlet.cpp
@@ -42,6 +42,7 @@
 #include "EdGraph/EdGraph.h"
 #include "ProfilingDebugging/ScopedTimers.h"
 #include "AssetRegistry/AssetRegistryModule.h"
+#include "AssetRegistry/AssetRegistryHelpers.h"
 #include "FoliageEditUtility.h"
 #include "FoliageHelper.h"
 #include "Engine/WorldComposition.h"
@@ -561,15 +562,7 @@ bool UWorldPartitionConvertCommandlet::RenameWorldPackageWithSuffix(UWorld* Worl
 	bool bRenamedSuccess = false;
 	UPackage* Package = World->GetPackage();
 
-	FString OldWorldName = World->GetName();
-	FString NewWorldName = OldWorldName + ConversionSuffix;
-	bRenamedSuccess = World->Rename(*NewWorldName, nullptr, REN_NonTransactional | REN_DontCreateRedirectors);
-	if (!bRenamedSuccess)
-	{
-		UE_LOG(LogWorldPartitionConvertCommandlet, Error, TEXT("Unable to rename world to %s"), *NewWorldName);
-		return false;
-	}
-
+	// First, rename the world package so code relying on it when renaming the world can work properly.
 	FString OldPackageName = Package->GetName();
 	FString NewPackageName = OldPackageName + ConversionSuffix;
 	FString NewPackageResourceName = Package->GetLoadedPath().GetPackageName().Replace(*OldPackageName, *NewPackageName);
@@ -581,6 +574,16 @@ bool UWorldPartitionConvertCommandlet::RenameWorldPackageWithSuffix(UWorld* Worl
 	}
 	Package->SetLoadedPath(FPackagePath::FromPackageNameChecked(NewPackageResourceName));
 
+	// Make sure the world object has the same name as its package.
+	FString OldWorldName = World->GetName();
+	FString NewWorldName = OldWorldName + ConversionSuffix;
+	bRenamedSuccess = World->Rename(*NewWorldName, nullptr, REN_NonTransactional | REN_DontCreateRedirectors);
+	if (!bRenamedSuccess)
+	{
+		UE_LOG(LogWorldPartitionConvertCommandlet, Error, TEXT("Unable to rename world to %s"), *NewWorldName);
+		return false;
+	}
+
 	return true;
 }
 
@@ -801,8 +804,24 @@ int32 UWorldPartitionConvertCommandlet::Main(const FString& Params)
 			}
 		}
 
-		bool bResult = PackageHelper.Delete(FilesToDelete);
-		if (!bResult)
+		if (TArray<FAssetData> EntityAssets = UAssetRegistryHelpers::GetAssetsWithOuterForPaths({ *OldLevelName }, *OldLevelName, true, true, false); !EntityAssets.IsEmpty())
+		{
+			for (const FAssetData& EntityAsset : EntityAssets)
+			{
+				FString EntityAssetPath;
+				if (FPackageName::TryConvertLongPackageNameToFilename(EntityAsset.PackageName.ToString(), EntityAssetPath, FPackageName::GetAssetPackageExtension()))
+				{
+					FilesToDelete.Add(EntityAssetPath);
+				}
+				else
+				{
+					UE_LOG(LogWorldPartitionConvertCommandlet, Error, TEXT("Failed to delete previous conversion package %s"), *EntityAsset.PackageName.ToString());
+					return 1;
+				}
+			}
+		}
+
+		if (!PackageHelper.Delete(FilesToDelete))
 		{
 			UE_LOG(LogWorldPartitionConvertCommandlet, Error, TEXT("Failed to delete previous conversion package(s)"));
 			return 1;
@@ -1436,7 +1455,20 @@ int32 UWorldPartitionConvertCommandlet::Main(const FString& Params)
 		TSet<FGuid> ActorGuids;
 		for(AActor* Actor: ActorList)
 		{
-			if (!IsValid(Actor) || !Actor->SupportsExternalPackaging())
+			if (!IsValid(Actor))
+			{
+				continue;
+			}
+
+			ForEachObjectWithOuter(Actor, [this](UObject* Object)
+			{
+				if (Object->IsPackageExternal())
+				{
+					PackagesToSave.Add(Object->GetPackage());
+				}
+			});
+
+			if (!Actor->SupportsExternalPackaging())
 			{
 				continue;
 			}
