MPCook: Remove call to ActivateMachineResourceReduction in StartCook. It was added to handle the case of starting cooks multiple times in one process, without relaunching the CookWorkers but it doesn't work because it needs to NOT be called in some circumstances in InitializeWorkers. Will reimplement in the future if and when we add a way to launch multiple MPCooks in a single process.
#rnx

diff --git a/Engine/Source/Editor/UnrealEd/Private/Cooker/CookDirector.cpp b/Engine/Source/Editor/UnrealEd/Private/Cooker/CookDirector.cpp
--- a/Engine/Source/Editor/UnrealEd/Private/Cooker/CookDirector.cpp
+++ b/Engine/Source/Editor/UnrealEd/Private/Cooker/CookDirector.cpp
@@ -322,7 +322,6 @@ void FCookDirector::StartCook(const FBeginCookContext& InBeginContext)
 	// work the Director has to do during the first tick, for PumpRequests and AssignRequests.
 	FScopeLock CommunicationScopeLock(&CommunicationLock);
 	InitializeWorkers();
-	ActivateMachineResourceReduction();
 }
 
 void FCookDirector::AssignRequests(TArrayView<FPackageData*> Requests, TArray<FWorkerId>& OutAssignments,
