Print global dependencies hash for each target platform when they are used so that we have a record to compare against on other runs.

diff --git a/Engine/Source/Editor/UnrealEd/Private/Cooker/CookGlobalDependencies.cpp b/Engine/Source/Editor/UnrealEd/Private/Cooker/CookGlobalDependencies.cpp
--- a/Engine/Source/Editor/UnrealEd/Private/Cooker/CookGlobalDependencies.cpp
+++ b/Engine/Source/Editor/UnrealEd/Private/Cooker/CookGlobalDependencies.cpp
@@ -26,21 +26,25 @@ namespace UE::Cook
 
 		Hasher.Update(&CookIncrementalVersion, sizeof(CookIncrementalVersion));
 
+		UE_LOG(LogCook, Display, TEXT("CalculateGlobalDependenciesHash(%s):"), *Platform->PlatformName());
 		{
 			bool bValue = false;
 			GConfig->GetBool(TEXT("Core.System"), TEXT("CanStripEditorOnlyExportsAndImports"), bValue, GEngineIni);
 			Hasher.Update(&bValue, sizeof(bValue));
+			UE_LOG(LogCook, Display, TEXT("\t[Engine]:Core.System.CanStripEditorOnlyExportsAndImports = %s"), bValue ? TEXT("True") : TEXT("False"));
 		}
 
 		{
 			bool bValue = false;
 			GConfig->GetBool(TEXT("Core.System"), TEXT("CanSkipEditorReferencedPackagesWhenCooking"), bValue, GEngineIni);
 			Hasher.Update(&bValue, sizeof(bValue));
+			UE_LOG(LogCook, Display, TEXT("\t[Engine]:Core.System.CanSkipEditorReferencedPackagesWhenCooking = %s"), bValue ? TEXT("True") : TEXT("False"));
 		}
 
 		{
 			bool bCanUseUnversionedPropertySerialization = CanUseUnversionedPropertySerialization(Platform);
 			Hasher.Update(&bCanUseUnversionedPropertySerialization, sizeof(bCanUseUnversionedPropertySerialization));
+			UE_LOG(LogCook, Display, TEXT("\tCanUseUnversionedPropertySerialization = %s"), bCanUseUnversionedPropertySerialization ? TEXT("True") : TEXT("False"));
 		}
 
 		{
@@ -50,14 +54,18 @@ namespace UE::Cook
 			{
 				Hasher.Update(*AssetType, AssetType.NumBytesWithoutNull());
 			}
+			UE_LOG(LogCook, Display, TEXT("\t[Editor]:CookSettings.AutomaticOptionalInclusionAssetType = %s"), *FString::Join(AssetList, TEXT(", ")));
 		}
 
 		{
 			bool bIsUnversioned = COTFS.IsCookFlagSet(ECookInitializationFlags::Unversioned);
 			Hasher.Update(&bIsUnversioned, sizeof(bIsUnversioned));
+			UE_LOG(LogCook, Display, TEXT("\tIsCookFlagSet(ECookInitializationFlags::Unversioned) = %s"), bIsUnversioned ? TEXT("True") : TEXT("False"));
 		}
 
-		GlobalHash.Add(Platform, Hasher.Finalize());
+		FBlake3Hash Hash = Hasher.Finalize();
+		UE_LOG(LogCook, Display, TEXT("\tGlobalDependenciesHash(%s) = %s"), *Platform->PlatformName(), *LexToString(Hash));
+		GlobalHash.Add(Platform, Hash);
 	}
 
 	FBlake3Hash GetGlobalDependenciesHash(const ITargetPlatform* Platform)
