SourceControl: Removed anti-pattern usage of doing an unnecessary ForceUpdate after just performing a source control operation on the file. After doing a source control operation, the source control state can be queried using 'Use' instead of 'ForceUpdate' as the source control operation is responsible for making sure the provider cache is populated accurately.

#rnx

diff --git a/Engine/Source/Developer/Localization/Private/LocalizationSourceControlUtil.cpp b/Engine/Source/Developer/Localization/Private/LocalizationSourceControlUtil.cpp
--- a/Engine/Source/Developer/Localization/Private/LocalizationSourceControlUtil.cpp
+++ b/Engine/Source/Developer/Localization/Private/LocalizationSourceControlUtil.cpp
@@ -68,7 +68,7 @@ bool FLocalizationSCC::CheckOutFile(const FString& InFile, FText& OutError)
 	{
 		// If it's deleted, we need to revert that first
 		SourceControlProvider.Execute(ISourceControlOperation::Create<FRevert>(), FilesToBeCheckedOut);
-		SourceControlState = SourceControlProvider.GetState(AbsoluteFilename, EStateCacheUsage::ForceUpdate);
+		SourceControlState = SourceControlProvider.GetState(AbsoluteFilename, EStateCacheUsage::Use);
 	}
 
 	if (SourceControlState.IsValid())
diff --git a/Engine/Source/Editor/UnrealEd/Private/FileHelpers.cpp b/Engine/Source/Editor/UnrealEd/Private/FileHelpers.cpp
--- a/Engine/Source/Editor/UnrealEd/Private/FileHelpers.cpp
+++ b/Engine/Source/Editor/UnrealEd/Private/FileHelpers.cpp
@@ -2434,8 +2434,9 @@ ECommandResult::Type FEditorFileUtils::CheckoutPackages(const TArray<UPackage*>&
 				IntermediateResult = SourceControlProvider.Execute(ISourceControlOperation::Create<FRevert>(), PkgsToRevert);
 				if (IntermediateResult == ECommandResult::Succeeded)
 				{
-					// Force update all states to checkout
-					IntermediateResult = SourceControlProvider.GetState(PkgsToCheckOut, SourceControlStates, EStateCacheUsage::ForceUpdate);
+					// Refetch states that have may have been updated by FRevert
+					SourceControlStates.Empty();
+					IntermediateResult = SourceControlProvider.GetState(PkgsToCheckOut, SourceControlStates, EStateCacheUsage::Use);
 				}
 			}
 
