SourceControl: Make sure UPackageTools::CheckOutRootPackages doesn't pop up an individual message box if multiple files fail to checkout.

#rnx

diff --git a/Engine/Source/Editor/UnrealEd/Private/PackageTools.cpp b/Engine/Source/Editor/UnrealEd/Private/PackageTools.cpp
--- a/Engine/Source/Editor/UnrealEd/Private/PackageTools.cpp
+++ b/Engine/Source/Editor/UnrealEd/Private/PackageTools.cpp
@@ -1368,7 +1368,7 @@ UPackageTools::UPackageTools(const FObjectInitializer& ObjectInitializer)
 			SourceControlProvider.Execute(ISourceControlOperation::Create<FUpdateStatus>(), Packages);
 
 			TArray<FString> TouchedPackageNames;
-			bool bCheckedSomethingOut = false;
+			TArray<FString> FailedPackageNames;
 			for( int32 PackageIndex = 0 ; PackageIndex < Packages.Num() ; ++PackageIndex )
 			{
 				UPackage* Package = Packages[PackageIndex];
@@ -1376,7 +1376,6 @@ UPackageTools::UPackageTools(const FObjectInitializer& ObjectInitializer)
 				if( SourceControlState.IsValid() && SourceControlState->CanCheckout() )
 				{
 					// The package is still available, so do the check out.
-					bCheckedSomethingOut = true;
 					TouchedPackageNames.Add(Package->GetName());
 				}
 				else
@@ -1385,13 +1384,21 @@ UPackageTools::UPackageTools(const FObjectInitializer& ObjectInitializer)
 					// Don't warn if the file isn't in the depot.
 					if (SourceControlState.IsValid() && SourceControlState->IsSourceControlled())
 					{			
-						FMessageDialog::Open( EAppMsgType::Ok, NSLOCTEXT("UnrealEd", "Error_PackageStatusChanged", "Package can't be checked out - status has changed!") );
+						FailedPackageNames.Add(Package->GetName());
 					}
 				}
 			}
 
-			// Synchronize source control state if something was checked out.
-			SourceControlProvider.Execute(ISourceControlOperation::Create<FCheckOut>(), SourceControlHelpers::PackageFilenames(TouchedPackageNames));
+			if (FailedPackageNames.Num() > 0)
+			{
+				FMessageDialog::Open(EAppMsgType::Ok, NSLOCTEXT("UnrealEd", "Error_PackageStatusChanged", "Some package(s) can't be checked out - status has changed!"));
+			}
+
+			if (TouchedPackageNames.Num() > 0)
+			{
+				// Synchronize source control state if something was checked out.
+				SourceControlProvider.Execute(ISourceControlOperation::Create<FCheckOut>(), SourceControlHelpers::PackageFilenames(TouchedPackageNames));
+			}
 		}
 	}
 
