SourceControl: The FEditorFileUtils::CheckoutPackages method should perform a single FRevert on multiple files instead of multiple FRevert on single files.

#rnx

diff --git a/Engine/Source/Editor/UnrealEd/Private/FileHelpers.cpp b/Engine/Source/Editor/UnrealEd/Private/FileHelpers.cpp
--- a/Engine/Source/Editor/UnrealEd/Private/FileHelpers.cpp
+++ b/Engine/Source/Editor/UnrealEd/Private/FileHelpers.cpp
@@ -2612,14 +2612,12 @@ ECommandResult::Type FEditorFileUtils::CheckoutPackages(const TArray<FString>& P
 
 	ISourceControlProvider& SourceControlProvider = ISourceControlModule::Get().GetProvider();
 
-	// Source control may have been enabled in the package checkout dialog.
-	// Ensure the status is up to date
+	// We have an array of package names, but the SCC needs an array of their corresponding filenames
+	TArray<FString> PkgsToCheckOutFilenames;
+	PkgsToCheckOutFilenames.Reserve(PkgsToCheckOut.Num());
+	
 	if(PkgsToCheckOut.Num() > 0)
 	{
-		// We have an array of package names, but the SCC needs an array of their corresponding filenames
-		TArray<FString> PkgsToCheckOutFilenames;
-		PkgsToCheckOutFilenames.Reserve(PkgsToCheckOut.Num());
-
 		for( auto PkgsToCheckOutIter = PkgsToCheckOut.CreateConstIterator(); PkgsToCheckOutIter; ++PkgsToCheckOutIter )
 		{
 			const FString& PackageToCheckOutName = *PkgsToCheckOutIter;
@@ -2630,10 +2628,37 @@ ECommandResult::Type FEditorFileUtils::CheckoutPackages(const TArray<FString>& P
 				PkgsToCheckOutFilenames.Add(PackageFilename);
 			}
 		}
+	}
 
+
+	// Source control may have been enabled in the package checkout dialog.
+	// Ensure the status is up to date
+	if(PkgsToCheckOutFilenames.Num() > 0)
+	{
 		CheckOutResult = SourceControlProvider.Execute(ISourceControlOperation::Create<FUpdateStatus>(), PkgsToCheckOutFilenames);
 	}
 
+	// If any file was marked for delete, revert it now so it can be checked out below
+	if(CheckOutResult != ECommandResult::Cancelled)
+	{
+		TArray<FString> PkgsToRevertFilenames;
+		PkgsToRevertFilenames.Reserve(PkgsToCheckOut.Num());
+	
+		for (const FString& PackageFilename : PkgsToCheckOutFilenames)
+		{
+			FSourceControlStatePtr SourceControlState = SourceControlProvider.GetState(PackageFilename, EStateCacheUsage::Use);
+			if ( SourceControlState.IsValid() && SourceControlState->IsDeleted() )
+			{
+				PkgsToRevertFilenames.Add(PackageFilename);
+			}
+		}
+		
+		if (PkgsToRevertFilenames.Num() > 0)
+		{
+			CheckOutResult = SourceControlProvider.Execute(ISourceControlOperation::Create<FRevert>(), PkgsToRevertFilenames);
+		}
+	}
+
 	TArray<FString> FinalPackageCheckoutList;
 	TArray<FString> FinalPackageMarkForAddList;
 	if(CheckOutResult != ECommandResult::Cancelled)
@@ -2653,13 +2678,6 @@ ECommandResult::Type FEditorFileUtils::CheckoutPackages(const TArray<FString>& P
 				SourceControlState = SourceControlProvider.GetState(PackageFilename, EStateCacheUsage::Use);
 			}
 
-			// If the file was marked for delete, revert it now so it can be checked out below
-			if ( SourceControlState.IsValid() && SourceControlState->IsDeleted() )
-			{
-				SourceControlProvider.Execute(ISourceControlOperation::Create<FRevert>(), PackageFilename);
-				SourceControlState = SourceControlProvider.GetState(PackageFilename, EStateCacheUsage::ForceUpdate);
-			}
-
 			// Mark the package for check out if possible
 			bool bShowCheckoutError = true;
 			if( SourceControlState.IsValid() )
