SourceControl: Fixed a 'View Changes' hyperlink appearing on a notification for SCC providers that do not support changelists.

#rnx

diff --git a/Engine/Source/Editor/UnrealEd/Private/EditorEngine.cpp b/Engine/Source/Editor/UnrealEd/Private/EditorEngine.cpp
--- a/Engine/Source/Editor/UnrealEd/Private/EditorEngine.cpp
+++ b/Engine/Source/Editor/UnrealEd/Private/EditorEngine.cpp
@@ -4245,8 +4245,12 @@ void UEditorEngine::RunDeferredMarkForAddFiles(bool)
 
 				Info.bFireAndForget = true;
 				Info.bUseSuccessFailIcons = true;
-				Info.Hyperlink = FSimpleDelegate::CreateLambda([]() { ISourceControlWindowsModule::Get().ShowChangelistsTab(); });
-				Info.HyperlinkText = LOCTEXT("ViewChangesHyperlink", "View Changes");
+
+				if (SourceControlProvider.UsesChangelists() || SourceControlProvider.UsesUncontrolledChangelists())
+				{
+					Info.Hyperlink = FSimpleDelegate::CreateLambda([]() { ISourceControlWindowsModule::Get().ShowChangelistsTab(); });
+					Info.HyperlinkText = LOCTEXT("ViewChangesHyperlink", "View Changes");
+				}
 
 				Info.ButtonDetails.Emplace(
 					LOCTEXT("Dismiss", "Dismiss"),
