Collection Manager Scripting Subsystem - "AddToCollection" updated to create a collection if it doesn't exist.

[RN] Collection Subscript System - AddToCollection will create a collection if the one provided doesnt exist

diff --git a/Engine/Source/Editor/UnrealEd/Private/Subsystems/CollectionManagerScriptingSubsystem.cpp b/Engine/Source/Editor/UnrealEd/Private/Subsystems/CollectionManagerScriptingSubsystem.cpp
--- a/Engine/Source/Editor/UnrealEd/Private/Subsystems/CollectionManagerScriptingSubsystem.cpp
+++ b/Engine/Source/Editor/UnrealEd/Private/Subsystems/CollectionManagerScriptingSubsystem.cpp
@@ -318,7 +318,23 @@ bool UCollectionManagerScriptingSubsystem::AddAssetsToCollection(const FCollecti
 	}
 
 	FText Error;
-	if (!CollectionContainer->AddToCollection(Collection.Name, CollectionManagerScriptingSubsystemUtil::ECollectionScriptingShareType_To_ECollectionShareType(Collection.ShareType), AssetPaths, nullptr, &Error))
+	const ECollectionShareType::Type ShareType = CollectionManagerScriptingSubsystemUtil::ECollectionScriptingShareType_To_ECollectionShareType(Collection.ShareType);
+	if (!CollectionContainer->CollectionExists(Collection.Name, ShareType))
+	{
+		if (!CollectionContainer->IsValidCollectionName(Collection.Name.ToString(), ECollectionShareType::CST_All, &Error))
+		{
+			CollectionManagerScriptingSubsystemUtil::LogLastCollectionManagerError(Error);
+			return false;
+		}
+
+		if (!CollectionContainer->CreateCollection(Collection.Name, ShareType, ECollectionStorageMode::Static, &Error))
+		{
+			CollectionManagerScriptingSubsystemUtil::LogLastCollectionManagerError(Error);
+			return false;
+		}
+	}
+
+	if (!CollectionContainer->AddToCollection(Collection.Name, ShareType, AssetPaths, nullptr, &Error))
 	{
 		CollectionManagerScriptingSubsystemUtil::LogLastCollectionManagerError(Error);
 		return false;
diff --git a/Engine/Source/Editor/UnrealEd/Public/Subsystems/CollectionManagerScriptingSubsystem.h b/Engine/Source/Editor/UnrealEd/Public/Subsystems/CollectionManagerScriptingSubsystem.h
--- a/Engine/Source/Editor/UnrealEd/Public/Subsystems/CollectionManagerScriptingSubsystem.h
+++ b/Engine/Source/Editor/UnrealEd/Public/Subsystems/CollectionManagerScriptingSubsystem.h
@@ -139,7 +139,7 @@ public:
 	/**
 	 * Add the given asset to the given collection.
 	 *
-	 * @param Collection The collection to modify.
+	 * @param Collection The collection to modify. Will be created if it does not exist.
 	 * @param AssetPath Asset to add (its path name, eg) /Game/MyFolder/MyAsset.MyAsset).
 	 *
 	 * @return True if the collection was modified, false otherwise (see the output log for details on error).
@@ -150,7 +150,7 @@ public:
 	/**
 	 * Add the given asset to the given collection.
 	 *
-	 * @param Collection The collection to modify.
+	 * @param Collection The collection to modify. Will be created if it does not exist.
 	 * @param AssetData Asset to add.
 	 *
 	 * @return True if the collection was modified, false otherwise (see the output log for details on error).
@@ -161,7 +161,7 @@ public:
 	/**
 	 * Add the given asset to the given collection.
 	 *
-	 * @param Collection The collection to modify.
+	 * @param Collection The collection to modify. Will be created if it does not exist.
 	 * @param AssetPtr Asset to add.
 	 *
 	 * @return True if the collection was modified, false otherwise (see the output log for details on error).
@@ -172,7 +172,7 @@ public:
 	/**
 	 * Add the given assets to the given collection.
 	 *
-	 * @param Collection The collection to modify.
+	 * @param Collection The collection to modify. Will be created if it does not exist.
 	 * @param AssetPaths Assets to add (their paths, eg) /Game/MyFolder/MyAsset.MyAsset).
 	 *
 	 * @return True if the collection was modified, false otherwise (see the output log for details on error).
@@ -183,7 +183,7 @@ public:
 	/**
 	 * Add the given assets to the given collection.
 	 *
-	 * @param Collection The collection to modify.
+	 * @param Collection The collection to modify. Will be created if it does not exist.
 	 * @param AssetDatas Assets to add.
 	 *
 	 * @return True if the collection was modified, false otherwise (see the output log for details on error).
@@ -194,7 +194,7 @@ public:
 	/**
 	 * Add the given assets to the given collection.
 	 *
-	 * @param Collection The collection to modify.
+	 * @param Collection The collection to modify. Will be created if it does not exist.
 	 * @param AssetPtrs Assets to add.
 	 *
 	 * @return True if the collection was modified, false otherwise (see the output log for details on error).
