Downgrade the header size mismatch error to a warning so the automated build can continue to subsequent steps.
- it's no more severe than the other items in the diff report

#rnx

diff --git a/Engine/Source/Editor/UnrealEd/Private/Cooker/DiffWriterArchive.cpp b/Engine/Source/Editor/UnrealEd/Private/Cooker/DiffWriterArchive.cpp
--- a/Engine/Source/Editor/UnrealEd/Private/Cooker/DiffWriterArchive.cpp
+++ b/Engine/Source/Editor/UnrealEd/Private/Cooker/DiffWriterArchive.cpp
@@ -661,7 +661,7 @@ FDiffOutputRecorder::FDiffOutputRecorder( FMessageCallback&& InMessageCallback,
 
 void FDiffOutputRecorder::RecordHeaderSizeMismatch( const TCHAR* Filename, int64 FirstHeaderSize, int64 SecondHeaderSize )
 {
-	MessageCallback(ELogVerbosity::Error, FString::Printf(
+	MessageCallback(ELogVerbosity::Warning, FString::Printf(
 		TEXT("%s: Indeterministic header size. When saving the package twice into memory, first header size %" INT64_FMT " != second header size %" INT64_FMT ". Callstacks for indeterminism in the exports will be incorrect.")
 		TEXT("\n\tDumping differences from first and second memory saves."),
 		Filename, FirstHeaderSize, SecondHeaderSize));
