Introduce camera-relative vetical viewport panning behind a CVar for evaluation.

#rnx

diff --git a/Engine/Source/Editor/UnrealEd/Private/EditorViewportClient.cpp b/Engine/Source/Editor/UnrealEd/Private/EditorViewportClient.cpp
--- a/Engine/Source/Editor/UnrealEd/Private/EditorViewportClient.cpp
+++ b/Engine/Source/Editor/UnrealEd/Private/EditorViewportClient.cpp
@@ -5062,6 +5062,16 @@ void FEditorViewportClient::UpdateGestureDelta()
 
 // Converts a generic movement delta into drag/rotation deltas based on the viewport and keys held down
 
+namespace UE::Editor::Private
+{
+	// TODO: This should only persist for the evaluation period.
+	static FAutoConsoleVariable CVarCameraRelativeVerticalMousePan(
+		TEXT("EditorViewport.CameraRelativeVerticalMousePan"),
+		false,
+		TEXT("Set to true to make dragging the mouse up and down move the view relative to the camera's orientation rather than in worldspace.")
+	);
+}
+
 void FEditorViewportClient::ConvertMovementToDragRot(const FVector& InDelta,
 													 FVector& InDragDelta,
 													 FRotator& InRotDelta) const
@@ -5122,11 +5132,17 @@ void FEditorViewportClient::ConvertMovementToDragRot(const FVector& InDelta,
 				// Pan left/right/up/down
 				const bool bInvert = !bIsUsingTrackpad && MiddleMouseButtonDown && GetDefault<ULevelEditorViewportSettings>()->bInvertMiddleMousePan;
 
-
-				float Direction = bInvert ? 1 : -1;
-				InDragDelta.X = InDelta.X * Direction * FMath::Sin( ViewRotation.Yaw * PI / 180.f );
-				InDragDelta.Y = InDelta.X * -Direction * FMath::Cos( ViewRotation.Yaw * PI / 180.f );
-				InDragDelta.Z = -Direction * InDelta.Y;
+				if (UE::Editor::Private::CVarCameraRelativeVerticalMousePan->GetBool())
+				{
+					InDragDelta = FQuat(ViewRotation) * FVector( 0.0f, InDelta.X, InDelta.Y) * (bInvert ? -1.0f : 1.0f);
+				}
+				else
+				{
+					float Direction = bInvert ? 1 : -1;
+					InDragDelta.X = InDelta.X * Direction * FMath::Sin( ViewRotation.Yaw * PI / 180.f );
+					InDragDelta.Y = InDelta.X * -Direction * FMath::Cos( ViewRotation.Yaw * PI / 180.f );
+					InDragDelta.Z = -Direction * InDelta.Y;
+				}
 			}
 			else if( RightMouseButtonDown && !LeftMouseButtonDown )
 			{
