Add setting to suppress search operation when changing a variable type

[RN] By default the blueprint editor no longer searches for references to variables when the variables type is changed via the blueprint editor- to reenable this behavior set "bSearchForReferencesWhenVariableTypeChanged" to true in your editor settings

diff --git a/Engine/Source/Editor/BlueprintGraph/Private/BlueprintEditorSettings.cpp b/Engine/Source/Editor/BlueprintGraph/Private/BlueprintEditorSettings.cpp
--- a/Engine/Source/Editor/BlueprintGraph/Private/BlueprintEditorSettings.cpp
+++ b/Engine/Source/Editor/BlueprintGraph/Private/BlueprintEditorSettings.cpp
@@ -73,6 +73,7 @@ UBlueprintEditorSettings::UBlueprintEditorSettings(const FObjectInitializer& Obj
 	, NodeTemplateCacheCapMB(20.f)
 	// Find-in-Blueprints Settings
 	, AllowIndexAllBlueprints(EFiBIndexAllPermission::LoadOnly)
+	, bSearchForReferencesWhenVariableTypeChanged(false)
 	// No category
 	, bShowInheritedVariables(false)
 	, bAlwaysShowInterfacesInOverrides(true)
diff --git a/Engine/Source/Editor/BlueprintGraph/Public/BlueprintEditorSettings.h b/Engine/Source/Editor/BlueprintGraph/Public/BlueprintEditorSettings.h
--- a/Engine/Source/Editor/BlueprintGraph/Public/BlueprintEditorSettings.h
+++ b/Engine/Source/Editor/BlueprintGraph/Public/BlueprintEditorSettings.h
@@ -283,7 +283,10 @@ public:
 	UPROPERTY(EditAnywhere, config, Category = FindInBlueprints)
 	EFiBIndexAllPermission AllowIndexAllBlueprints;
 
-public:
+	/** If true we'll automatically search for references to a variable when its type is changed */
+	UPROPERTY(EditAnywhere, config, Category = FindInBlueprints)
+	bool bSearchForReferencesWhenVariableTypeChanged;
+
 	/** If set we'll show the inherited variables in the My Blueprint view. */
 	UPROPERTY(config)
 	bool bShowInheritedVariables;
diff --git a/Engine/Source/Editor/UnrealEd/Private/Kismet2/BlueprintEditorUtils.cpp b/Engine/Source/Editor/UnrealEd/Private/Kismet2/BlueprintEditorUtils.cpp
--- a/Engine/Source/Editor/UnrealEd/Private/Kismet2/BlueprintEditorUtils.cpp
+++ b/Engine/Source/Editor/UnrealEd/Private/Kismet2/BlueprintEditorUtils.cpp
@@ -5046,6 +5046,25 @@ bool FBlueprintEditorUtils::VerifyUserWantsVariableTypeChanged(const FName& InVa
 	return RetCode == FSuppressableWarningDialog::Confirm || RetCode == FSuppressableWarningDialog::Suppressed;
 }
 
+bool FBlueprintEditorUtils::VerifyUserWantsToSearch(FText Reason, const FString& SearchString)
+{
+	const UBlueprintEditorSettings* LocalSettings = GetDefault<UBlueprintEditorSettings>();
+	// if the user has requested search operations on variable type changes then trigger the 
+	// search operation
+	 if(LocalSettings->bSearchForReferencesWhenVariableTypeChanged)
+	 {
+		return true;
+	 }
+
+	// else just log the search string so the user can search themself if they care to:
+	FText LocalizedMessage = FText::Format( 
+		LOCTEXT( "SearchSkipped", "{0} - search skipped as requested by editor settings, consider searching \"{1}\""),
+		Reason, FText::FromString(SearchString));
+	UE_LOG(LogBlueprint, Warning, TEXT("%s"), *LocalizedMessage.ToString());
+	
+	return false;
+}
+
 bool FBlueprintEditorUtils::VerifyUserWantsRepNotifyVariableNameChanged(const FName& InVarName, const FName& InFuncName)
 {
 	FFormatNamedArguments Args;
@@ -5134,7 +5153,8 @@ void FBlueprintEditorUtils::ChangeMemberVariableType(UBlueprint* Blueprint, cons
 					bBreakingVariableConnections = true;
 				}
 
-				const FScopedTransaction Transaction( LOCTEXT("ChangeVariableType", "Change Variable Type") );
+				FText OperationIdentifier = LOCTEXT("ChangeVariableType", "Change Variable Type");
+				const FScopedTransaction Transaction( OperationIdentifier );
 				Blueprint->Modify();
 
 				/** Only change the variable type if type selection is valid, some unloaded Blueprints will turn out to be bad */
@@ -5246,13 +5266,18 @@ void FBlueprintEditorUtils::ChangeMemberVariableType(UBlueprint* Blueprint, cons
 									break;
 								}
 							}
-
+							
 							if (FirstVariableNode)
 							{
-								constexpr bool bSetFindWithinBlueprint = false;
-								constexpr bool bSelectFirstResult = false;
 								constexpr EGetFindReferenceSearchStringFlags Flags = EGetFindReferenceSearchStringFlags::UseSearchSyntax;
-								BlueprintEditor->SummonSearchUI(bSetFindWithinBlueprint, FirstVariableNode->GetFindReferenceSearchString(Flags), bSelectFirstResult);
+								const FString SearchString = FirstVariableNode->GetFindReferenceSearchString(Flags);
+								// At scale indexing is slow, allow the user to disable it:
+								if(VerifyUserWantsToSearch(OperationIdentifier, SearchString))
+								{
+									constexpr bool bSetFindWithinBlueprint = false;
+									constexpr bool bSelectFirstResult = false;
+									BlueprintEditor->SummonSearchUI(bSetFindWithinBlueprint, SearchString, bSelectFirstResult);
+								}
 							}
 						}
 					}
@@ -5704,7 +5729,8 @@ void FBlueprintEditorUtils::ChangeLocalVariableType(UBlueprint* InBlueprint, con
 					}
 				}
 
-				const FScopedTransaction Transaction( LOCTEXT("ChangeLocalVariableType", "Change Local Variable Type") );
+				FText OperationIdentifier = LOCTEXT("ChangeLocalVariableType", "Change Local Variable Type");
+				const FScopedTransaction Transaction( OperationIdentifier );
 				InBlueprint->Modify();
 				FunctionEntry->Modify();
 
@@ -5761,12 +5787,18 @@ void FBlueprintEditorUtils::ChangeLocalVariableType(UBlueprint* InBlueprint, con
 						}
 					}
 
+					// At scale indexing is slow, allow the user to disable it:
 					if (FirstVariableNode)
 					{
-						constexpr bool bSetFindWithinBlueprint = true;
-						constexpr bool bSelectFirstResult = false;
 						constexpr EGetFindReferenceSearchStringFlags Flags = EGetFindReferenceSearchStringFlags::UseSearchSyntax;
-						BlueprintEditor->SummonSearchUI(bSetFindWithinBlueprint, VariableNodes[0]->GetFindReferenceSearchString(Flags), bSelectFirstResult);
+						const FString SearchString = VariableNodes[0]->GetFindReferenceSearchString(Flags);
+						// At scale indexing is slow, allow the user to disable it:
+						if(VerifyUserWantsToSearch(OperationIdentifier, SearchString))
+						{
+							constexpr bool bSetFindWithinBlueprint = true;
+							constexpr bool bSelectFirstResult = false;
+							BlueprintEditor->SummonSearchUI(bSetFindWithinBlueprint, SearchString, bSelectFirstResult);
+						}
 					}
 				}
 			}
diff --git a/Engine/Source/Editor/UnrealEd/Public/Kismet2/BlueprintEditorUtils.h b/Engine/Source/Editor/UnrealEd/Public/Kismet2/BlueprintEditorUtils.h
--- a/Engine/Source/Editor/UnrealEd/Public/Kismet2/BlueprintEditorUtils.h
+++ b/Engine/Source/Editor/UnrealEd/Public/Kismet2/BlueprintEditorUtils.h
@@ -1857,6 +1857,16 @@ protected:
 	 */
 	static UNREALED_API bool VerifyUserWantsVariableTypeChanged(const FName& InVarName);
 
+	/**
+	 * Helper function to warn user that we are about to search for a reference, which can incur an expensive
+	 * indexing operation
+	 *
+	 * @param Reason		The reason the search would be requested, logged if search is suppressed
+	 * @param SearchString	The search string that will be used if the user wants to search, logged if search is suppressed
+	 * @return				TRUE if the user wants to perform the search operation
+	 */
+	static UNREALED_API bool VerifyUserWantsToSearch(FText Reason, const FString& SearchString);
+
 	/**
 	 * Helper function to warn user of the results of changing a RepNotify variable name by displaying a suppressible dialog
 	 *
