SourceControl: Don't perform an FSave on package save.

#rnx

diff --git a/Engine/Source/Editor/UnrealEd/Private/FileHelpers.cpp b/Engine/Source/Editor/UnrealEd/Private/FileHelpers.cpp
--- a/Engine/Source/Editor/UnrealEd/Private/FileHelpers.cpp
+++ b/Engine/Source/Editor/UnrealEd/Private/FileHelpers.cpp
@@ -4487,13 +4487,7 @@ FEditorFileUtils::EPromptReturnCode InternalPromptForCheckoutAndSave(const TArra
 	PackagesToSave.Reserve(FinalSaveList.Num());
 
 	{
-		auto SaveOperation = ISourceControlOperation::Create<FSave>();
-		bool bCanSaveToRevisionControl = 
-			ISourceControlModule::Get().IsEnabled() &&
-			ISourceControlModule::Get().GetProvider().IsAvailable() &&
-			ISourceControlModule::Get().GetProvider().CanExecuteOperation(SaveOperation);
-
-		FScopedSlowTask SlowTask(static_cast<float>(FinalSaveList.Num() * (bCanSaveToRevisionControl ? 3 : 2)), NSLOCTEXT("UnrealEd", "SavingPackages", "Saving packages..."));
+		FScopedSlowTask SlowTask(static_cast<float>(FinalSaveList.Num() * 2), NSLOCTEXT("UnrealEd", "SavingPackages", "Saving packages..."));
 		SlowTask.MakeDialog();
 
 		UWorld* ActorsWorld = nullptr;
@@ -4585,13 +4579,6 @@ FEditorFileUtils::EPromptReturnCode InternalPromptForCheckoutAndSave(const TArra
 		{
 			FEditorDelegates::PostSaveExternalActors.Broadcast(ActorsWorld);
 		}
-
-		if (bCanSaveToRevisionControl)
-		{
-			SlowTask.EnterProgressFrame(PackagesToSave.Num() + PackagesToClean.Num());
-
-			ISourceControlModule::Get().GetProvider().Execute(SaveOperation, USourceControlHelpers::PackageFilenames(PackagesToSave), EConcurrency::Synchronous);
-		}
 	}
 
 	// Add all files that needs to be marked for add in one command, if any
