[UI/Toolbars] New PIE/Simulate box has a discrepancy between Simulate icons

diff --git a/Engine/Source/Editor/UnrealEd/Private/Kismet2/DebuggerCommands.cpp b/Engine/Source/Editor/UnrealEd/Private/Kismet2/DebuggerCommands.cpp
--- a/Engine/Source/Editor/UnrealEd/Private/Kismet2/DebuggerCommands.cpp
+++ b/Engine/Source/Editor/UnrealEd/Private/Kismet2/DebuggerCommands.cpp
@@ -937,7 +937,7 @@ TSharedRef< SWidget > FPlayWorldCommands::GeneratePlayMenuContent(TSharedRef<FUI
 					"Simulate",
 					FPlayWorldCommands::Get().Simulate->GetLabel(),
 					FPlayWorldCommands::Get().Simulate->GetDescription(),
-					FPlayWorldCommands::Get().Simulate->GetIcon(),
+					FSlateIcon(FAppStyle::GetAppStyleSetName(), "PlayWorld.PlaySimulate"),
 					FExecuteAction::CreateLambda(
 						[]()
 						{
