BSP: Fixed crash on new empty level when vertex snapping is enabled.

diff --git a/Engine/Source/Editor/UnrealEd/Private/SnappingUtils.cpp b/Engine/Source/Editor/UnrealEd/Private/SnappingUtils.cpp
--- a/Engine/Source/Editor/UnrealEd/Private/SnappingUtils.cpp
+++ b/Engine/Source/Editor/UnrealEd/Private/SnappingUtils.cpp
@@ -334,7 +334,7 @@ bool FEditorViewportSnapping::SnapToBSPVertex(FVector& Location, FVector GridBas
 		FVector3f	SrcPoint = (FVector3f)Location;
 		FVector3f	DestPoint;
 		int32 Temp;
-		if( GWorld->GetModel()->FindNearestVertex(SrcPoint, DestPoint, GetDefault<ULevelEditorViewportSettings>()->SnapDistance, Temp ) >= 0.0)
+		if( GWorld && GWorld->GetModel() && GWorld->GetModel()->FindNearestVertex(SrcPoint, DestPoint, GetDefault<ULevelEditorViewportSettings>()->SnapDistance, Temp ) >= 0.0)
 		{
 			Location = (FVector)DestPoint;
 			bSnapped = true;
