Switch to camera-relative vertical panning by default. Add an option for world-based vertical panning.

diff --git a/Engine/Config/BaseEditorPerProjectUserSettings.ini b/Engine/Config/BaseEditorPerProjectUserSettings.ini
--- a/Engine/Config/BaseEditorPerProjectUserSettings.ini
+++ b/Engine/Config/BaseEditorPerProjectUserSettings.ini
@@ -353,6 +353,7 @@ MouseSensitivty=.2f
 bInvertMouseLookYAxis=False
 bInvertOrbitYAxis=False
 bInvertMiddleMousePan=False
+bWorldSpaceVerticalPan=False
 bInvertRightMouseDollyYAxis=False
 ; Whether to use mouse position as direct widget position
 bUseAbsoluteTranslation=True
diff --git a/Engine/Source/Editor/UnrealEd/Classes/Settings/LevelEditorViewportSettings.h b/Engine/Source/Editor/UnrealEd/Classes/Settings/LevelEditorViewportSettings.h
--- a/Engine/Source/Editor/UnrealEd/Classes/Settings/LevelEditorViewportSettings.h
+++ b/Engine/Source/Editor/UnrealEd/Classes/Settings/LevelEditorViewportSettings.h
@@ -385,6 +385,10 @@ class ULevelEditorViewportSettings
 	UPROPERTY(EditAnywhere, config, Category=Controls)
 	bool bInvertMiddleMousePan;
 
+	/** If enabled, vertical movements of the mouse while panning with the middle mouse button will move the camera vertically in world space. Otherwise, vertical mouse movements will be camera-relative. */
+	UPROPERTY(EditAnywhere, config, Category=Controls)
+	bool bWorldSpaceVerticalPan;
+
 	/** Whether or not to invert the direction of right mouse dolly on the Y axis in orbit mode */
 	UPROPERTY(EditAnywhere, config, Category = Controls)
 	bool bInvertRightMouseDollyYAxis;
diff --git a/Engine/Source/Editor/UnrealEd/Private/EditorViewportClient.cpp b/Engine/Source/Editor/UnrealEd/Private/EditorViewportClient.cpp
--- a/Engine/Source/Editor/UnrealEd/Private/EditorViewportClient.cpp
+++ b/Engine/Source/Editor/UnrealEd/Private/EditorViewportClient.cpp
@@ -5028,17 +5028,6 @@ void FEditorViewportClient::UpdateGestureDelta()
 }
 
 // Converts a generic movement delta into drag/rotation deltas based on the viewport and keys held down
-
-namespace UE::Editor::Private
-{
-	// TODO: This should only persist for the evaluation period.
-	static FAutoConsoleVariable CVarCameraRelativeVerticalMousePan(
-		TEXT("EditorViewport.CameraRelativeVerticalMousePan"),
-		false,
-		TEXT("Set to true to make dragging the mouse up and down move the view relative to the camera's orientation rather than in worldspace.")
-	);
-}
-
 void FEditorViewportClient::ConvertMovementToDragRot(const FVector& InDelta,
 													 FVector& InDragDelta,
 													 FRotator& InRotDelta) const
@@ -5099,17 +5088,17 @@ void FEditorViewportClient::ConvertMovementToDragRot(const FVector& InDelta,
 				// Pan left/right/up/down
 				const bool bInvert = !bIsUsingTrackpad && MiddleMouseButtonDown && GetDefault<ULevelEditorViewportSettings>()->bInvertMiddleMousePan;
 
-				if (UE::Editor::Private::CVarCameraRelativeVerticalMousePan->GetBool())
-				{
-					InDragDelta = FQuat(ViewRotation) * FVector( 0.0f, InDelta.X, InDelta.Y) * (bInvert ? -1.0f : 1.0f);
-				}
-				else
+				if (ViewportSettings->bWorldSpaceVerticalPan)
 				{
 					float Direction = bInvert ? 1 : -1;
 					InDragDelta.X = InDelta.X * Direction * FMath::Sin( ViewRotation.Yaw * PI / 180.f );
 					InDragDelta.Y = InDelta.X * -Direction * FMath::Cos( ViewRotation.Yaw * PI / 180.f );
 					InDragDelta.Z = -Direction * InDelta.Y;
 				}
+				else
+				{
+					InDragDelta = FQuat(ViewRotation) * FVector( 0.0f, InDelta.X, InDelta.Y) * (bInvert ? -1.0f : 1.0f);
+				}
 			}
 			else if( RightMouseButtonDown && !LeftMouseButtonDown )
 			{
