Respect the "Use Mouse For Touch" editor preference/project setting when launching a new standalone PIE window. Before, this would only be respected for the non-standalone window in the editor or mobile preview.

[RN] Bug fix: Make sure to respect the "Use mouse for touch" editor preference and project setting in the editor when playing in a standalone or new PIE window.

diff --git a/Engine/Source/Editor/UnrealEd/Private/PlayLevel.cpp b/Engine/Source/Editor/UnrealEd/Private/PlayLevel.cpp
--- a/Engine/Source/Editor/UnrealEd/Private/PlayLevel.cpp
+++ b/Engine/Source/Editor/UnrealEd/Private/PlayLevel.cpp
@@ -50,6 +50,7 @@
 #include "GameMapsSettings.h"
 #include "GeneralProjectSettings.h"
 #include "Engine/NavigationObjectBase.h"
+#include "GameFramework/InputSettings.h"
 #include "GameFramework/PlayerStart.h"
 #include "GameFramework/GameModeBase.h"
 #include "Components/AudioComponent.h"
@@ -3032,6 +3033,9 @@ UGameInstance* UEditorEngine::CreateInnerProcessPIEGameInstance(FRequestPlaySess
 		ULevelEditorPlaySettings* PlayInSettings = GetMutableDefault<ULevelEditorPlaySettings>();
 		ViewportClient->EngineShowFlags.SetServerDrawDebug(PlayInSettings->ShowServerDebugDrawingByDefault());
 
+		const bool bUseMouseForTouch = InParams.EditorPlaySettings->UseMouseForTouch || GetDefault<UInputSettings>()->bUseMouseForTouch; 
+		FSlateApplication::Get().SetGameIsFakingTouchEvents(bUseMouseForTouch);
+
 		if (!InParams.EditorPlaySettings->EnableGameSound)
 		{
 			if (FAudioDeviceHandle GameInstanceAudioDevice = GameInstance->GetWorld()->GetAudioDevice())
diff --git a/Engine/Source/Editor/UnrealEd/Private/PlayLevelNewProcess.cpp b/Engine/Source/Editor/UnrealEd/Private/PlayLevelNewProcess.cpp
--- a/Engine/Source/Editor/UnrealEd/Private/PlayLevelNewProcess.cpp
+++ b/Engine/Source/Editor/UnrealEd/Private/PlayLevelNewProcess.cpp
@@ -5,6 +5,7 @@
 #include "HeadMountedDisplayTypes.h"
 #include "Editor.h"
 #include "GameFramework/GameModeBase.h"
+#include "GameFramework/InputSettings.h"
 #include "HAL/PlatformApplicationMisc.h"
 #include "Settings/LevelEditorPlaySettings.h"
 #include "DataDrivenShaderPlatformInfo.h"
@@ -160,6 +161,20 @@ void UEditorEngine::LaunchNewProcess(const FRequestPlaySessionParams& InParams,
 		// ToDo: Should this be on all multi-client launches?
 		CommandLine += TEXT(" -MultiprocessSaveConfig");
 	}
+	// If we are not launching a mobile preview, them we need to check the editor/project settings for using mouse touch as well.
+	else
+	{
+#if WITH_EDITOR
+		const bool bUseMouseForTouch = GetDefault<ULevelEditorPlaySettings>()->UseMouseForTouch || GetDefault<UInputSettings>()->bUseMouseForTouch;
+#else
+		const bool bUseMouseForTouch = GetDefault<UInputSettings>()->bUseMouseForTouch;
+#endif
+
+		if (bUseMouseForTouch)
+		{
+			CommandLine += TEXT(" -faketouches");
+		}
+	}
 
 	// In order for the previewer to adjust its safe zone according to the device profile specified in the editor play settings,
 	// we need to pass the PIESafeZoneOverride's values as command line variables to the new process that we are about to launch.
