FEditorFileUtils::IsMapPackageAsset will now get the object asset data from the asset registry and check if the package has the PKG_ContainsMapData flag.
This will effectively report external entities as map package assets. Did not remove the special case for external actors but it should also cover that case.

diff --git a/Engine/Source/Editor/UnrealEd/Private/FileHelpers.cpp b/Engine/Source/Editor/UnrealEd/Private/FileHelpers.cpp
--- a/Engine/Source/Editor/UnrealEd/Private/FileHelpers.cpp
+++ b/Engine/Source/Editor/UnrealEd/Private/FileHelpers.cpp
@@ -5482,6 +5482,22 @@ bool FEditorFileUtils::IsMapPackageAsset(const FString& ObjectPath, FString& Map
 					return true;
 				}
 			}
+
+			FAssetRegistryModule& AssetRegistryModule = FModuleManager::LoadModuleChecked<FAssetRegistryModule>("AssetRegistry");
+			IAssetRegistry& AssetRegistry = AssetRegistryModule.Get();
+
+			TArray<FAssetData> Assets;
+			if (AssetRegistry.GetAssetsByPackageName(*ObjectPath, Assets))
+			{
+				for (const FAssetData& Asset : Assets)
+				{
+					if (Asset.HasAnyPackageFlags(PKG_ContainsMapData))
+					{
+						MapFilePath = PackagePath;
+						return true;
+					}
+				}
+			}
 		}
 	}
 
