[UnrealEd] Added overload to SnapActorsToNearestActor that takes a World and outputs deltas, used by the new gizmo to support legacy snapping

diff --git a/Engine/Source/Editor/UnrealEd/Private/SnappingUtils.cpp b/Engine/Source/Editor/UnrealEd/Private/SnappingUtils.cpp
--- a/Engine/Source/Editor/UnrealEd/Private/SnappingUtils.cpp
+++ b/Engine/Source/Editor/UnrealEd/Private/SnappingUtils.cpp
@@ -67,6 +67,15 @@ public:
 	 */
 	bool SnapActorsToNearestActor( FVector& DragDelta, FLevelEditorViewportClient* ViewportClient );
 
+	/**
+     * Attempts to snap the selected actors to the nearest other actor
+     *
+     * @param DragDelta			The current world space drag amount
+     * @param InWorld			The world the user is dragging in
+     * @return the location delta to the nearest actor, if any
+     */
+    TOptional<FVector> GetDeltaToNearestActor( FVector& DragDelta, const UWorld* InWorld );
+
 	/**
 	 * Snaps actors to the nearest vertex on another actor
 	 *
@@ -185,6 +194,20 @@ void FEditorViewportSnapping::SetActorSnapDistance(float Distance)
 
 bool FEditorViewportSnapping::SnapActorsToNearestActor( FVector& Drag, FLevelEditorViewportClient* ViewportClient )
 {
+	const TOptional<FVector> DeltaToNearestActor = GetDeltaToNearestActor(Drag, ViewportClient->GetWorld());
+	const bool bWasSnapped = DeltaToNearestActor.IsSet();
+	if (bWasSnapped)
+	{
+		ViewportClient->ApplyDeltaToActors(DeltaToNearestActor.GetValue(), FRotator::ZeroRotator, FVector::ZeroVector);
+	}
+
+	return bWasSnapped;
+}
+
+TOptional<FVector> FEditorViewportSnapping::GetDeltaToNearestActor(FVector& Drag, const UWorld* InWorld)
+{
+	TOptional<FVector> DeltaToNearestActor;
+
 	FEditorModeTools& Tools = GLevelEditorModeTools();
 
 	// Does the user have actor snapping enabled?
@@ -202,7 +225,7 @@ bool FEditorViewportSnapping::SnapActorsToNearestActor( FVector& Drag, FLevelEdi
 
 			// Find the nearest actor to the pivot point that isn't part of the selection
 			const FVector PivotLocation = Tools.PivotLocation;
-			for (FActorIterator It(ViewportClient->GetWorld()); It; ++It)	// Actor iterator :( [Note: Also, can't use BoxOverlapMulti to as the pivot may lie outside the bounds of the actor!]
+			for (FActorIterator It(InWorld); It; ++It)	// Actor iterator :( [Note: Also, can't use BoxOverlapMulti to as the pivot may lie outside the bounds of the actor!]
 			{
 				AActor* Actor = static_cast<AActor*>( *It );
 				checkSlow( Actor->IsA(AActor::StaticClass()) );
@@ -253,7 +276,7 @@ bool FEditorViewportSnapping::SnapActorsToNearestActor( FVector& Drag, FLevelEdi
 						// Calculate the delta between the snapped location and the current pivot and apply to all the selected actors
 						const FScopedTransaction Transaction( NSLOCTEXT("UnrealEd", "SnapActorsToActor", "Snap Actors To Actor") );
 						const FVector PivotDelta = ( BestPoint - PivotLocation );
-						ViewportClient->ApplyDeltaToActors( PivotDelta, FRotator::ZeroRotator, FVector::ZeroVector );
+						DeltaToNearestActor = PivotDelta;
 						Tools.SetPivotLocation( BestPoint, false );	// Overwrite the location for next time we check
 						Drag = FVector::ZeroVector;	// Reset the drag so the pivot doesn't jump
 					}
@@ -261,7 +284,7 @@ bool FEditorViewportSnapping::SnapActorsToNearestActor( FVector& Drag, FLevelEdi
 				else if ( Tools.SnappedActor && !Tools.CachedLocation.Equals( PivotLocation ) )
 				{
 					const FVector PivotDelta = ( PivotLocation - BestPoint );
-					ViewportClient->ApplyDeltaToActors( PivotDelta, FRotator::ZeroRotator, FVector::ZeroVector );
+					DeltaToNearestActor = PivotDelta;
 					//GUnrealEd->UpdatePivotLocationForSelection();	// Calling this ends up forcing the pivot back inside the threshold?!
 					Tools.SetPivotLocation( PivotLocation, false );	// Overwrite the location for next time we check
 					Drag = FVector::ZeroVector;	// Reset the drag so the pivot doesn't jump
@@ -271,7 +294,7 @@ bool FEditorViewportSnapping::SnapActorsToNearestActor( FVector& Drag, FLevelEdi
 	}
 
 	Tools.SnappedActor = bSnapped;
-	return bSnapped;	// Whether or not the selection is snapped in place
+	return DeltaToNearestActor;	// Returns the location delta to the nearest actor, if snapped
 }
 
 void FEditorViewportSnapping::SnapPointToGrid(FVector& Point, const FVector& GridBase)
@@ -448,6 +471,11 @@ bool FSnappingUtils::SnapActorsToNearestActor(FVector& DragDelta, FLevelEditorVi
 	return EditorViewportSnapper.IsValid() ? EditorViewportSnapper->SnapActorsToNearestActor(DragDelta, ViewportClient) : false;
 }
 
+TOptional<FVector> FSnappingUtils::GetDeltaToNearestActor(FVector& DragDelta, const UWorld* InWorld)
+{
+	return EditorViewportSnapper.IsValid() ? EditorViewportSnapper->GetDeltaToNearestActor(DragDelta, InWorld) : TOptional<FVector>();
+}
+
 bool FSnappingUtils::SnapDraggedActorsToNearestVertex(FVector& DragDelta, FLevelEditorViewportClient* ViewportClient)
 {
 	return EditorViewportSnapper.IsValid() ? EditorViewportSnapper->SnapDraggedActorsToNearestVertex(DragDelta, ViewportClient) : false;
diff --git a/Engine/Source/Editor/UnrealEd/Public/SnappingUtils.h b/Engine/Source/Editor/UnrealEd/Public/SnappingUtils.h
--- a/Engine/Source/Editor/UnrealEd/Public/SnappingUtils.h
+++ b/Engine/Source/Editor/UnrealEd/Public/SnappingUtils.h
@@ -49,6 +49,15 @@ public:
 	 */
 	static UNREALED_API bool SnapActorsToNearestActor( FVector& DragDelta, FLevelEditorViewportClient* ViewportClient );
 
+	/**
+	 * Attempts to snap the selected actors to the nearest other actor
+	 *
+	 * @param DragDelta			The current world space drag amount
+	 * @param InWorld			The world the user is dragging in
+	 * @return the location delta to the nearest actor, if any
+	 */
+	static UNREALED_API TOptional<FVector> GetDeltaToNearestActor( FVector& DragDelta, const UWorld* InWorld );
+
 	/**
 	 * Snaps actors to the nearest vertex on another actor
 	 *
