Adds a new "Toggle Selected Actor Visibility" command & behavior. This command takes the default "H" binding from "Hide Selected Actors".
The command:
- Hides selected actors if any actors in the selection are visible.
- Shows selected actors if all actors in the selection are hidden.
- Does not affect selection.

[RN] minor editor

diff --git a/Engine/Source/Editor/LevelEditor/Private/LevelEditor.cpp b/Engine/Source/Editor/LevelEditor/Private/LevelEditor.cpp
--- a/Engine/Source/Editor/LevelEditor/Private/LevelEditor.cpp
+++ b/Engine/Source/Editor/LevelEditor/Private/LevelEditor.cpp
@@ -1111,6 +1111,12 @@ void FLevelEditorModule::BindGlobalLevelEditorCommands()
 		FExecuteAction::CreateStatic( &FLevelEditorActionCallbacks::ExecuteExecCommand, FString( TEXT("ACTOR HIDE SELECTED") ) ),
 		FCanExecuteAction::CreateStatic( &FLevelEditorActionCallbacks::CanHideSelectedActors )
 		);
+		
+	ActionList.MapAction(
+		Commands.ToggleSelectedVisibility,
+		FExecuteAction::CreateStatic( &FLevelEditorActionCallbacks::ExecuteExecCommand, FString( TEXT("ACTOR TOGGLE SELECTED VISIBILITY") ) ),
+		FCanExecuteAction::CreateStatic( &FLevelEditorActionCallbacks::CanToggleSelectedActorsVisibility )
+		);
 
 	ActionList.MapAction(
 		Commands.ShowAllStartup,
diff --git a/Engine/Source/Editor/LevelEditor/Private/LevelEditorActions.cpp b/Engine/Source/Editor/LevelEditor/Private/LevelEditorActions.cpp
--- a/Engine/Source/Editor/LevelEditor/Private/LevelEditorActions.cpp
+++ b/Engine/Source/Editor/LevelEditor/Private/LevelEditorActions.cpp
@@ -2233,6 +2233,11 @@ bool FLevelEditorActionCallbacks::CanShowSelectedActors()
 	return false;
 }
 
+bool FLevelEditorActionCallbacks::CanToggleSelectedActorsVisibility()
+{
+	return GEditor->GetSelectedActorCount() > 0;
+}
+
 void FLevelEditorActionCallbacks::OnSelectOwningHLODCluster()
 {
 	if (GEditor->GetSelectedActorCount() > 0)
@@ -4304,7 +4309,8 @@ void FLevelEditorCommands::RegisterCommands()
 #endif
 	UI_COMMAND( ShowSelectedOnly, "Show Only Selected", "Shows only the selected actors", EUserInterfaceActionType::Button, FInputChord() );
 	UI_COMMAND( ShowSelected, "Show Selected", "Shows the selected actors", EUserInterfaceActionType::Button, FInputChord( EModifierKey::Shift, EKeys::H ) );
-	UI_COMMAND( HideSelected, "Hide Selected", "Hides the selected actors", EUserInterfaceActionType::Button, FInputChord( EKeys::H ) );
+	UI_COMMAND( HideSelected, "Hide Selected", "Hides the selected actors", EUserInterfaceActionType::Button, FInputChord() );
+	UI_COMMAND( ToggleSelectedVisibility, "Toggle Selected Visibility", "Hides or shows the selected actors", EUserInterfaceActionType::Button, FInputChord( EKeys::H ) );
 	UI_COMMAND( ShowAllStartup, "Show All At Startup", "Shows all actors at startup", EUserInterfaceActionType::Button, FInputChord() );
 	UI_COMMAND( ShowSelectedStartup, "Show Selected At Startup", "Shows selected actors at startup", EUserInterfaceActionType::Button, FInputChord() );
 	UI_COMMAND( HideSelectedStartup, "Hide Selected At Startup", "Hide selected actors at startup", EUserInterfaceActionType::Button, FInputChord() );
diff --git a/Engine/Source/Editor/LevelEditor/Private/LevelEditorContextMenu.cpp b/Engine/Source/Editor/LevelEditor/Private/LevelEditorContextMenu.cpp
--- a/Engine/Source/Editor/LevelEditor/Private/LevelEditorContextMenu.cpp
+++ b/Engine/Source/Editor/LevelEditor/Private/LevelEditorContextMenu.cpp
@@ -1070,8 +1070,9 @@ void FLevelEditorContextMenuImpl::FillActorVisibilityMenu(UToolMenu* Menu)
 {
 	{
 		FToolMenuSection& Section = Menu->AddSection("VisibilitySelected", LOCTEXT("VisibilitySelectedHeader", "Selection"));
-		Section.AddMenuEntry(FLevelEditorCommands::Get().ShowSelected);
+		Section.AddMenuEntry(FLevelEditorCommands::Get().ToggleSelectedVisibility);
 		Section.AddMenuEntry(FLevelEditorCommands::Get().HideSelected);
+		Section.AddMenuEntry(FLevelEditorCommands::Get().ShowSelected);
 	}
 
 	{
diff --git a/Engine/Source/Editor/LevelEditor/Public/LevelEditorActions.h b/Engine/Source/Editor/LevelEditor/Public/LevelEditorActions.h
--- a/Engine/Source/Editor/LevelEditor/Public/LevelEditorActions.h
+++ b/Engine/Source/Editor/LevelEditor/Public/LevelEditorActions.h
@@ -359,6 +359,9 @@ public:
 
 	/** Hides selected actors */
 	TSharedPtr< FUICommandInfo > HideSelected;
+	
+	/** Hides or shows selected actors */
+	TSharedPtr< FUICommandInfo > ToggleSelectedVisibility;
 
 	/** Shows all actors at startup */
 	TSharedPtr< FUICommandInfo > ShowAllStartup;
@@ -1133,6 +1136,9 @@ public:
 	
 	/** Called to see if selected actors can be hidden */
 	static UE_API bool CanShowSelectedActors();
+	
+	/** Called to see if the visibility of seleced actors can be toggled */
+	static UE_API bool CanToggleSelectedActorsVisibility();
 
 	/**
 	 * Called to select all lights
diff --git a/Engine/Source/Editor/UnrealEd/Classes/Editor/UnrealEdEngine.h b/Engine/Source/Editor/UnrealEd/Classes/Editor/UnrealEdEngine.h
--- a/Engine/Source/Editor/UnrealEd/Classes/Editor/UnrealEdEngine.h
+++ b/Engine/Source/Editor/UnrealEd/Classes/Editor/UnrealEdEngine.h
@@ -597,6 +597,13 @@ public:
 	 * modify/dirty actors/BSP.
 	 */
 	UNREALED_API virtual void edactHideUnselected( UWorld* InWorld );
+	
+	/**
+	 * Toggles visibility of selected actors. If any selected actors are visible, they will be made invisible.
+	 * If all selected actors are invisible, the selection will be made visible.
+	 * Will not modify/dirty actors/BSP.
+	 */
+	UNREALED_API virtual void edactToggleSelectedVisibility( UWorld* InWorld );
 
 	/**
 	 * Attempt to unhide all actors and BSP models by setting their bHiddenEdTemporary flags to false if they
diff --git a/Engine/Source/Editor/UnrealEd/Private/EditorActor.cpp b/Engine/Source/Editor/UnrealEd/Private/EditorActor.cpp
--- a/Engine/Source/Editor/UnrealEd/Private/EditorActor.cpp
+++ b/Engine/Source/Editor/UnrealEd/Private/EditorActor.cpp
@@ -1510,26 +1510,11 @@ void UUnrealEdEngine::edactReplaceClassWithClass(UWorld* InWorld, UClass* SrcCla
 
 static void SetActorVisibility(AActor* Actor, bool bVisible)
 {
-	using namespace UE::Editor::DataStorage;
-
 	// Save the actor to the transaction buffer to support undo/redo, but do
 	// not call Modify, as we do not want to dirty the actor's package and
 	// we're only editing temporary, transient values
 	SaveToTransactionBuffer(Actor, false);
-
-	if (const ICompatibilityProvider* DataStoreCompat = GetDataStorageFeature<ICompatibilityProvider>(CompatibilityFeatureName))
-	{
-		RowHandle Row = DataStoreCompat->FindRowWithCompatibleObject(Actor);
-		if (ICoreProvider* DataStore = GetMutableDataStorageFeature<ICoreProvider>(StorageFeatureName))
-		{
-			FVisibleInEditorColumn* VisibilityColumn = DataStore->GetColumn<FVisibleInEditorColumn>(Row);
-			if (VisibilityColumn)
-			{
-				VisibilityColumn->bIsVisibleInEditor = bVisible;
-			}
-			DataStore->AddColumn<FTypedElementSyncBackToWorldTag>(Row);
-		}
-	}
+	Actor->SetIsTemporarilyHiddenInEditor(!bVisible);
 }
 
 void UUnrealEdEngine::edactHideSelected( UWorld* InWorld )
@@ -1624,6 +1609,76 @@ void UUnrealEdEngine::edactHideUnselected( UWorld* InWorld )
 }
 
 
+void UUnrealEdEngine::edactToggleSelectedVisibility(UWorld* InWorld)
+{
+	// Assemble a list of actors to operate on.
+	TArray<AActor*> ActorsToConsider;
+	bool bTargetVisibility = true;
+	
+	for ( FSelectionIterator It( GetSelectedActorIterator() ) ; It ; ++It )
+	{
+		AActor* Actor = static_cast<AActor*>( *It );
+		checkSlow( Actor->IsA(AActor::StaticClass()) );
+		
+		// Don't consider already hidden actors or the builder brush
+		if ( !FActorEditorUtils::IsABuilderBrush(Actor) )
+		{
+			ActorsToConsider.Add( Actor );
+			bTargetVisibility = bTargetVisibility && Actor->IsHiddenEd();
+		}
+	}
+	
+	// Check BSP to confirm their visibility
+	if ( InWorld )
+	{
+		for ( TArray<ULevel*>::TConstIterator LevelIterator = InWorld->GetLevels().CreateConstIterator(); LevelIterator; ++LevelIterator )
+		{
+			UModel& CurLevelModel = *( ( *LevelIterator )->Model );
+			for ( TArray<FBspSurf>::TIterator SurfaceIterator( CurLevelModel.Surfs ); SurfaceIterator; ++SurfaceIterator )
+			{
+				FBspSurf& CurSurface = *SurfaceIterator;
+				if ( ( CurSurface.PolyFlags & PF_Selected ) )
+				{
+					bTargetVisibility = bTargetVisibility && CurSurface.IsHiddenEd();
+				}
+			}
+		}
+	}
+	
+	// Apply the target visibility to the target actors
+	if ( ActorsToConsider.Num() > 0 )
+	{
+		for( int32 ActorIndex = 0 ; ActorIndex < ActorsToConsider.Num() ; ++ActorIndex )
+		{
+			AActor* Actor = ActorsToConsider[ ActorIndex ];
+			if (Actor->IsHiddenEd() != !bTargetVisibility)
+			{
+				SetActorVisibility(Actor, bTargetVisibility);
+			}
+		}
+	}
+	
+	// Apply the target visibility to BSP
+	if ( InWorld )
+	{
+		for ( TArray<ULevel*>::TConstIterator LevelIterator = InWorld->GetLevels().CreateConstIterator(); LevelIterator; ++LevelIterator )
+		{
+			UModel& CurLevelModel = *( ( *LevelIterator )->Model );
+			for ( TArray<FBspSurf>::TIterator SurfaceIterator( CurLevelModel.Surfs ); SurfaceIterator; ++SurfaceIterator )
+			{
+				FBspSurf& CurSurface = *SurfaceIterator;
+				if ( ( CurSurface.PolyFlags & PF_Selected ) )
+				{
+					CurLevelModel.ModifySurf( SurfaceIterator.GetIndex(), false );
+					CurSurface.bHiddenEdTemporary = !bTargetVisibility;
+				}
+			}
+		}
+	}
+	RedrawLevelEditingViewports();
+}
+
+
 void UUnrealEdEngine::edactUnHideAll( UWorld* InWorld )
 {
 	// Iterate through all of the actors and unhide them
diff --git a/Engine/Source/Editor/UnrealEd/Private/UnrealEdSrv.cpp b/Engine/Source/Editor/UnrealEd/Private/UnrealEdSrv.cpp
--- a/Engine/Source/Editor/UnrealEd/Private/UnrealEdSrv.cpp
+++ b/Engine/Source/Editor/UnrealEd/Private/UnrealEdSrv.cpp
@@ -2372,6 +2372,12 @@ bool UUnrealEdEngine::Exec_Actor( UWorld* InWorld, const TCHAR* Str, FOutputDevi
 			}
 		}
 	}
+	else if ( FParse::Command(&Str, TEXT("TOGGLE SELECTED VISIBILITY")) )
+	{
+		const FScopedTransaction Transaction( NSLOCTEXT("UnrealEd", "ToggleSelectedVisibility", "Toggle Selected Visibility") );
+		edactToggleSelectedVisibility( InWorld );
+		return true;
+	}
 	else if( FParse::Command(&Str, TEXT("APPLYTRANSFORM")) )
 	{
 		CommandIsDeprecated( TEXT("ACTOR APPLYTRANSFORM"), Ar );
