Fix missing text localization in pose asset widget

diff --git a/Engine/Source/Editor/UnrealEd/Private/Factories/PoseAssetFactory.cpp b/Engine/Source/Editor/UnrealEd/Private/Factories/PoseAssetFactory.cpp
--- a/Engine/Source/Editor/UnrealEd/Private/Factories/PoseAssetFactory.cpp
+++ b/Engine/Source/Editor/UnrealEd/Private/Factories/PoseAssetFactory.cpp
@@ -85,7 +85,7 @@ public:
 					.Padding(2, 2)
 					[
 						SNew(STextBlock)
-						.Text(FText::FromString(TEXT("Select Source Animation")))
+						.Text(LOCTEXT("SelectSourceAnimation", "Select Source Animation"))
 					]
 
 					+SVerticalBox::Slot()
@@ -113,14 +113,14 @@ public:
 					.Padding(2, 2)
 					[
 						SNew(STextBlock)
-						.Text(FText::FromString(TEXT("[OPTIONAL] Pose Names (one name for each line)")))
+						.Text(LOCTEXT("PoseNamesDescription", "[OPTIONAL] Pose Names (one name for each line)"))
 					]
 
 					+ SVerticalBox::Slot()
 					.Padding(2, 2)
 					[
 						SAssignNew(TextBlock, SMultiLineEditableTextBox)
-						.HintText(FText::FromString(TEXT("Type one pose name for each line...")))
+						.HintText(LOCTEXT("PoseNamesPlaceholder", "Type one pose name for each line..."))
 					]
 				]
 
