Cloth - Tidy up DebugSkelMeshComponent:
- Removed deprecated functions
- Updated RefreshSelectedClothingSkinnedPositions so that it always keeps the painted skinned vertices in a good state with incoming change to the Clothing Asset.

#rnx

diff --git a/Engine/Source/Editor/UnrealEd/Classes/Animation/DebugSkelMeshComponent.h b/Engine/Source/Editor/UnrealEd/Classes/Animation/DebugSkelMeshComponent.h
--- a/Engine/Source/Editor/UnrealEd/Classes/Animation/DebugSkelMeshComponent.h
+++ b/Engine/Source/Editor/UnrealEd/Classes/Animation/DebugSkelMeshComponent.h
@@ -284,10 +284,6 @@ class UDebugSkelMeshComponent : public USkeletalMeshComponent
 	UPROPERTY(transient)
 	FLinearColor WireframeMeshOverlayColor;
 
-	UE_DEPRECATED(5.0, "This variable is no longer used. Use ProcessRootMotionMode instead.")
-	UPROPERTY()
-	uint32 bPreviewRootMotion_DEPRECATED : 1;
-
 	/** Requested Process root motion mode, ProcessRootMotionMode gets set based on requested mode and what is supported. */
 	UPROPERTY(transient)
 	EProcessRootMotionMode RequestedProcessRootMotionMode;
@@ -495,14 +491,6 @@ class UDebugSkelMeshComponent : public USkeletalMeshComponent
 	/** Sets the flag used to determine whether or not the current active cloth sim mesh should be rendered */
 	UNREALED_API void SetShowClothProperty(bool bState);
 
-	/** Get whether we should be previewing root motion */
-	UE_DEPRECATED(5.0, "Please use IsProcessingRootMotion or GetProcessRootMotionMode")
-	bool GetPreviewRootMotion() const { return IsProcessingRootMotion(); }
-
-	/** Set whether we should be previewing root motion. Note: disabling root motion preview resets transform. */
-	UE_DEPRECATED(5.0, "Please use SetProcessRootMotionMode")
-	void SetPreviewRootMotion(bool bInPreviewRootMotion) { SetProcessRootMotionMode(bInPreviewRootMotion ? EProcessRootMotionMode::Loop : EProcessRootMotionMode::Ignore); }
-
 	/** Whether we are processing root motion or not */
 	UNREALED_API bool IsProcessingRootMotion() const;
 
@@ -530,8 +518,6 @@ class UDebugSkelMeshComponent : public USkeletalMeshComponent
 	void SetShowNotificationVisualizations(const bool bShow) { bShowNotificationVisualizations = bShow; }
 	bool IsNotificationVisualizationsEnabled() const { return bShowNotificationVisualizations; }
 
-	UE_DEPRECATED(5.5, "Use VisualizeRootMotionMode functions instead.")
-	void SetShowRootMotionVisualizations(const bool bShow) { VisualizeRootMotionMode = bShow ? EVisualizeRootMotionMode::Trajectory : EVisualizeRootMotionMode::None; }
 	bool IsRootMotionVisualizationsEnabled() const { return VisualizeRootMotionMode != EVisualizeRootMotionMode::None; }
 
 	/** Sets flags whether we AssetUserData visualizations should be drawn in the viewport. */
@@ -629,21 +615,10 @@ public:
 	// ToggleMeshSectionForCloth will make these get out of sync
 	UNREALED_API void ResetMeshSectionVisibility();
 
-	// Rebuilds the fixed parameter on the mesh to mesh data, to be used if the editor has
-	// changed a vert to be fixed or unfixed otherwise the simulation will not work
-	// bInvalidateDerivedDataCache can only be false during previewing as otherwise the changes won't be correctly saved
-	UE_DEPRECATED(5.0, "This function is redundant, since it is always called after ApplyParameterMasks and therefore will be removed.")
-	UNREALED_API void RebuildClothingSectionsFixedVerts(bool bInvalidateDerivedDataCache = true);
-
 	TArray<FVector3f> SkinnedSelectedClothingPositions;
 	TArray<FVector3f> SkinnedSelectedClothingNormals;
 
 private:
-	// Rebuilds the fixed vertex attribute on any cloth deformer mappings,
-	// including LOD bias mappings, that reference the specified LOD section.
-	UE_DEPRECATED(5.0, "This function is redundant, since it is always called after ApplyParameterMasks and therefore will be removed.")
-	UNREALED_API void RebuildClothingSectionFixedVerts(int32 LODIndex, int32 SectionIndex);
-
 	// Helper function to generate space bases for current frame
 	UNREALED_API void GenSpaceBases(TArray<FTransform>& OutSpaceBases);
 
diff --git a/Engine/Source/Editor/UnrealEd/Private/Animation/DebugSkelMeshComponent.cpp b/Engine/Source/Editor/UnrealEd/Private/Animation/DebugSkelMeshComponent.cpp
--- a/Engine/Source/Editor/UnrealEd/Private/Animation/DebugSkelMeshComponent.cpp
+++ b/Engine/Source/Editor/UnrealEd/Private/Animation/DebugSkelMeshComponent.cpp
@@ -1161,75 +1161,6 @@ void UDebugSkelMeshComponent::ResetMeshSectionVisibility()
 	}
 }
 
-void UDebugSkelMeshComponent::RebuildClothingSectionsFixedVerts(bool bInvalidateDerivedDataCache)
-{
-	// TODO: There is no need to rebuild all section/LODs at once.
-	//       It should only do the section associated to the current cloth asset being
-	//        painted instead, and only when the MaxDistance mask changes.
-	FScopedSkeletalMeshPostEditChange ScopedSkeletalMeshPostEditChange(GetSkeletalMeshAsset());
-
-	GetSkeletalMeshAsset()->PreEditChange(nullptr);
-
-	TIndirectArray<FSkeletalMeshLODModel>& LODModels = GetSkeletalMeshAsset()->GetImportedModel()->LODModels;
-
-	for (int32 LODIndex = 0; LODIndex < LODModels.Num(); ++LODIndex)
-	{
-		for (int32 SectionIndex = 0; SectionIndex < LODModels[LODIndex].Sections.Num(); ++SectionIndex)
-		{
-			PRAGMA_DISABLE_DEPRECATION_WARNINGS
-			RebuildClothingSectionFixedVerts(LODIndex, SectionIndex);
-			PRAGMA_ENABLE_DEPRECATION_WARNINGS
-		}
-	}
-
-	if (bInvalidateDerivedDataCache)
-	{
-		GetSkeletalMeshAsset()->InvalidateDeriveDataCacheGUID();  // Dirty the DDC key unless previewing
-	}
-
-	ReregisterComponent();
-}
-
-void UDebugSkelMeshComponent::RebuildClothingSectionFixedVerts(int32 LODIndex, int32 SectionIndex)
-{
-	FSkeletalMeshModel* const SkeletalMeshModel = GetSkeletalMeshAsset()->GetImportedModel();
-	if (!ensure(SkeletalMeshModel && LODIndex < SkeletalMeshModel->LODModels.Num()))
-	{
-		return;
-	}
-
-	TIndirectArray<FSkeletalMeshLODModel>& LODModels = SkeletalMeshModel->LODModels;
-	if (!ensure(SectionIndex < LODModels[LODIndex].Sections.Num()))
-	{
-		return;
-	}
-
-	FSkelMeshSection& UpdatedSection = LODModels[LODIndex].Sections[SectionIndex];
-	if (!UpdatedSection.HasClothingData() || !UpdatedSection.ClothingData.IsValid())
-	{
-		return;
-	}
-
-	const UClothingAssetCommon* const ClothingAsset = Cast<UClothingAssetCommon>(GetSkeletalMeshAsset()->GetClothingAsset(UpdatedSection.ClothingData.AssetGuid));
-	check(ClothingAsset);  // Must have a valid clothing asset at this point, or something has gone terribly wrong
-
-	const FClothLODDataCommon& ClothLODData = ClothingAsset->LodData[UpdatedSection.ClothingData.AssetLodIndex];
-	const FPointWeightMap* const MaxDistances = ClothLODData.PhysicalMeshData.FindWeightMap(EWeightMapTargetCommon::MaxDistance);
-
-	// Iterate through all LOD sections that might contain mapping to the updated clothing asset (can only be higher or same LOD)
-	for (int32 BiasedLODIndex = LODIndex; BiasedLODIndex < LODModels.Num(); ++BiasedLODIndex)
-	{
-		const int32 LODBias = BiasedLODIndex - LODIndex;
-
-		FSkelMeshSection& BiasedSection = LODModels[BiasedLODIndex].Sections[SectionIndex];
-		if (LODBias < BiasedSection.ClothMappingDataLODs.Num() && BiasedSection.ClothMappingDataLODs[LODBias].Num())
-		{
-			// Update vertex contributions for this LOD bias
-			ClothingMeshUtils::ComputeVertexContributions(BiasedSection.ClothMappingDataLODs[LODBias], MaxDistances, ClothLODData.bSmoothTransition, ClothLODData.bUseMultipleInfluences);
-		}
-	}
-}
-
 void UDebugSkelMeshComponent::CheckClothTeleport()
 {
 	// do nothing to avoid clothing reset while modifying properties
@@ -1305,38 +1236,39 @@ void UDebugSkelMeshComponent::TickComponent(float DeltaTime, enum ELevelTick Tic
 
 void UDebugSkelMeshComponent::RefreshSelectedClothingSkinnedPositions()
 {
-	if(GetSkeletalMeshAsset() && SelectedClothingGuidForPainting.IsValid())
+	if (GetSkeletalMeshAsset() && SelectedClothingGuidForPainting.IsValid())
 	{
-		auto* Asset = GetSkeletalMeshAsset()->GetMeshClothingAssets().FindByPredicate([&](UClothingAssetBase* Item)
-		{
-			return Item && SelectedClothingGuidForPainting == Item->GetAssetGuid();
-		});
+		const TObjectPtr<UClothingAssetBase>* const Asset = GetSkeletalMeshAsset()->GetMeshClothingAssets().FindByPredicate(
+			[this](const TObjectPtr<UClothingAssetBase>& Item) -> bool
+			{
+				return Item && Item->GetAssetGuid() == SelectedClothingGuidForPainting;
+			});
 
-		if(Asset)
+		if (const UClothingAssetCommon* const ConcreteAsset = Cast<const UClothingAssetCommon>(Asset ? *Asset : nullptr))
 		{
-			UClothingAssetCommon* ConcreteAsset = Cast<UClothingAssetCommon>(*Asset);
-
-			if(ConcreteAsset->LodData.IsValidIndex(SelectedClothingLodForPainting))
+			if (ConcreteAsset->LodData.IsValidIndex(SelectedClothingLodForPainting))
 			{
-				SkinnedSelectedClothingPositions.Reset();
-				SkinnedSelectedClothingNormals.Reset();
-
 				TArray<FMatrix44f> RefToLocals;
-				// Pass LOD0 to collect all bones
-				GetCurrentRefToLocalMatrices(RefToLocals, 0);
+				GetCurrentRefToLocalMatrices(RefToLocals, 0);  // Pass LOD0 to collect all bones
 
 				const FClothLODDataCommon& LodData = ConcreteAsset->LodData[SelectedClothingLodForPainting];
 
-				ClothingMeshUtils::SkinPhysicsMesh(ConcreteAsset->UsedBoneIndices, LodData.PhysicalMeshData, FTransform::Identity, RefToLocals.GetData(), RefToLocals.Num(), SkinnedSelectedClothingPositions, SkinnedSelectedClothingNormals);
+				ClothingMeshUtils::SkinPhysicsMesh(
+					ConcreteAsset->UsedBoneIndices,
+					LodData.PhysicalMeshData,
+					FTransform::Identity,
+					RefToLocals.GetData(),
+					RefToLocals.Num(),
+					SkinnedSelectedClothingPositions,
+					SkinnedSelectedClothingNormals);
+
 				RebuildCachedClothBounds();
+				return;
 			}
 		}
 	}
-	else
-	{
-		SkinnedSelectedClothingNormals.Reset();
-		SkinnedSelectedClothingPositions.Reset();
-	}
+	SkinnedSelectedClothingNormals.Reset();
+	SkinnedSelectedClothingPositions.Reset();
 }
 
 void UDebugSkelMeshComponent::GetUsedMaterials(TArray<UMaterialInterface *>& OutMaterials, bool bGetDebugMaterials /*= false*/) const
