Splits viewport panning logic so that the middle mouse button and the left+right mouse button gestures are distinct. Middle mouse can be camera relative (and continues to be so by default) while left+right mouse button always move up and down.

diff --git a/Engine/Source/Editor/UnrealEd/Private/EditorViewportClient.cpp b/Engine/Source/Editor/UnrealEd/Private/EditorViewportClient.cpp
--- a/Engine/Source/Editor/UnrealEd/Private/EditorViewportClient.cpp
+++ b/Engine/Source/Editor/UnrealEd/Private/EditorViewportClient.cpp
@@ -5100,17 +5100,19 @@ void FEditorViewportClient::ConvertMovementToDragRot(const FVector& InDelta,
 				// Pan left/right/up/down
 				const bool bInvert = !bIsUsingTrackpad && MiddleMouseButtonDown && GetDefault<ULevelEditorViewportSettings>()->bInvertMiddleMousePan;
 
-				if (ViewportSettings->bWorldSpaceVerticalPan)
+				// Treat the middle mouse pan as a distinct gesture from the left+right mouse button 
+				if (MiddleMouseButtonDown && !ViewportSettings->bWorldSpaceVerticalPan)
+				{
+					// Pan relative to the camera
+					InDragDelta = FQuat(ViewRotation) * FVector( 0.0f, InDelta.X, InDelta.Y) * (bInvert ? -1.0f : 1.0f);
+				}
+				else
 				{
 					float Direction = bInvert ? 1 : -1;
 					InDragDelta.X = InDelta.X * Direction * FMath::Sin( ViewRotation.Yaw * PI / 180.f );
 					InDragDelta.Y = InDelta.X * -Direction * FMath::Cos( ViewRotation.Yaw * PI / 180.f );
 					InDragDelta.Z = -Direction * InDelta.Y;
 				}
-				else
-				{
-					InDragDelta = FQuat(ViewRotation) * FVector( 0.0f, InDelta.X, InDelta.Y) * (bInvert ? -1.0f : 1.0f);
-				}
 			}
 			else if( RightMouseButtonDown && !LeftMouseButtonDown )
 			{
