Cooker: DiffOnly: Fix crash when deserializing TArray properties for display. When deserializing the property that wrote to a given offset, we find the start offset of the property's serialization by walking backwards along the serialize calls until we find a serialize call that wrote to a different property. But for that to work, we need to record a callstack entry, with SerializedProperty recorded, for every serialize call. We were recording the callstack entry for each serialize call, but were not storing the SerializedProperty for each entry unless we were collecting callstacks due to a detected difference. Modified the AddUniqueCallstack code to instead record the SerializedProperty for every Serialize call.
#rnx

diff --git a/Engine/Source/Editor/UnrealEd/Private/Cooker/DiffPackageWriter.cpp b/Engine/Source/Editor/UnrealEd/Private/Cooker/DiffPackageWriter.cpp
--- a/Engine/Source/Editor/UnrealEd/Private/Cooker/DiffPackageWriter.cpp
+++ b/Engine/Source/Editor/UnrealEd/Private/Cooker/DiffPackageWriter.cpp
@@ -105,8 +105,8 @@ public:
 		Diff.Size = 1;
 		if (DifferenceCallstackData)
 		{
-			Diff.ObjectName = DifferenceCallstackData->SerializedObjectName;
-			Diff.PropertyName = DifferenceCallstackData->SerializedPropertyName;
+			Diff.ObjectName = DifferenceCallstackData->GetObjectName();
+			Diff.PropertyName = DifferenceCallstackData->GetPropertyName();
 			Diff.Callstack = DifferenceCallstackData->ToString(CallstackCutoffString);
 		}
 		CurrentPackage.bHasDiffs = true;
diff --git a/Engine/Source/Editor/UnrealEd/Private/Cooker/DiffWriterArchive.cpp b/Engine/Source/Editor/UnrealEd/Private/Cooker/DiffWriterArchive.cpp
--- a/Engine/Source/Editor/UnrealEd/Private/Cooker/DiffWriterArchive.cpp
+++ b/Engine/Source/Editor/UnrealEd/Private/Cooker/DiffWriterArchive.cpp
@@ -104,18 +104,12 @@ FLinkerLoad* CreateLinkerForPackage(
 	const FPackageData& PackageData);
 
 
-FCallstacks::FCallstackData::FCallstackData(TUniquePtr<ANSICHAR[]>&& InCallstack, UObject* InSerializedObject, FProperty* InSerializedProperty)
+FCallstacks::FCallstackData::FCallstackData(TUniquePtr<ANSICHAR[]>&& InCallstack, uint32 InCppCallstackHash, UObject* InSerializedObject, FProperty* InSerializedProperty)
 	: Callstack(MoveTemp(InCallstack))
+	, SerializedObject(InSerializedObject)
 	, SerializedProp(InSerializedProperty)
+	, CppCallstackHash(InCppCallstackHash)
 {
-	if (InSerializedObject)
-	{
-		SerializedObjectName = InSerializedObject->GetFullName();
-	}
-	if (InSerializedProperty)
-	{
-		SerializedPropertyName = InSerializedProperty->GetFullName();
-	}
 }
 
 FString FCallstacks::FCallstackData::ToString(const TCHAR* CallstackCutoffText) const
@@ -154,88 +148,106 @@ FString FCallstacks::FCallstackData::ToString(const TCHAR* CallstackCutoffText)
 		HumanReadableString += NewLineToken;
 	}
 
-	if (!SerializedObjectName.IsEmpty())
+	FString ObjectName = GetObjectName();
+	if (!ObjectName.IsEmpty())
 	{
 		HumanReadableString += NewLineToken;
 		HumanReadableString += IndentToken;
 		HumanReadableString += TEXT("Serialized Object: ");
-		HumanReadableString += SerializedObjectName;
+		HumanReadableString += ObjectName;
 		HumanReadableString += NewLineToken;
 	}
-	if (!SerializedPropertyName.IsEmpty())
+	FString PropertyName = GetPropertyName();
+	if (!PropertyName.IsEmpty())
 	{
-		if (SerializedObjectName.IsEmpty())
+		if (ObjectName.IsEmpty())
 		{
 			HumanReadableString += NewLineToken;
 		}
 		HumanReadableString += IndentToken;
 		HumanReadableString += TEXT("Serialized Property: ");
-		HumanReadableString += SerializedPropertyName;
+		HumanReadableString += PropertyName;
 		HumanReadableString += NewLineToken;
 	}
 	return HumanReadableString;
 }
 
-FCallstacks::FCallstackData FCallstacks::FCallstackData::Clone() const
+TRefCountPtr<FCallstacks::FCallstackData> FCallstacks::FCallstackData::Clone() const
 {
 	TUniquePtr<ANSICHAR[]> CallstackCopy;
-	if (const int32 Len = FCStringAnsi::Strlen(Callstack.Get()); Len > 0)
+	if (Callstack)
 	{
-		CallstackCopy = MakeUnique<ANSICHAR[]>(Len + 1);
-		FMemory::Memcpy(CallstackCopy.Get(), Callstack.Get(), Len + 1);
+		if (const int32 Len = FCStringAnsi::Strlen(Callstack.Get()); Len > 0)
+		{
+			CallstackCopy = MakeUnique<ANSICHAR[]>(Len + 1);
+			FMemory::Memcpy(CallstackCopy.Get(), Callstack.Get(), Len + 1);
+		}
 	}
 
-	FCallstackData Clone(MoveTemp(CallstackCopy), nullptr, SerializedProp);
-	Clone.SerializedObjectName = SerializedObjectName;
+	TRefCountPtr<FCallstackData> Clone = new FCallstackData(MoveTemp(CallstackCopy),
+		CppCallstackHash, SerializedObject, SerializedProp);
 
 	return Clone;
 }
 
+FCallstackKey FCallstacks::FCallstackData::GetKey() const
+{
+	return FCallstackKey{ SerializedObject, SerializedProp, CppCallstackHash };
+}
+
+FString FCallstacks::FCallstackData::GetObjectName() const
+{
+	return SerializedObject ? SerializedObject->GetFullName() : FString();
+}
+
+FString FCallstacks::FCallstackData::GetPropertyName() const
+{
+	return SerializedProp ? SerializedProp->GetFullName() : FString();
+}
+
 FCallstacks::FCallstacks()
-	: bCallstacksDirty(true)
-	, StackTraceSize(65535)
-	, LastSerializeCallstack(nullptr)
+	: StackTraceSize(65535)
+	, PreviousCppCallstackHash(MAX_uint32)
+	, bCppCallstackDirty(true)
 	, EndOffset(0)
 {
 	StackTrace = MakeUnique<ANSICHAR[]>(StackTraceSize);
 	StackTrace[0] = 0;
+	// We create one entry in CallstackAtOffsetMap for every Serialize call,
+	// and one entry in UniqueCallstacks for every unique callstack key in those Serialize calls.
+	// This adds up to hundreds or thousands of entries for most packages. Reserve a large constant
+	// initial size to avoid the reallocation costs on the first 1000 entries.
+	CallstackAtOffsetMap.Reserve(1000);
+	UniqueCallstacks.Reserve(1000);
 }
 
 void FCallstacks::Reset()
 {
 	CallstackAtOffsetMap.Reset();
 	UniqueCallstacks.Reset();
-	bCallstacksDirty = true;
-	LastSerializeCallstack = nullptr;
+	bCppCallstackDirty = true;
+	PreviousCppCallstackHash = MAX_uint32;
 	StackTrace[0] = 0;
 	EndOffset = 0;
 }
 
-ANSICHAR* FCallstacks::AddUniqueCallstack(bool bIsCollectingCallstacks, UObject* SerializedObject, FProperty* SerializedProperty, uint32& OutCallstackCRC)
+TRefCountPtr<FCallstacks::FCallstackData> FCallstacks::AddUniqueCallstack(const FCallstackKey& Key)
 {
 	ANSICHAR* Callstack = nullptr;
-	if (bIsCollectingCallstacks)
+	TRefCountPtr<FCallstackData>& Existing = UniqueCallstacks.FindOrAdd(Key);
+	if (!Existing)
 	{
-		OutCallstackCRC = FCrc::StrCrc32(StackTrace.Get());
-
-		if (FCallstackData* ExistingCallstack = UniqueCallstacks.Find(OutCallstackCRC))
-		{
-			Callstack = ExistingCallstack->Callstack.Get();
-		}
-		else
+		TUniquePtr<ANSICHAR[]> NewCallstack;
+		if (StackTrace.Get()[0] != '\0')
 		{
 			const int32 Len = FCStringAnsi::Strlen(StackTrace.Get()) + 1;
-			TUniquePtr<ANSICHAR[]> NewCallstack = MakeUnique<ANSICHAR[]>(Len);
+			NewCallstack = MakeUnique<ANSICHAR[]>(Len);
 			FCStringAnsi::Strncpy(NewCallstack.Get(), StackTrace.Get(), Len);
-			FCallstackData& NewEntry = UniqueCallstacks.Add(OutCallstackCRC, FCallstackData(MoveTemp(NewCallstack), SerializedObject, SerializedProperty));
-			Callstack = NewEntry.Callstack.Get();
 		}
+		Existing = new FCallstackData(MoveTemp(NewCallstack), Key.CppCallstackHash,
+			Key.SerializedObject, Key.SerializedProperty);
 	}
-	else
-	{
-		OutCallstackCRC = 0;
-	}
-	return Callstack;
+	return Existing;
 }
 
 void FCallstacks::Add(
@@ -288,111 +300,118 @@ void FCallstacks::Add(
 				bIsIndenting = !bIsPropertyLabel;
 			}
 		}
-		// Make sure we compare the new stack trace with the last one in the next if statement
-		bCallstacksDirty = true;
+		bCppCallstackDirty = true;
+	}
+	else
+	{
+		bCppCallstackDirty = StackTrace[0] != '\0';
+		StackTrace[0] = '\0';
 	}
 
-	if (LastSerializeCallstack == nullptr || (bCallstacksDirty && FCStringAnsi::Strcmp(LastSerializeCallstack, StackTrace.Get()) != 0))
+	FCallstackKey CallstackKey;
+	CallstackKey.SerializedObject = SerializedObject;
+	CallstackKey.SerializedProperty = SerializedProperty;
+	if (bCppCallstackDirty)
+	{
+		CallstackKey.CppCallstackHash = FCrc::StrCrc32(StackTrace.Get());
+		bCppCallstackDirty = false;
+		PreviousCppCallstackHash = CallstackKey.CppCallstackHash;
+	}
+	else
 	{
-		uint32 CallstackCRC = 0;
-		bool bSuppressLogging = UE::ArchiveStackTrace::ShouldIgnoreDiff();
-		LastSerializeCallstack = AddUniqueCallstack(bIsCollectingCallstacks, SerializedObject, SerializedProperty, CallstackCRC);
-		check(CallstackCRC != 0 || !bShouldCollectCallstack);
-		FCallstackAtOffset NewBlock{ CurrentOffset, Length, CurrentOffset, Length, CallstackCRC, bSuppressLogging };
+		CallstackKey.CppCallstackHash = PreviousCppCallstackHash;
+	}
+	TRefCountPtr<FCallstackData> CallstackData = AddUniqueCallstack(CallstackKey);
 
-		FCallstackAtOffset* LastBlock = CallstackAtOffsetMap.Num() ? &CallstackAtOffsetMap.Last() : nullptr;
-		if (!LastBlock || CurrentOffset >= LastBlock->Offset + LastBlock->Length)
-		{
-			// New block serialized at the end of archive buffer
-			CallstackAtOffsetMap.Add(NewBlock);
-		}
-		else
-		{
-			// This happens after a Seek(). We need to modify or replace the old block that covered the range written
-			// by this new Serialize call.
-			const int32 OldBlockIndex = GetCallstackIndexAtOffset(CurrentOffset);
-			check(OldBlockIndex != -1);
+	bool bSuppressLogging = UE::ArchiveStackTrace::ShouldIgnoreDiff();
+	FCallstackAtOffset NewBlock{ CurrentOffset, Length, CurrentOffset, Length, CallstackData, bSuppressLogging };
+	FCallstackAtOffset* LastBlock = CallstackAtOffsetMap.Num() ? &CallstackAtOffsetMap.Last() : nullptr;
+	if (!LastBlock || CurrentOffset >= LastBlock->Offset + LastBlock->Length)
+	{
+		// New block serialized at the end of archive buffer
+		CallstackAtOffsetMap.Add(NewBlock);
+	}
+	else
+	{
+		// This happens after a Seek(). We need to modify or replace the old block that covered the range written
+		// by this new Serialize call.
+		const int32 OldBlockIndex = GetCallstackIndexAtOffset(CurrentOffset);
+		check(OldBlockIndex != -1);
 
-			FCallstackAtOffset* OldBlock = &CallstackAtOffsetMap[OldBlockIndex];
+		FCallstackAtOffset* OldBlock = &CallstackAtOffsetMap[OldBlockIndex];
 
-			int64 OldEnd = OldBlock->Offset + OldBlock->Length;
-			int64 NewEnd = NewBlock.Offset + NewBlock.Length;
-			if (OldEnd <= NewEnd)
+		int64 OldEnd = OldBlock->Offset + OldBlock->Length;
+		int64 NewEnd = NewBlock.Offset + NewBlock.Length;
+		if (OldEnd <= NewEnd)
+		{
+			// The new block overwrites the end of the old block, and possibly overwrites blocks after it
+			check(OldBlock->Offset <= NewBlock.Offset); // GetCallstackIndexAtOffset guarantees this
+			bool bNewEntirelyContainsOld = OldBlock->Offset == NewBlock.Offset;
+			int32 StartRemoveIndex;
+			if (bNewEntirelyContainsOld)
 			{
-				// The new block overwrites the end of the old block, and possibly overwrites blocks after it
-				check(OldBlock->Offset <= NewBlock.Offset); // GetCallstackIndexAtOffset guarantees this
-				bool bNewEntirelyContainsOld = OldBlock->Offset == NewBlock.Offset;
-				int32 StartRemoveIndex;
-				if (bNewEntirelyContainsOld)
+				// The new block completely overwrites the old block; delete the old block and replace it with the
+				// the new block. Still need to check whether new also overwrites old blocks after the first.
+				*OldBlock = NewBlock;
+				StartRemoveIndex = OldBlockIndex + 1;
+			}
+			else
+			{
+				// The new block does not overwrite the old block, so keep the old block, clamp it to end at the
+				// new block, and add the new block after it. Still need to check whether new also overwrites old
+				// blocks after the first.
+				// There might be a gap in between the end of old block and the beginning of the new block. This
+				// can occur when we are not recording every serialize call. Leave the old block unmodified in
+				// that case.
+				if (OldEnd > NewBlock.Offset)
 				{
-					// The new block completely overwrites the old block; delete the old block and replace it with the
-					// the new block. Still need to check whether new also overwrites old blocks after the first.
-					*OldBlock = NewBlock;
-					StartRemoveIndex = OldBlockIndex + 1;
+					OldBlock->Length = NewBlock.Offset - OldBlock->Offset;
 				}
-				else
+				CallstackAtOffsetMap.Insert(NewBlock, OldBlockIndex + 1);
+				OldBlock = nullptr; // Our pointer for OldBlock is now possibly invalidated by reallocation.
+				StartRemoveIndex = OldBlockIndex + 2;
+			}
+			int32 EndRemoveIndex = StartRemoveIndex;
+			while (EndRemoveIndex < CallstackAtOffsetMap.Num())
+			{
+				OldBlock = &CallstackAtOffsetMap[EndRemoveIndex];
+				if (OldBlock->Offset >= NewEnd)
 				{
-					// The new block does not overwrite the old block, so keep the old block, clamp it to end at the
-					// new block, and add the new block after it. Still need to check whether new also overwrites old
-					// blocks after the first.
-					// There might be a gap in between the end of old block and the beginning of the new block. This
-					// can occur when we are not recording every serialize call. Leave the old block unmodified in
-					// that case.
-					if (OldEnd > NewBlock.Offset)
-					{
-						OldBlock->Length = NewBlock.Offset - OldBlock->Offset;
-					}
-					CallstackAtOffsetMap.Insert(NewBlock, OldBlockIndex + 1);
-					OldBlock = nullptr; // Our pointer for OldBlock is now possibly invalidated by reallocation.
-					StartRemoveIndex = OldBlockIndex + 2;
+					break;
 				}
-				int32 EndRemoveIndex = StartRemoveIndex;
-				while (EndRemoveIndex < CallstackAtOffsetMap.Num())
+				OldEnd = OldBlock->Offset + OldBlock->Length;
+				if (OldEnd > NewEnd)
 				{
-					OldBlock = &CallstackAtOffsetMap[EndRemoveIndex];
-					if (OldBlock->Offset >= NewEnd)
-					{
-						break;
-					}
-					OldEnd = OldBlock->Offset + OldBlock->Length;
-					if (OldEnd > NewEnd)
-					{
-						// The beginning of this followup block is overwritten by the new block; shorten it
-						OldBlock->Length = OldEnd - NewEnd;
-						OldBlock->Offset = NewEnd;
-						break;
-					}
-					else
-					{
-						// This followup block is completely inside the new block; mark it for delete and move to next.
-						++EndRemoveIndex;
-					}
+					// The beginning of this followup block is overwritten by the new block; shorten it
+					OldBlock->Length = OldEnd - NewEnd;
+					OldBlock->Offset = NewEnd;
+					break;
 				}
-				if (EndRemoveIndex > StartRemoveIndex)
+				else
 				{
-					CallstackAtOffsetMap.RemoveAt(StartRemoveIndex, EndRemoveIndex - StartRemoveIndex,
-						EAllowShrinking::No);
+					// This followup block is completely inside the new block; mark it for delete and move to next.
+					++EndRemoveIndex;
 				}
 			}
-			else // OldEnd > NewEnd
+			if (EndRemoveIndex > StartRemoveIndex)
 			{
-				// The new block is completely inside the old block, which extends after it. Shorten the beginning of
-				// the old block, and add an additional new block containing the portion of the old block that extends
-				// after the new block.
-				FCallstackAtOffset SegmentAfterNewBlock = *OldBlock;
-				SegmentAfterNewBlock.Offset = NewEnd;
-				SegmentAfterNewBlock.Length = OldEnd - NewEnd;
-				OldBlock->Length = NewBlock.Offset - OldBlock->Offset;
-				CallstackAtOffsetMap.Insert(NewBlock, OldBlockIndex + 1);
-				CallstackAtOffsetMap.Insert(SegmentAfterNewBlock, OldBlockIndex + 2);
-				OldBlock = nullptr; // Our pointer for OldBlock is now possibly invalidated by reallocation.
+				CallstackAtOffsetMap.RemoveAt(StartRemoveIndex, EndRemoveIndex - StartRemoveIndex,
+					EAllowShrinking::No);
 			}
 		}
-	}
-	else if (LastSerializeCallstack)
-	{
-		// Skip callstack comparison on next serialize call unless we grab a stack trace
-		bCallstacksDirty = false;
+		else // OldEnd > NewEnd
+		{
+			// The new block is completely inside the old block, which extends after it. Shorten the beginning of
+			// the old block, and add an additional new block containing the portion of the old block that extends
+			// after the new block.
+			FCallstackAtOffset SegmentAfterNewBlock = *OldBlock;
+			SegmentAfterNewBlock.Offset = NewEnd;
+			SegmentAfterNewBlock.Length = OldEnd - NewEnd;
+			OldBlock->Length = NewBlock.Offset - OldBlock->Offset;
+			CallstackAtOffsetMap.Insert(NewBlock, OldBlockIndex + 1);
+			CallstackAtOffsetMap.Insert(SegmentAfterNewBlock, OldBlockIndex + 2);
+			OldBlock = nullptr; // Our pointer for OldBlock is now possibly invalidated by reallocation.
+		}
 	}
 }
 
@@ -480,6 +499,16 @@ void FCallstacks::Append(const FCallstacks& Other, int64 OtherStartOffset)
 		FCallstackAtOffset& New = CallstackAtOffsetMap.Add_GetRef(OtherOffset);
 		New.Offset += OtherStartOffset;
 		New.SerializeCallOffset += OtherStartOffset;
+		if (New.Callstack)
+		{
+			FCallstackKey Key = New.Callstack->GetKey();
+			TRefCountPtr<FCallstackData>& Existing = UniqueCallstacks.FindOrAdd(Key);
+			if (!Existing)
+			{
+				Existing = New.Callstack->Clone();
+			}
+			New.Callstack = Existing;
+		}
 	}
 
 	CallstackAtOffsetMap.Sort([](const FCallstackAtOffset& LHS,const FCallstackAtOffset& RHS)
@@ -487,19 +516,6 @@ void FCallstacks::Append(const FCallstacks& Other, int64 OtherStartOffset)
 		return LHS.Offset < RHS.Offset;
 	});
 
-	for (const TPair<uint32, FCallstackData>& Kv : Other.UniqueCallstacks)
-	{
-		if (FCallstackData* Existing = UniqueCallstacks.Find(Kv.Key))
-		{
-			if (LastSerializeCallstack == Existing->Callstack.Get())
-			{
-				LastSerializeCallstack = nullptr;
-			}
-			UniqueCallstacks.Remove(Kv.Key);
-		}
-		UniqueCallstacks.Emplace(Kv.Key, Kv.Value.Clone());
-	}
-
 	EndOffset = FMath::Max(EndOffset, Other.EndOffset + OtherStartOffset);
 }
 
@@ -1256,7 +1272,7 @@ void FAccumulator::CompareWithPreviousForSection(const FPackageData& SourcePacka
 			CallstackAtOffsetPtr = &Callstacks.GetCallstack(DifferenceCallstackOffsetIndex);
 			bCallstackSuppressLogging = CallstackAtOffsetPtr->bSuppressLogging;
 
-			DifferenceCallstackDataPtr = &Callstacks.GetCallstackData(*CallstackAtOffsetPtr);
+			DifferenceCallstackDataPtr = Callstacks.GetCallstackData(*CallstackAtOffsetPtr).GetReference();
 			DifferenceCallstackDataText = DifferenceCallstackDataPtr->ToString(CallstackCutoffText);
 			if (NumDiffsLocal > 0 && 
 				LastDifferenceCallstackDataText.Compare(DifferenceCallstackDataText, ESearchCase::CaseSensitive) == 0)
@@ -1315,7 +1331,7 @@ void FAccumulator::CompareWithPreviousForSection(const FPackageData& SourcePacka
 			const FCallstacks::FCallstackData& DifferenceCallstackData = *DifferenceCallstackDataPtr;
 			if (DeterminismManager)
 			{
-				DeterminismManager->RecordExportModified(DifferenceCallstackData.SerializedObjectName);
+				DeterminismManager->RecordExportModified(DifferenceCallstackData.GetObjectName());
 			}
 
 			FString BeforePropertyVal;
@@ -1347,7 +1363,7 @@ void FAccumulator::CompareWithPreviousForSection(const FPackageData& SourcePacka
 							break;
 						}
 						const FCallstacks::FCallstackAtOffset& PreviousCallstack = Callstacks.GetCallstack(CallstackIndex);
-						const FCallstacks::FCallstackData& PreviousCallstackData = Callstacks.GetCallstackData(PreviousCallstack);
+						const FCallstacks::FCallstackData& PreviousCallstackData = *Callstacks.GetCallstackData(PreviousCallstack);
 						if (PreviousCallstackData.SerializedProp != SerProp)
 						{
 							break;
diff --git a/Engine/Source/Editor/UnrealEd/Private/Cooker/DiffWriterArchive.h b/Engine/Source/Editor/UnrealEd/Private/Cooker/DiffWriterArchive.h
--- a/Engine/Source/Editor/UnrealEd/Private/Cooker/DiffWriterArchive.h
+++ b/Engine/Source/Editor/UnrealEd/Private/Cooker/DiffWriterArchive.h
@@ -16,6 +16,7 @@
 #include "Serialization/PackageWriter.h"
 #include "Templates/RefCounting.h"
 #include "Templates/Function.h"
+#include "Templates/RefCounting.h"
 #include "Templates/UniquePtr.h"
 #include "UObject/NameTypes.h"
 
@@ -96,10 +97,75 @@ public:
 	}
 };
 
+/**
+ * Key for the data that is stored per unique callstack by FCallstacks.
+ * Each FCallstackAtOffset constructs the key based on c++ callstack and other data from the serialize
+ * call. If the key does not already exist, more information is collected and stored under the key.
+ */
+struct FCallstackKey
+{
+	UObject* SerializedObject = nullptr;
+	FProperty* SerializedProperty = nullptr;
+	uint32 CppCallstackHash = MAX_uint32;
+	bool operator==(const FCallstackKey& Other) const
+	{
+		return Other.SerializedObject == SerializedObject &&
+			Other.SerializedProperty == SerializedProperty &&
+			Other.CppCallstackHash == Other.CppCallstackHash;
+	}
+	friend uint32 GetTypeHash(const FCallstackKey& Key)
+	{
+		constexpr uint32 Prime = 101;
+		const uint32* WordPtr = reinterpret_cast<const uint32*>(&Key);
+		const uint32* WordPtrEnd = WordPtr + sizeof(Key) / sizeof(uint32);
+		uint32 Hash = 0;
+		while (WordPtr < WordPtrEnd)
+		{
+			Hash = Hash * Prime + *WordPtr;
+			++WordPtr;
+		}
+		return Hash;
+	}
+};
+
 /** Holds offsets to captured callstacks. */
 class FCallstacks
 {
 public:
+	/** Struct to hold the actual Serialize call callstack and any associated data */
+	struct FCallstackData : public FThreadSafeRefCountedObject
+	{
+		/** Full callstack */
+		TUniquePtr<ANSICHAR[]> Callstack;
+		/**
+		 * The export being serialized. Garbage collections do not occur during the savepackage call,
+		 * so it is safe for us to hold a raw pointer.
+		 */
+		UObject* SerializedObject = nullptr;
+		/** The currently serialized property */
+		FProperty* SerializedProp = nullptr;
+		/** Hash of this->Callstack. */
+		uint32 CppCallstackHash = MAX_uint32;
+
+		FCallstackData() = default;
+		FCallstackData(TUniquePtr<ANSICHAR[]>&& InCallstack, uint32 InCppCallstackHash, UObject* InSerializedObject, FProperty* InSerializedProperty);
+		FCallstackData(FCallstackData&&) = delete; // FThreadSafeRefCountedObject doesn't allow it
+		FCallstackData(const FCallstackData&) = delete;
+
+		FCallstackData& operator=(FCallstackData&&) = delete; // FThreadSafeRefCountedObject doesn't allow it
+		FCallstackData& operator=(const FCallstackData&) = delete;
+
+		/** Converts the callstack and associated data to human readable string */
+		FString ToString(const TCHAR* CallstackCutoffText) const;
+
+		/** Clone the callstack data */
+		TRefCountPtr<FCallstackData> Clone() const;
+		/** Get the key under which this FCallstackData is stored. */
+		FCallstackKey GetKey() const;
+		FString GetObjectName() const;
+		FString GetPropertyName() const;
+	};
+
 	/** Offset and callstack pair */
 	struct FCallstackAtOffset
 	{
@@ -118,38 +184,11 @@ public:
 		/** The length written by the Serialize call. */
 		int64 SerializeCallLength = -1;
 		/** Callstack CRC for the Serialize call */
-		uint32 Callstack = 0;
+		TRefCountPtr<FCallstackData> Callstack;
 		/** Collected inside of a scope that indicates diff should be recorded but logging should be suppressed */
 		bool bSuppressLogging = false;
 	};
 
-	/** Struct to hold the actual Serialize call callstack and any associated data */
-	struct FCallstackData
-	{
-		/** Full callstack */
-		TUniquePtr<ANSICHAR[]> Callstack;
-		/** Full name of the currently serialized object */
-		FString SerializedObjectName;
-		/** The currently serialized property */
-		FProperty* SerializedProp = nullptr;
-		/** Name of the currently serialized property */
-		FString SerializedPropertyName;
-
-		FCallstackData() = default;
-		FCallstackData(TUniquePtr<ANSICHAR[]>&& InCallstack, UObject* InSerializedObject, FProperty* InSerializedProperty);
-		FCallstackData(FCallstackData&&) = default;
-		FCallstackData(const FCallstackData&) = delete;
-
-		FCallstackData& operator=(FCallstackData&&) = default;
-		FCallstackData& operator=(const FCallstackData&) = delete;
-
-		/** Converts the callstack and associated data to human readable string */
-		FString ToString(const TCHAR* CallstackCutoffText) const;
-
-		/** Clone the callstack data */
-		FCallstackData Clone() const;
-	};
-	
 	FCallstacks();
 
 	/** Returns the total number of callstacks. */
@@ -191,9 +230,9 @@ public:
 		return CallstackAtOffsetMap[CallstackIndex];
 	}
 	
-	const FCallstackData& GetCallstackData(const FCallstackAtOffset& CallstackOffset) const
+	const TRefCountPtr<FCallstackData>& GetCallstackData(const FCallstackAtOffset& CallstackOffset) const
 	{
-		return UniqueCallstacks[CallstackOffset.Callstack];
+		return CallstackOffset.Callstack;
 	}
 
 	int64 GetEndOffset() const
@@ -203,20 +242,24 @@ public:
 
 private:
 	/** Adds a unique callstack to UniqueCallstacks map */
-	ANSICHAR* AddUniqueCallstack(bool bIsCollectingCallstacks, UObject* SerializedObject, FProperty* SerializedProperty, uint32& OutCallstackCRC);
+	TRefCountPtr<FCallstackData> AddUniqueCallstack(const FCallstackKey& Key);
 
 	/** List of offsets and their respective callstacks */
 	TArray<FCallstackAtOffset> CallstackAtOffsetMap;
 	/** Contains all unique callstacks for all Serialize calls */
-	TMap<uint32, FCallstackData> UniqueCallstacks;
-	/** Optimizes callstack comparison. If true the current and the last callstack should be compared as it may have changed */
-	bool bCallstacksDirty;
+	TMap<FCallstackKey, TRefCountPtr<FCallstackData>> UniqueCallstacks;
 	/** Maximum size of the stack trace */
 	const SIZE_T StackTraceSize;
 	/** Buffer for getting the current stack trace */
 	TUniquePtr<ANSICHAR[]> StackTrace;
 	/** Callstack associated with the previous Serialize call */
-	ANSICHAR* LastSerializeCallstack;
+	uint32 PreviousCppCallstackHash;
+	/**
+	 * Optimizes callstack comparison. If false the Cpp callstack has not changed and does not need to be compared when
+	 * checking whether the current UniqueCallstackHash has changed.
+	 */
+	bool bCppCallstackDirty;
+
 	/** Total serialized bytes */
 	int64 EndOffset;
 };
