Capture handling in SColorGradientEditor

#tests Tested in application loop with the down/up happening normally. Tested with the alt-tab before the up in another app. Cleanup is happening with the changed code
#rnx

diff --git a/Engine/Source/Editor/UnrealEd/Private/SColorGradientEditor.cpp b/Engine/Source/Editor/UnrealEd/Private/SColorGradientEditor.cpp
--- a/Engine/Source/Editor/UnrealEd/Private/SColorGradientEditor.cpp
+++ b/Engine/Source/Editor/UnrealEd/Private/SColorGradientEditor.cpp
@@ -624,6 +624,18 @@ FReply SColorGradientEditor::OnMouseButtonUp( const FGeometry& MyGeometry, const
 	return FReply::Unhandled();
 }
 
+void SColorGradientEditor::OnMouseCaptureLost(const FCaptureLostEvent& CaptureLostEvent)
+{
+	if (bDraggingStop == true)
+	{
+		GEditor->EndTransaction();
+		DistanceDragged = 0;
+		bDraggingStop = false;
+		CurveOwner->OnCurveChanged(CurveOwner->GetCurves());
+	}
+
+}
+
 FReply SColorGradientEditor::OnKeyDown( const FGeometry& MyGeometry, const FKeyEvent& InKeyEvent )
 {
 	if ( IsEditingEnabled.Get() == true && InKeyEvent.GetKey() == EKeys::Platform_Delete )
diff --git a/Engine/Source/Editor/UnrealEd/Public/SColorGradientEditor.h b/Engine/Source/Editor/UnrealEd/Public/SColorGradientEditor.h
--- a/Engine/Source/Editor/UnrealEd/Public/SColorGradientEditor.h
+++ b/Engine/Source/Editor/UnrealEd/Public/SColorGradientEditor.h
@@ -73,6 +73,7 @@ public:
 	UNREALED_API virtual FReply OnKeyDown( const FGeometry& MyGeometry, const FKeyEvent& InKeyEvent ) override;
 	UNREALED_API virtual void OnMouseLeave(const FPointerEvent& MouseEvent) override;
 	UNREALED_API virtual FVector2D ComputeDesiredSize(float) const override;
+	UNREALED_API virtual void OnMouseCaptureLost(const FCaptureLostEvent& CaptureLostEvent) override;
 
 	/**
 	 * Sets the curve that is being edited by this editor
