Cooker: DiffOnly: Fix bug in compare function found by static analysis compiler.
#rnx

diff --git a/Engine/Source/Editor/UnrealEd/Private/Cooker/DiffWriterArchive.h b/Engine/Source/Editor/UnrealEd/Private/Cooker/DiffWriterArchive.h
--- a/Engine/Source/Editor/UnrealEd/Private/Cooker/DiffWriterArchive.h
+++ b/Engine/Source/Editor/UnrealEd/Private/Cooker/DiffWriterArchive.h
@@ -111,7 +111,7 @@ struct FCallstackKey
 	{
 		return Other.SerializedObject == SerializedObject &&
 			Other.SerializedProperty == SerializedProperty &&
-			Other.CppCallstackHash == Other.CppCallstackHash;
+			Other.CppCallstackHash == CppCallstackHash;
 	}
 	friend uint32 GetTypeHash(const FCallstackKey& Key)
 	{
