Fix static analysis null deref issue introduced by change to reset preview audio component
#rnx

diff --git a/Engine/Source/Editor/UnrealEd/Private/EditorEngine.cpp b/Engine/Source/Editor/UnrealEd/Private/EditorEngine.cpp
--- a/Engine/Source/Editor/UnrealEd/Private/EditorEngine.cpp
+++ b/Engine/Source/Editor/UnrealEd/Private/EditorEngine.cpp
@@ -2865,6 +2865,7 @@ UAudioComponent* UEditorEngine::ResetPreviewAudioComponent(USoundBase* Sound, US
 
 	if (FAudioDevice* AudioDevice = GetMainAudioDeviceRaw())
 	{
+		// 1. Initialize AudioComponent for previewing
 		if (PreviewAudioComponent)
 		{
 			PreviewAudioComponent->Stop();
@@ -2897,20 +2898,13 @@ UAudioComponent* UEditorEngine::ResetPreviewAudioComponent(USoundBase* Sound, US
 		PreviewAudioComponent->bReverb = false;
 		PreviewAudioComponent->bCenterChannelOnly = false;
 
-		if (!::IsValid(Sound) || Sound->EditorData.OverrideUseSubmixesOnPreview == EUseSubmixOnPreviewMode::UseEditorPreference)
-		{
-			const UAudioEditorSettings* AudioEditorSettings = GetDefault<UAudioEditorSettings>();
-			PreviewAudioComponent->bIsDryPreviewPlayback = !AudioEditorSettings->bUseSubmixesForPreviewSound;
-		}
-		else
-		{
-			const bool bIsDryPreviewPlayback = Sound->EditorData.OverrideUseSubmixesOnPreview == EUseSubmixOnPreviewMode::PlayPreviewsWithoutSubmixes;
-			PreviewAudioComponent->bIsDryPreviewPlayback = bIsDryPreviewPlayback;
-		}
-
+		PreviewAudioComponent->AudioDeviceID = INDEX_NONE;
 		if (FAudioDeviceManager* DeviceManager = GEngine->GetAudioDeviceManager())
 		{
-			PreviewAudioComponent->AudioDeviceID = DeviceManager->GetActiveAudioDevice()->DeviceID;
+			if (FAudioDeviceHandle ActiveDevice = DeviceManager->GetActiveAudioDevice(); ActiveDevice.IsValid())
+			{
+				PreviewAudioComponent->AudioDeviceID = ActiveDevice->DeviceID;
+			}
 		}
 
 		if (Sound)
@@ -2924,6 +2918,20 @@ UAudioComponent* UEditorEngine::ResetPreviewAudioComponent(USoundBase* Sound, US
 			PreviewAudioComponent->Sound = PreviewSoundCue;
 			PreviewSoundCue->CacheAggregateValues();
 		}
+
+		// 2. Determine whether or not to apply submix behavior
+		{
+			const UAudioEditorSettings* AudioEditorSettings = GetDefault<UAudioEditorSettings>();
+			check(AudioEditorSettings);
+			PreviewAudioComponent->bIsDryPreviewPlayback = !AudioEditorSettings->bUseSubmixesForPreviewSound;
+			if (const USoundBase* PreviewSound = PreviewAudioComponent->Sound)
+			{
+				if (PreviewSound->EditorData.OverrideUseSubmixesOnPreview != EUseSubmixOnPreviewMode::UseEditorPreference)
+				{
+					PreviewAudioComponent->bIsDryPreviewPlayback = PreviewSound->EditorData.OverrideUseSubmixesOnPreview == EUseSubmixOnPreviewMode::PlayPreviewsWithoutSubmixes;
+				}
+			}
+		}
 	}
 
 	return PreviewAudioComponent;
