Add method to force AllChunksInfo to be generated for non-chunked builds.
Add new option:
-ForceGenerateAllChunksInfo

Which in a non chunked build will force the assetRegistryGenerator to produce an AllChunksInfo.csv

diff --git a/Engine/Source/Editor/UnrealEd/Private/Commandlets/AssetRegistryGenerator.cpp b/Engine/Source/Editor/UnrealEd/Private/Commandlets/AssetRegistryGenerator.cpp
--- a/Engine/Source/Editor/UnrealEd/Private/Commandlets/AssetRegistryGenerator.cpp
+++ b/Engine/Source/Editor/UnrealEd/Private/Commandlets/AssetRegistryGenerator.cpp
@@ -1979,7 +1979,7 @@ bool FAssetRegistryGenerator::SaveAssetRegistry(const FString& SandboxPath, bool
 			OutDevelopmentAssetRegistryHash = DevArXxHash;
 		}
 
-		if (bGenerateChunks)
+		if (bGenerateChunks || FParse::Param(FCommandLine::Get(), TEXT("ForceGenerateAllChunksInfo")))
 		{
 			FString ChunkListsPath = PlatformSandboxPath.Replace(*FString::Printf(TEXT("/%s"), DevelopmentAssetRegistryFilename), TEXT(""));
 
