Fix for GC issues when resaving sound asset while playing in preview component
#rnx

diff --git a/Engine/Source/Editor/UnrealEd/Private/EditorEngine.cpp b/Engine/Source/Editor/UnrealEd/Private/EditorEngine.cpp
--- a/Engine/Source/Editor/UnrealEd/Private/EditorEngine.cpp
+++ b/Engine/Source/Editor/UnrealEd/Private/EditorEngine.cpp
@@ -2861,75 +2861,74 @@ UAudioComponent* UEditorEngine::GetPreviewAudioComponent()
 
 UAudioComponent* UEditorEngine::ResetPreviewAudioComponent(USoundBase* Sound, USoundNode* SoundNode)
 {
-	PreviewSoundCue = nullptr;
+	if (PreviewSoundCue)
+	{
+		PreviewSoundCue->FirstNode = nullptr;
+		PreviewSoundCue = nullptr;
+	}
 
-	if (FAudioDevice* AudioDevice = GetMainAudioDeviceRaw())
+	// 1. Initialize & overwrite new component for previewing
+	if (PreviewAudioComponent)
 	{
-		// 1. Initialize AudioComponent for previewing
-		if (PreviewAudioComponent)
-		{
-			PreviewAudioComponent->Stop();
-		}
-		else
-		{
-			if (Sound)
-			{
-				PreviewAudioComponent = FAudioDevice::CreateComponent(Sound);
-			}
-			else if (SoundNode)
-			{
-				PreviewSoundCue = NewObject<USoundCue>();
-				PreviewAudioComponent = FAudioDevice::CreateComponent(PreviewSoundCue);
-			}
-			else
-			{
-				PreviewAudioComponent = NewObject<UAudioComponent>();
-			}
-		}
+		PreviewAudioComponent->Stop();
+	}
+
+	if (Sound)
+	{
+		PreviewAudioComponent = FAudioDevice::CreateComponent(Sound);
+	}
+	else if (SoundNode)
+	{
+		PreviewSoundCue = NewObject<USoundCue>();
+		PreviewAudioComponent = FAudioDevice::CreateComponent(PreviewSoundCue);
+	}
+	else
+	{
+		PreviewAudioComponent = NewObject<UAudioComponent>();
+	}
 
-		check(PreviewAudioComponent);
+	check(PreviewAudioComponent);
 
-		// Mark as a preview component so the distance calculations is ignored
-		PreviewAudioComponent->bPreviewComponent = true;
-		PreviewAudioComponent->bIsPreviewSound = true;
-		PreviewAudioComponent->bAutoDestroy = false;
-		PreviewAudioComponent->bIsUISound = true;
-		PreviewAudioComponent->bAllowSpatialization = false;
-		PreviewAudioComponent->bReverb = false;
-		PreviewAudioComponent->bCenterChannelOnly = false;
+	// Mark as a preview component so the distance calculations is ignored
+	PreviewAudioComponent->bPreviewComponent = true;
+	PreviewAudioComponent->bIsPreviewSound = true;
+	PreviewAudioComponent->bAutoDestroy = false;
+	PreviewAudioComponent->bIsUISound = true;
+	PreviewAudioComponent->bAllowSpatialization = false;
+	PreviewAudioComponent->bReverb = false;
+	PreviewAudioComponent->bCenterChannelOnly = false;
 
-		PreviewAudioComponent->AudioDeviceID = INDEX_NONE;
-		if (FAudioDeviceManager* DeviceManager = GEngine->GetAudioDeviceManager())
+	PreviewAudioComponent->AudioDeviceID = INDEX_NONE;
+	if (FAudioDeviceManager* DeviceManager = GEngine->GetAudioDeviceManager())
+	{
+		if (FAudioDeviceHandle ActiveDevice = DeviceManager->GetActiveAudioDevice(); ActiveDevice.IsValid())
 		{
-			if (FAudioDeviceHandle ActiveDevice = DeviceManager->GetActiveAudioDevice(); ActiveDevice.IsValid())
-			{
-				PreviewAudioComponent->AudioDeviceID = ActiveDevice->DeviceID;
-			}
+			PreviewAudioComponent->AudioDeviceID = ActiveDevice->DeviceID;
 		}
+	}
 
-		if (Sound)
-		{
-			PreviewAudioComponent->Sound = Sound;
-		}
-		else if (SoundNode)
-		{
-			check(PreviewSoundCue);
-			PreviewSoundCue->FirstNode = SoundNode;
-			PreviewAudioComponent->Sound = PreviewSoundCue;
-			PreviewSoundCue->CacheAggregateValues();
-		}
+	if (Sound)
+	{
+		PreviewAudioComponent->Sound = Sound;
+	}
+	else if (SoundNode)
+	{
+		check(PreviewSoundCue);
+		PreviewSoundCue->FirstNode = SoundNode;
+		PreviewAudioComponent->Sound = PreviewSoundCue;
+		PreviewSoundCue->CacheAggregateValues();
+	}
 
-		// 2. Determine whether or not to apply submix behavior
+	// 2. Determine whether or not to apply submix behavior
+	{
+		const UAudioEditorSettings* AudioEditorSettings = GetDefault<UAudioEditorSettings>();
+		check(AudioEditorSettings);
+		PreviewAudioComponent->bIsDryPreviewPlayback = !AudioEditorSettings->bUseSubmixesForPreviewSound;
+		if (const USoundBase* PreviewSound = PreviewAudioComponent->Sound)
 		{
-			const UAudioEditorSettings* AudioEditorSettings = GetDefault<UAudioEditorSettings>();
-			check(AudioEditorSettings);
-			PreviewAudioComponent->bIsDryPreviewPlayback = !AudioEditorSettings->bUseSubmixesForPreviewSound;
-			if (const USoundBase* PreviewSound = PreviewAudioComponent->Sound)
+			if (PreviewSound->EditorData.OverrideUseSubmixesOnPreview != EUseSubmixOnPreviewMode::UseEditorPreference)
 			{
-				if (PreviewSound->EditorData.OverrideUseSubmixesOnPreview != EUseSubmixOnPreviewMode::UseEditorPreference)
-				{
-					PreviewAudioComponent->bIsDryPreviewPlayback = PreviewSound->EditorData.OverrideUseSubmixesOnPreview == EUseSubmixOnPreviewMode::PlayPreviewsWithoutSubmixes;
-				}
+				PreviewAudioComponent->bIsDryPreviewPlayback = PreviewSound->EditorData.OverrideUseSubmixesOnPreview == EUseSubmixOnPreviewMode::PlayPreviewsWithoutSubmixes;
 			}
 		}
 	}
@@ -2985,21 +2984,24 @@ bool UEditorEngine::CanPlayEditorSound() const
 
 void UEditorEngine::ClearPreviewComponents()
 {
-	if( PreviewAudioComponent )
+	if (PreviewAudioComponent)
 	{
 		PreviewAudioComponent->Stop();
 
-		// Just null out so they get GC'd
-		PreviewSoundCue->FirstNode = NULL;
-		PreviewSoundCue = NULL;
-		PreviewAudioComponent->Sound = NULL;
-		PreviewAudioComponent = NULL;
+		if (PreviewSoundCue)
+		{
+			PreviewSoundCue->FirstNode = nullptr;
+			PreviewSoundCue = nullptr;
+		}
+
+		PreviewAudioComponent->Sound = nullptr;
+		PreviewAudioComponent = nullptr;
 	}
 
 	if (PreviewMeshComp)
 	{
 		PreviewMeshComp->UnregisterComponent();
-		PreviewMeshComp = NULL;
+		PreviewMeshComp = nullptr;
 	}
 }
 
