Fix compilation issues with 44292583

diff --git a/Engine/Source/Editor/UnrealEd/Private/Cooker/PackageBuildDependencyTracker.h b/Engine/Source/Editor/UnrealEd/Private/Cooker/PackageBuildDependencyTracker.h
--- a/Engine/Source/Editor/UnrealEd/Private/Cooker/PackageBuildDependencyTracker.h
+++ b/Engine/Source/Editor/UnrealEd/Private/Cooker/PackageBuildDependencyTracker.h
@@ -8,6 +8,7 @@
 #include "Async/Mutex.h"
 #include "Templates/UnrealTemplate.h"
 #include "UObject/ObjectHandle.h"
+#include "UObject/TopLevelAssetPath.h"
 #endif
 
 #if UE_WITH_PACKAGE_ACCESS_TRACKING
