Firefight: Bump CookIncrementalVersion

diff --git a/Engine/Source/Editor/UnrealEd/Private/Cooker/CookIncrementalVersion.cpp b/Engine/Source/Editor/UnrealEd/Private/Cooker/CookIncrementalVersion.cpp
--- a/Engine/Source/Editor/UnrealEd/Private/Cooker/CookIncrementalVersion.cpp
+++ b/Engine/Source/Editor/UnrealEd/Private/Cooker/CookIncrementalVersion.cpp
@@ -8,6 +8,6 @@ namespace UE::Cook
 {
 
 // Change CookIncrementalVersion to a new guid when all packages in an incremental cook need to be invalidated
-FGuid CookIncrementalVersion(0x93DE0E19, 0x721143D5, 0x9487FF18, 0xEF03D6B1);
+FGuid CookIncrementalVersion(0x93DE0E1A, 0x721143D5, 0x9487FF18, 0xEF03D6B1);
 
 }
