Asset Thumbnails performance improvements

To improve the performance of asset thumbnails and avoid hitches, this CL removes the synchronization of the thumbnail rendering with the RenderThread.

The current implementation uses a fence to force a synchronization before rendering the thumbnail which creates unnecessary hitches. This is especially visible when shaders are compiling as the thumbnails are re-rendered when the shaders are done compiling. This can happen frequently when navigating inside a newly opened map.

The removed fence was fixing an issue of thumbnail resizing. A render command was issued to resize the texture on RenderThread, but the size of the texture was read and used on GameThread just after to inform the sizer of the Canvas. In some cases, reading the size on GameThread returned the old size and ended up creating visual inconsistencies, as the real size of the texture was different from the size used to render its content.

To avoid this, we adjusted the textures in the pool to have an immutable size instead of resizing them when being reused. This also allows the pool to look first for an available texture of the right size, avoiding unnecessary reallocations.

diff --git a/Engine/Source/Editor/UnrealEd/Private/AssetThumbnail.cpp b/Engine/Source/Editor/UnrealEd/Private/AssetThumbnail.cpp
--- a/Engine/Source/Editor/UnrealEd/Private/AssetThumbnail.cpp
+++ b/Engine/Source/Editor/UnrealEd/Private/AssetThumbnail.cpp
@@ -1966,16 +1966,26 @@ FAssetThumbnailPool::~FAssetThumbnailPool()
 
 FAssetThumbnailPool::FThumbnailInfo::~FThumbnailInfo()
 {
-	if( ThumbnailTexture )
+	if( ThumbnailTexture || ThumbnailRenderTarget )
 	{
-		delete ThumbnailTexture;
-		ThumbnailTexture = NULL;
-	}
-
-	if( ThumbnailRenderTarget )
-	{
-		delete ThumbnailRenderTarget;
-		ThumbnailRenderTarget = NULL;
+		ENQUEUE_RENDER_COMMAND(ReleaseThumbnailTextures)(
+		[ThumbnailTexture = ThumbnailTexture, ThumbnailRenderTarget = ThumbnailRenderTarget](FRHICommandListImmediate& RHICmdList)
+		{
+			if (ThumbnailTexture)
+			{
+				ThumbnailTexture->ClearTextureData();
+				ThumbnailTexture->ReleaseResource();
+				delete ThumbnailTexture;
+			}
+			if (ThumbnailRenderTarget)
+			{
+				ThumbnailRenderTarget->ReleaseResource();
+				delete ThumbnailRenderTarget;
+			}
+		});
+		
+		ThumbnailTexture = nullptr;
+		ThumbnailRenderTarget = nullptr;
 	}
 }
 
@@ -1986,47 +1996,29 @@ void FAssetThumbnailPool::ReleaseResources()
 	RealTimeThumbnails.Empty();
 	RealTimeThumbnailsToRender.Empty();
 
-	TArray< TSharedRef<FThumbnailInfo> > ThumbnailsToRelease;
-
-	for( auto ThumbIt = ThumbnailToTextureMap.CreateConstIterator(); ThumbIt; ++ThumbIt )
-	{
-		ThumbnailsToRelease.Add(ThumbIt.Value());
-	}
-	ThumbnailToTextureMap.Empty();
-
-	for( auto ThumbIt = FreeThumbnails.CreateConstIterator(); ThumbIt; ++ThumbIt )
-	{
-		ThumbnailsToRelease.Add(*ThumbIt);
-	}
-	FreeThumbnails.Empty();
-
-	for ( auto ThumbIt = ThumbnailsToRelease.CreateConstIterator(); ThumbIt; ++ThumbIt )
-	{
-		const TSharedRef<FThumbnailInfo>& Thumb = *ThumbIt;
-
-			// Release rendering resources
-			FThumbnailInfo_RenderThread ThumbInfo = Thumb.Get();
-			ENQUEUE_RENDER_COMMAND(ReleaseThumbnailResources)(
-				[ThumbInfo](FRHICommandListImmediate& RHICmdList)
-				{
-					ThumbInfo.ThumbnailTexture->ClearTextureData();
-					ThumbInfo.ThumbnailTexture->ReleaseResource();
-					ThumbInfo.ThumbnailRenderTarget->ReleaseResource();
-				});
-		}
-
 	// Wait for all resources to be released
 	FlushRenderingCommands();
 
-	// Make sure there are no more references to any of our thumbnails now that rendering commands have been flushed
-	for ( auto ThumbIt = ThumbnailsToRelease.CreateConstIterator(); ThumbIt; ++ThumbIt )
+	auto EnsureThumbnailIsUnique = [](const TSharedRef<FThumbnailInfo>& Thumb)
 	{
-		const TSharedRef<FThumbnailInfo>& Thumb = *ThumbIt;
 		if ( !Thumb.IsUnique() )
 		{
 			ensureMsgf(0, TEXT("Thumbnail info for '%s' is still referenced by '%d' other objects"), *Thumb->AssetData.GetObjectPathString(), Thumb.GetSharedReferenceCount());
 		}
+	};
+	
+	// Make sure there are no more references to any of our thumbnails now that rendering commands have been flushed
+	for(const TPair<FThumbId, TSharedRef<FThumbnailInfo>>& Pair : ThumbnailToTextureMap)
+	{
+		EnsureThumbnailIsUnique(Pair.Value);
 	}
+	ThumbnailToTextureMap.Empty();
+	
+	for(const TSharedRef<FThumbnailInfo>& Thumbnail : FreeThumbnails)
+	{
+		EnsureThumbnailIsUnique(Thumbnail);
+	}
+	FreeThumbnails.Empty();
 }
 
 TStatId FAssetThumbnailPool::GetStatId() const
@@ -2039,6 +2031,11 @@ bool FAssetThumbnailPool::IsTickable() const
 	return RecentlyLoadedAssets.Num() > 0 || ThumbnailsToRenderStack.Num() > 0 || RealTimeThumbnails.Num() > 0 || bWereShadersCompilingLastFrame || (GShaderCompilingManager && GShaderCompilingManager->IsCompiling());
 }
 
+UE_TRACE_EVENT_BEGIN(Cpu, FAssetThumbnailPool_Tick_RenderingThumbnail, NoSync)
+	UE_TRACE_EVENT_FIELD(UE::Trace::WideString, ObjectPath)
+	UE_TRACE_EVENT_FIELD(bool, IsRealTime)
+UE_TRACE_EVENT_END()
+
 void FAssetThumbnailPool::Tick( float DeltaTime )
 {
 	// If throttling do not tick unless drag dropping which could have a thumbnail as the cursor decorator
@@ -2103,6 +2100,7 @@ void FAssetThumbnailPool::Tick( float DeltaTime )
 		// Render as many thumbnails as we are allowed to
 		while( ThumbnailsToRenderStack.Num() + RealTimeThumbnailsToRender.Num() > 0 && FPlatformTime::Seconds() - FrameStartTime < MaxFrameTimeAllowance )
 		{
+			bool bIsRealTimeThumbnail = false;
 			TSharedPtr<FThumbnailInfo> Info;
 			if ( ThumbnailsToRenderStack.Num() > 0 )
 			{
@@ -2112,6 +2110,7 @@ void FAssetThumbnailPool::Tick( float DeltaTime )
 			{
 				Info = RealTimeThumbnailsToRender.Pop();
 				NumRealTimeThumbnailsRenderedThisFrame++;
+				bIsRealTimeThumbnail = true;
 			}
 			else
 			{
@@ -2126,6 +2125,12 @@ void FAssetThumbnailPool::Tick( float DeltaTime )
 
 				if ( InfoRef->AssetData.IsValid() )
 				{
+#if CPUPROFILERTRACE_ENABLED
+					UE_TRACE_LOG_SCOPED_T(Cpu, FAssetThumbnailPool_Tick_RenderingThumbnail, CpuChannel)
+						<< FAssetThumbnailPool_Tick_RenderingThumbnail.ObjectPath(*InfoRef->AssetData.GetObjectPathString())
+						<< FAssetThumbnailPool_Tick_RenderingThumbnail.IsRealTime(bIsRealTimeThumbnail);
+#endif // CPUPROFILERTRACE_ENABLED
+
 					FAssetData CustomThumbnailAsset;
 					// Check if a different asset should be used to generate the thumbnail for this asset
 					FString CustomThumbnailTagValue;
@@ -2226,22 +2231,16 @@ bool FAssetThumbnailPool::LoadThumbnail(TSharedRef<FThumbnailInfo> ThumbnailInfo
 					auto EnqueueThumbnailRender = [Asset, ThumbInfo, ThumbnailInfo]()
 					{
 						ENQUEUE_RENDER_COMMAND(SyncSlateTextureCommand)(
-							[ThumbInfo](FRHICommandListImmediate& RHICmdList)
+						[ThumbInfo](FRHICommandListImmediate& RHICmdList)
+						{
+							if (ThumbInfo.ThumbnailTexture->GetTypedResource() != ThumbInfo.ThumbnailRenderTarget->GetTextureRHI())
 							{
-								if (ThumbInfo.ThumbnailTexture->GetTypedResource() != ThumbInfo.ThumbnailRenderTarget->GetTextureRHI())
-								{
-									ThumbInfo.ThumbnailTexture->ClearTextureData();
-									ThumbInfo.ThumbnailTexture->ReleaseRHI();
-									ThumbInfo.ThumbnailTexture->SetRHIRef(ThumbInfo.ThumbnailRenderTarget->GetTextureRHI(), ThumbInfo.Width, ThumbInfo.Height);
-								}
-							});
-
-						// // We have to wait for the render command to finish since thumbnail rendering cannot be done on the GPU currently
-						FRenderCommandFence Fence;
-						Fence.BeginFence();
-						Fence.Wait();
+								ThumbInfo.ThumbnailTexture->ClearTextureData();
+								ThumbInfo.ThumbnailTexture->ReleaseRHI();
+								ThumbInfo.ThumbnailTexture->SetRHIRef(ThumbInfo.ThumbnailRenderTarget->GetTextureRHI(), ThumbInfo.Width, ThumbInfo.Height);
+							}
+						});
 
-						//@todo: this should be done on the GPU only but it is not supported by thumbnail tools yet
 						ThumbnailTools::RenderThumbnail(
 							Asset,
 							ThumbnailInfo->Width,
@@ -2334,26 +2333,46 @@ FSlateTexture2DRHIRef* FAssetThumbnailPool::AccessTexture( const FAssetData& Ass
 {
 	if(!AssetData.IsValid() || Width == 0 || Height == 0)
 	{
-		return NULL;
+		return nullptr;
+	}
+	
+	FThumbId ThumbId( AssetData.GetSoftObjectPath(), Width, Height ) ;
+	// Check to see if a thumbnail for this asset exists.  If so we don't need to render it
+	const TSharedRef<FThumbnailInfo>* ThumbnailInfoPtr = ThumbnailToTextureMap.Find( ThumbId );
+	TSharedPtr<FThumbnailInfo> ThumbnailInfo;
+	if( ThumbnailInfoPtr )
+	{
+		ThumbnailInfo = *ThumbnailInfoPtr;
 	}
 	else
 	{
-		FThumbId ThumbId( AssetData.GetSoftObjectPath(), Width, Height ) ;
-		// Check to see if a thumbnail for this asset exists.  If so we don't need to render it
-		const TSharedRef<FThumbnailInfo>* ThumbnailInfoPtr = ThumbnailToTextureMap.Find( ThumbId );
-		TSharedPtr<FThumbnailInfo> ThumbnailInfo;
-		if( ThumbnailInfoPtr )
+		// Check if there is any thumbnail in the free pool that can be used
+		// If there is, use it. This does not change the overall number of resources used
+		for (TArray<TSharedRef<FThumbnailInfo>>::TIterator It(FreeThumbnails); It; ++It)
 		{
-			ThumbnailInfo = *ThumbnailInfoPtr;
+			TSharedRef<FThumbnailInfo>& FreeThumbnail = *It;
+			if (FreeThumbnail->Height == Height && FreeThumbnail->Width == Width)
+			{
+				ThumbnailInfo = FreeThumbnail;
+				It.RemoveCurrentSwap();
+				break;
+			}
 		}
-		else
+		
+		// If we didn't find one and we have reached the pool limit, remove any thumbnail from the free thumbnails if any. This will free up a slot
+		// Otherwise, if there are no free thumbnails, evict the LRU used thumbnail to free up a slot.
+		// Note, with luck, the LRU element might be suitable for reuse
+		if( !ThumbnailInfo && (FreeThumbnails.Num() + ThumbnailToTextureMap.Num() == NumInPool) )
 		{
-			// If a the max number of thumbnails allowed by the pool exists then reuse its rendering resource for the new thumbnail
-			if( FreeThumbnails.Num() == 0 && ThumbnailToTextureMap.Num() == NumInPool )
+			if (!FreeThumbnails.IsEmpty())
 			{
-				// Find the thumbnail which was accessed last and use it for the new thumbnail
+				FreeThumbnails.Pop();
+			}
+			else
+			{
+				// Find the thumbnail which was accessed last and potentially use it for the new thumbnail
 				double LastAccessTime = std::numeric_limits<double>::max();
-				const FThumbId* AssetToRemove = NULL;
+				const FThumbId* AssetToRemove = nullptr;
 				for( TMap< FThumbId, TSharedRef<FThumbnailInfo> >::TConstIterator It(ThumbnailToTextureMap); It; ++It )
 				{
 					if( It.Value()->LastAccessTime < LastAccessTime )
@@ -2364,64 +2383,54 @@ FSlateTexture2DRHIRef* FAssetThumbnailPool::AccessTexture( const FAssetData& Ass
 				}
 
 				check( AssetToRemove );
-
-				// Remove the old mapping 
+				// Remove the old mapping. If matching our desired size, we reuse it, otherwise we discard it
 				ThumbnailInfo = ThumbnailToTextureMap.FindAndRemoveChecked( *AssetToRemove );
+				if (ThumbnailInfo->Height != Height || ThumbnailInfo->Width != Width)
+				{
+					ThumbnailInfo = nullptr;
+				}
 			}
-			else if( FreeThumbnails.Num() > 0 )
-			{
-				ThumbnailInfo = FreeThumbnails.Pop();
-
-				FSlateTextureRenderTarget2DResource* ThumbnailRenderTarget = ThumbnailInfo->ThumbnailRenderTarget;
-				ENQUEUE_RENDER_COMMAND(SlateUpdateThumbSizeCommand)(
-					[ThumbnailRenderTarget, Width, Height](FRHICommandListImmediate& RHICmdList)
-					{
-						ThumbnailRenderTarget->SetSize(Width, Height);
-					});
-			}
-			else
-			{
-				// There are no free thumbnail resources
-				check( (uint32)ThumbnailToTextureMap.Num() <= NumInPool );
-				check( !ThumbnailInfo.IsValid() );
-				// The pool isn't used up so just make a new texture 
-
-				// Make new thumbnail info if it doesn't exist
-				// This happens when the pool is not yet full
-				ThumbnailInfo = MakeShareable( new FThumbnailInfo );
-				
-				// Set the thumbnail and asset on the info. It is NOT safe to change or NULL these pointers until ReleaseResources.
-				ThumbnailInfo->ThumbnailTexture = new FSlateTexture2DRHIRef( Width, Height, PF_B8G8R8A8, NULL, TexCreate_None );
-				ThumbnailInfo->ThumbnailRenderTarget = new FSlateTextureRenderTarget2DResource(FLinearColor::Black, Width, Height, PF_B8G8R8A8, SF_Point, TA_Wrap, TA_Wrap, 0.0f);
-
-				BeginInitResource( ThumbnailInfo->ThumbnailTexture );
-				BeginInitResource( ThumbnailInfo->ThumbnailRenderTarget );
-			}
+		}
+		
+		// If there are no free thumbnail resources or space in the pool, create a new one
+		if( !ThumbnailInfo )
+		{
+			// There are no free thumbnail resources
+			check( (uint32)(FreeThumbnails.Num() + ThumbnailToTextureMap.Num()) < NumInPool );
+			// The pool isn't used up so just make a new texture 
 
+			// Make new thumbnail info if it doesn't exist
+			ThumbnailInfo = MakeShareable( new FThumbnailInfo(Width, Height) );
+			
+			// Set the thumbnail and asset on the info. It is NOT safe to change or NULL these pointers until ReleaseResources.
+			ThumbnailInfo->ThumbnailTexture = new FSlateTexture2DRHIRef( Width, Height, PF_B8G8R8A8, NULL, TexCreate_None );
+			ThumbnailInfo->ThumbnailRenderTarget = new FSlateTextureRenderTarget2DResource(FLinearColor::Black, Width, Height, PF_B8G8R8A8, SF_Point, TA_Wrap, TA_Wrap, 0.0f);
 
-			check( ThumbnailInfo.IsValid() );
-			TSharedRef<FThumbnailInfo> ThumbnailRef = ThumbnailInfo.ToSharedRef();
+			BeginInitResource( ThumbnailInfo->ThumbnailTexture );
+			BeginInitResource( ThumbnailInfo->ThumbnailRenderTarget );
+		}
 
-			// Map the object to its thumbnail info
-			ThumbnailToTextureMap.Add( ThumbId, ThumbnailRef );
 
-			ThumbnailInfo->AssetData = AssetData;
-			ThumbnailInfo->Width = Width;
-			ThumbnailInfo->Height = Height;
-		
-			// Mark this thumbnail as needing to be updated
-			ThumbnailInfo->LastUpdateTime = -1.f;
+		check( ThumbnailInfo.IsValid() );
+		TSharedRef<FThumbnailInfo> ThumbnailRef = ThumbnailInfo.ToSharedRef();
 
-			// Request that the thumbnail be rendered as soon as possible
-			ThumbnailsToRenderStack.Push( ThumbnailRef );
-		}
+		// Map the object to its thumbnail info
+		ThumbnailToTextureMap.Add( ThumbId, ThumbnailRef );
 
-		// This thumbnail was accessed, update its last time to the current time
-		// We'll use LastAccessTime to determine the order to recycle thumbnails if the pool is full
-		ThumbnailInfo->LastAccessTime = FPlatformTime::Seconds();
+		ThumbnailInfo->AssetData = AssetData;
+	
+		// Mark this thumbnail as needing to be updated
+		ThumbnailInfo->LastUpdateTime = -1.f;
 
-		return ThumbnailInfo->ThumbnailTexture;
+		// Request that the thumbnail be rendered as soon as possible
+		ThumbnailsToRenderStack.Push( ThumbnailRef );
 	}
+
+	// This thumbnail was accessed, update its last time to the current time
+	// We'll use LastAccessTime to determine the order to recycle thumbnails if the pool is full
+	ThumbnailInfo->LastAccessTime = FPlatformTime::Seconds();
+
+	return ThumbnailInfo->ThumbnailTexture;
 }
 
 void FAssetThumbnailPool::AddReferencer( const FAssetThumbnail& AssetThumbnail )
diff --git a/Engine/Source/Editor/UnrealEd/Public/AssetThumbnail.h b/Engine/Source/Editor/UnrealEd/Public/AssetThumbnail.h
--- a/Engine/Source/Editor/UnrealEd/Public/AssetThumbnail.h
+++ b/Engine/Source/Editor/UnrealEd/Public/AssetThumbnail.h
@@ -331,9 +331,11 @@ private:
 		/** The time since last update */
 		double LastUpdateTime = 0.0;
 		/** Width of the thumbnail */
-		uint32 Width = 0;
+		const uint32 Width;
 		/** Height of the thumbnail */
-		uint32 Height = 0;
+		const uint32 Height;
+
+		FThumbnailInfo(uint32 InWidth, uint32 InHeight) : Width(InWidth), Height(InHeight) {}
 		~FThumbnailInfo();
 	};
 	/**
diff --git a/Engine/Source/Editor/UnrealEd/Public/ObjectTools.h b/Engine/Source/Editor/UnrealEd/Public/ObjectTools.h
--- a/Engine/Source/Editor/UnrealEd/Public/ObjectTools.h
+++ b/Engine/Source/Editor/UnrealEd/Public/ObjectTools.h
@@ -750,7 +750,7 @@ namespace ThumbnailTools
 	 * @param	InImageWidth			the maximum width of the thumbnail. It may be smaller when supplying a OutThumbnail.
 	 * @param	InImageHeight			the maximum height of the thumbnail. It may be smaller when supplying a OutThumbnail.
 	 * @param	InRenderTargetResource	if non-NULL, the render target that will receive the thumbnail.
-	 * @param	OutThumbnail			if non-NULL, the FObjectThumbnail that will receive the thumbnail data.
+	 * @param	OutThumbnail			if non-NULL, the FObjectThumbnail that will receive the thumbnail data. This would also force a synchronization with the RenderThread
 	 */
 	UNREALED_API void RenderThumbnail( UObject* InObject, const uint32 InImageWidth, const uint32 InImageHeight, EThumbnailTextureFlushMode::Type InFlushMode, FTextureRenderTargetResource* InRenderTargetResource = NULL, FObjectThumbnail* OutThumbnail = NULL );
 
