IncrementalCook: Abort the cook rather than allowing it to wipe state, when we are running -incrementalvalidate but we find that an incremental cook is not possible. Wiping state is an irreversable operation that could cause a large time loss for the developer.
#rnx

diff --git a/Engine/Source/Editor/UnrealEd/Private/CookOnTheFlyServer.cpp b/Engine/Source/Editor/UnrealEd/Private/CookOnTheFlyServer.cpp
--- a/Engine/Source/Editor/UnrealEd/Private/CookOnTheFlyServer.cpp
+++ b/Engine/Source/Editor/UnrealEd/Private/CookOnTheFlyServer.cpp
@@ -5698,9 +5698,11 @@ public:
 		Initialize();
 		if (DiffMode == EDiffMode::None)
 		{
+			bAllowWrite.Emplace(true);
 			return;
 		}
 
+		bool bCurrentWriterAllowWrite = true;
 		ICookedPackageWriter::FCookCapabilities Capabilities = CookedPackageWriter->GetCookCapabilities();
 		if (!Capabilities.bDiffModeSupported)
 		{
@@ -5710,24 +5712,26 @@ public:
 		}
 
 		// Wrap the incoming writer inside the feature-specific-functionality writer
+		FIncrementalValidatePackageWriter* IncrementalValidatePackageWriter = nullptr;
 		switch (DiffMode)
 		{
 		case EDiffMode::DiffOnly:
 			CookedPackageWriter = new FDiffPackageWriter(TUniquePtr<ICookedPackageWriter>(CookedPackageWriter),
 				InDeterminismManager);
+			bCurrentWriterAllowWrite = false;
 			break;
 		case EDiffMode::IncrementalValidate:
-			CookedPackageWriter = new FIncrementalValidatePackageWriter(COTFS,
+			IncrementalValidatePackageWriter = new FIncrementalValidatePackageWriter(COTFS,
 				TUniquePtr<ICookedPackageWriter>(CookedPackageWriter), FIncrementalValidatePackageWriter::EPhase::AllInOnePhase,
 				ResolvedMetadataPath, InDeterminismManager);
 			break;
 		case EDiffMode::IncrementalValidatePhase1:
-			CookedPackageWriter = new FIncrementalValidatePackageWriter(COTFS,
+			IncrementalValidatePackageWriter = new FIncrementalValidatePackageWriter(COTFS,
 				TUniquePtr<ICookedPackageWriter>(CookedPackageWriter), FIncrementalValidatePackageWriter::EPhase::Phase1,
 				ResolvedMetadataPath, InDeterminismManager);
 			break;
 		case EDiffMode::IncrementalValidatePhase2:
-			CookedPackageWriter = new FIncrementalValidatePackageWriter(COTFS,
+			IncrementalValidatePackageWriter = new FIncrementalValidatePackageWriter(COTFS,
 				TUniquePtr<ICookedPackageWriter>(CookedPackageWriter), FIncrementalValidatePackageWriter::EPhase::Phase2,
 				ResolvedMetadataPath, InDeterminismManager);
 			break;
@@ -5735,6 +5739,16 @@ public:
 			checkNoEntry();
 			break;
 		}
+
+		if (IncrementalValidatePackageWriter)
+		{
+			CookedPackageWriter = IncrementalValidatePackageWriter;
+			bCurrentWriterAllowWrite = !IncrementalValidatePackageWriter->IsReadOnly();
+		}
+		if (!bAllowWrite.IsSet())
+		{
+			bAllowWrite.Emplace(bCurrentWriterAllowWrite);
+		}
 	}
 
 	bool IsDeterminismDebug() const
@@ -5752,6 +5766,14 @@ public:
 		}
 	}
 
+	bool IsFullBuildAllowed() const
+	{
+		// We require that bAllowWrite has been set - which is done during FindOrCreatePackageWriter - before
+		// IsFullBuildAllowed is read.
+		check(bAllowWrite.IsSet());
+		return *bAllowWrite;
+	}
+
 	void Initialize()
 	{
 		if (bInitialized)
@@ -5815,6 +5837,7 @@ public:
 private:
 	bool bInitialized = false;
 	EDiffMode DiffMode = EDiffMode::None;
+	TOptional<bool> bAllowWrite;
 };
 
 bool UCookOnTheFlyServer::IsDebugRecordUnsolicited() const
@@ -10676,6 +10699,11 @@ void UCookOnTheFlyServer::LoadBeginCookIncrementalFlagsLocal(FBeginCookContext&
 {
 	const TCHAR* CommandLine = FCommandLine::Get();
 	const bool bIsDiffOnly = FParse::Param(CommandLine, TEXT("DIFFONLY"));
+	// When running with incrementalvalidate without allowwrite, we have a contract with our user that
+	// we will not wipe their state, even if we find that an incremental build is not normally allowed.
+	// Abort the cook in that case instead.
+	const bool bFullBuildAllowed = DiffModeHelper->IsFullBuildAllowed();
+
 	bool bDefaultIncremental = true;
 	if (bCookIncremental)
 	{
@@ -10845,6 +10873,13 @@ void UCookOnTheFlyServer::LoadBeginCookIncrementalFlagsLocal(FBeginCookContext&
 		PlatformData->bAllowIncrementalResults = PlatformContext.bAllowIncrementalResults;
 		PlatformData->bLegacyIterativeSharedBuild = PlatformContext.bLegacyIterativeSharedBuild;
 		PlatformData->bWorkerOnSharedSandbox = PlatformContext.bWorkerOnSharedSandbox;
+
+		if (PlatformData->bFullBuild && !bFullBuildAllowed)
+		{
+			UE_LOG(LogCook, Fatal,
+				TEXT("INCREMENTAL COOK: A diffonly cook mode was specified (-diffonly or -incrementalvalidate), ")
+				TEXT("and we guarantee state will not be wiped. But an incremental cook was found to not be possible (see logs above). Aborting the cook."));
+		}
 	}
 }
 
diff --git a/Engine/Source/Editor/UnrealEd/Private/Cooker/IncrementalValidatePackageWriter.cpp b/Engine/Source/Editor/UnrealEd/Private/Cooker/IncrementalValidatePackageWriter.cpp
--- a/Engine/Source/Editor/UnrealEd/Private/Cooker/IncrementalValidatePackageWriter.cpp
+++ b/Engine/Source/Editor/UnrealEd/Private/Cooker/IncrementalValidatePackageWriter.cpp
@@ -341,6 +341,22 @@ FIncrementalValidatePackageWriter::FIncrementalValidatePackageWriter(UCookOnTheF
 
 	LoggingSoftMaximum = -1;
 	GConfig->GetValue(TEXT("IncrementalValidate"), TEXT("LoggingSoftMaximum"), LoggingSoftMaximum, GEditorIni);
+
+	switch (Phase)
+	{
+	case EPhase::AllInOnePhase:
+		bReadOnly = !FParse::Param(FCommandLine::Get(), TEXT("IncrementalValidateAllowWrite"));
+		break;
+	case EPhase::Phase1:
+		bReadOnly = false;
+		break;
+	case EPhase::Phase2:
+		bReadOnly = false;
+		break;
+	default:
+		checkNoEntry();
+		break;
+	}
 }
 
 void FIncrementalValidatePackageWriter::BeginPackage(const FBeginPackageInfo& Info)
@@ -592,11 +608,7 @@ void FIncrementalValidatePackageWriter::Initialize(const FCookInfo& CookInfo)
 			UE_CLOG(COTFS.GetCookMode() != ECookMode::CookWorker,
 				LogIncrementalValidate, Display,
 				TEXT("The cook is running non-incrementally. All packages are reported \"modified\" and will be resaved as in a normal cook."));
-			bReadOnly = false;
-		}
-		else
-		{
-			bReadOnly = !FParse::Param(FCommandLine::Get(), TEXT("IncrementalValidateAllowWrite"));
+			check(!bReadOnly); // We should have set this due to incrementalvalidateallowwrite, or aborted the cook earlier
 		}
 		break;
 	case EPhase::Phase1:
@@ -606,7 +618,6 @@ void FIncrementalValidatePackageWriter::Initialize(const FCookInfo& CookInfo)
 				LogIncrementalValidate, Display,
 				TEXT("The cook is running non-incrementally. All packages are reported \"modified\" and will be resaved during the final IncrementalValidate phase."));
 		}
-		bReadOnly = false;
 		break;
 	case EPhase::Phase2:
 		if (CookInfo.bFullBuild)
@@ -615,7 +626,6 @@ void FIncrementalValidatePackageWriter::Initialize(const FCookInfo& CookInfo)
 				LogIncrementalValidate, Display,
 				TEXT("The cook is running non-incrementally. Packages that were incrementally skipped and found valid will be resaved anyway."));
 		}
-		bReadOnly = false;
 		break;
 	default:
 		checkNoEntry();
@@ -1081,6 +1091,11 @@ bool FIncrementalValidatePackageWriter::IsAnotherSaveNeededInternal(FSavePackage
 	}
 }
 
+bool FIncrementalValidatePackageWriter::IsReadOnly() const
+{
+	return bReadOnly;
+}
+
 void FIncrementalValidatePackageWriter::OnDiffWriterMessage(ELogVerbosity::Type Verbosity, FStringView Message)
 {
 	PackageMessageMap.FindOrAdd(BeginInfo.PackageName).Add(FMessage{ FString(Message), Verbosity });
diff --git a/Engine/Source/Editor/UnrealEd/Private/Cooker/IncrementalValidatePackageWriter.h b/Engine/Source/Editor/UnrealEd/Private/Cooker/IncrementalValidatePackageWriter.h
--- a/Engine/Source/Editor/UnrealEd/Private/Cooker/IncrementalValidatePackageWriter.h
+++ b/Engine/Source/Editor/UnrealEd/Private/Cooker/IncrementalValidatePackageWriter.h
@@ -51,6 +51,9 @@ public:
 	virtual void UpdateSaveArguments(FSavePackageArgs& SaveArgs) override;
 	virtual bool IsAnotherSaveNeeded(FSavePackageResultStruct& PreviousResult, FSavePackageArgs& SaveArgs) override;
 
+	// FIncrementalValidatePackageWriter-specific public functions
+	bool IsReadOnly() const;
+
 protected:
 	enum class ESaveAction : uint8
 	{
