PR #11888: In LoadLastDirectories() reset a directory value if the last directory no longer exists.

diff --git a/Engine/Source/Editor/UnrealEd/Private/EditorDirectories.cpp b/Engine/Source/Editor/UnrealEd/Private/EditorDirectories.cpp
--- a/Engine/Source/Editor/UnrealEd/Private/EditorDirectories.cpp
+++ b/Engine/Source/Editor/UnrealEd/Private/EditorDirectories.cpp
@@ -52,7 +52,7 @@ void FEditorDirectories::LoadLastDirectories()
 
 	for( int32 CurDirectoryIndex = 0; CurDirectoryIndex < UE_ARRAY_COUNT( LastDir ); ++CurDirectoryIndex )
 	{
-		if (LastDir[ CurDirectoryIndex ].IsEmpty())
+		if (LastDir[ CurDirectoryIndex ].IsEmpty() || !FPaths::DirectoryExists(LastDir[CurDirectoryIndex]))
 		{
 			// Default all directories to the game content folder
 			if (CurDirectoryIndex == ELastDirectory::LEVEL)
