Add option for DLC AssetRegistry to be filtered via the same serialization method a non-DLC asset registry would be.
-ApplyFiltersToDLCAssetRegistry
Will have the AssetRegistry.bin also filtered rather than the default which is to not filter the AssetRegistry at all.

diff --git a/Engine/Source/Editor/UnrealEd/Private/CookOnTheFlyServer.cpp b/Engine/Source/Editor/UnrealEd/Private/CookOnTheFlyServer.cpp
--- a/Engine/Source/Editor/UnrealEd/Private/CookOnTheFlyServer.cpp
+++ b/Engine/Source/Editor/UnrealEd/Private/CookOnTheFlyServer.cpp
@@ -9807,7 +9807,7 @@ void UCookOnTheFlyServer::CookByTheBookFinishedInternal()
 	// if we are cooking DLC, the DevelopmentAR isn't needed - it's used when making DLC against shipping, so there's no need to make it
 	// again, as we don't make DLC against DLC (but allow an override just in case)
 	bool bSaveDevelopmentAssetRegistry = !FParse::Param(FCommandLine::Get(), TEXT("NoSaveDevAR"));
-	bool bForceNoFilterAssetsFromAssetRegistry = IsCookingDLC();
+	bool bForceNoFilterAssetsFromAssetRegistry = IsCookingDLC() ? !FParse::Param(FCommandLine::Get(), TEXT("ApplyFiltersToDLCAssetRegistry")) : false;
 	bool bSaveManifests = true;
 	bool bSaveIniSettings = !FParse::Param(FCommandLine::Get(), TEXT("SkipSaveCookSettings"));
 	bool bSaveCookerOpenOrder = true;
