Fix PostEditMove not being called when surface snapping and using multiple axes on the translate tool

PR #11028: Fix PostEditMove not being called when surface snapping and using mul?

diff --git a/Engine/Source/Editor/UnrealEd/Private/LevelEditorViewport.cpp b/Engine/Source/Editor/UnrealEd/Private/LevelEditorViewport.cpp
--- a/Engine/Source/Editor/UnrealEd/Private/LevelEditorViewport.cpp
+++ b/Engine/Source/Editor/UnrealEd/Private/LevelEditorViewport.cpp
@@ -3214,6 +3214,7 @@ void FLevelEditorViewportClient::ProjectActorsIntoWorld(const TArray<AActor*>& A
 			if (USceneComponent* RootComponent = Actor->GetRootComponent())
 			{
 				RootComponent->SetWorldTransform(ActorTransform);
+				Actor->PostEditMove(false);
 			}
 		}
 		else
