PR #10668: Fix crash when using seamless travel in PIE with selected objects and pending transactions

diff --git a/Engine/Source/Editor/UnrealEd/Classes/Editor/EditorEngine.h b/Engine/Source/Editor/UnrealEd/Classes/Editor/EditorEngine.h
--- a/Engine/Source/Editor/UnrealEd/Classes/Editor/EditorEngine.h
+++ b/Engine/Source/Editor/UnrealEd/Classes/Editor/EditorEngine.h
@@ -2775,6 +2775,9 @@ private:
 	/** Delegate callback for when a streaming level is removed from world. */
 	UNREALED_API void OnLevelRemovedFromWorld(ULevel* InLevel, UWorld* InWorld);
 
+	/** Delegate callback when doing a seamless transition to a new world */
+	void OnSeamlessTravelTransition(UWorld* InWorld);
+
 	/** Puts the currently loaded project file at the top of the recents list and trims and files that fell off the list */
 	UNREALED_API void UpdateRecentlyLoadedProjectFiles();
 
diff --git a/Engine/Source/Editor/UnrealEd/Private/EditorEngine.cpp b/Engine/Source/Editor/UnrealEd/Private/EditorEngine.cpp
--- a/Engine/Source/Editor/UnrealEd/Private/EditorEngine.cpp
+++ b/Engine/Source/Editor/UnrealEd/Private/EditorEngine.cpp
@@ -1259,6 +1259,7 @@ void UEditorEngine::Init(IEngineLoop* InEngineLoop)
 	FCoreUObjectDelegates::OnAssetLoaded.AddUObject(this, &UEditorEngine::OnAssetLoaded);
 	FWorldDelegates::LevelAddedToWorld.AddUObject(this, &UEditorEngine::OnLevelAddedToWorld);
 	FWorldDelegates::LevelRemovedFromWorld.AddUObject(this, &UEditorEngine::OnLevelRemovedFromWorld);
+	FWorldDelegates::OnSeamlessTravelTransition.AddUObject(this, &UEditorEngine::OnSeamlessTravelTransition);
 
 	IAssetRegistry::GetChecked().OnInMemoryAssetCreated().AddUObject(this, &UEditorEngine::OnAssetCreated);
 
@@ -1628,6 +1629,7 @@ void UEditorEngine::FinishDestroy()
 		FCoreUObjectDelegates::OnAssetLoaded.RemoveAll(this);
 		FWorldDelegates::LevelAddedToWorld.RemoveAll(this);
 		FWorldDelegates::LevelRemovedFromWorld.RemoveAll(this);
+		FWorldDelegates::OnSeamlessTravelTransition.RemoveAll(this);
 		GetMutableDefault<UEditorStyleSettings>()->OnSettingChanged().RemoveAll(this);
 
 		FAssetRegistryModule* AssetRegistryModule = FModuleManager::GetModulePtr<FAssetRegistryModule>("AssetRegistry");
@@ -5525,6 +5527,15 @@ void UEditorEngine::OnLevelRemovedFromWorld(ULevel* InLevel, UWorld* InWorld)
 	}
 }
 
+void UEditorEngine::OnSeamlessTravelTransition(UWorld* InWorld)
+{
+	if (InWorld->IsPlayInEditor())
+	{
+		EndTransaction();
+		SelectNone(true, true, false);
+	}
+}
+
 void UEditorEngine::UpdateRecentlyLoadedProjectFiles()
 {
 	if (FPaths::IsProjectFilePathSet())
