[WPBuilderCommandlet] Fix logging when processing additional maps

diff --git a/Engine/Source/Editor/UnrealEd/Private/Commandlets/WorldPartitionBuilderCommandlet.cpp b/Engine/Source/Editor/UnrealEd/Private/Commandlets/WorldPartitionBuilderCommandlet.cpp
--- a/Engine/Source/Editor/UnrealEd/Private/Commandlets/WorldPartitionBuilderCommandlet.cpp
+++ b/Engine/Source/Editor/UnrealEd/Private/Commandlets/WorldPartitionBuilderCommandlet.cpp
@@ -203,7 +203,7 @@ int32 UWorldPartitionBuilderCommandlet::Main(const FString& Params)
 
 			if (!RunBuilder(BuilderClass, AdditionalPackageName))
 			{
-				UE_LOG(LogWorldPartitionBuilderCommandlet, Warning, TEXT("Failed to execute %s on additional map %s."), *BuilderClassName, *MapPackageName);
+				UE_LOG(LogWorldPartitionBuilderCommandlet, Warning, TEXT("Failed to execute %s on additional map %s."), *BuilderClassName, *AdditionalPackageName);
 				Result = 1;
 			}
 		}
