[UnrealEd] Added cvar "Editor.AlwaysSnapPercentageBasedScale" (defaults to 1) to allow disabling forced grid snapping for percentage-based scaling, which appears to be an implementation error but is already prolific

diff --git a/Engine/Plugins/Runtime/nDisplay/Source/DisplayClusterConfigurator/Private/Views/Viewport/DisplayClusterConfiguratorSCSEditorViewportClient.cpp b/Engine/Plugins/Runtime/nDisplay/Source/DisplayClusterConfigurator/Private/Views/Viewport/DisplayClusterConfiguratorSCSEditorViewportClient.cpp
--- a/Engine/Plugins/Runtime/nDisplay/Source/DisplayClusterConfigurator/Private/Views/Viewport/DisplayClusterConfiguratorSCSEditorViewportClient.cpp
+++ b/Engine/Plugins/Runtime/nDisplay/Source/DisplayClusterConfigurator/Private/Views/Viewport/DisplayClusterConfiguratorSCSEditorViewportClient.cpp
@@ -35,6 +35,7 @@
 #include "Materials/MaterialInstanceConstant.h"
 #include "PhysicsEngine/PhysicsConstraintComponent.h"
 #include "SKismetInspector.h"
+#include "SnappingUtils.h"
 
 FDisplayClusterConfiguratorSCSEditorViewportClient::FDisplayClusterConfiguratorSCSEditorViewportClient(
 	TWeakPtr<FBlueprintEditor>& InBlueprintEditorPtr, FPreviewScene* InPreviewScene,
@@ -541,10 +542,7 @@ bool FDisplayClusterConfiguratorSCSEditorViewportClient::InputWidgetDelta(FViewp
 				// (mirrored from Level Editor VPC) - we don't scale components when we only have a very small scale change
 				if (!Scale.IsNearlyZero())
 				{
-					if (GEditor->UsePercentageBasedScaling())
-					{
-						ModifiedScale = Scale * ((GEditor->GetScaleGridSize() / 100.0f) / GEditor->GetGridSize());
-					}
+					FSnappingUtils::AdjustPercentageBasedScale(ModifiedScale);
 				}
 				else
 				{
diff --git a/Engine/Source/Editor/Kismet/Private/SCSEditorViewportClient.cpp b/Engine/Source/Editor/Kismet/Private/SCSEditorViewportClient.cpp
--- a/Engine/Source/Editor/Kismet/Private/SCSEditorViewportClient.cpp
+++ b/Engine/Source/Editor/Kismet/Private/SCSEditorViewportClient.cpp
@@ -60,6 +60,7 @@
 #include "ScopedTransaction.h"
 #include "Settings/LevelEditorViewportSettings.h"
 #include "ShowFlags.h"
+#include "SnappingUtils.h"
 #include "StaticMeshResources.h"
 #include "SubobjectData.h"
 #include "SubobjectDataSubsystem.h"
@@ -521,10 +522,7 @@ bool FSCSEditorViewportClient::InputWidgetDelta( FViewport* InViewport, EAxisLis
 				// (mirrored from Level Editor VPC) - we don't scale components when we only have a very small scale change
 				if (!Scale.IsNearlyZero())
 				{
-					if (GEditor->UsePercentageBasedScaling())
-					{
-						ModifiedScale = Scale * ((GEditor->GetScaleGridSize() / 100.0f) / GEditor->GetGridSize());
-					}
+					FSnappingUtils::AdjustPercentageBasedScale(ModifiedScale);
 				}
 				else
 				{
diff --git a/Engine/Source/Editor/PhysicsAssetEditor/Private/PhysicsAssetEditorEditMode.cpp b/Engine/Source/Editor/PhysicsAssetEditor/Private/PhysicsAssetEditorEditMode.cpp
--- a/Engine/Source/Editor/PhysicsAssetEditor/Private/PhysicsAssetEditorEditMode.cpp
+++ b/Engine/Source/Editor/PhysicsAssetEditor/Private/PhysicsAssetEditorEditMode.cpp
@@ -26,6 +26,7 @@
 #include "Engine/Font.h"
 
 #include "SceneView.h"
+#include "SnappingUtils.h"
 
 #define LOCTEXT_NAMESPACE "PhysicsAssetEditorEditMode"
 
@@ -1212,7 +1213,7 @@ void FPhysicsAssetEditorEditMode::ModifyPrimitiveSize(int32 BodyIndex, EAggColli
 		FVector ModifiedSize = DeltaSize;
 		if (GEditor->UsePercentageBasedScaling())
 		{
-			ModifiedSize = DeltaSize * ((GEditor->GetScaleGridSize() / 100.0f) / GEditor->GetGridSize());
+			FSnappingUtils::AdjustPercentageBasedScale(ModifiedSize);
 		}
 
 		AggGeom->ConvexElems[PrimIndex].ScaleElem(ModifiedSize, MinPrimSize);
diff --git a/Engine/Source/Editor/StaticMeshEditor/Private/StaticMeshEditor.cpp b/Engine/Source/Editor/StaticMeshEditor/Private/StaticMeshEditor.cpp
--- a/Engine/Source/Editor/StaticMeshEditor/Private/StaticMeshEditor.cpp
+++ b/Engine/Source/Editor/StaticMeshEditor/Private/StaticMeshEditor.cpp
@@ -63,6 +63,7 @@
 #include "AssetEditorModeManager.h"
 #include "StaticMeshEditorViewportClient.h"
 #include "AdvancedPreviewScene.h"
+#include "SnappingUtils.h"
 
 #define LOCTEXT_NAMESPACE "StaticMeshEditor"
 
@@ -1446,10 +1447,7 @@ void FStaticMeshEditor::ScaleSelectedPrims(const FVector& InScale)
 	FKAggregateGeom* AggGeom = &StaticMesh->GetBodySetup()->AggGeom;
 
 	FVector ModifiedScale = InScale;
-	if (GEditor->UsePercentageBasedScaling())
-	{
-		ModifiedScale = InScale * ((GEditor->GetScaleGridSize() / 100.0f) / GEditor->GetGridSize());
-	}
+	FSnappingUtils::AdjustPercentageBasedScale(ModifiedScale);
 
 	//Multiply in estimated size of the mesh so scaling of sphere, box and sphyl is similar speed to other scaling
 	float SimplePrimitiveScaleSpeedFactor = static_cast<float>( StaticMesh->GetBounds().SphereRadius );
diff --git a/Engine/Source/Editor/UnrealEd/Private/LevelEditorViewport.cpp b/Engine/Source/Editor/UnrealEd/Private/LevelEditorViewport.cpp
--- a/Engine/Source/Editor/UnrealEd/Private/LevelEditorViewport.cpp
+++ b/Engine/Source/Editor/UnrealEd/Private/LevelEditorViewport.cpp
@@ -4528,12 +4528,7 @@ void FLevelEditorViewportClient::ApplyDeltaToSelectedElements(const FTransform&
 
 	{
 		FVector AdjustedScale = ModifiedDeltaTransform.GetScale3D();
-
-		// If we are scaling, we need to change the scaling factor a bit to properly align to grid
-		if (GEditor->UsePercentageBasedScaling() && !AdjustedScale.IsNearlyZero())
-		{
-			AdjustedScale *= ((GEditor->GetScaleGridSize() / 100.0f) / GEditor->GetGridSize());
-		}
+		FSnappingUtils::AdjustPercentageBasedScale(AdjustedScale);
 
 		ModifiedDeltaTransform.SetScale3D(AdjustedScale);
 	}
diff --git a/Engine/Source/Editor/UnrealEd/Private/SnappingUtils.cpp b/Engine/Source/Editor/UnrealEd/Private/SnappingUtils.cpp
--- a/Engine/Source/Editor/UnrealEd/Private/SnappingUtils.cpp
+++ b/Engine/Source/Editor/UnrealEd/Private/SnappingUtils.cpp
@@ -22,6 +22,15 @@
 #include "ViewportSnappingModule.h"
 #include "ActorGroupingUtils.h"
 
+namespace SnappingUtilsLocals
+{
+	static bool bAlwaysSnapPercentageBasedScaleToGrid = true;
+	static FAutoConsoleVariableRef CVarAlwaysSnapPercentageBasedScaleToGrid(
+		TEXT("Editor.AlwaysSnapPercentageBasedScale"),
+		bAlwaysSnapPercentageBasedScaleToGrid,
+		TEXT("When true (default), grid snapping is always used when scaling an object with Percentage based scaling, regardless of snap state."));
+}
+
 //////////////////////////////////////////////////////////////////////////
 // FEditorViewportSnapping
 
@@ -496,6 +505,23 @@ void FSnappingUtils::SnapScale(FVector& Point, const FVector& GridBase)
 	IViewportSnappingModule::GetSnapManager()->SnapScale(Point, GridBase);
 }
 
+void FSnappingUtils::AdjustPercentageBasedScale(FVector& InOutScale)
+{
+	// If we are scaling, we need to change the scaling factor a bit to properly align to grid
+	if (!InOutScale.IsNearlyZero() && GEditor->UsePercentageBasedScaling())
+	{
+		// In both cases, the scale is assumed to be in the 0-100 range and is adjusted to 0-1
+		if (SnappingUtilsLocals::bAlwaysSnapPercentageBasedScaleToGrid)
+		{
+			InOutScale *= ((GEditor->GetScaleGridSize() / 100.0f) / GEditor->GetGridSize());
+		}
+		else
+		{
+			InOutScale /= 100.0f;
+		}
+	}
+}
+
 void FSnappingUtils::SnapPointToGrid(FVector& Point, const FVector& GridBase)
 {
 	IViewportSnappingModule::GetSnapManager()->SnapPointToGrid(Point, GridBase);
diff --git a/Engine/Source/Editor/UnrealEd/Public/SnappingUtils.h b/Engine/Source/Editor/UnrealEd/Public/SnappingUtils.h
--- a/Engine/Source/Editor/UnrealEd/Public/SnappingUtils.h
+++ b/Engine/Source/Editor/UnrealEd/Public/SnappingUtils.h
@@ -96,6 +96,14 @@ public:
 	 */
 	static UNREALED_API void SnapScale( FVector& Point, const FVector& GridBase );
 
+	/**
+	 * Adjusts the given scale value for percentage based scaling, if enabled
+	 * The input should be 0-100, and the output will be adjusted to be 0-1
+	 *
+	 * @param InOutScale	The scale value to adjust
+	 */
+	static UNREALED_API void AdjustPercentageBasedScale( FVector& InOutScale );
+
 	/**
 	 * Snaps a point value to the positional grid
 	 *
