Fixed components being uninitialized when running cook-time validation on worlds

#rnx

diff --git a/Engine/Source/Editor/UnrealEd/Private/CookOnTheFlyServer.cpp b/Engine/Source/Editor/UnrealEd/Private/CookOnTheFlyServer.cpp
--- a/Engine/Source/Editor/UnrealEd/Private/CookOnTheFlyServer.cpp
+++ b/Engine/Source/Editor/UnrealEd/Private/CookOnTheFlyServer.cpp
@@ -3368,6 +3368,7 @@ EDataValidationResult UCookOnTheFlyServer::ValidateSourcePackage(UE::Cook::FPack
 		IVS.CreateNavigation(false);
 
 		World->InitWorld(IVS);
+		World->UpdateWorldComponents(true, false);
 		bRunCleanupWorld = true;
 	}
 
