PR #12295: Add option for MakeEditWidget to turn off local space positioning for scene components

diff --git a/Engine/Source/Editor/UnrealEd/Private/Tools/LegacyEdModeWidgetHelpers.cpp b/Engine/Source/Editor/UnrealEd/Private/Tools/LegacyEdModeWidgetHelpers.cpp
--- a/Engine/Source/Editor/UnrealEd/Private/Tools/LegacyEdModeWidgetHelpers.cpp
+++ b/Engine/Source/Editor/UnrealEd/Private/Tools/LegacyEdModeWidgetHelpers.cpp
@@ -17,6 +17,7 @@
 
 const FName FLegacyEdModeWidgetHelper::MD_MakeEditWidget(TEXT("MakeEditWidget"));
 const FName FLegacyEdModeWidgetHelper::MD_ValidateWidgetUsing(TEXT("ValidateWidgetUsing"));
+const FName FLegacyEdModeWidgetHelper::MD_WidgetLocalSpace(TEXT("LocalSpaceEditWidget"));
 
 /** Hit proxy used for editable properties */
 struct HPropertyWidgetProxy : public HHitProxy
@@ -32,11 +33,15 @@ struct HPropertyWidgetProxy : public HHitProxy
 	/** This property is a transform */
 	bool	bPropertyIsTransform;
 
-	HPropertyWidgetProxy(FString InPropertyName, int32 InPropertyIndex, bool bInPropertyIsTransform)
+	/** This property is in local space */
+	bool	bLocalSpace;
+
+	HPropertyWidgetProxy(FString InPropertyName, int32 InPropertyIndex, bool bInPropertyIsTransform, bool bInLocalSpace = true)
 		: HHitProxy(HPP_Foreground)
 		, PropertyName(InPropertyName)
 		, PropertyIndex(InPropertyIndex)
 		, bPropertyIsTransform(bInPropertyIsTransform)
+		, bLocalSpace(bInLocalSpace)
 	{}
 
 	/** Show cursor as cross when over this handle */
@@ -166,7 +171,13 @@ static bool IsTransformProperty(FProperty* InProp)
 {
 	FStructProperty* StructProp = CastField<FStructProperty>(InProp);
 	return (StructProp != NULL && StructProp->Struct->GetFName() == NAME_Transform);
+}
 
+static bool IsLocalSpace(const FProperty* InProp)
+{
+	return !InProp->HasMetaData(FEdMode::MD_WidgetLocalSpace)
+		|| InProp->GetMetaData(FEdMode::MD_WidgetLocalSpace).IsEmpty()
+		|| InProp->GetBoolMetaData(FEdMode::MD_WidgetLocalSpace);
 }
 
 struct FPropertyWidgetInfoChainElement
@@ -209,13 +220,14 @@ struct FPropertyWidgetInfoChainElement
 		return false;
 	}
 
-	static FLegacyEdModeWidgetHelper::FPropertyWidgetInfo CreateWidgetInfo(const TArray<FPropertyWidgetInfoChainElement>& Chain, bool bIsTransform, FProperty* CurrentProp, int32 Index = INDEX_NONE)
+	static FLegacyEdModeWidgetHelper::FPropertyWidgetInfo CreateWidgetInfo(const TArray<FPropertyWidgetInfoChainElement>& Chain, bool bIsTransform, bool bIsLocalSpace, FProperty* CurrentProp, int32 Index = INDEX_NONE)
 	{
 		check(CurrentProp);
 		FEdMode::FPropertyWidgetInfo WidgetInfo;
 		WidgetInfo.PropertyValidationName = FName(*CurrentProp->GetMetaData(FEdMode::MD_ValidateWidgetUsing));
 		WidgetInfo.bIsTransform = bIsTransform;
 		WidgetInfo.PropertyIndex = Index;
+		WidgetInfo.bLocalSpace = bIsLocalSpace;
 
 		const FString SimplePostFix(TEXT("."));
 		for (int32 ChainIndex = 0; ChainIndex < Chain.Num(); ++ChainIndex)
@@ -257,12 +269,12 @@ struct FPropertyWidgetInfoChainElement
 					const uint32 ArrayDim = ArrayHelper.Num();
 					for (uint32 Index = 0; Index < ArrayDim; ++Index)
 					{
-						OutInfos.Add(FPropertyWidgetInfoChainElement::CreateWidgetInfo(Chain, IsTransformProperty(ArrayProp->Inner), CurrentProp, Index));
+						OutInfos.Add(FPropertyWidgetInfoChainElement::CreateWidgetInfo(Chain, IsTransformProperty(ArrayProp->Inner), IsLocalSpace(ArrayProp->Inner), CurrentProp, Index));
 					}
 				}
 				else
 				{
-					OutInfos.Add(FPropertyWidgetInfoChainElement::CreateWidgetInfo(Chain, IsTransformProperty(CurrentProp), CurrentProp));
+					OutInfos.Add(FPropertyWidgetInfoChainElement::CreateWidgetInfo(Chain, IsTransformProperty(CurrentProp), IsLocalSpace(CurrentProp), CurrentProp));
 				}
 			}
 			else if (FStructProperty* StructProp = CastField<FStructProperty>(CurrentProp))
@@ -333,6 +345,7 @@ void FLegacyEdModeWidgetHelper::FPropertyWidgetInfo::GetTransformAndColor(UObjec
 FLegacyEdModeWidgetHelper::FLegacyEdModeWidgetHelper()
 	: EditedPropertyIndex(INDEX_NONE)
 	, bEditedPropertyIsTransform(false)
+	, bLocalSpace(true)
 	, CurrentWidgetAxis(EAxisList::None)
 	, Owner(nullptr)
 {
@@ -368,19 +381,24 @@ FVector FLegacyEdModeWidgetHelper::GetWidgetLocation() const
 		{
 			if (!EditedPropertyName.IsEmpty())
 			{
-				FVector LocalPos = FVector::ZeroVector;
+				FVector PropertyPosition = FVector::ZeroVector;
 
 				if (bEditedPropertyIsTransform)
 				{
-					FTransform LocalTM = GetPropertyValueByName<FTransform>(BestSelectedItem, EditedPropertyName, EditedPropertyIndex);
-					LocalPos = LocalTM.GetLocation();
+					FTransform PropertyTransform = GetPropertyValueByName<FTransform>(BestSelectedItem, EditedPropertyName, EditedPropertyIndex);
+					PropertyPosition = PropertyTransform.GetLocation();
 				}
 				else
 				{
-					LocalPos = GetPropertyValueByName<FVector>(BestSelectedItem, EditedPropertyName, EditedPropertyIndex);
+					PropertyPosition = GetPropertyValueByName<FVector>(BestSelectedItem, EditedPropertyName, EditedPropertyIndex);
 				}
 
-				const FVector WorldPos = DisplayWidgetToWorld.TransformPosition(LocalPos);
+				if (!bLocalSpace)
+				{
+					return PropertyPosition;
+				}
+				
+				const FVector WorldPos = DisplayWidgetToWorld.TransformPosition(PropertyPosition);
 				return WorldPos;
 			}
 		}
@@ -531,6 +549,7 @@ void FLegacyEdModeWidgetHelper::ActorSelectionChangeNotify()
 	EditedPropertyName = TEXT("");
 	EditedPropertyIndex = INDEX_NONE;
 	bEditedPropertyIsTransform = false;
+	bLocalSpace = true;
 }
 
 bool FLegacyEdModeWidgetHelper::AllowsViewportDragTool() const
@@ -551,44 +570,46 @@ bool FLegacyEdModeWidgetHelper::InputDelta(FEditorViewportClient* InViewportClie
 
 			if (!EditedPropertyName.IsEmpty())
 			{
-				FTransform LocalTM = FTransform::Identity;
-
+				FTransform PropertyTransform = FTransform::Identity;
 				if (bEditedPropertyIsTransform)
 				{
-					LocalTM = GetPropertyValueByName<FTransform>(BestSelectedItem, EditedPropertyName, EditedPropertyIndex);
+					PropertyTransform = GetPropertyValueByName<FTransform>(BestSelectedItem, EditedPropertyName, EditedPropertyIndex);
+
 				}
 				else
 				{
-					FVector LocalPos = GetPropertyValueByName<FVector>(BestSelectedItem, EditedPropertyName, EditedPropertyIndex);
-					LocalTM = FTransform(LocalPos);
+					FVector PropertyTranslation = GetPropertyValueByName<FVector>(BestSelectedItem, EditedPropertyName, EditedPropertyIndex);
+					PropertyTransform = FTransform(PropertyTranslation);
+				}
+
+				// Convert local transform into world space
+				if (bLocalSpace)
+				{
+					PropertyTransform = PropertyTransform * DisplayWidgetToWorld;
 				}
 
-				// Calculate world transform
-				FTransform WorldTM = LocalTM * DisplayWidgetToWorld;
-				// Calc delta specified by drag
-				//FTransform DeltaTM(InRot.Quaternion(), InDrag);
 				// Apply delta in world space
-				WorldTM.SetTranslation(WorldTM.GetTranslation() + InDrag);
-				WorldTM.SetRotation(InRot.Quaternion() * WorldTM.GetRotation());
-				// Convert new world transform back into local space
-				LocalTM = WorldTM.GetRelativeTransform(DisplayWidgetToWorld);
+				PropertyTransform.SetTranslation(PropertyTransform.GetTranslation() + InDrag);
+				PropertyTransform.SetRotation(InRot.Quaternion() * PropertyTransform.GetRotation());
 				// Apply delta scale
-				LocalTM.SetScale3D(LocalTM.GetScale3D() + InScale);
-
-				BestSelectedItem->PreEditChange(NULL);
+				PropertyTransform.SetScale3D(PropertyTransform.GetScale3D() + InScale);
 
+				// Convert world transform back into local space
+				if (bLocalSpace)
+				{
+					PropertyTransform = PropertyTransform.GetRelativeTransform(DisplayWidgetToWorld);
+				}
+				BestSelectedItem->PreEditChange(nullptr);
 				// Property that we actually change
-				FProperty* SetProperty = NULL;
-
+				FProperty* SetProperty = nullptr;
 				if (bEditedPropertyIsTransform)
 				{
-					SetPropertyValueByName<FTransform>(BestSelectedItem, EditedPropertyName, EditedPropertyIndex, LocalTM, SetProperty);
+					SetPropertyValueByName<FTransform>(BestSelectedItem, EditedPropertyName, EditedPropertyIndex, PropertyTransform, SetProperty);
 				}
 				else
 				{
-					SetPropertyValueByName<FVector>(BestSelectedItem, EditedPropertyName, EditedPropertyIndex, LocalTM.GetLocation(), SetProperty);
+					SetPropertyValueByName<FVector>(BestSelectedItem, EditedPropertyName, EditedPropertyIndex, PropertyTransform.GetLocation(), SetProperty);
 				}
-
 				FPropertyChangedEvent PropertyChangeEvent(SetProperty);
 				BestSelectedItem->PostEditChangeProperty(PropertyChangeEvent);
 
@@ -610,6 +631,7 @@ bool FLegacyEdModeWidgetHelper::HandleClick(FEditorViewportClient* InViewportCli
 			EditedPropertyName = PropertyProxy->PropertyName;
 			EditedPropertyIndex = PropertyProxy->PropertyIndex;
 			bEditedPropertyIsTransform = PropertyProxy->bPropertyIsTransform;
+			bLocalSpace = PropertyProxy->bLocalSpace;
 			return true;
 		}
 		// Left clicking on an actor, stop editing a property
@@ -618,6 +640,7 @@ bool FLegacyEdModeWidgetHelper::HandleClick(FEditorViewportClient* InViewportCli
 			EditedPropertyName = FString();
 			EditedPropertyIndex = INDEX_NONE;
 			bEditedPropertyIsTransform = false;
+			bLocalSpace = true;
 		}
 	}
 
@@ -650,15 +673,19 @@ void FLegacyEdModeWidgetHelper::Render(const FSceneView* View, FViewport* Viewpo
 				FString ValidationMessage;
 				FColor WidgetColor;
 				WidgetInfo.GetTransformAndColor(BestSelectedItem, bSelected, /*out*/ LocalWidgetTransform, /*out*/ ValidationMessage, /*out*/ WidgetColor);
-
-				const FTransform WorldWidgetTransform = LocalWidgetTransform * DisplayWidgetToWorld;
+				
+				FTransform WorldWidgetTransform = LocalWidgetTransform;
+				if (WidgetInfo.bLocalSpace)
+				{
+					WorldWidgetTransform = LocalWidgetTransform * DisplayWidgetToWorld;
+				}
 				const FMatrix WidgetTM(WorldWidgetTransform.ToMatrixWithScale());
 
 				const double WidgetSize = 0.035;
 				const double ZoomFactor = FMath::Min(View->ViewMatrices.GetProjectionMatrix().M[0][0], View->ViewMatrices.GetProjectionMatrix().M[1][1]);
 				const float WidgetRadius = static_cast<float>(View->Project(WorldWidgetTransform.GetTranslation()).W * (WidgetSize / ZoomFactor));
 
-				if (bHitTesting) PDI->SetHitProxy(new HPropertyWidgetProxy(WidgetInfo.PropertyName, WidgetInfo.PropertyIndex, WidgetInfo.bIsTransform));
+				if (bHitTesting) PDI->SetHitProxy(new HPropertyWidgetProxy(WidgetInfo.PropertyName, WidgetInfo.PropertyIndex, WidgetInfo.bIsTransform, WidgetInfo.bLocalSpace));
 				DrawWireDiamond(PDI, WidgetTM, WidgetRadius, WidgetColor, SDPG_Foreground);
 				if (bAllowEditWidgetAxisDisplay && WidgetInfo.bIsTransform)
 				{
@@ -700,8 +727,12 @@ void FLegacyEdModeWidgetHelper::DrawHUD(FEditorViewportClient* ViewportClient, F
 				FColor IgnoredWidgetColor;
 				WidgetInfo.GetTransformAndColor(BestSelectedItem, /*bSelected=*/ false, /*out*/ LocalWidgetTransform, /*out*/ ValidationMessage, /*out*/ IgnoredWidgetColor);
 
-				const FTransform WorldWidgetTransform = LocalWidgetTransform * DisplayWidgetToWorld;
-
+				FTransform WorldWidgetTransform = LocalWidgetTransform;
+				if (WidgetInfo.bLocalSpace)
+				{
+					WorldWidgetTransform = LocalWidgetTransform * DisplayWidgetToWorld;
+				}
+				
 				const FPlane Proj = View->Project(WorldWidgetTransform.GetTranslation());
 				if (Proj.W > 0.f)
 				{
diff --git a/Engine/Source/Editor/UnrealEd/Public/Tools/LegacyEdModeWidgetHelpers.h b/Engine/Source/Editor/UnrealEd/Public/Tools/LegacyEdModeWidgetHelpers.h
--- a/Engine/Source/Editor/UnrealEd/Public/Tools/LegacyEdModeWidgetHelpers.h
+++ b/Engine/Source/Editor/UnrealEd/Public/Tools/LegacyEdModeWidgetHelpers.h
@@ -59,13 +59,18 @@ public:
 		FName PropertyValidationName;
 		FString DisplayName;
 		bool bIsTransform;
+		bool bLocalSpace;
 
 		FPropertyWidgetInfo()
 			: PropertyIndex(INDEX_NONE)
 			, bIsTransform(false)
+			, bLocalSpace(true)
 		{
 		}
 
+		/**
+		 * Returns the appropriate color, and transform of the widget in local space
+		 */
 		void GetTransformAndColor(UObject* BestSelectedItem, bool bIsSelected, FTransform& OutLocalTransform, FString& OutValidationMessage, FColor& OutDrawColor) const;
 	};
 
@@ -79,6 +84,8 @@ public:
 	static UNREALED_API const FName MD_MakeEditWidget;
 	/** Specifies a function used for validation of the current value of a property.  The function returns a string that is empty if the value is valid, or contains an error description if the value is invalid */
 	static UNREALED_API const FName MD_ValidateWidgetUsing;
+	/** If the value of a UPROPERTY's edit widget's transform or translation is stored on the property in local space (defaults true) */
+	static UNREALED_API const FName MD_WidgetLocalSpace;
 	/** Returns true if this structure can support creating a widget in the editor */
 	static UNREALED_API bool CanCreateWidgetForStructure(const UStruct* InPropStruct);
 	/** Returns true if this property can support creating a widget in the editor */
@@ -103,16 +110,19 @@ protected:
 
 	/** Finds the best item to display widgets for (preferring selected components over selected actors) */
 	UNREALED_API virtual UObject* GetItemToTryDisplayingWidgetsFor(FTransform& OutWidgetToWorld) const;
-
+	
 	/** Name of the property currently being edited */
 	FString EditedPropertyName;
 
 	/** If the property being edited is an array property, this is the index of the element we're currently dealing with */
 	int32 EditedPropertyIndex;
 
-	/** Indicates  */
+	/** Indicates the property being edited is a transform */
 	bool bEditedPropertyIsTransform;
 
+	/** World or local space */
+	bool bLocalSpace;
+
 	/** The current axis that is being dragged on the widget. */
 	EAxisList::Type CurrentWidgetAxis;
 
