[ITF Viewport Interactions] Modular Viewport Interactions

First batch of ITF-based interactions for camera movements and viewport inputs.
Interactions combine Behaviors and Behavior Targets, as part of the effort to convert all legacy viewport inputs to ITF-based solutions.

Each interactions handles a specific behavior, e.g.: Orbiting, Panning, Movement, Zooming, etc
The idea is to have a list of "interactions" enabled by each editor, based on its needs.

A "default camera movements interactions" set is available including standard behaviors, replicating the legacy level editor ones (`UE::Editor::ViewportInteractions::AddDefaultCameraMovementInteractions`).
This can be used to quickly add a standard set of inputs to a viewport, but an editor could specify even just a single interaction.

Interactions are handled by a `UViewportInteractionsBehaviorSource` instance, hosted by `UModeManagerInteractiveToolsContext`.

List of viewport interactions added by this CL:

- Zoom
- View Angle
- Panning
- Ortho Panning
- Orbit
- Move + Yaw
- FOV change
- Camera Translate
- Camera Rotate
- Camera Speed update with mouse wheel
- Commands Interaction (to handle viewport FUICommands via ITF)

Editor Viewports can specify their interactions by overriding the new methods:
- `OnRegisterViewportInteractions`
- `GetModeManagerInteractiveToolsContext`

There are a few changes/addition to ITF, including:
- `UMouseWheelModifierInputBehavior`: treats mouse wheel events as modifiers - input will not be captured
- `UMouseButtonAsModifierInputBehavior`: same as the above - will just read mouse button changes, without capturing the input

These interactions can be enabled with CVar `ViewportInteractions.EnableITFInteractions 1`.
This CL does not add interactions to any editor. A future CL will add the default set to Level Editor.

diff --git a/Engine/Source/Editor/Experimental/EditorInteractiveToolsFramework/Private/EditorDragTools/DragToolsBehaviourSource.cpp b/Engine/Source/Editor/Experimental/EditorInteractiveToolsFramework/Private/EditorDragTools/DragToolsBehaviourSource.cpp
--- a/Engine/Source/Editor/Experimental/EditorInteractiveToolsFramework/Private/EditorDragTools/DragToolsBehaviourSource.cpp
+++ b/Engine/Source/Editor/Experimental/EditorInteractiveToolsFramework/Private/EditorDragTools/DragToolsBehaviourSource.cpp
@@ -17,20 +17,6 @@
 
 namespace UE::Editor::DragTools
 {
-	// In case New TRS Gizmos are enabled, we need Mouse Drag + Key modifier Drag Tools Behaviors to be processed before
-	// the Gizmo ones, otherwise Drag Behaviors will never be triggered. See:
-	// * UTransformGizmo::SetupBehaviors()			[DEFAULT_GIZMO_PRIORITY]
-	// * UTransformGizmo::SetupIndirectBehaviors()	[DEFAULT_GIZMO_PRIORITY - 1]
-	static constexpr int HANDLED_BEFORE_GIZMO_PRIORITY = FInputCapturePriority::DEFAULT_GIZMO_PRIORITY - 3;
-
-	// Conversely, tools which are activated with just mouse buttons input (no key modifiers, mouse drag behavior only) need to be processed after Gizmo Behaviors
-	static constexpr int HANDLED_AFTER_GIZMO_PRIORITY = FInputCapturePriority::DEFAULT_GIZMO_PRIORITY + 1;
-
-	static constexpr int FRUSTUM_SELECT_PRIORITY = HANDLED_BEFORE_GIZMO_PRIORITY;
-	static constexpr int VIEWPORT_CHANGE_PRIORITY = HANDLED_BEFORE_GIZMO_PRIORITY;
-	static constexpr int BOX_SELECT_PRIORITY = HANDLED_AFTER_GIZMO_PRIORITY;
-	static constexpr int MEASURE_PRIORITY = VIEWPORT_CHANGE_PRIORITY + 1;
-
 	// Used to group Viewport Change behaviors inside the BehaviorSet and remove them if needed
 	static const FString ViewportChangeBehaviorGroup = TEXT("ViewportChange");
 
@@ -58,6 +44,9 @@ namespace UE::Editor::DragTools
 
 void UDragToolsBehaviorSource::Initialize(const UEditorInteractiveToolsContext* InInteractiveToolsContext)
 {
+	// Used to reach priority settings
+	using namespace UE::Editor::DragTools;
+
 	if (!InInteractiveToolsContext)
 	{
 		return;
@@ -84,7 +73,7 @@ void UDragToolsBehaviorSource::Initialize(const UEditorInteractiveToolsContext*
 		// Add drag input behavior
 		UClickDragInputBehavior* MouseBehavior = NewObject<UClickDragInputBehavior>();
 		MouseBehavior->Initialize(BoxSelectTool.Get());
-		MouseBehavior->SetDefaultPriority(UE::Editor::DragTools::BOX_SELECT_PRIORITY);
+		MouseBehavior->SetDefaultPriority(BOX_SELECT_PRIORITY);
 		BehaviorSet->Add(MouseBehavior);
 
 		// Register modifier listeners
@@ -98,7 +87,7 @@ void UDragToolsBehaviorSource::Initialize(const UEditorInteractiveToolsContext*
 		KeyModifierInputBehavior->Initialize(
 			BoxSelectTool.Get(), UE::Editor::DragTools::AltKeyMod, FInputDeviceState::IsAltKeyDown
 		);
-		KeyModifierInputBehavior->SetDefaultPriority(UE::Editor::DragTools::BOX_SELECT_PRIORITY);
+		KeyModifierInputBehavior->SetDefaultPriority(BOX_SELECT_PRIORITY);
 		BehaviorSet->Add(KeyModifierInputBehavior);
 
 		BoxSelectTool->OnActivateTool().AddUObject(this, &UDragToolsBehaviorSource::OnActivateTool, EDragTool::BoxSelect);
@@ -115,7 +104,7 @@ void UDragToolsBehaviorSource::Initialize(const UEditorInteractiveToolsContext*
 		// Add drag input behavior
 		UClickDragInputBehavior* MouseBehavior = NewObject<UClickDragInputBehavior>();
 		MouseBehavior->Initialize(FrustumSelectTool.Get());
-		MouseBehavior->SetDefaultPriority(UE::Editor::DragTools::FRUSTUM_SELECT_PRIORITY);
+		MouseBehavior->SetDefaultPriority(FRUSTUM_SELECT_PRIORITY);
 		MouseBehavior->SetUseLeftMouseButton();
 		BehaviorSet->Add(MouseBehavior);
 
@@ -130,7 +119,7 @@ void UDragToolsBehaviorSource::Initialize(const UEditorInteractiveToolsContext*
 		KeyModifierInputBehavior->Initialize(
 			FrustumSelectTool.Get(), UE::Editor::DragTools::AltKeyMod, FInputDeviceState::IsAltKeyDown
 		);
-		KeyModifierInputBehavior->SetDefaultPriority(UE::Editor::DragTools::FRUSTUM_SELECT_PRIORITY);
+		KeyModifierInputBehavior->SetDefaultPriority(FRUSTUM_SELECT_PRIORITY);
 		BehaviorSet->Add(KeyModifierInputBehavior);
 
 		FrustumSelectTool->OnActivateTool().AddUObject(
@@ -150,7 +139,7 @@ void UDragToolsBehaviorSource::Initialize(const UEditorInteractiveToolsContext*
 		// Add drag input behavior
 		UClickDragInputBehavior* MouseBehavior = NewObject<UClickDragInputBehavior>();
 		MouseBehavior->Initialize(MeasureTool.Get());
-		MouseBehavior->SetDefaultPriority(UE::Editor::DragTools::MEASURE_PRIORITY);
+		MouseBehavior->SetDefaultPriority(MEASURE_PRIORITY);
 		MouseBehavior->SetUseMiddleMouseButton();
 		BehaviorSet->Add(MouseBehavior);
 
@@ -176,7 +165,7 @@ void UDragToolsBehaviorSource::Initialize(const UEditorInteractiveToolsContext*
 		KeyModifierInputBehavior->Initialize(
 			DuplicateDragSelection.Get(), UE::Editor::DragTools::AltKeyMod, FInputDeviceState::IsAltKeyDown
 		);
-		KeyModifierInputBehavior->SetDefaultPriority(UE::Editor::DragTools::HANDLED_BEFORE_GIZMO_PRIORITY);
+		KeyModifierInputBehavior->SetDefaultPriority(DUPLICATE_SELECTION_PRIORITY);
 
 		BehaviorSet->Add(KeyModifierInputBehavior);
 	}
@@ -187,7 +176,7 @@ void UDragToolsBehaviorSource::Initialize(const UEditorInteractiveToolsContext*
 
 		UKeyAsModifierInputBehavior* KeyModifierInputBehavior = NewObject<UKeyAsModifierInputBehavior>();
 		KeyModifierInputBehavior->Initialize(MoveCameraWithObject.Get(), UE::Editor::DragTools::ShiftKeyMod, FInputDeviceState::IsShiftKeyDown);
-		KeyModifierInputBehavior->SetDefaultPriority(UE::Editor::DragTools::HANDLED_BEFORE_GIZMO_PRIORITY);
+		KeyModifierInputBehavior->SetDefaultPriority(HANDLED_BEFORE_VIEWPORT_INTERACTIONS);
 
 		BehaviorSet->Add(KeyModifierInputBehavior);
 	}
@@ -304,6 +293,9 @@ void UDragToolsBehaviorSource::RenderTools(const FSceneView* InSceneView, FCanva
 
 void UDragToolsBehaviorSource::ActivateViewportChangeTool()
 {
+	// Used to reach priority settings
+	using namespace UE::Editor::DragTools;
+
 	// If tool is already active, ignore
 	if (DragTools.Contains(EDragTool::ViewportChange))
 	{
@@ -333,7 +325,7 @@ void UDragToolsBehaviorSource::ActivateViewportChangeTool()
 			// Add drag input behavior
 			UClickDragInputBehavior* MouseBehavior = NewObject<UClickDragInputBehavior>();
 			MouseBehavior->Initialize(ViewportChangeTool.Get());
-			MouseBehavior->SetDefaultPriority(UE::Editor::DragTools::VIEWPORT_CHANGE_PRIORITY);
+			MouseBehavior->SetDefaultPriority(VIEWPORT_CHANGE_PRIORITY);
 			MouseBehavior->SetUseMiddleMouseButton();
 			BehaviorSet->Add(MouseBehavior, this, UE::Editor::DragTools::ViewportChangeBehaviorGroup);
 
@@ -342,7 +334,7 @@ void UDragToolsBehaviorSource::ActivateViewportChangeTool()
 			KeyModifierInputBehavior->Initialize(
 				ViewportChangeTool.Get(), UE::Editor::DragTools::CtrlKeyMod, FInputDeviceState::IsCtrlKeyDown
 			);
-			KeyModifierInputBehavior->SetDefaultPriority(UE::Editor::DragTools::VIEWPORT_CHANGE_PRIORITY);
+			KeyModifierInputBehavior->SetDefaultPriority(VIEWPORT_CHANGE_PRIORITY);
 			BehaviorSet->Add(KeyModifierInputBehavior, this, UE::Editor::DragTools::ViewportChangeBehaviorGroup);
 
 			ViewportChangeTool->OnActivateTool().AddUObject(
diff --git a/Engine/Source/Editor/Experimental/EditorInteractiveToolsFramework/Private/ViewportInteractions/ViewportCameraRotateInteraction.cpp b/Engine/Source/Editor/Experimental/EditorInteractiveToolsFramework/Private/ViewportInteractions/ViewportCameraRotateInteraction.cpp
new file mode 100644
--- /dev/null
+++ b/Engine/Source/Editor/Experimental/EditorInteractiveToolsFramework/Private/ViewportInteractions/ViewportCameraRotateInteraction.cpp
@@ -0,0 +1,122 @@
+// Copyright Epic Games, Inc. All Rights Reserved.
+
+#include "ViewportInteractions/ViewportCameraRotateInteraction.h"
+#include "BaseBehaviors/KeyInputBehavior.h"
+#include "CameraController.h"
+#include "EditorViewportClient.h"
+#include "ViewportClientNavigationHelper.h"
+#include "ViewportInteractions/ViewportInteractionsBehaviorSource.h"
+
+UViewportCameraRotateInteraction::UViewportCameraRotateInteraction()
+{
+	InteractionName = TEXT("Camera Rotate");
+	ToolType = UE::Editor::ViewportInteractions::CameraRotate;
+
+	UKeyInputBehavior* KeyInputBehavior = NewObject<UKeyInputBehavior>();
+	KeyInputBehavior->Initialize(this, GetKeys());
+	KeyInputBehavior->SetDefaultPriority(UE::Editor::ViewportInteractions::VIEWPORT_INTERACTIONS_DEFAULT_PRIORITY);
+	KeyInputBehavior->bRequireAllKeys = false;
+	KeyInputBehaviorWeak = KeyInputBehavior;
+
+	RegisterInputBehavior(KeyInputBehavior);
+}
+
+void UViewportCameraRotateInteraction::OnKeyPressed(const FKey& InKeyID)
+{
+	if (IsEnabled())
+	{
+		constexpr bool bIsPressed = true;
+		UpdateKeyState(InKeyID, bIsPressed);
+	}
+}
+
+void UViewportCameraRotateInteraction::OnKeyReleased(const FKey& InKeyID)
+{
+	if (IsEnabled())
+	{
+		constexpr bool bIsPressed = false;
+		UpdateKeyState(InKeyID, bIsPressed);
+	}
+}
+
+void UViewportCameraRotateInteraction::Tick(float InDeltaTime) const
+{
+	if (RotateYawImpulse != 0 || RotatePitchImpulse != 0)
+	{
+		if (FEditorViewportClient* EditorViewportClient = GetEditorViewportClient())
+		{
+			if (FViewportClientNavigationHelper* NavigationHelper = EditorViewportClient->GetViewportNavigationHelper())
+			{
+				NavigationHelper->ImpulseDataDelta.RotateYawImpulse += RotateYawImpulse;
+				NavigationHelper->ImpulseDataDelta.RotatePitchImpulse += RotatePitchImpulse;
+			}
+		}
+	}
+}
+
+void UViewportCameraRotateInteraction::UpdateKeyState(const FKey& InKeyID, bool bInIsPressed)
+{
+	using namespace UE::Editor::ViewportInteractions;
+
+	constexpr float ImpulseValue = 1.0f;
+
+	const bool bMouseButtonDownOrRelease = IsAnyMouseButtonDown() || !bInIsPressed;
+
+	// Rotate Up/Down
+	if (CommandMatchesKey(FViewportNavigationCommands::Get().RotateUp, InKeyID) && bMouseButtonDownOrRelease)
+	{
+		RotatePitchImpulse = bInIsPressed ? ImpulseValue : 0.0f;
+	}
+	else if (CommandMatchesKey(FViewportNavigationCommands::Get().RotateDown, InKeyID) && bMouseButtonDownOrRelease)
+	{
+		RotatePitchImpulse = bInIsPressed ? -ImpulseValue : 0.0f;
+	}
+	// Rotate Left/Right
+	else if (CommandMatchesKey(FViewportNavigationCommands::Get().RotateLeft, InKeyID) && bMouseButtonDownOrRelease)
+	{
+		RotateYawImpulse = bInIsPressed ? -ImpulseValue : 0.0f;
+	}
+	else if (CommandMatchesKey(FViewportNavigationCommands::Get().RotateRight, InKeyID) && bMouseButtonDownOrRelease)
+	{
+		RotateYawImpulse = bInIsPressed ? ImpulseValue : 0.0f;
+	}
+}
+
+TArray<FKey> UViewportCameraRotateInteraction::GetKeys() const
+{
+	TArray<FKey> Keys;
+	for (int32 i = 0; i < static_cast<uint8>(EMultipleKeyBindingIndex::NumChords); ++i)
+	{
+		EMultipleKeyBindingIndex ChordIndex = static_cast<EMultipleKeyBindingIndex>(i);
+
+		TArray<FKey> KeysFromChords;
+		for (const TSharedPtr<FUICommandInfo>& Command : GetCommands())
+		{
+			KeysFromChords.Add(Command->GetActiveChord(ChordIndex)->Key);
+		}
+
+		Keys.Append(KeysFromChords);
+	}
+
+	return Keys;
+}
+
+void UViewportCameraRotateInteraction::OnCommandChordChanged()
+{
+	if (TStrongObjectPtr<UKeyInputBehavior> KeyInputBehaviorPinned = KeyInputBehaviorWeak.Pin())
+	{
+		KeyInputBehaviorPinned->Initialize(this, GetKeys());
+	}
+}
+
+TArray<TSharedPtr<FUICommandInfo>> UViewportCameraRotateInteraction::GetCommands() const
+{
+	if (FViewportNavigationCommands::IsRegistered())
+	{
+		return { FViewportNavigationCommands::Get().RotateUp,
+				 FViewportNavigationCommands::Get().RotateDown,
+				 FViewportNavigationCommands::Get().RotateLeft,
+				 FViewportNavigationCommands::Get().RotateRight };
+	}
+	return {};
+}
diff --git a/Engine/Source/Editor/Experimental/EditorInteractiveToolsFramework/Private/ViewportInteractions/ViewportCameraSpeedMouseWheelInteraction.cpp b/Engine/Source/Editor/Experimental/EditorInteractiveToolsFramework/Private/ViewportInteractions/ViewportCameraSpeedMouseWheelInteraction.cpp
new file mode 100644
--- /dev/null
+++ b/Engine/Source/Editor/Experimental/EditorInteractiveToolsFramework/Private/ViewportInteractions/ViewportCameraSpeedMouseWheelInteraction.cpp
@@ -0,0 +1,77 @@
+﻿// Copyright Epic Games, Inc. All Rights Reserved.
+
+#include "ViewportInteractions/ViewportCameraSpeedMouseWheelInteraction.h"
+
+#include "BaseBehaviors/MouseWheelModifierInputBehavior.h"
+#include "ViewportInteractions/ViewportInteractionsBehaviorSource.h"
+#include "EditorViewportClient.h"
+
+UViewportCameraSpeedMouseWheelInteraction::UViewportCameraSpeedMouseWheelInteraction()
+{
+	ToolType = UE::Editor::ViewportInteractions::CameraTranslate;
+	InteractionName = TEXT("Camera Speed");
+
+	// Set camera speed using mouse wheel
+	UMouseWheelInputBehavior* MouseWheelInputBehavior = NewObject<UMouseWheelModifierInputBehavior>();
+	MouseWheelInputBehavior->Initialize(this);
+	MouseWheelInputBehavior->SetDefaultPriority(UE::Editor::ViewportInteractions::VIEWPORT_INTERACTIONS_HIGH_PRIORITY);
+
+	RegisterInputBehavior(MouseWheelInputBehavior);
+}
+
+FInputRayHit UViewportCameraSpeedMouseWheelInteraction::ShouldRespondToMouseWheel(const FInputDeviceRay& CurrentPos)
+{
+	return IsAnyMouseButtonDown() ? FInputRayHit(TNumericLimits<float>::Max()) : FInputRayHit();
+}
+
+void UViewportCameraSpeedMouseWheelInteraction::OnMouseWheelScrollUp(const FInputDeviceRay& CurrentPos)
+{
+	if (!IsAnyMouseButtonDown())
+	{
+		return;
+	}
+
+	if (const FEditorViewportClient* const EditorViewportClient = GetEditorViewportClient())
+	{
+		if (EditorViewportClient->IsOrtho())
+		{
+			return;
+		}
+	}
+
+	UE::Editor::ViewportInteractions::LOG(TEXT("INCREASE CAMERA SPEED"));
+	UpdateCameraSpeed(0.1f);
+}
+
+void UViewportCameraSpeedMouseWheelInteraction::OnMouseWheelScrollDown(const FInputDeviceRay& CurrentPos)
+{
+	if (!IsAnyMouseButtonDown())
+	{
+		return;
+	}
+
+	if (const FEditorViewportClient* const EditorViewportClient = GetEditorViewportClient())
+	{
+		if (EditorViewportClient->IsOrtho())
+		{
+			return;
+		}
+	}
+
+	UE::Editor::ViewportInteractions::LOG(TEXT("DECREASE CAMERA SPEED"));
+	UpdateCameraSpeed(-0.1f);
+}
+
+void UViewportCameraSpeedMouseWheelInteraction::UpdateCameraSpeed(float InUpdateFactor) const
+{
+	if (FEditorViewportClient* EditorViewportClient = GetEditorViewportClient())
+	{
+		FEditorViewportCameraSpeedSettings SpeedSettings = EditorViewportClient->GetCameraSpeedSettings();
+		const float Speed = SpeedSettings.GetCurrentSpeed();
+
+		float Delta = Speed * InUpdateFactor;
+
+		SpeedSettings.SetCurrentSpeed(Speed + Delta);
+		EditorViewportClient->SetCameraSpeedSettings(SpeedSettings);
+	}
+}
diff --git a/Engine/Source/Editor/Experimental/EditorInteractiveToolsFramework/Private/ViewportInteractions/ViewportCameraTranslateInteraction.cpp b/Engine/Source/Editor/Experimental/EditorInteractiveToolsFramework/Private/ViewportInteractions/ViewportCameraTranslateInteraction.cpp
new file mode 100644
--- /dev/null
+++ b/Engine/Source/Editor/Experimental/EditorInteractiveToolsFramework/Private/ViewportInteractions/ViewportCameraTranslateInteraction.cpp
@@ -0,0 +1,176 @@
+// Copyright Epic Games, Inc. All Rights Reserved.
+
+#include "ViewportInteractions/ViewportCameraTranslateInteraction.h"
+#include "BaseBehaviors/KeyInputBehavior.h"
+#include "CameraController.h"
+#include "EditorViewportClient.h"
+#include "ViewportClientNavigationHelper.h"
+#include "ViewportInteractions/ViewportInteractionsBehaviorSource.h"
+
+UViewportCameraTranslateInteraction::UViewportCameraTranslateInteraction()
+{
+	ToolType = UE::Editor::ViewportInteractions::CameraTranslate;
+	InteractionName = TEXT("Camera Translate");
+
+	// Camera translate
+	UKeyInputBehavior* KeyInputBehavior = NewObject<UKeyInputBehavior>();
+	KeyInputBehavior->Initialize(this, GetKeys());
+	KeyInputBehavior->SetDefaultPriority(UE::Editor::ViewportInteractions::VIEWPORT_INTERACTIONS_DEFAULT_PRIORITY);
+	KeyInputBehavior->bRequireAllKeys = false;
+	KeyInputBehaviorWeak = KeyInputBehavior;
+
+	RegisterInputBehavior(KeyInputBehavior);
+}
+
+void UViewportCameraTranslateInteraction::OnKeyPressed(const FKey& InKeyID)
+{
+	if (IsEnabled())
+	{
+		constexpr bool bIsPressed = true;
+		UpdateKeyState(InKeyID, bIsPressed);
+	}
+}
+
+void UViewportCameraTranslateInteraction::OnKeyReleased(const FKey& InKeyID)
+{
+	if (IsEnabled())
+	{
+		constexpr bool bIsPressed = false;
+		UpdateKeyState(InKeyID, bIsPressed);
+	}
+}
+void UViewportCameraTranslateInteraction::OnForceEndCapture()
+{
+	ForwardBackwardImpulse = 0.0f;
+	RightLeftImpulse = 0.0f;
+	WorldUpDownImpulse = 0.0f;
+	LocalUpDownImpulse = 0.0f;
+}
+
+void UViewportCameraTranslateInteraction::Tick(float InDeltaTime) const
+{
+	if (ForwardBackwardImpulse != 0 || RightLeftImpulse != 0 || WorldUpDownImpulse != 0 || LocalUpDownImpulse != 0)
+	{
+		if (FEditorViewportClient* EditorViewportClient = GetEditorViewportClient())
+		{
+			if (FViewportClientNavigationHelper* NavigationHelper = EditorViewportClient->GetViewportNavigationHelper())
+			{
+				FCameraControllerUserImpulseData& ImpulseDataDelta = NavigationHelper->ImpulseDataDelta;
+				ImpulseDataDelta.MoveForwardBackwardImpulse += ForwardBackwardImpulse;
+				ImpulseDataDelta.MoveRightLeftImpulse += RightLeftImpulse;
+				ImpulseDataDelta.MoveWorldUpDownImpulse += WorldUpDownImpulse;
+				ImpulseDataDelta.MoveLocalUpDownImpulse += LocalUpDownImpulse;
+			}
+		}
+	}
+}
+
+void UViewportCameraTranslateInteraction::UpdateKeyState(const FKey& InKeyID, bool bInIsPressed)
+{
+	using namespace UE::Editor::ViewportInteractions;
+
+	constexpr float ImpulseValue = 1.0f;
+
+	const bool bMouseButtonDownOrRelease = IsAnyMouseButtonDown() || !bInIsPressed;
+
+	// Forward/Backward
+	if ((CommandMatchesKey(FViewportNavigationCommands::Get().Forward, InKeyID) && bMouseButtonDownOrRelease)
+		|| InKeyID == EKeys::NumPadEight || InKeyID == EKeys::Up)
+	{
+		ForwardBackwardImpulse = bInIsPressed ? ImpulseValue : 0.0f;
+	}
+	else if ((CommandMatchesKey(FViewportNavigationCommands::Get().Backward, InKeyID) && bMouseButtonDownOrRelease)
+			 || InKeyID == EKeys::NumPadTwo || InKeyID == EKeys::Down)
+	{
+		ForwardBackwardImpulse = bInIsPressed ? -ImpulseValue : 0.0f;
+	}
+	// Right/Left
+	else if ((CommandMatchesKey(FViewportNavigationCommands::Get().Right, InKeyID) && bMouseButtonDownOrRelease)
+			 || InKeyID == EKeys::NumPadSix || InKeyID == EKeys::Right)
+	{
+		RightLeftImpulse = bInIsPressed ? ImpulseValue : 0.0f;
+	}
+	else if ((CommandMatchesKey(FViewportNavigationCommands::Get().Left, InKeyID) && bMouseButtonDownOrRelease)
+			 || InKeyID == EKeys::NumPadFour || InKeyID == EKeys::Left)
+	{
+		RightLeftImpulse = bInIsPressed ? -ImpulseValue : 0.0f;
+	}
+	// World Up/Down
+	else if ((CommandMatchesKey(FViewportNavigationCommands::Get().WorldUp, InKeyID) && bMouseButtonDownOrRelease)
+			 || InKeyID == EKeys::NumPadNine)
+	{
+		WorldUpDownImpulse = bInIsPressed ? ImpulseValue : 0.0f;
+	}
+	else if ((CommandMatchesKey(FViewportNavigationCommands::Get().WorldDown, InKeyID) && bMouseButtonDownOrRelease)
+			 || InKeyID == EKeys::NumPadSeven)
+	{
+		WorldUpDownImpulse = bInIsPressed ? -ImpulseValue : 0.0f;
+	}
+	// Local Up/Down
+	else if (CommandMatchesKey(FViewportNavigationCommands::Get().LocalUp, InKeyID) && bMouseButtonDownOrRelease)
+	{
+		LocalUpDownImpulse = bInIsPressed ? ImpulseValue : 0.0f;
+	}
+	else if (CommandMatchesKey(FViewportNavigationCommands::Get().LocalDown, InKeyID) && bMouseButtonDownOrRelease)
+	{
+		LocalUpDownImpulse = bInIsPressed ? -ImpulseValue : 0.0f;
+	}
+}
+
+TArray<FKey> UViewportCameraTranslateInteraction::GetKeys() const
+{
+	TArray<FKey> Keys = { EKeys::Add, EKeys::Subtract };
+
+	// Key Mappings (Command Info based)
+	// Iterate through all key mappings to generate key state flags
+	for (int32 i = 0; i < static_cast<uint8>(EMultipleKeyBindingIndex::NumChords); ++i)
+	{
+		EMultipleKeyBindingIndex ChordIndex = static_cast<EMultipleKeyBindingIndex>(i);
+		TArray<FKey> KeysFromChords;
+
+		// Retrieve command based keys
+		for (const TSharedPtr<FUICommandInfo>& Command : GetCommands())
+		{
+			KeysFromChords.Add(Command->GetActiveChord(ChordIndex)->Key);
+		}
+
+		Keys.Append(KeysFromChords);
+	}
+
+	// Hardcoded numpad keys
+	if (bUseNumpadKey)
+	{
+		TArray<FKey> OtherKeys = { EKeys::NumPadTwo,   EKeys::NumPadFour,  EKeys::NumPadSix,
+								   EKeys::NumPadSeven, EKeys::NumPadEight, EKeys::NumPadNine };
+
+		Keys.Append(OtherKeys);
+	}
+
+	// Hardcoded
+	Keys.Append({ EKeys::Up, EKeys::Down, EKeys::Right, EKeys::Left });
+
+	return Keys;
+}
+
+void UViewportCameraTranslateInteraction::OnCommandChordChanged()
+{
+	if (TStrongObjectPtr<UKeyInputBehavior> KeyInputBehaviorPinned = KeyInputBehaviorWeak.Pin())
+	{
+		KeyInputBehaviorPinned->Initialize(this, GetKeys());
+	}
+}
+
+TArray<TSharedPtr<FUICommandInfo>> UViewportCameraTranslateInteraction::GetCommands() const
+{
+	if (FViewportNavigationCommands::IsRegistered())
+	{
+		return {
+			FViewportNavigationCommands::Get().Forward, FViewportNavigationCommands::Get().Backward,
+			FViewportNavigationCommands::Get().Right,   FViewportNavigationCommands::Get().Left,
+			FViewportNavigationCommands::Get().WorldUp, FViewportNavigationCommands::Get().WorldDown,
+			FViewportNavigationCommands::Get().LocalUp, FViewportNavigationCommands::Get().LocalDown,
+		};
+	}
+
+	return {};
+}
diff --git a/Engine/Source/Editor/Experimental/EditorInteractiveToolsFramework/Private/ViewportInteractions/ViewportCommandsInteraction.cpp b/Engine/Source/Editor/Experimental/EditorInteractiveToolsFramework/Private/ViewportInteractions/ViewportCommandsInteraction.cpp
new file mode 100644
--- /dev/null
+++ b/Engine/Source/Editor/Experimental/EditorInteractiveToolsFramework/Private/ViewportInteractions/ViewportCommandsInteraction.cpp
@@ -0,0 +1,137 @@
+// Copyright Epic Games, Inc. All Rights Reserved.
+
+#include "ViewportInteractions/ViewportCommandsInteraction.h"
+#include "BaseBehaviors/KeyInputBehavior.h"
+#include "Framework/Commands/UICommandInfo.h"
+#include "Framework/Commands/UICommandList.h"
+
+UViewportCommandsInteraction::UViewportCommandsInteraction()
+{
+	ToolType = UE::Editor::ViewportInteractions::Commands;
+
+	UKeyInputBehavior* KeyInputBehavior = NewObject<UKeyInputBehavior>();
+	KeyInputBehavior->Initialize(this, GetKeys());
+
+	KeyInputBehavior->SetDefaultPriority(UE::Editor::ViewportInteractions::VIEWPORT_INTERACTIONS_HIGH_PRIORITY);
+	KeyInputBehavior->bRequireAllKeys = false;
+	KeyInputBehavior->ModifierCheckFunc = [this](const FInputDeviceState& InputDeviceState)
+	{
+		// We don't want commands to work while controlling the camera with the mouse.
+		return CanBeActivated();
+	};
+	KeyInputBehaviorWeak = KeyInputBehavior;
+
+	RegisterInputBehavior(KeyInputBehavior);
+
+}
+
+void UViewportCommandsInteraction::SetCommands(
+	const TSharedPtr<FUICommandList>& InEditorCommandList, const TArray<TSharedPtr<FUICommandInfo>>& InCommands
+)
+{
+	EditorCommandList = InEditorCommandList;
+	Commands = InCommands;
+
+	if (TStrongObjectPtr<UKeyInputBehavior> KeyInputBehaviorPinned = KeyInputBehaviorWeak.Pin())
+	{
+		KeyInputBehaviorPinned->Initialize(this, GetKeys());
+	}
+}
+
+void UViewportCommandsInteraction::OnKeyPressed(const FKey& InKeyID)
+{
+	if (IsEnabled() && bExecuteOnPress)
+	{
+		ExecuteCommand(InKeyID);
+	}
+}
+
+void UViewportCommandsInteraction::OnKeyReleased(const FKey& InKeyID)
+{
+	if (IsEnabled() && !bExecuteOnPress)
+	{
+		ExecuteCommand(InKeyID);
+	}
+}
+
+void UViewportCommandsInteraction::OnCommandChordChanged()
+{
+	if (TStrongObjectPtr<UKeyInputBehavior> KeyInputBehaviorPinned = KeyInputBehaviorWeak.Pin())
+	{
+		KeyInputBehaviorPinned->Initialize(this, GetKeys());
+	}
+}
+
+TArray<TSharedPtr<FUICommandInfo>> UViewportCommandsInteraction::GetCommands() const
+{
+	return Commands;
+}
+
+TArray<FKey> UViewportCommandsInteraction::GetKeys()
+{
+	TArray<FKey> Keys;
+	for (int32 i = 0; i < static_cast<uint8>(EMultipleKeyBindingIndex::NumChords); ++i)
+	{
+		EMultipleKeyBindingIndex ChordIndex = static_cast<EMultipleKeyBindingIndex>(i);
+		for (const TSharedPtr<FUICommandInfo>& Command : Commands)
+		{
+			if (Command)
+			{
+				TSharedRef<const FInputChord> Chord = Command->GetActiveChord(ChordIndex);
+				Keys.AddUnique(Chord->Key);
+			}
+		}
+	}
+
+	return Keys;
+}
+
+void UViewportCommandsInteraction::ExecuteCommand(const FKey& InKeyID)
+{
+	if (!EditorCommandList)
+	{
+		return;
+	}
+
+	for (const TSharedPtr<FUICommandInfo>& Command : Commands)
+	{
+		if (Command.IsValid())
+		{
+			if (MatchesCommandKey(Command, InKeyID))
+			{
+				if (EditorCommandList->CanExecuteAction(Command.ToSharedRef()))
+				{
+					EditorCommandList->ExecuteAction(Command.ToSharedRef());
+					return;
+				}
+			}
+		}
+	}
+}
+
+bool UViewportCommandsInteraction::MatchesCommandKey(const TSharedPtr<FUICommandInfo>& InCommandInfo, const FKey& InKeyID) const
+{
+	for (int32 i = 0; i < static_cast<uint8>(EMultipleKeyBindingIndex::NumChords); ++i)
+	{
+		EMultipleKeyBindingIndex ChordIndex = static_cast<EMultipleKeyBindingIndex>(i);
+		TSharedRef<const FInputChord> InputChord = InCommandInfo->GetActiveChord(ChordIndex);
+		if (InputChord->Key == InKeyID)
+		{
+			if (!InputChord->HasAnyModifierKeys())
+			{
+				return true;
+			}
+			bool bSuccess = false;
+			bSuccess |= InputChord->NeedsShift() && IsShiftDown();
+			bSuccess |= InputChord->NeedsControl() && IsCtrlDown();
+			bSuccess |= InputChord->NeedsAlt() && IsAltDown();
+			return bSuccess;
+		}
+	}
+	return false;
+}
+
+bool UViewportCommandsInteraction::CanBeActivated() const
+{
+	return !IsAnyMouseButtonDown() && !IsMouseLooking();
+}
diff --git a/Engine/Source/Editor/Experimental/EditorInteractiveToolsFramework/Private/ViewportInteractions/ViewportDragInteraction.cpp b/Engine/Source/Editor/Experimental/EditorInteractiveToolsFramework/Private/ViewportInteractions/ViewportDragInteraction.cpp
new file mode 100644
--- /dev/null
+++ b/Engine/Source/Editor/Experimental/EditorInteractiveToolsFramework/Private/ViewportInteractions/ViewportDragInteraction.cpp
@@ -0,0 +1,115 @@
+// Copyright Epic Games, Inc. All Rights Reserved.
+
+#include "ViewportInteractions/ViewportDragInteraction.h"
+#include "BaseBehaviors/ClickDragBehavior.h"
+#include "BaseBehaviors/SingleClickOrDragBehavior.h"
+#include "ViewportInteractions/ViewportInteractionsBehaviorSource.h"
+
+UViewportDragInteraction::UViewportDragInteraction()
+{
+	bIsDragging = false;
+	MouseX = 0;
+	MouseY = 0;
+
+	USingleClickOrDragInputBehavior* ClickDragBehavior = NewObject<USingleClickOrDragInputBehavior>();
+
+	// Click Target can be specified separately if needed
+	ClickDragBehavior->Initialize(nullptr, this);
+	ClickDragBehavior->SetDefaultPriority(UE::Editor::ViewportInteractions::VIEWPORT_INTERACTIONS_DEFAULT_PRIORITY);
+
+	ClickOrDragInputBehavior = ClickDragBehavior;
+
+	RegisterInputBehavior(ClickDragBehavior);
+}
+
+FInputRayHit UViewportDragInteraction::CanBeginClickDragSequence(const FInputDeviceRay& InPressPos)
+{
+	return FInputRayHit();
+}
+
+void UViewportDragInteraction::OnClickPress(const FInputDeviceRay& InPressPos)
+{
+	if (!IsEnabled())
+	{
+		bIsDragging = false;
+		return;
+	}
+
+	MouseX = static_cast<int32>(InPressPos.ScreenPosition.X);
+	MouseY = static_cast<int32>(InPressPos.ScreenPosition.Y);
+
+	bIsDragging = true;
+
+	GetViewportInteractionsBehaviorSource()->SetMouseCursorOverride(EMouseCursor::None);
+}
+
+void UViewportDragInteraction::OnClickDrag(const FInputDeviceRay& InDragPos)
+{
+	if (!IsEnabled())
+	{
+		return;
+	}
+
+	const int32 ScreenPosX = static_cast<int32>(InDragPos.ScreenPosition.X);
+	const int32 ScreenPosY = static_cast<int32>(InDragPos.ScreenPosition.Y);
+
+	const int32 MouseDeltaX = ScreenPosX - MouseX;
+	const int32 MouseDeltaY = ScreenPosY - MouseY;
+
+	MouseX = ScreenPosX;
+	MouseY = ScreenPosY;
+
+	OnDragDelta(static_cast<float>(MouseDeltaX), static_cast<float>(MouseDeltaY));
+}
+
+void UViewportDragInteraction::OnClickRelease(const FInputDeviceRay& InReleasePos)
+{
+	MouseX = static_cast<int32>(InReleasePos.ScreenPosition.X);
+	MouseY = static_cast<int32>(InReleasePos.ScreenPosition.Y);
+
+	bIsDragging = false;
+}
+
+void UViewportDragInteraction::OnForceEndCapture()
+{
+	OnClickRelease(FInputDeviceRay());
+}
+
+void UViewportDragInteraction::Initialize(UViewportInteractionsBehaviorSource* const InViewportInteractionsBehaviorSource)
+{
+	UViewportInteraction::Initialize(InViewportInteractionsBehaviorSource);
+
+	if (ClickOrDragInputBehavior.IsValid())
+	{
+		ClickOrDragInputBehavior->SetUseCustomMouseButton([this](const FInputDeviceState& InInputDeviceState)
+		{
+			return GetActiveMouseButtonState(InInputDeviceState);
+		});
+
+		ClickOrDragInputBehavior->ModifierCheckFunc = [this](const FInputDeviceState& InInputDeviceState)
+		{
+			return CanBeActivated(InInputDeviceState);
+		};
+	}
+}
+
+void UViewportDragInteraction::SetClickTarget(IClickBehaviorTarget* InClickBehaviorTarget)
+{
+	if (TStrongObjectPtr<USingleClickOrDragInputBehavior> ClickOrDragBehavior = ClickOrDragInputBehavior.Pin())
+	{
+		ClickOrDragBehavior->Initialize(InClickBehaviorTarget, this);
+	}
+}
+
+void UViewportDragInteraction::SetClickOrDragPriority(int InPriority) const
+{
+	if (TStrongObjectPtr<USingleClickOrDragInputBehavior> ClickOrDragBehavior = ClickOrDragInputBehavior.Pin())
+	{
+		ClickOrDragBehavior->SetDefaultPriority(InPriority);
+	}
+}
+
+FDeviceButtonState UViewportDragInteraction::GetActiveMouseButtonState(const FInputDeviceState& Input)
+{
+	return Input.Mouse.Left;
+}
diff --git a/Engine/Source/Editor/Experimental/EditorInteractiveToolsFramework/Private/ViewportInteractions/ViewportFOVInteraction.cpp b/Engine/Source/Editor/Experimental/EditorInteractiveToolsFramework/Private/ViewportInteractions/ViewportFOVInteraction.cpp
new file mode 100644
--- /dev/null
+++ b/Engine/Source/Editor/Experimental/EditorInteractiveToolsFramework/Private/ViewportInteractions/ViewportFOVInteraction.cpp
@@ -0,0 +1,156 @@
+// Copyright Epic Games, Inc. All Rights Reserved.
+
+#include "ViewportInteractions/ViewportFOVInteraction.h"
+#include "BaseBehaviors/KeyInputBehavior.h"
+#include "EditorViewportClient.h"
+#include "ViewportClientNavigationHelper.h"
+#include "ViewportInteractions/ViewportInteractionsBehaviorSource.h"
+
+UViewportFOVInteraction::UViewportFOVInteraction()
+{
+	ToolType = UE::Editor::ViewportInteractions::FOV;
+
+	InteractionName = TEXT("FOV");
+
+	// Separating keys which work only when Mouse Looking from numpad ones and arrows, which always work.
+	// This prevents this IKeyInputBehaviorTarget from capturing unwanted inputs
+
+	// Mouse Looking only FOV controls
+	UKeyInputBehavior* KeyInputBehavior = NewObject<UKeyInputBehavior>();
+	KeyInputBehavior->Initialize(this, GetKeys());
+	KeyInputBehavior->SetDefaultPriority(UE::Editor::ViewportInteractions::VIEWPORT_INTERACTIONS_DEFAULT_PRIORITY);
+	KeyInputBehavior->bRequireAllKeys = false;
+	KeyInputBehaviorWeak = KeyInputBehavior;
+
+	// Always on FOV controls (currently numpad based)
+	UKeyInputBehavior* NumpadKeyInputBehavior = NewObject<UKeyInputBehavior>();
+	NumpadKeyInputBehavior->Initialize(this, GetNumpadKeys());
+	NumpadKeyInputBehavior->SetDefaultPriority(UE::Editor::ViewportInteractions::VIEWPORT_INTERACTIONS_DEFAULT_PRIORITY);
+	NumpadKeyInputBehavior->bRequireAllKeys = false;
+
+	RegisterInputBehaviors({ KeyInputBehavior, NumpadKeyInputBehavior });
+}
+
+UViewportFOVInteraction::~UViewportFOVInteraction()
+{
+	if (UViewportInteractionsBehaviorSource* BehaviorSource = GetViewportInteractionsBehaviorSource())
+	{
+		BehaviorSource->OnMouseLookingStateChanged().RemoveAll(this);
+	}
+}
+
+void UViewportFOVInteraction::OnKeyPressed(const FKey& InKeyID)
+{
+	if (IsEnabled())
+	{
+		constexpr bool bIsPressed = true;
+		UpdateKeyState(InKeyID, bIsPressed);
+	}
+}
+
+void UViewportFOVInteraction::OnKeyReleased(const FKey& InKeyID)
+{
+	if (IsEnabled())
+	{
+		constexpr bool bIsPressed = false;
+		UpdateKeyState(InKeyID, bIsPressed);
+	}
+}
+
+void UViewportFOVInteraction::OnForceEndCapture()
+{
+	ResetImpulse();
+}
+
+void UViewportFOVInteraction::Tick(float InDeltaTime) const
+{
+	if (const FEditorViewportClient* const EditorViewportClient = GetEditorViewportClient())
+	{
+		if (FViewportClientNavigationHelper* const NavigationHelper = EditorViewportClient->GetViewportNavigationHelper())
+		{
+			if (ZoomOutInImpulse != 0)
+			{
+				NavigationHelper->ImpulseDataDelta.ZoomOutInImpulse += ZoomOutInImpulse;
+			}
+		}
+	}
+}
+void UViewportFOVInteraction::Initialize(UViewportInteractionsBehaviorSource* const InViewportInteractionsBehaviorSource)
+{
+	Super::Initialize(InViewportInteractionsBehaviorSource);
+
+	if (UViewportInteractionsBehaviorSource* BehaviorSource = GetViewportInteractionsBehaviorSource())
+	{
+		BehaviorSource->OnMouseLookingStateChanged().AddUObject(this, &UViewportFOVInteraction::OnMouseLookingChanged);
+	}
+}
+
+void UViewportFOVInteraction::UpdateKeyState(const FKey& InKeyID, bool bInIsPressed)
+{
+	using namespace UE::Editor::ViewportInteractions;
+
+	if ((CommandMatchesKey(FViewportNavigationCommands::Get().FovZoomOut, InKeyID) && IsMouseLooking())
+		|| InKeyID == EKeys::NumPadOne)
+	{
+		ZoomOutInImpulse = bInIsPressed ? 1.0f : 0.0f;
+	}
+	else if ((CommandMatchesKey(FViewportNavigationCommands::Get().FovZoomIn, InKeyID) && IsMouseLooking())
+			 || InKeyID == EKeys::NumPadThree)
+	{
+		ZoomOutInImpulse = bInIsPressed ? -1.0f : 0.0f;
+	}
+}
+
+TArray<FKey> UViewportFOVInteraction::GetKeys() const
+{
+	TArray<FKey> Keys;
+
+	for (int32 i = 0; i < static_cast<uint8>(EMultipleKeyBindingIndex::NumChords); ++i)
+	{
+		EMultipleKeyBindingIndex ChordIndex = static_cast<EMultipleKeyBindingIndex>(i);
+
+		for (const TSharedPtr<FUICommandInfo>& Command : GetCommands())
+		{
+			Keys.Add(Command->GetActiveChord(ChordIndex)->Key);
+		}
+	}
+
+	// Hardcoded
+	Keys.Append({ EKeys::Add, EKeys::Subtract });
+
+	return Keys;
+}
+
+TArray<FKey> UViewportFOVInteraction::GetNumpadKeys() const
+{
+	//TODO: currently hardcoded
+	return { EKeys::NumPadOne, EKeys::NumPadThree };
+}
+
+void UViewportFOVInteraction::OnCommandChordChanged()
+{
+	if (TStrongObjectPtr<UKeyInputBehavior> KeyInputBehaviorPinned = KeyInputBehaviorWeak.Pin())
+	{
+		KeyInputBehaviorPinned->Initialize(this, GetKeys());
+	}
+}
+
+TArray<TSharedPtr<FUICommandInfo>> UViewportFOVInteraction::GetCommands() const
+{
+	if (FViewportNavigationCommands::IsRegistered())
+	{
+		return { FViewportNavigationCommands::Get().FovZoomOut, FViewportNavigationCommands::Get().FovZoomIn };
+	}
+
+	return {};
+}
+
+void UViewportFOVInteraction::OnMouseLookingChanged(bool bInIsMouseLooking)
+{
+	ResetImpulse();
+}
+
+void UViewportFOVInteraction::ResetImpulse()
+{
+	ZoomOutInImpulse = 0.0f;
+}
diff --git a/Engine/Source/Editor/Experimental/EditorInteractiveToolsFramework/Private/ViewportInteractions/ViewportInteraction.cpp b/Engine/Source/Editor/Experimental/EditorInteractiveToolsFramework/Private/ViewportInteractions/ViewportInteraction.cpp
new file mode 100644
--- /dev/null
+++ b/Engine/Source/Editor/Experimental/EditorInteractiveToolsFramework/Private/ViewportInteractions/ViewportInteraction.cpp
@@ -0,0 +1,182 @@
+// Copyright Epic Games, Inc. All Rights Reserved.
+
+#include "ViewportInteractions/ViewportInteraction.h"
+#include "Framework/Commands/InputBindingManager.h"
+#include "Framework/Commands/UICommandInfo.h"
+#include "ViewportInteractions/ViewportInteractionsBehaviorSource.h"
+
+UViewportInteraction::UViewportInteraction()
+{
+	// Register to chord changes, so we can re-initialize behaviors with new inputs accordingly
+	OnChordChangedDelegateHandle = FInputBindingManager::Get().RegisterUserDefinedChordChanged(
+		FOnUserDefinedChordChanged::FDelegate::CreateUObject(this, &UViewportInteraction::OnUserDefinedChordChanged)
+	);
+}
+
+bool UViewportInteraction::IsShiftDown() const
+{
+	if (UViewportInteractionsBehaviorSource* BehaviorSource = GetViewportInteractionsBehaviorSource())
+	{
+		return BehaviorSource->IsShiftDown();
+	}
+
+	return false;
+}
+
+bool UViewportInteraction::IsAltDown() const
+{
+	if (UViewportInteractionsBehaviorSource* BehaviorSource = GetViewportInteractionsBehaviorSource())
+	{
+		return BehaviorSource->IsAltDown();
+	}
+
+	return false;
+}
+
+bool UViewportInteraction::IsCtrlDown() const
+{
+	if (UViewportInteractionsBehaviorSource* BehaviorSource = GetViewportInteractionsBehaviorSource())
+	{
+		return BehaviorSource->IsCtrlDown();
+	}
+
+	return false;
+}
+
+bool UViewportInteraction::IsLeftMouseButtonDown() const
+{
+	if (UViewportInteractionsBehaviorSource* BehaviorSource = GetViewportInteractionsBehaviorSource())
+	{
+		return BehaviorSource->IsLeftMouseButtonDown();
+	}
+
+	return false;
+}
+
+bool UViewportInteraction::IsMiddleMouseButtonDown() const
+{
+	if (UViewportInteractionsBehaviorSource* BehaviorSource = GetViewportInteractionsBehaviorSource())
+	{
+		return BehaviorSource->IsMiddleMouseButtonDown();
+	}
+
+	return false;
+}
+
+bool UViewportInteraction::IsRightMouseButtonDown() const
+{
+	if (UViewportInteractionsBehaviorSource* BehaviorSource = GetViewportInteractionsBehaviorSource())
+	{
+		return BehaviorSource->IsRightMouseButtonDown();
+	}
+
+	return false;
+}
+
+bool UViewportInteraction::IsAnyMouseButtonDown() const
+{
+	return IsLeftMouseButtonDown() || IsMiddleMouseButtonDown() || IsRightMouseButtonDown();
+}
+
+bool UViewportInteraction::IsMouseLooking() const
+{
+	if (UViewportInteractionsBehaviorSource* BehaviorSource = GetViewportInteractionsBehaviorSource())
+	{
+		return BehaviorSource->IsMouseLooking();
+	}
+
+	return false;
+}
+
+void UViewportInteraction::SetEnabled(bool bInEnabled)
+{
+	bEnabled = bInEnabled;
+
+	if (UViewportInteractionsBehaviorSource* BehaviorSource = GetViewportInteractionsBehaviorSource())
+	{
+		if (ToolType != NAME_None)
+		{
+			BehaviorSource->SetViewportInteractionActive(ToolType, bEnabled);
+		}
+	}
+	else
+	{
+		UE_LOG(LogTemp,Log,TEXT("Trying to initialize an instance of %s with an invalid UViewportInteractionsBehaviorSource"), *GetClass()->GetName());
+	}
+}
+
+bool UViewportInteraction::IsEnabled() const
+{
+	return bEnabled;
+}
+
+FName UViewportInteraction::GetInteractionName() const
+{
+	return InteractionName.IsSet() ? InteractionName.GetValue() : StaticClass()->GetFName();
+}
+
+void UViewportInteraction::Initialize(UViewportInteractionsBehaviorSource* const InViewportInteractionsBehaviorSource)
+{
+	ViewportInteractionsBehaviorSource = InViewportInteractionsBehaviorSource;
+
+	SetEnabled(true);
+}
+
+UViewportInteractionsBehaviorSource* UViewportInteraction::GetViewportInteractionsBehaviorSource() const
+{
+	if (TStrongObjectPtr<UViewportInteractionsBehaviorSource> BehaviorSource = ViewportInteractionsBehaviorSource.Pin())
+	{
+		return BehaviorSource.Get();
+	}
+
+	return nullptr;
+}
+
+void UViewportInteraction::BeginDestroy()
+{
+	Super::BeginDestroy();
+
+	if (OnChordChangedDelegateHandle.IsValid())
+	{
+		FInputBindingManager::Get().UnregisterUserDefinedChordChanged(OnChordChangedDelegateHandle);
+	}
+}
+
+FEditorViewportClient* UViewportInteraction::GetEditorViewportClient() const
+{
+	if (UViewportInteractionsBehaviorSource* BehaviorSource = GetViewportInteractionsBehaviorSource())
+	{
+		return BehaviorSource->GetEditorViewportClient();
+	}
+
+	return nullptr;
+}
+
+void UViewportInteraction::OnUserDefinedChordChanged(const FUICommandInfo& InCommandInfo)
+{
+	for (const TSharedPtr<FUICommandInfo>& Command : GetCommands())
+	{
+		if (Command.IsValid())
+		{
+			// Is this enough to ensure we're only triggering for the right commands - and we're not skipping when required?
+			if (Command->GetBindingContext() == InCommandInfo.GetBindingContext()
+				&& Command->GetCommandName() == InCommandInfo.GetCommandName())
+			{
+				OnCommandChordChanged();
+			}
+		}
+	}
+}
+
+void UViewportInteraction::RegisterInputBehavior(UInputBehavior* InBehavior)
+{
+	InputBehaviors.AddUnique(InBehavior);
+}
+
+void UViewportInteraction::RegisterInputBehaviors(TArray<UInputBehavior*> InBehaviors)
+{
+	for (UInputBehavior* InputBehavior : InBehaviors)
+	{
+		RegisterInputBehavior(InputBehavior);
+	}
+}
diff --git a/Engine/Source/Editor/Experimental/EditorInteractiveToolsFramework/Private/ViewportInteractions/ViewportInteractionsBehaviorSource.cpp b/Engine/Source/Editor/Experimental/EditorInteractiveToolsFramework/Private/ViewportInteractions/ViewportInteractionsBehaviorSource.cpp
new file mode 100644
--- /dev/null
+++ b/Engine/Source/Editor/Experimental/EditorInteractiveToolsFramework/Private/ViewportInteractions/ViewportInteractionsBehaviorSource.cpp
@@ -0,0 +1,361 @@
+// Copyright Epic Games, Inc. All Rights Reserved.
+
+#include "ViewportInteractions/ViewportInteractionsBehaviorSource.h"
+#include "BaseBehaviors/KeyAsModifierInputBehavior.h"
+#include "EditorModeManager.h"
+#include "EditorViewportClient.h"
+#include "Tools/EdModeInteractiveToolsContext.h"
+#include "UnrealClient.h"
+#include "ViewportClientNavigationHelper.h"
+#include "ViewportInteractions/ViewportInteraction.h"
+
+namespace UE::Editor::ViewportInteractions
+{
+// CVar initializer
+static int32 UseITFViewportInteractions = 0;
+static FAutoConsoleVariableRef CVarEnableITFViewportInteractions(
+	TEXT("ViewportInteractions.EnableITFInteractions"),
+	UseITFViewportInteractions,
+	TEXT("Are ITF Viewport Interactions enabled?"),
+	FConsoleVariableDelegate::CreateLambda(
+		[](const IConsoleVariable* InVariable)
+		{
+			ToggleEditorViewportInteractions(UseITFViewportInteractions == 1);
+		}
+	)
+);
+
+// CVar initializer
+static int32 VerboseITFViewportInteractions = 1;
+static FAutoConsoleVariableRef CVarEnableVerboseITFViewportInteractions(
+	TEXT("ViewportInteractions.Verbose"),
+	VerboseITFViewportInteractions,
+	TEXT("Enables verbose logging for ITF Viewport Interactions"),
+	FConsoleVariableDelegate::CreateLambda(
+		[](const IConsoleVariable* InVariable)
+		{
+		}
+	)
+);
+
+bool UseEditorViewportInteractions()
+{
+	return UseITFViewportInteractions == 1;
+}
+
+bool IsVerbose()
+{
+	return VerboseITFViewportInteractions == 1;
+}
+
+void ToggleEditorViewportInteractions(bool bInEnable)
+{
+	UseITFViewportInteractions = bInEnable;
+	if (UseITFViewportInteractions)
+	{
+		OnEditorViewportInteractionsActivated().Broadcast();
+	}
+	else
+	{
+		OnEditorViewportInteractionsDeactivated().Broadcast();
+	}
+}
+
+FOnEditorViewportInteractionsToggleDelegate& OnEditorViewportInteractionsActivated()
+{
+	static FOnEditorViewportInteractionsToggleDelegate OnViewportInteractionsActivated;
+	return OnViewportInteractionsActivated;
+}
+
+FOnEditorViewportInteractionsToggleDelegate& OnEditorViewportInteractionsDeactivated()
+{
+	static FOnEditorViewportInteractionsToggleDelegate OnViewportInteractionsDeactivated;
+	return OnViewportInteractionsDeactivated;
+}
+
+bool CommandMatchesKey(const TSharedPtr<FUICommandInfo>& InCommandInfo, const FKey& InKeyID)
+{
+	for (int32 i = 0; i < static_cast<uint8>(EMultipleKeyBindingIndex::NumChords); ++i)
+	{
+		EMultipleKeyBindingIndex ChordIndex = static_cast<EMultipleKeyBindingIndex>(i);
+		if (InCommandInfo->GetActiveChord(ChordIndex)->Key == InKeyID)
+		{
+			return true;
+		}
+	}
+
+	return false;
+}
+
+void LOG(const TCHAR* InMessage)
+{
+	if (IsVerbose())
+	{
+		UE_LOG(LogTemp, Log, TEXT("%s"), InMessage);
+	}
+}
+
+static constexpr int MODIFIERS_PRIORITY = VIEWPORT_INTERACTIONS_HIGH_PRIORITY;
+
+// We use this string to group Behavior Inputs so that we can selectively remove them if needed, without touching other inputs we might add (e.g. the ones added by Initialize)
+static const FString ViewportInteractionsGroupName = TEXT("ViewportInteractions");
+
+} // namespace UE::Editor::ViewportInteractions
+
+void UViewportInteractionsBehaviorSource::OnUpdateModifierState(int InModifierID, bool bInIsOn)
+{
+	if (InModifierID == UE::Editor::ViewportInteractions::ShiftKeyMod)
+	{
+		bIsShiftDown = bInIsOn;
+	}
+	else if (InModifierID == UE::Editor::ViewportInteractions::CtrlKeyMod)
+	{
+		bIsCtrlDown = bInIsOn;
+	}
+	else if (InModifierID == UE::Editor::ViewportInteractions::AltKeyMod)
+	{
+		bIsAltDown = bInIsOn;
+	}
+	else if (InModifierID == UE::Editor::ViewportInteractions::LeftMouseButtonMod)
+	{
+		bIsLeftMouseButtonDown = bInIsOn;
+	}
+	else if (InModifierID == UE::Editor::ViewportInteractions::MiddleMouseButtonMod)
+	{
+		bIsMiddleMouseButtonDown = bInIsOn;
+	}
+	else if (InModifierID == UE::Editor::ViewportInteractions::RightMouseButtonMod)
+	{
+		bIsRightMouseButtonDown = bInIsOn;
+	}
+}
+
+void UViewportInteractionsBehaviorSource::OnForceEndCapture()
+{
+	bIsShiftDown = false;
+	bIsCtrlDown = false;
+	bIsAltDown = false;
+
+	bIsLeftMouseButtonDown = false;
+	bIsMiddleMouseButtonDown = false;
+	bIsRightMouseButtonDown = false;
+
+	SetIsMouseLooking(false);
+}
+
+void UViewportInteractionsBehaviorSource::Initialize(const UEditorInteractiveToolsContext* InInteractiveToolsContext)
+{
+	if (!InInteractiveToolsContext)
+	{
+		return;
+	}
+
+	EditorInteractiveToolsContextWeak = InInteractiveToolsContext;
+
+	BehaviorSet = NewObject<UInputBehaviorSet>();
+
+	// Modifiers handling, directly from "this"
+	{
+		UKeyAsModifierInputBehavior* KeyAsModifierInputBehavior = NewObject<UKeyAsModifierInputBehavior>();
+		KeyAsModifierInputBehavior->Initialize(
+			this, UE::Editor::ViewportInteractions::ShiftKeyMod, FInputDeviceState::IsShiftKeyDown
+		);
+		KeyAsModifierInputBehavior->Initialize(
+			this, UE::Editor::ViewportInteractions::CtrlKeyMod, FInputDeviceState::IsCtrlKeyDown
+		);
+		KeyAsModifierInputBehavior->Initialize(
+			this, UE::Editor::ViewportInteractions::AltKeyMod, FInputDeviceState::IsAltKeyDown
+		);
+		KeyAsModifierInputBehavior->SetDefaultPriority(UE::Editor::ViewportInteractions::MODIFIERS_PRIORITY);
+		BehaviorSet->Add(KeyAsModifierInputBehavior, this);
+
+		UMouseButtonAsModifierInputBehavior* MouseButtonAsModifierInputBehavior =
+			NewObject<UMouseButtonAsModifierInputBehavior>();
+		MouseButtonAsModifierInputBehavior->Initialize(
+			this,
+			UE::Editor::ViewportInteractions::LeftMouseButtonMod,
+			[](const FInputDeviceState& InputState)
+			{
+				return InputState.Mouse.Left.bDown || InputState.Mouse.Left.bPressed;
+			}
+		);
+		MouseButtonAsModifierInputBehavior->Initialize(
+			this,
+			UE::Editor::ViewportInteractions::MiddleMouseButtonMod,
+			[](const FInputDeviceState& InputState)
+			{
+				return InputState.Mouse.Middle.bDown || InputState.Mouse.Middle.bPressed;
+			}
+		);
+		MouseButtonAsModifierInputBehavior->Initialize(
+			this,
+			UE::Editor::ViewportInteractions::RightMouseButtonMod,
+			[](const FInputDeviceState& InputState)
+			{
+				return InputState.Mouse.Right.bDown || InputState.Mouse.Right.bPressed;
+			}
+		);
+		MouseButtonAsModifierInputBehavior->SetDefaultPriority(UE::Editor::ViewportInteractions::MODIFIERS_PRIORITY);
+		BehaviorSet->Add(MouseButtonAsModifierInputBehavior, this);
+	}
+}
+
+void UViewportInteractionsBehaviorSource::RegisterBehaviorSources()
+{
+	if (const UEditorInteractiveToolsContext* InteractiveToolsContext = EditorInteractiveToolsContextWeak.Get())
+	{
+		if (UInputRouter* InputRouter = InteractiveToolsContext->InputRouter)
+		{
+			InputRouter->RegisterSource(this);
+		}
+	}
+}
+
+void UViewportInteractionsBehaviorSource::DeregisterBehaviorSources()
+{
+	if (const UEditorInteractiveToolsContext* InteractiveToolsContext = EditorInteractiveToolsContextWeak.Get())
+	{
+		if (UInputRouter* InputRouter = InteractiveToolsContext->InputRouter)
+		{
+			InputRouter->DeregisterSource(this);
+		}
+	}
+}
+
+void UViewportInteractionsBehaviorSource::Tick(float InDeltaTime) const
+{
+	for (const TPair<FName, TObjectPtr<UViewportInteraction>>& Pair : ViewportInteractions)
+	{
+		if (UViewportInteraction* Interaction = Pair.Value)
+		{
+			Interaction->Tick(InDeltaTime);
+		}
+	}
+}
+
+void UViewportInteractionsBehaviorSource::AddInteractions(const TArray<const UClass*>& InInteractions, bool bInReregister)
+{
+	for (const UClass* InteractionClass : InInteractions)
+	{
+		AddInteraction(InteractionClass);
+	}
+
+	if (bInReregister)
+	{
+		// Refreshing registered behaviors
+		DeregisterBehaviorSources();
+		RegisterBehaviorSources();
+	}
+}
+
+UViewportInteraction* UViewportInteractionsBehaviorSource::AddInteraction(const UClass* InInteractionClass, bool bInReregister)
+{
+	if (!BehaviorSet && EditorInteractiveToolsContextWeak.IsValid())
+	{
+		Initialize(EditorInteractiveToolsContextWeak.Get());
+	}
+
+	UViewportInteraction* Interaction = NewObject<UViewportInteraction>(GetTransientPackage(), InInteractionClass);
+	Interaction->Initialize(this);
+
+	ViewportInteractions.Add(Interaction->GetInteractionName(), Interaction);
+
+	for (UInputBehavior* const InputBehavior : Interaction->GetInputBehaviors())
+	{
+		BehaviorSet->Add(InputBehavior, Interaction, UE::Editor::ViewportInteractions::ViewportInteractionsGroupName);
+	}
+
+	if (bInReregister)
+	{
+		// Refreshing registered behaviors
+		DeregisterBehaviorSources();
+		RegisterBehaviorSources();
+	}
+
+	return Interaction;
+}
+
+void UViewportInteractionsBehaviorSource::Reset()
+{
+	DeregisterBehaviorSources();
+	ViewportInteractions.Empty();
+	BehaviorSet->RemoveByGroup(UE::Editor::ViewportInteractions::ViewportInteractionsGroupName);
+}
+
+void UViewportInteractionsBehaviorSource::SetMouseCursorOverride(EMouseCursor::Type InMouseCursor)
+{
+	CursorOverride = InMouseCursor;
+}
+
+bool UViewportInteractionsBehaviorSource::IsMouseLooking() const
+{
+	if (FEditorViewportClient* const EditorViewportClient = GetEditorViewportClient())
+	{
+		return EditorViewportClient->GetViewportNavigationHelper()->bIsMouseLooking;
+	}
+
+	return false;
+}
+
+void UViewportInteractionsBehaviorSource::SetIsMouseLooking(bool bInIsLooking)
+{
+	if (FEditorViewportClient* const EditorViewportClient = GetEditorViewportClient())
+	{
+		if (bInIsLooking)
+		{
+			EditorViewportClient->SetRequiredCursorOverride(true, CursorOverride);
+		}
+		else
+		{
+			EditorViewportClient->SetRequiredCursorOverride(false);
+		}
+
+		EditorViewportClient->GetViewportNavigationHelper()->bIsMouseLooking = bInIsLooking;
+	}
+
+	OnMouseLookingStateChanged().Broadcast(bInIsLooking);
+}
+
+FEditorViewportClient* UViewportInteractionsBehaviorSource::GetEditorViewportClient() const
+{
+	FEditorViewportClient* EditorViewportClient = nullptr;
+
+	if (const UEditorInteractiveToolsContext* const InteractiveToolsContext = EditorInteractiveToolsContextWeak.Get())
+	{
+		if (const FEditorModeTools* const ModeManager = InteractiveToolsContext->GetParentEditorModeManager())
+		{
+			EditorViewportClient = ModeManager->GetFocusedViewportClient();
+		}
+
+		if (!EditorViewportClient)
+		{
+			if (UInteractiveToolManager* const ToolManager = InteractiveToolsContext->ToolManager)
+			{
+				if (const IToolsContextQueriesAPI* const ContextQueriesAPI = ToolManager->GetContextQueriesAPI())
+				{
+					if (const FViewport* const Viewport = ContextQueriesAPI->GetFocusedViewport())
+					{
+						EditorViewportClient = static_cast<FEditorViewportClient*>(Viewport->GetClient());
+					}
+				}
+			}
+		}
+	}
+
+	return EditorViewportClient;
+}
+
+void UViewportInteractionsBehaviorSource::SetViewportInteractionActive(FName InViewportInteraction, bool bInActive)
+{
+	ViewportInteractionsStatusMap.Emplace(InViewportInteraction, bInActive);
+}
+
+bool UViewportInteractionsBehaviorSource::IsViewportInteractionActive(FName InViewportInteraction)
+{
+	if (ViewportInteractionsStatusMap.Contains(InViewportInteraction))
+	{
+		return ViewportInteractionsStatusMap[InViewportInteraction];
+	}
+
+	return false;
+}
diff --git a/Engine/Source/Editor/Experimental/EditorInteractiveToolsFramework/Private/ViewportInteractions/ViewportInteractionsUtilities.cpp b/Engine/Source/Editor/Experimental/EditorInteractiveToolsFramework/Private/ViewportInteractions/ViewportInteractionsUtilities.cpp
new file mode 100644
--- /dev/null
+++ b/Engine/Source/Editor/Experimental/EditorInteractiveToolsFramework/Private/ViewportInteractions/ViewportInteractionsUtilities.cpp
@@ -0,0 +1,36 @@
+﻿// Copyright Epic Games, Inc. All Rights Reserved.
+
+#include "ViewportInteractions/ViewportInteractionsUtilities.h"
+
+#include "ViewportInteractions/ViewportCameraSpeedMouseWheelInteraction.h"
+#include "ViewportInteractions/ViewportCameraRotateInteraction.h"
+#include "ViewportInteractions/ViewportCameraTranslateInteraction.h"
+#include "ViewportInteractions/ViewportFOVInteraction.h"
+#include "ViewportInteractions/ViewportInteractionsBehaviorSource.h"
+#include "ViewportInteractions/ViewportMoveYawInteraction.h"
+#include "ViewportInteractions/ViewportOrbitInteraction.h"
+#include "ViewportInteractions/ViewportOrthoPanInteraction.h"
+#include "ViewportInteractions/ViewportPanInteraction.h"
+#include "ViewportInteractions/ViewportViewAngleInteraction.h"
+#include "ViewportInteractions/ViewportZoomInteraction.h"
+
+void UE::Editor::ViewportInteractions::AddDefaultCameraMovementInteractions(
+	UViewportInteractionsBehaviorSource* InInteractionsBehaviorSource
+)
+{
+	if (InInteractionsBehaviorSource)
+	{
+		InInteractionsBehaviorSource->AddInteractions(
+			{UViewportCameraSpeedMouseWheelInteraction::StaticClass(),
+				UViewportOrbitInteraction::StaticClass(),
+				UViewportMoveYawInteraction::StaticClass(),
+				UViewportViewAngleInteraction::StaticClass(),
+				UViewportOrthoPanInteraction::StaticClass(),
+				UViewportPanInteraction::StaticClass(),
+				UViewportZoomInteraction::StaticClass(),
+				UViewportCameraTranslateInteraction::StaticClass(),
+				UViewportCameraRotateInteraction::StaticClass(),
+				UViewportFOVInteraction::StaticClass() }
+		);
+	}
+}
diff --git a/Engine/Source/Editor/Experimental/EditorInteractiveToolsFramework/Private/ViewportInteractions/ViewportMoveYawInteraction.cpp b/Engine/Source/Editor/Experimental/EditorInteractiveToolsFramework/Private/ViewportInteractions/ViewportMoveYawInteraction.cpp
new file mode 100644
--- /dev/null
+++ b/Engine/Source/Editor/Experimental/EditorInteractiveToolsFramework/Private/ViewportInteractions/ViewportMoveYawInteraction.cpp
@@ -0,0 +1,90 @@
+// Copyright Epic Games, Inc. All Rights Reserved.
+
+#include "ViewportInteractions/ViewportMoveYawInteraction.h"
+#include "BaseBehaviors/ClickDragBehavior.h"
+#include "BaseBehaviors/SingleClickOrDragBehavior.h"
+#include "EditorViewportClient.h"
+#include "Settings/LevelEditorViewportSettings.h"
+#include "ViewportClientNavigationHelper.h"
+#include "EditorDragTools/DragToolsBehaviourSource.h"
+#include "ViewportInteractions/ViewportInteractionsBehaviorSource.h"
+
+UViewportMoveYawInteraction::UViewportMoveYawInteraction()
+{
+	ToolType = UE::Editor::ViewportInteractions::CameraDrag;
+
+	InteractionName = TEXT("Move-Yaw");
+
+	SetClickOrDragPriority(FInputCapturePriority::DEFAULT_GIZMO_PRIORITY + 1);
+}
+
+FInputRayHit UViewportMoveYawInteraction::CanBeginClickDragSequence(const FInputDeviceRay& InPressPos)
+{
+	bool bSuccess = CanBeActivated();
+	if (FEditorViewportClient* EditorViewportClient = GetEditorViewportClient())
+	{
+		bSuccess &= EditorViewportClient->IsPerspective();
+	}
+
+	return bSuccess ? FInputRayHit(TNumericLimits<float>::Max()) : FInputRayHit();
+}
+
+void UViewportMoveYawInteraction::OnClickPress(const FInputDeviceRay& PressPos)
+{
+	UE::Editor::ViewportInteractions::LOG(TEXT("MOVE YAW 1"));
+
+	UViewportDragInteraction::OnClickPress(PressPos);
+
+	if (UViewportInteractionsBehaviorSource* BehaviorSource = GetViewportInteractionsBehaviorSource())
+	{
+		BehaviorSource->SetIsMouseLooking(true);
+	}
+}
+
+void UViewportMoveYawInteraction::OnClickRelease(const FInputDeviceRay& ReleasePos)
+{
+	UE::Editor::ViewportInteractions::LOG(TEXT("MOVE YAW 0"));
+
+	UViewportDragInteraction::OnClickRelease(ReleasePos);
+
+	if (UViewportInteractionsBehaviorSource* BehaviorSource = GetViewportInteractionsBehaviorSource())
+	{
+		BehaviorSource->SetIsMouseLooking(false);
+	}
+}
+
+void UViewportMoveYawInteraction::OnDragDelta(float InMouseDeltaX, float InMouseDeltaY)
+{
+	if (FEditorViewportClient* const EditorViewportClient = GetEditorViewportClient())
+	{
+		FVector LocationDelta = FVector::ZeroVector;
+		FRotator RotDelta = FRotator::ZeroRotator;
+
+		LocationDelta.X = -InMouseDeltaY * FMath::Cos(EditorViewportClient->GetViewRotation().Yaw * PI / 180.f);
+		LocationDelta.Y = -InMouseDeltaY * FMath::Sin(EditorViewportClient->GetViewRotation().Yaw * PI / 180.f);
+
+		const ULevelEditorViewportSettings* ViewportSettings = GetDefault<ULevelEditorViewportSettings>();
+		RotDelta.Yaw = InMouseDeltaX * ViewportSettings->MouseSensitivty;
+
+		FVector& ViewportLocationDelta = EditorViewportClient->GetViewportNavigationHelper()->LocationDelta;
+		ViewportLocationDelta += LocationDelta;
+
+		FRotator& ViewportRotationDelta = EditorViewportClient->GetViewportNavigationHelper()->RotationDelta;
+		ViewportRotationDelta += RotDelta;
+	}
+}
+
+FDeviceButtonState UViewportMoveYawInteraction::GetActiveMouseButtonState(const FInputDeviceState& Input)
+{
+	return Super::GetActiveMouseButtonState(Input);
+}
+
+bool UViewportMoveYawInteraction::CanBeActivated(const FInputDeviceState& InInputDeviceState) const
+{
+	if (const FEditorViewportClient* const EditorViewportClient = GetEditorViewportClient())
+	{
+		return EditorViewportClient->IsPerspective();
+	}
+
+	return false;
+}
diff --git a/Engine/Source/Editor/Experimental/EditorInteractiveToolsFramework/Private/ViewportInteractions/ViewportOrbitInteraction.cpp b/Engine/Source/Editor/Experimental/EditorInteractiveToolsFramework/Private/ViewportInteractions/ViewportOrbitInteraction.cpp
new file mode 100644
--- /dev/null
+++ b/Engine/Source/Editor/Experimental/EditorInteractiveToolsFramework/Private/ViewportInteractions/ViewportOrbitInteraction.cpp
@@ -0,0 +1,92 @@
+// Copyright Epic Games, Inc. All Rights Reserved.
+
+#include "ViewportInteractions/ViewportOrbitInteraction.h"
+#include "BaseBehaviors/ClickDragBehavior.h"
+#include "EditorViewportClient.h"
+#include "ViewportClientNavigationHelper.h"
+#include "BaseBehaviors/SingleClickOrDragBehavior.h"
+#include "ViewportInteractions/ViewportInteractionsBehaviorSource.h"
+
+UViewportOrbitInteraction::UViewportOrbitInteraction()
+{
+	ToolType = UE::Editor::ViewportInteractions::Orbit;
+
+	InteractionName = TEXT("Orbit");
+
+	if (TStrongObjectPtr<USingleClickOrDragInputBehavior> ClickOrDragBehavior = ClickOrDragInputBehavior.Pin())
+	{
+		// Priority lower than gizmos, so that we can still Alt - Duplicate.
+		// This might need adjustments once we have all the pieces together (gizmo, drag tools, viewport interactions)
+		ClickOrDragBehavior->SetDefaultPriority(FInputCapturePriority::DEFAULT_GIZMO_PRIORITY + 1);
+	}
+}
+
+FInputRayHit UViewportOrbitInteraction::CanBeginClickDragSequence(const FInputDeviceRay& InPressPos)
+{
+	if (CanBeActivated())
+	{
+		if (FEditorViewportClient* EditorViewportClient = GetEditorViewportClient())
+		{
+			return EditorViewportClient->IsPerspective() ? FInputRayHit(TNumericLimits<float>::Max()) : FInputRayHit();
+		}
+	}
+
+	return FInputRayHit();
+}
+
+void UViewportOrbitInteraction::OnClickPress(const FInputDeviceRay& InPressPos)
+{
+	UE::Editor::ViewportInteractions::LOG(TEXT("ORBIT 1"));
+
+	if (FEditorViewportClient* EditorViewportClient = GetEditorViewportClient())
+	{
+		EditorViewportClient->ToggleOrbitCamera(true);
+	}
+
+	if (UViewportInteractionsBehaviorSource* BehaviorSource = GetViewportInteractionsBehaviorSource())
+	{
+		BehaviorSource->SetIsMouseLooking(true);
+	}
+
+	UViewportDragInteraction::OnClickPress(InPressPos);
+}
+
+void UViewportOrbitInteraction::OnClickRelease(const FInputDeviceRay& InReleasePos)
+{
+	UE::Editor::ViewportInteractions::LOG(TEXT("ORBIT 0"));
+
+	if (FEditorViewportClient* EditorViewportClient = GetEditorViewportClient())
+	{
+		EditorViewportClient->ToggleOrbitCamera(false);
+	}
+
+	if (UViewportInteractionsBehaviorSource* BehaviorSource = GetViewportInteractionsBehaviorSource())
+	{
+		BehaviorSource->SetIsMouseLooking(false);
+	}
+
+	UViewportDragInteraction::OnClickRelease(InReleasePos);
+}
+
+void UViewportOrbitInteraction::OnDragDelta(float InMouseDeltaX, float InMouseDeltaY)
+{
+	if (FEditorViewportClient* const EditorViewportClient = GetEditorViewportClient())
+	{
+		constexpr bool bNudge = false;
+
+		EditorViewportClient->SetCurrentWidgetAxis(EAxisList::None);
+
+		// TODO: we can probably extract just the part dedicated to perspective view from EditorViewportClient->TranslateDelta
+		FVector Delta = FVector::ZeroVector;
+		Delta += EditorViewportClient->TranslateDelta(EKeys::MouseX, InMouseDeltaX, bNudge);
+		Delta += EditorViewportClient->TranslateDelta(EKeys::MouseY, -InMouseDeltaY, bNudge);
+
+		FVector& ViewportOrbitDelta = EditorViewportClient->GetViewportNavigationHelper()->OrbitDelta;
+		ViewportOrbitDelta += Delta;
+	}
+}
+
+bool UViewportOrbitInteraction::CanBeActivated(const FInputDeviceState& InInputDeviceState) const
+{
+	return IsEnabled() && IsAltDown();
+}
diff --git a/Engine/Source/Editor/Experimental/EditorInteractiveToolsFramework/Private/ViewportInteractions/ViewportOrthoPanInteraction.cpp b/Engine/Source/Editor/Experimental/EditorInteractiveToolsFramework/Private/ViewportInteractions/ViewportOrthoPanInteraction.cpp
new file mode 100644
--- /dev/null
+++ b/Engine/Source/Editor/Experimental/EditorInteractiveToolsFramework/Private/ViewportInteractions/ViewportOrthoPanInteraction.cpp
@@ -0,0 +1,105 @@
+// Copyright Epic Games, Inc. All Rights Reserved.
+
+#include "ViewportInteractions/ViewportOrthoPanInteraction.h"
+#include "BaseBehaviors/ClickDragBehavior.h"
+#include "EditorViewportClient.h"
+#include "ViewportClientNavigationHelper.h"
+#include "ViewportInteractions/ViewportInteractionsBehaviorSource.h"
+
+UViewportOrthoPanInteraction::UViewportOrthoPanInteraction()
+{
+	ToolType = UE::Editor::ViewportInteractions::CameraDrag;
+
+	InteractionName = TEXT("Orthographic Panning");
+}
+
+FInputRayHit UViewportOrthoPanInteraction::CanBeginClickDragSequence(const FInputDeviceRay& InPressPos)
+{
+	bool bSuccess = CanBeActivated();
+	if (FEditorViewportClient* EditorViewportClient = GetEditorViewportClient())
+	{
+		bSuccess &= EditorViewportClient->IsOrtho();
+	}
+
+	return bSuccess ? FInputRayHit(TNumericLimits<float>::Max()) : FInputRayHit();
+}
+
+void UViewportOrthoPanInteraction::OnClickPress(const FInputDeviceRay& PressPos)
+{
+	UE::Editor::ViewportInteractions::LOG(TEXT("ORTHO PAN 1"));
+
+	UViewportDragInteraction::OnClickPress(PressPos);
+
+	if (UViewportInteractionsBehaviorSource* BehaviorSource = GetViewportInteractionsBehaviorSource())
+	{
+		GetViewportInteractionsBehaviorSource()->SetMouseCursorOverride(EMouseCursor::Type::GrabHandClosed);
+		BehaviorSource->SetIsMouseLooking(true);
+	}
+}
+
+void UViewportOrthoPanInteraction::OnClickRelease(const FInputDeviceRay& ReleasePos)
+{
+	UE::Editor::ViewportInteractions::LOG(TEXT("ORTHO PAN 0"));
+
+	UViewportDragInteraction::OnClickRelease(ReleasePos);
+
+	if (UViewportInteractionsBehaviorSource* BehaviorSource = GetViewportInteractionsBehaviorSource())
+	{
+		BehaviorSource->SetIsMouseLooking(false);
+	}
+}
+
+void UViewportOrthoPanInteraction::OnDragDelta(float InMouseDeltaX, float InMouseDeltaY)
+{
+	if (FEditorViewportClient* const EditorViewportClient = GetEditorViewportClient())
+	{
+		// based off FEditorViewportClient::TranslateDelta
+
+		const float UnitsPerPixel = EditorViewportClient->GetOrthoUnitsPerPixel(EditorViewportClient->Viewport);
+		FVector LocationDelta(-InMouseDeltaX * UnitsPerPixel, InMouseDeltaY * UnitsPerPixel, 0.0f);
+
+		switch (EditorViewportClient->GetViewportType())
+		{
+		case LVT_OrthoTop:
+			LocationDelta = FVector(-InMouseDeltaY * UnitsPerPixel, InMouseDeltaX * UnitsPerPixel, 0.0f);
+			break;
+		case LVT_OrthoLeft:
+			LocationDelta = FVector(-InMouseDeltaX * UnitsPerPixel, 0.0f, InMouseDeltaY * UnitsPerPixel);
+			break;
+		case LVT_OrthoBack:
+			LocationDelta = FVector(0.0f, -InMouseDeltaX * UnitsPerPixel, InMouseDeltaY * UnitsPerPixel);
+			break;
+		case LVT_OrthoBottom:
+			LocationDelta = FVector(-InMouseDeltaY * UnitsPerPixel, -InMouseDeltaX * UnitsPerPixel, 0.0f);
+			break;
+		case LVT_OrthoRight:
+			LocationDelta = FVector(InMouseDeltaX * UnitsPerPixel, 0.0f, InMouseDeltaY * UnitsPerPixel);
+			break;
+		case LVT_OrthoFront:
+			LocationDelta = FVector(0.0f, InMouseDeltaX * UnitsPerPixel, InMouseDeltaY * UnitsPerPixel);
+			break;
+		case LVT_OrthoFreelook:
+		case LVT_Perspective:
+			break;
+		}
+
+		// Invert when Alt or Shift are down
+		if (IsAltDown() || IsShiftDown())
+		{
+			LocationDelta *= -1.0f;
+		}
+
+		FVector& ViewportLocationDelta = EditorViewportClient->GetViewportNavigationHelper()->LocationDelta;
+		ViewportLocationDelta += LocationDelta;
+	}
+}
+
+FDeviceButtonState UViewportOrthoPanInteraction::GetActiveMouseButtonState(const FInputDeviceState& Input)
+{
+	return Input.Mouse.Right;
+}
+bool UViewportOrthoPanInteraction::CanBeActivated(const FInputDeviceState& InInputDeviceState) const
+{
+	//TODO: we might be able to remove the check for LMB
+	return IsEnabled() && !IsLeftMouseButtonDown();
+}
diff --git a/Engine/Source/Editor/Experimental/EditorInteractiveToolsFramework/Private/ViewportInteractions/ViewportPanInteraction.cpp b/Engine/Source/Editor/Experimental/EditorInteractiveToolsFramework/Private/ViewportInteractions/ViewportPanInteraction.cpp
new file mode 100644
--- /dev/null
+++ b/Engine/Source/Editor/Experimental/EditorInteractiveToolsFramework/Private/ViewportInteractions/ViewportPanInteraction.cpp
@@ -0,0 +1,130 @@
+// Copyright Epic Games, Inc. All Rights Reserved.
+
+#include "ViewportInteractions/ViewportPanInteraction.h"
+#include "BaseBehaviors/ClickDragBehavior.h"
+#include "Behaviors/MultiButtonClickDragBehavior.h"
+#include "EditorViewportClient.h"
+#include "Settings/LevelEditorViewportSettings.h"
+#include "ViewportClientNavigationHelper.h"
+#include "ViewportInteractions/ViewportInteractionsBehaviorSource.h"
+
+UViewportPanInteraction::UViewportPanInteraction()
+{
+	ToolType = UE::Editor::ViewportInteractions::CameraDrag;
+
+	InteractionName = TEXT("Perspective Panning");
+
+	// LMB + RMB + Drag to Pan
+	UMultiButtonClickDragBehavior* MultiButtonClickDragBehavior = NewObject<UMultiButtonClickDragBehavior>();
+	MultiButtonClickDragBehavior->Initialize();
+	MultiButtonClickDragBehavior->EnableButton(EKeys::LeftMouseButton);
+	MultiButtonClickDragBehavior->EnableButton(EKeys::RightMouseButton);
+
+	MultiButtonClickDragBehavior->CanBeginClickDragFunc = [this](const FInputDeviceRay& InPressPos)
+	{
+		bool bSuccess = CanBeActivated();
+		if (FEditorViewportClient* EditorViewportClient = GetEditorViewportClient())
+		{
+			bSuccess &= EditorViewportClient->IsPerspective();
+		}
+
+		return bSuccess ? FInputRayHit(TNumericLimits<float>::Max()) : FInputRayHit();
+	};
+
+	MultiButtonClickDragBehavior->ModifierCheckFunc = [this](const FInputDeviceState& InInputDeviceState)
+	{
+		return InInputDeviceState.Mouse.Left.bDown && InInputDeviceState.Mouse.Right.bDown;
+	};
+
+	MultiButtonClickDragBehavior->OnClickPressFunc = [this](const FInputDeviceRay& InPressPos)
+	{
+		OnClickPress(InPressPos);
+		bMultiButtonPanInteraction = true;
+	};
+
+	MultiButtonClickDragBehavior->OnClickReleaseFunc = [this](const FInputDeviceRay& InReleasePos)
+	{
+		OnClickRelease(InReleasePos);
+	};
+
+	MultiButtonClickDragBehavior->OnClickDragFunc = [this](const FInputDeviceRay& InDragPos)
+	{
+		OnClickDrag(InDragPos);
+	};
+
+	MultiButtonClickDragBehavior->SetDefaultPriority(UE::Editor::ViewportInteractions::VIEWPORT_INTERACTIONS_HIGH_PRIORITY);
+
+	RegisterInputBehavior(MultiButtonClickDragBehavior);
+}
+
+FInputRayHit UViewportPanInteraction::CanBeginClickDragSequence(const FInputDeviceRay& InPressPos)
+{
+	bool bSuccess = CanBeActivated();
+
+	if (FEditorViewportClient* EditorViewportClient = GetEditorViewportClient())
+	{
+		bSuccess &= EditorViewportClient->IsPerspective();
+	}
+
+	return bSuccess ? FInputRayHit(TNumericLimits<float>::Max()) : FInputRayHit();
+}
+
+void UViewportPanInteraction::OnClickPress(const FInputDeviceRay& PressPos)
+{
+	UE::Editor::ViewportInteractions::LOG(TEXT("PAN 1"));
+
+	UViewportDragInteraction::OnClickPress(PressPos);
+
+	if (UViewportInteractionsBehaviorSource* BehaviorSource = GetViewportInteractionsBehaviorSource())
+	{
+		BehaviorSource->SetIsMouseLooking(true);
+	}
+
+	bMultiButtonPanInteraction = false;
+}
+
+void UViewportPanInteraction::OnClickRelease(const FInputDeviceRay& ReleasePos)
+{
+	UE::Editor::ViewportInteractions::LOG(TEXT("PAN 0"));
+
+	UViewportDragInteraction::OnClickRelease(ReleasePos);
+
+	if (UViewportInteractionsBehaviorSource* BehaviorSource = GetViewportInteractionsBehaviorSource())
+	{
+		BehaviorSource->SetIsMouseLooking(false);
+	}
+}
+
+void UViewportPanInteraction::OnDragDelta(float InMouseDeltaX, float InMouseDeltaY)
+{
+	if (FEditorViewportClient* const EditorViewportClient = GetEditorViewportClient())
+	{
+		// Pan left/right/up/down
+		// TODO: In legacy, this bool will also consider if using trackpad or not - possibly need to handle this in a different way with ITF?
+
+		bool bInvert = GetDefault<ULevelEditorViewportSettings>()->bInvertMiddleMousePan;
+
+		FVector LocationDelta;
+
+		const FRotator& ViewRotation = EditorViewportClient->GetViewRotation();
+		if (bMultiButtonPanInteraction)
+		{
+			float Direction = bInvert ? 1.0f : -1.0f;
+			LocationDelta.X = -InMouseDeltaX * Direction * FMath::Sin( ViewRotation.Yaw * PI / 180.0f );
+			LocationDelta.Y = -InMouseDeltaX * -Direction * FMath::Cos( ViewRotation.Yaw * PI / 180.0f );
+			LocationDelta.Z = -Direction * InMouseDeltaY;
+		}
+		else
+		{
+			LocationDelta = FQuat(ViewRotation) * FVector( 0.0f, InMouseDeltaX, -InMouseDeltaY) * (bInvert ? 1.0f : -1.0f);
+		}
+
+		FVector& ViewportLocationDelta = EditorViewportClient->GetViewportNavigationHelper()->LocationDelta;
+		ViewportLocationDelta += LocationDelta;
+	}
+}
+
+FDeviceButtonState UViewportPanInteraction::GetActiveMouseButtonState(const FInputDeviceState& Input)
+{
+	return Input.Mouse.Middle;
+}
diff --git a/Engine/Source/Editor/Experimental/EditorInteractiveToolsFramework/Private/ViewportInteractions/ViewportViewAngleInteraction.cpp b/Engine/Source/Editor/Experimental/EditorInteractiveToolsFramework/Private/ViewportInteractions/ViewportViewAngleInteraction.cpp
new file mode 100644
--- /dev/null
+++ b/Engine/Source/Editor/Experimental/EditorInteractiveToolsFramework/Private/ViewportInteractions/ViewportViewAngleInteraction.cpp
@@ -0,0 +1,86 @@
+// Copyright Epic Games, Inc. All Rights Reserved.
+
+#include "ViewportInteractions/ViewportViewAngleInteraction.h"
+#include "BaseBehaviors/ClickDragBehavior.h"
+#include "BaseBehaviors/SingleClickOrDragBehavior.h"
+#include "EditorViewportClient.h"
+#include "Settings/LevelEditorViewportSettings.h"
+#include "ViewportClientNavigationHelper.h"
+#include "ViewportInteractions/ViewportInteractionsBehaviorSource.h"
+
+UViewportViewAngleInteraction::UViewportViewAngleInteraction()
+{
+	ToolType = UE::Editor::ViewportInteractions::CameraDrag;
+
+	InteractionName = TEXT("View Angle");
+
+	SetClickOrDragPriority(UE::Editor::ViewportInteractions::VIEWPORT_INTERACTIONS_AFTER_GIZMOS_PRIORITY);
+}
+
+FInputRayHit UViewportViewAngleInteraction::CanBeginClickDragSequence(const FInputDeviceRay& InPressPos)
+{
+	bool bSuccess = CanBeActivated();
+	if (FEditorViewportClient* EditorViewportClient = GetEditorViewportClient())
+	{
+		bSuccess &= EditorViewportClient->IsPerspective();
+	}
+
+	return bSuccess ? FInputRayHit(TNumericLimits<float>::Max()) : FInputRayHit();
+}
+
+void UViewportViewAngleInteraction::OnClickPress(const FInputDeviceRay& PressPos)
+{
+	UE::Editor::ViewportInteractions::LOG(TEXT("VIEW ANGLE 1"));
+
+	UViewportDragInteraction::OnClickPress(PressPos);
+
+	if (UViewportInteractionsBehaviorSource* BehaviorSource = GetViewportInteractionsBehaviorSource())
+	{
+		BehaviorSource->SetIsMouseLooking(true);
+	}
+}
+
+void UViewportViewAngleInteraction::OnClickRelease(const FInputDeviceRay& ReleasePos)
+{
+	UE::Editor::ViewportInteractions::LOG(TEXT("VIEW ANGLE 0"));
+
+	UViewportDragInteraction::OnClickRelease(ReleasePos);
+
+	if (UViewportInteractionsBehaviorSource* BehaviorSource = GetViewportInteractionsBehaviorSource())
+	{
+		BehaviorSource->SetIsMouseLooking(false);
+	}
+}
+
+void UViewportViewAngleInteraction::OnDragDelta(float InMouseDeltaX, float InMouseDeltaY)
+{
+	if (FEditorViewportClient* const EditorViewportClient = GetEditorViewportClient())
+	{
+		// Change viewing angle
+
+		// Inverting orbit axis is handled elsewhere
+		const bool bInvertY = !EditorViewportClient->ShouldOrbitCamera()
+						   && GetDefault<ULevelEditorViewportSettings>()->bInvertMouseLookYAxis;
+		float Direction = bInvertY ? -1.0f : 1.0f;
+
+		const ULevelEditorViewportSettings* ViewportSettings = GetDefault<ULevelEditorViewportSettings>();
+
+		FRotator RotDelta;
+		RotDelta.Yaw = InMouseDeltaX * ViewportSettings->MouseSensitivty;
+		RotDelta.Pitch = -InMouseDeltaY * ViewportSettings->MouseSensitivty * Direction;
+		RotDelta.Roll = 0.0f;
+
+		FRotator& ViewportRotationDelta = EditorViewportClient->GetViewportNavigationHelper()->RotationDelta;
+		ViewportRotationDelta += RotDelta;
+	}
+}
+
+FDeviceButtonState UViewportViewAngleInteraction::GetActiveMouseButtonState(const FInputDeviceState& Input)
+{
+	return Input.Mouse.Right;
+}
+
+void UViewportViewAngleInteraction::OnUpdateModifierState(int ModifierID, bool bIsOn)
+{
+	Super::OnUpdateModifierState(ModifierID, bIsOn);
+}
diff --git a/Engine/Source/Editor/Experimental/EditorInteractiveToolsFramework/Private/ViewportInteractions/ViewportZoomInteraction.cpp b/Engine/Source/Editor/Experimental/EditorInteractiveToolsFramework/Private/ViewportInteractions/ViewportZoomInteraction.cpp
new file mode 100644
--- /dev/null
+++ b/Engine/Source/Editor/Experimental/EditorInteractiveToolsFramework/Private/ViewportInteractions/ViewportZoomInteraction.cpp
@@ -0,0 +1,330 @@
+// Copyright Epic Games, Inc. All Rights Reserved.
+
+#include "ViewportInteractions/ViewportZoomInteraction.h"
+#include "BaseBehaviors/ClickDragBehavior.h"
+#include "BaseBehaviors/KeyInputBehavior.h"
+#include "Behaviors/MultiButtonClickDragBehavior.h"
+#include "EditorViewportClient.h"
+#include "Settings/LevelEditorViewportSettings.h"
+#include "Templates/Function.h"
+#include "ViewportClientNavigationHelper.h"
+#include "ViewportInteractions/ViewportInteractionsBehaviorSource.h"
+
+UViewportZoomInteraction::UViewportZoomInteraction()
+{
+	ToolType = UE::Editor::ViewportInteractions::Zoom;
+
+	InteractionName = TEXT("Zoom");
+
+	// Mousewheel zoom
+	UMouseWheelInputBehavior* MouseWheelInputBehavior = NewObject<UMouseWheelInputBehavior>();
+	MouseWheelInputBehavior->Initialize(this);
+	MouseWheelInputBehavior->SetDefaultPriority(UE::Editor::ViewportInteractions::VIEWPORT_INTERACTIONS_DEFAULT_PRIORITY);
+
+	// Keyboard zoom
+	UKeyInputBehavior* KeyInputBehavior = NewObject<UKeyInputBehavior>();
+	// Hardcoded
+	KeyInputBehavior->Initialize(this, { ZoomIn, ZoomOut });
+	KeyInputBehavior->bRequireAllKeys = false;
+	KeyInputBehavior->SetDefaultPriority(UE::Editor::ViewportInteractions::VIEWPORT_INTERACTIONS_DEFAULT_PRIORITY);
+
+	// LMB + RMB + Drag zoom.
+	UMultiButtonClickDragBehavior* MultiButtonClickDragBehavior = NewObject<UMultiButtonClickDragBehavior>();
+	MultiButtonClickDragBehavior->Initialize();
+	MultiButtonClickDragBehavior->EnableButton(EKeys::LeftMouseButton);
+	MultiButtonClickDragBehavior->EnableButton(EKeys::RightMouseButton);
+
+	MultiButtonClickDragBehavior->CanBeginClickDragFunc = [this](const FInputDeviceRay&)
+	{
+		bool bSuccess = IsEnabled();
+		if (FEditorViewportClient* EditorViewportClient = GetEditorViewportClient())
+		{
+			bSuccess &= EditorViewportClient->IsOrtho();
+		}
+
+		return bSuccess ? FInputRayHit(TNumericLimits<float>::Max()) : FInputRayHit();
+	};
+
+	MultiButtonClickDragBehavior->ModifierCheckFunc = [this](const FInputDeviceState& InInputDeviceState)
+	{
+		return InInputDeviceState.Mouse.Left.bDown && InInputDeviceState.Mouse.Right.bDown;
+	};
+
+	MultiButtonClickDragBehavior->OnClickPressFunc = [this](const FInputDeviceRay& InPressPos)
+	{
+		OnClickPress(InPressPos);
+	};
+
+	MultiButtonClickDragBehavior->OnClickReleaseFunc = [this](const FInputDeviceRay& InReleasePos)
+	{
+		OnClickRelease(InReleasePos);
+	};
+
+	MultiButtonClickDragBehavior->OnClickDragFunc = [this](const FInputDeviceRay& InDragPos)
+	{
+		OnClickDrag(InDragPos);
+	};
+
+	MultiButtonClickDragBehavior->SetDefaultPriority(UE::Editor::ViewportInteractions::VIEWPORT_INTERACTIONS_HIGH_PRIORITY);
+
+	RegisterInputBehaviors({ MouseWheelInputBehavior, KeyInputBehavior, MultiButtonClickDragBehavior });
+}
+
+FInputRayHit UViewportZoomInteraction::ShouldRespondToMouseWheel(const FInputDeviceRay& InCurrentPos)
+{
+	return IsEnabled() && !IsMouseLooking() ? FInputRayHit(TNumericLimits<float>::Max()) : FInputRayHit();
+}
+
+void UViewportZoomInteraction::OnMouseWheelScrollUp(const FInputDeviceRay& InCurrentPos)
+{
+	if (!IsMouseLooking())
+	{
+		constexpr bool bZoomIn = false;
+		MouseWheelZoom(bZoomIn);
+	}
+}
+
+void UViewportZoomInteraction::OnMouseWheelScrollDown(const FInputDeviceRay& InCurrentPos)
+{
+	if (!IsMouseLooking())
+	{
+		constexpr bool bZoomIn = true;
+		MouseWheelZoom(bZoomIn);
+	}
+}
+
+void UViewportZoomInteraction::OnKeyPressed(const FKey& InKeyID)
+{
+	if (InKeyID == ZoomIn)
+	{
+		ZoomDirection = -1;
+	}
+	else if (InKeyID == ZoomOut)
+	{
+		ZoomDirection = 1;
+	}
+}
+
+void UViewportZoomInteraction::OnKeyReleased(const FKey& InKeyID)
+{
+	if (InKeyID == EKeys::Add || InKeyID == EKeys::Subtract)
+	{
+		ZoomDirection = 0;
+	}
+}
+
+void UViewportZoomInteraction::OnForceEndCapture()
+{
+	ZoomDirection = 0;
+	OnClickRelease(FInputDeviceRay());
+}
+
+void UViewportZoomInteraction::Tick(float InDeltaTime) const
+{
+	if (FEditorViewportClient* EditorViewportClient = GetEditorViewportClient())
+	{
+		if (ZoomDirection != 0)
+		{
+			const bool bZoomIn = ZoomDirection > 0;
+			if (EditorViewportClient->IsPerspective())
+			{
+				PerspectiveZoom(bZoomIn);
+			}
+			else
+			{
+				// Slower zoom speed, and always centered when done in orthographic view
+				OrthographicZoom(bZoomIn, 0.25f, true);
+			}
+		}
+	}
+
+	UViewportInteraction::Tick(InDeltaTime);
+}
+
+void UViewportZoomInteraction::OrthographicZoom(bool bInZoomIn, float InZoomMultiplier, bool bInForceZoomToCenter) const
+{
+	FEditorViewportClient* EditorViewportClient = GetEditorViewportClient();
+	if (!EditorViewportClient)
+	{
+		return;
+	}
+
+	if (const FViewport* const Viewport = EditorViewportClient->Viewport)
+	{
+		constexpr float Scale = 1.0f;
+
+		// Scrolling the mousewheel up/down zooms the orthogonal viewport in/out.
+		float NewOrthoRatio = 1.0f + ((25.0f * Scale) / CAMERA_ZOOM_DAMPEN) * InZoomMultiplier;
+		if (!bInZoomIn)
+		{
+			NewOrthoRatio = 1.0f / NewOrthoRatio;
+		}
+
+		// Extract current state
+		int32 ViewportWidth = Viewport->GetSizeXY().X;
+		int32 ViewportHeight = Viewport->GetSizeXY().Y;
+
+		FVector OldOffsetFromCenter;
+
+		const bool bCenterZoomAroundCursor = GetDefault<ULevelEditorViewportSettings>()->bCenterZoomAroundCursor;
+
+		if (bCenterZoomAroundCursor && !bInForceZoomToCenter)
+		{
+			// Y is actually backwards, but since we're move the camera opposite the cursor to center, we negate both
+			// therefore the x is negated
+			// X Is backwards, negate it
+			// default to viewport mouse position
+			int32 CenterX = Viewport->GetMouseX();
+			int32 CenterY = Viewport->GetMouseY();
+
+			//TODO: check this one out
+			/*if (EditorViewportClient->ShouldUseMoveCanvasMovement())
+			{
+				// use virtual mouse while dragging (normal mouse is clamped when invisible)
+				CenterX = EditorViewportClient->LastMouseX;
+				CenterY = EditorViewportClient->LastMouseY;
+			}*/
+			int32 DeltaFromCenterX = -(CenterX - (ViewportWidth >> 1));
+			int32 DeltaFromCenterY = (CenterY - (ViewportHeight >> 1));
+
+			switch (EditorViewportClient->GetViewportType())
+			{
+			case LVT_OrthoTop:
+				OldOffsetFromCenter.Set(-DeltaFromCenterY, -DeltaFromCenterX, 0.0f);
+				break;
+			case LVT_OrthoLeft:
+				OldOffsetFromCenter.Set(DeltaFromCenterX, 0.0f, DeltaFromCenterY);
+				break;
+			case LVT_OrthoBack:
+				OldOffsetFromCenter.Set(0.0f, DeltaFromCenterX, DeltaFromCenterY);
+				break;
+			case LVT_OrthoBottom:
+				OldOffsetFromCenter.Set(-DeltaFromCenterY, DeltaFromCenterX, 0.0f);
+				break;
+			case LVT_OrthoRight:
+				OldOffsetFromCenter.Set(-DeltaFromCenterX, 0.0f, DeltaFromCenterY);
+				break;
+			case LVT_OrthoFront:
+				OldOffsetFromCenter.Set(0.0f, -DeltaFromCenterX, DeltaFromCenterY);
+				break;
+			case LVT_OrthoFreelook:
+				//@TODO: CAMERA: How to handle this (todo copied from legacy viewport inputs)
+				break;
+			case LVT_Perspective:
+				break;
+			}
+		}
+
+		// Save off old zoom
+		const float OldUnitsPerPixel = EditorViewportClient->GetOrthoUnitsPerPixel(Viewport);
+
+		// Update zoom based on input
+		const float Zoom = EditorViewportClient->GetOrthoZoom() * NewOrthoRatio;
+
+		const float MinOrthoZoom = static_cast<float>(
+			FMath::Max(GetDefault<ULevelEditorViewportSettings>()->MinimumOrthographicZoom, MIN_ORTHOZOOM)
+		);
+
+		EditorViewportClient->SetOrthoZoom(static_cast<float>(FMath::Clamp(Zoom, MinOrthoZoom, MAX_ORTHOZOOM)));
+
+		if (bCenterZoomAroundCursor)
+		{
+			// This is the equivalent to moving the viewport to center about the cursor, zooming, and moving it back a proportional amount towards the cursor
+			const FVector FinalDelta = (EditorViewportClient->GetOrthoUnitsPerPixel(Viewport) - OldUnitsPerPixel)
+									 * OldOffsetFromCenter;
+
+			// Now move the view location proportionally
+			EditorViewportClient->SetViewLocation(EditorViewportClient->GetViewLocation() + FinalDelta);
+		}
+
+		constexpr bool bInvalidateViews = true;
+
+		// Update linked ortho viewport movement based on updated zoom and view location,
+		EditorViewportClient->UpdateLinkedOrthoViewports(bInvalidateViews);
+		EditorViewportClient->Invalidate(true, true);
+
+		UE::Editor::ViewportInteractions::LOG(TEXT("ITF CAMERA ORTHO ZOOM"));
+	}
+}
+
+void UViewportZoomInteraction::PerspectiveZoom(bool bInZoomIn) const
+{
+	FEditorViewportClient* EditorViewportClient = GetEditorViewportClient();
+
+	// Scrolling the mousewheel up/down moves the perspective viewport forwards/backwards.
+	FVector Drag;
+
+	const FRotator& ViewRotation = EditorViewportClient->GetViewRotation();
+	Drag.X = FMath::Cos(ViewRotation.Yaw * PI / 180.0f) * FMath::Cos(ViewRotation.Pitch * PI / 180.0f);
+	Drag.Y = FMath::Sin(ViewRotation.Yaw * PI / 180.0f) * FMath::Cos(ViewRotation.Pitch * PI / 180.0f);
+	Drag.Z = FMath::Sin(ViewRotation.Pitch * PI / 180.0f);
+
+	if (bInZoomIn)
+	{
+		Drag = -Drag;
+	}
+
+	const float CameraSpeed = EditorViewportClient->GetCameraSpeedSettings().GetCurrentSpeed();
+	Drag *= CameraSpeed * 32.0f;
+
+	constexpr bool bDollyCamera = true;
+	EditorViewportClient->MoveViewportCamera(Drag, FRotator::ZeroRotator, bDollyCamera);
+
+	FEditorDelegates::OnDollyPerspectiveCamera.Broadcast(Drag, EditorViewportClient->ViewIndex);
+
+	EditorViewportClient->Invalidate(true, true);
+
+	UE::Editor::ViewportInteractions::LOG(TEXT("ITF CAMERA PERSPECTIVE ZOOM"));
+}
+
+void UViewportZoomInteraction::MouseWheelZoom(bool bInZoomIn) const
+{
+	FEditorViewportClient* EditorViewportClient = GetEditorViewportClient();
+
+	if (EditorViewportClient->IsPerspective())
+	{
+		PerspectiveZoom(bInZoomIn);
+	}
+	else
+	{
+		OrthographicZoom(bInZoomIn);
+	}
+}
+
+void UViewportZoomInteraction::OnClickPress(const FInputDeviceRay& InPressPos)
+{
+	UE::Editor::ViewportInteractions::LOG(TEXT("ORTHO ZOOM LMB + RMB 1"));
+
+	if (UViewportInteractionsBehaviorSource* BehaviorSource = GetViewportInteractionsBehaviorSource())
+	{
+		BehaviorSource->SetMouseCursorOverride(EMouseCursor::None);
+		BehaviorSource->SetIsMouseLooking(true);
+	}
+
+	LastZoomDragY = static_cast<float>(InPressPos.ScreenPosition.Y);
+}
+
+void UViewportZoomInteraction::OnClickRelease(const FInputDeviceRay& InReleasePos)
+{
+	UE::Editor::ViewportInteractions::LOG(TEXT("ORTHO ZOOM LMB + RMB 0"));
+
+	if (UViewportInteractionsBehaviorSource* BehaviorSource = GetViewportInteractionsBehaviorSource())
+	{
+		BehaviorSource->SetIsMouseLooking(false);
+	}
+}
+
+void UViewportZoomInteraction::OnClickDrag(const FInputDeviceRay& InDragPos)
+{
+	const float ScreenPosY = static_cast<float>(InDragPos.ScreenPosition.Y);
+	const float DeltaY = ScreenPosY - LastZoomDragY;
+
+	if (FEditorViewportClient* const EditorViewportClient = GetEditorViewportClient())
+	{
+		FVector LocationDelta(0.0f, 0.0f, -DeltaY);
+		FVector& ViewportLocationDelta = EditorViewportClient->GetViewportNavigationHelper()->LocationDelta;
+		ViewportLocationDelta += LocationDelta;
+	}
+
+	LastZoomDragY = ScreenPosY;
+}
diff --git a/Engine/Source/Editor/Experimental/EditorInteractiveToolsFramework/Public/EditorDragTools/DragToolsBehaviourSource.h b/Engine/Source/Editor/Experimental/EditorInteractiveToolsFramework/Public/EditorDragTools/DragToolsBehaviourSource.h
--- a/Engine/Source/Editor/Experimental/EditorInteractiveToolsFramework/Public/EditorDragTools/DragToolsBehaviourSource.h
+++ b/Engine/Source/Editor/Experimental/EditorInteractiveToolsFramework/Public/EditorDragTools/DragToolsBehaviourSource.h
@@ -5,8 +5,35 @@
 #include "EditorViewportClient.h"
 #include "InputBehaviorSet.h"
 #include "UObject/Object.h"
+#include "ViewportInteractions/ViewportInteraction.h"
 #include "DragToolsBehaviourSource.generated.h"
 
+
+namespace UE::Editor::DragTools
+{
+
+// In case New TRS Gizmos are enabled, we need Mouse Drag + Key modifier Drag Tools Behaviors to be processed before
+// the Gizmo ones, and also before viewport interactions - otherwise Drag Behaviors will never be triggered. See:
+// * UTransformGizmo::SetupBehaviors()			[DEFAULT_GIZMO_PRIORITY]
+// * UTransformGizmo::SetupIndirectBehaviors()	[DEFAULT_GIZMO_PRIORITY - 1]
+
+static constexpr int HANDLED_BEFORE_VIEWPORT_INTERACTIONS =
+	UE::Editor::ViewportInteractions::VIEWPORT_INTERACTIONS_HIGH_PRIORITY - 1;
+
+// Conversely, tools which are activated with just mouse buttons input (no key modifiers, mouse drag behavior only) need to be processed after Gizmo Behaviors
+static constexpr int HANDLED_AFTER_GIZMO_PRIORITY = FInputCapturePriority::DEFAULT_GIZMO_PRIORITY + 1;
+
+static constexpr int DRAG_TOOLS_DEFAULT_PRIORITY = HANDLED_BEFORE_VIEWPORT_INTERACTIONS - 10;
+static constexpr int DRAG_TOOLS_LOW_PRIORITY = HANDLED_BEFORE_VIEWPORT_INTERACTIONS - 1;
+
+static constexpr int DUPLICATE_SELECTION_PRIORITY = DRAG_TOOLS_DEFAULT_PRIORITY;
+static constexpr int FRUSTUM_SELECT_PRIORITY = DRAG_TOOLS_DEFAULT_PRIORITY;
+static constexpr int VIEWPORT_CHANGE_PRIORITY = DRAG_TOOLS_DEFAULT_PRIORITY;
+static constexpr int BOX_SELECT_PRIORITY = HANDLED_AFTER_GIZMO_PRIORITY;
+static constexpr int MEASURE_PRIORITY = DRAG_TOOLS_LOW_PRIORITY;
+
+}
+
 class FEditorDuplicateDragSelection;
 class FEditorDragToolBehaviorTarget;
 class FEditorMoveCameraWithObject;
diff --git a/Engine/Source/Editor/Experimental/EditorInteractiveToolsFramework/Public/ViewportInteractions/ViewportCameraRotateInteraction.h b/Engine/Source/Editor/Experimental/EditorInteractiveToolsFramework/Public/ViewportInteractions/ViewportCameraRotateInteraction.h
new file mode 100644
--- /dev/null
+++ b/Engine/Source/Editor/Experimental/EditorInteractiveToolsFramework/Public/ViewportInteractions/ViewportCameraRotateInteraction.h
@@ -0,0 +1,49 @@
+// Copyright Epic Games, Inc. All Rights Reserved.
+
+#pragma once
+
+#include "BaseBehaviors/BehaviorTargetInterfaces.h"
+#include "ViewportInteraction.h"
+
+#include "ViewportCameraRotateInteraction.generated.h"
+
+class UKeyInputBehavior;
+struct FViewportClientNavigationHelper;
+
+/**
+ * A viewport interaction to Rotate the camera using keyboard inputs. Only works while the camera is being controlled using the mouse
+ */
+UCLASS(Transient)
+class EDITORINTERACTIVETOOLSFRAMEWORK_API UViewportCameraRotateInteraction final
+	: public UViewportInteraction
+	, public IKeyInputBehaviorTarget
+{
+	GENERATED_BODY()
+
+public:
+	UViewportCameraRotateInteraction();
+
+	//~ Begin IKeyInputBehaviorTarget
+	virtual void OnKeyPressed(const FKey& InKeyID) override;
+	virtual void OnKeyReleased(const FKey& InKeyID) override;
+	//~ End IKeyInputBehaviorTarget
+
+	//~ Begin UViewportInteraction
+	virtual void Tick(float InDeltaTime) const override;
+	//~ End UViewportInteraction
+
+protected:
+	//~ Begin UViewportInteractionBase
+	virtual void OnCommandChordChanged() override;
+	virtual TArray<TSharedPtr<FUICommandInfo>> GetCommands() const override;
+	//~ End UViewportInteractionBase
+
+	void UpdateKeyState(const FKey& InKeyID, bool bInIsPressed);
+	TArray<FKey> GetKeys() const;
+
+private:
+	TWeakObjectPtr<UKeyInputBehavior> KeyInputBehaviorWeak;
+
+	float RotatePitchImpulse = 0.0f;
+	float RotateYawImpulse = 0.0f;
+};
diff --git a/Engine/Source/Editor/Experimental/EditorInteractiveToolsFramework/Public/ViewportInteractions/ViewportCameraSpeedMouseWheelInteraction.h b/Engine/Source/Editor/Experimental/EditorInteractiveToolsFramework/Public/ViewportInteractions/ViewportCameraSpeedMouseWheelInteraction.h
new file mode 100644
--- /dev/null
+++ b/Engine/Source/Editor/Experimental/EditorInteractiveToolsFramework/Public/ViewportInteractions/ViewportCameraSpeedMouseWheelInteraction.h
@@ -0,0 +1,33 @@
+﻿// Copyright Epic Games, Inc. All Rights Reserved.
+
+#pragma once
+
+#include "ViewportInteraction.h"
+#include "BaseBehaviors/BehaviorTargetInterfaces.h"
+#include "ViewportCameraSpeedMouseWheelInteraction.generated.h"
+
+/**
+ *
+ */
+UCLASS()
+class EDITORINTERACTIVETOOLSFRAMEWORK_API UViewportCameraSpeedMouseWheelInteraction
+	: public UViewportInteraction
+	, public IMouseWheelBehaviorTarget
+{
+	GENERATED_BODY()
+public:
+
+	UViewportCameraSpeedMouseWheelInteraction();
+
+	//~ Begin IMouseWheelBehaviorTarget
+	virtual FInputRayHit ShouldRespondToMouseWheel(const FInputDeviceRay& CurrentPos) override;
+	virtual void OnMouseWheelScrollUp(const FInputDeviceRay& CurrentPos) override;
+	virtual void OnMouseWheelScrollDown(const FInputDeviceRay& CurrentPos) override;
+	//~ Begin IMouseWheelBehaviorTarget
+
+	/**
+ * Increments the current camera speed using the specified change factor, e.g. Speed += Speed * UpdateFactor
+ * @param InUpdateFactor positive or negative factor
+ */
+	void UpdateCameraSpeed(float InUpdateFactor) const;
+};
diff --git a/Engine/Source/Editor/Experimental/EditorInteractiveToolsFramework/Public/ViewportInteractions/ViewportCameraTranslateInteraction.h b/Engine/Source/Editor/Experimental/EditorInteractiveToolsFramework/Public/ViewportInteractions/ViewportCameraTranslateInteraction.h
new file mode 100644
--- /dev/null
+++ b/Engine/Source/Editor/Experimental/EditorInteractiveToolsFramework/Public/ViewportInteractions/ViewportCameraTranslateInteraction.h
@@ -0,0 +1,55 @@
+// Copyright Epic Games, Inc. All Rights Reserved.
+
+#pragma once
+
+#include "BaseBehaviors/BehaviorTargetInterfaces.h"
+#include "ViewportInteraction.h"
+
+#include "ViewportCameraTranslateInteraction.generated.h"
+
+class FUICommandInfo;
+class UKeyInputBehavior;
+struct FViewportClientNavigationHelper;
+
+/**
+ * A viewport interaction used to Translate the camera using keyboard inputs while the camera is being moved using the mouse
+ */
+UCLASS(Transient)
+class EDITORINTERACTIVETOOLSFRAMEWORK_API UViewportCameraTranslateInteraction final
+	: public UViewportInteraction
+	, public IKeyInputBehaviorTarget
+{
+	GENERATED_BODY()
+
+public:
+	UViewportCameraTranslateInteraction();
+
+	//~ Begin IKeyInputBehaviorTarget
+	virtual void OnKeyPressed(const FKey& InKeyID) override;
+	virtual void OnKeyReleased(const FKey& InKeyID) override;
+	virtual void OnForceEndCapture() override;
+	//~ End IKeyInputBehaviorTarget
+
+	//~ Begin UViewportInteraction
+	virtual void Tick(float InDeltaTime) const override;
+	//~ End UViewportInteraction
+
+protected:
+	//~ Begin UViewportInteractionBase
+	virtual void OnCommandChordChanged() override;
+	virtual TArray<TSharedPtr<FUICommandInfo>> GetCommands() const override;
+	//~ End UViewportInteractionBase
+
+	void UpdateKeyState(const FKey& InKeyID, bool bInIsPressed);
+	TArray<FKey> GetKeys() const;
+
+private:
+	TWeakObjectPtr<UKeyInputBehavior> KeyInputBehaviorWeak;
+
+	bool bUseNumpadKey = true;
+
+	float ForwardBackwardImpulse = 0.0f;
+	float RightLeftImpulse = 0.0f;
+	float WorldUpDownImpulse = 0.0f;
+	float LocalUpDownImpulse = 0.0f;
+};
diff --git a/Engine/Source/Editor/Experimental/EditorInteractiveToolsFramework/Public/ViewportInteractions/ViewportCommandsInteraction.h b/Engine/Source/Editor/Experimental/EditorInteractiveToolsFramework/Public/ViewportInteractions/ViewportCommandsInteraction.h
new file mode 100644
--- /dev/null
+++ b/Engine/Source/Editor/Experimental/EditorInteractiveToolsFramework/Public/ViewportInteractions/ViewportCommandsInteraction.h
@@ -0,0 +1,55 @@
+// Copyright Epic Games, Inc. All Rights Reserved.
+
+#pragma once
+
+#include "BaseBehaviors/BehaviorTargetInterfaces.h"
+#include "ViewportInteraction.h"
+#include "ViewportCommandsInteraction.generated.h"
+
+class FUICommandInfo;
+class FUICommandList;
+class UKeyInputBehavior;
+
+/**
+ * A viewport interaction class used to add support for a list of FUICommandInfo Commands to the current viewport
+ */
+UCLASS()
+class EDITORINTERACTIVETOOLSFRAMEWORK_API UViewportCommandsInteraction : public UViewportInteraction, public IKeyInputBehaviorTarget
+{
+	GENERATED_BODY()
+
+public:
+	UViewportCommandsInteraction();
+
+	/**
+	 * Sets which commands will be handled by this Command Interaction
+	 */
+	void SetCommands(const TSharedPtr<FUICommandList>& InEditorCommandList, const TArray<TSharedPtr<FUICommandInfo>>& InCommands);
+
+	// ~ Begin IKeyInputBehaviorTarget
+	virtual void OnKeyPressed(const FKey& InKeyID) override;
+	virtual void OnKeyReleased(const FKey& InKeyID) override;
+	// ~ End IKeyInputBehaviorTarget
+
+	/** Set to true to execute commands on Press instead of Release */
+	bool bExecuteOnPress = false;
+
+protected:
+	// ~ Begin UViewportInteractionBase
+	virtual void OnCommandChordChanged() override;
+	virtual TArray<TSharedPtr<FUICommandInfo>> GetCommands() const override;
+	// ~ End UViewportInteractionBase
+
+private:
+	TArray<FKey> GetKeys();
+
+	void ExecuteCommand(const FKey& InKeyID);
+	bool MatchesCommandKey(const TSharedPtr<FUICommandInfo>& InCommandInfo, const FKey& InKeyID) const;
+
+	bool CanBeActivated() const;
+
+	TSharedPtr<FUICommandList> EditorCommandList;
+	TArray<TSharedPtr<FUICommandInfo>> Commands;
+
+	TWeakObjectPtr<UKeyInputBehavior> KeyInputBehaviorWeak;
+};
diff --git a/Engine/Source/Editor/Experimental/EditorInteractiveToolsFramework/Public/ViewportInteractions/ViewportDragInteraction.h b/Engine/Source/Editor/Experimental/EditorInteractiveToolsFramework/Public/ViewportInteractions/ViewportDragInteraction.h
new file mode 100644
--- /dev/null
+++ b/Engine/Source/Editor/Experimental/EditorInteractiveToolsFramework/Public/ViewportInteractions/ViewportDragInteraction.h
@@ -0,0 +1,70 @@
+// Copyright Epic Games, Inc. All Rights Reserved.
+
+#pragma once
+
+#include "BaseBehaviors/BehaviorTargetInterfaces.h"
+#include "ViewportInteraction.h"
+
+#include "ViewportDragInteraction.generated.h"
+
+class UClickDragInputBehavior;
+class USingleClickOrDragInputBehavior;
+
+/**
+ * A base class which can be used to implement interactions based on mouse drag actions
+ */
+UCLASS(Transient, Abstract)
+class EDITORINTERACTIVETOOLSFRAMEWORK_API UViewportDragInteraction
+	: public UViewportInteraction
+	, public IClickDragBehaviorTarget
+{
+	GENERATED_BODY()
+public:
+	UViewportDragInteraction();
+
+	//~ Begin IClickDragBehaviorTarget
+	virtual FInputRayHit CanBeginClickDragSequence(const FInputDeviceRay& InPressPos) override;
+	virtual void OnClickPress(const FInputDeviceRay& InPressPos) override;
+	virtual void OnClickDrag(const FInputDeviceRay& InDragPos) override;
+	virtual void OnClickRelease(const FInputDeviceRay& InReleasePos) override;
+	virtual void OnTerminateDragSequence() override {};
+	//~ End IClickDragBehaviorTarget
+
+	//~ Begin IModifierToggleBehaviorTarget
+	virtual void OnForceEndCapture() override;
+	//~ End IModifierToggleBehaviorTarget
+
+	//~ Begin FViewportInputToolBase
+	virtual void Initialize(UViewportInteractionsBehaviorSource* const InViewportInteractionsBehaviorSource) override;
+	//~ End FViewportInputToolBase
+
+	virtual void OnDragDelta(float InMouseDeltaX, float InMouseDeltaY)
+	{
+	}
+
+	virtual bool CanBeActivated(const FInputDeviceState& InInputDeviceState = FInputDeviceState()) const
+	{
+		return IsEnabled();
+	}
+
+	/**
+	 * Provides a click behavior target to be called when drag interaction fails
+	 * @param InClickBehaviorTarget the Click behavior target responding to clicks
+	 */
+	void SetClickTarget(IClickBehaviorTarget* InClickBehaviorTarget);
+
+protected:
+
+	void SetClickOrDragPriority(int InPriority) const;
+	
+	/**
+	 * Returns the state of the mouse button currently used by this Interaction. Left mouse by default, override for other buttons.
+	 */
+	virtual FDeviceButtonState GetActiveMouseButtonState(const FInputDeviceState& Input);
+
+	TWeakObjectPtr<USingleClickOrDragInputBehavior> ClickOrDragInputBehavior;
+
+	bool bIsDragging;
+	int32 MouseX;
+	int32 MouseY;
+};
diff --git a/Engine/Source/Editor/Experimental/EditorInteractiveToolsFramework/Public/ViewportInteractions/ViewportFOVInteraction.h b/Engine/Source/Editor/Experimental/EditorInteractiveToolsFramework/Public/ViewportInteractions/ViewportFOVInteraction.h
new file mode 100644
--- /dev/null
+++ b/Engine/Source/Editor/Experimental/EditorInteractiveToolsFramework/Public/ViewportInteractions/ViewportFOVInteraction.h
@@ -0,0 +1,55 @@
+// Copyright Epic Games, Inc. All Rights Reserved.
+
+#pragma once
+
+#include "BaseBehaviors/BehaviorTargetInterfaces.h"
+#include "ViewportInteraction.h"
+
+#include "ViewportFOVInteraction.generated.h"
+
+class UKeyInputBehavior;
+
+/**
+ * A viewport interaction used to change the camera Field of View using either keyboard inputs or the mouse wheel
+ */
+UCLASS(Transient)
+class EDITORINTERACTIVETOOLSFRAMEWORK_API UViewportFOVInteraction
+	: public UViewportInteraction
+	, public IKeyInputBehaviorTarget
+{
+	GENERATED_BODY()
+
+public:
+	UViewportFOVInteraction();
+	virtual ~UViewportFOVInteraction() override;
+
+	//~ Begin IKeyInputBehaviorTarget
+	virtual void OnKeyPressed(const FKey& InKeyID) override;
+	virtual void OnKeyReleased(const FKey& InKeyID) override;
+	virtual void OnForceEndCapture() override;
+	//~ End IKeyInputBehaviorTarget
+
+	//~ Begin UViewportInteraction
+	virtual void Tick(float InDeltaTime) const override;
+	virtual void Initialize(UViewportInteractionsBehaviorSource* const InViewportInteractionsBehaviorSource) override;
+	//~ End UViewportInteraction
+
+protected:
+	//~ Begin UViewportInteractionBase
+	virtual void OnCommandChordChanged() override;
+	virtual TArray<TSharedPtr<FUICommandInfo>> GetCommands() const override;
+	//~ End UViewportInteractionBase
+
+	void UpdateKeyState(const FKey& InKeyID, bool bInIsPressed);
+	TArray<FKey> GetKeys() const;
+	TArray<FKey> GetNumpadKeys() const;
+
+private:
+	void OnMouseLookingChanged(bool bInIsMouseLooking);
+	void ResetImpulse();
+
+	bool bUseNumpadKey = true;
+	float ZoomOutInImpulse = 0;
+
+	TWeakObjectPtr<UKeyInputBehavior> KeyInputBehaviorWeak;
+};
diff --git a/Engine/Source/Editor/Experimental/EditorInteractiveToolsFramework/Public/ViewportInteractions/ViewportInteraction.h b/Engine/Source/Editor/Experimental/EditorInteractiveToolsFramework/Public/ViewportInteractions/ViewportInteraction.h
new file mode 100644
--- /dev/null
+++ b/Engine/Source/Editor/Experimental/EditorInteractiveToolsFramework/Public/ViewportInteractions/ViewportInteraction.h
@@ -0,0 +1,141 @@
+// Copyright Epic Games, Inc. All Rights Reserved.
+
+#pragma once
+
+#include "InputBehavior.h"
+#include "Tools/EdModeInteractiveToolsContext.h"
+
+#include "ViewportInteraction.generated.h"
+
+class FEditorViewportClient;
+class FUICommandInfo;
+class FViewportCameraMover;
+class FViewportClickHandler;
+class FViewportCommandsHandler;
+class UInputBehavior;
+class UViewportInteractionsBehaviorSource;
+
+namespace UE::Editor::ViewportInteractions
+{
+constexpr int ShiftKeyMod = 1;
+constexpr int AltKeyMod = 2;
+constexpr int CtrlKeyMod = 3;
+
+constexpr int LeftMouseButtonMod = 4;
+constexpr int RightMouseButtonMod = 5;
+constexpr int MiddleMouseButtonMod = 6;
+
+static constexpr int VIEWPORT_INTERACTIONS_DEFAULT_PRIORITY = FInputCapturePriority::DEFAULT_GIZMO_PRIORITY - 4;
+static constexpr int VIEWPORT_INTERACTIONS_LOW_PRIORITY = FInputCapturePriority::DEFAULT_GIZMO_PRIORITY - 3;
+static constexpr int VIEWPORT_INTERACTIONS_HIGH_PRIORITY = FInputCapturePriority::DEFAULT_GIZMO_PRIORITY - 5;
+static constexpr int VIEWPORT_INTERACTIONS_AFTER_GIZMOS_PRIORITY = FInputCapturePriority::DEFAULT_GIZMO_PRIORITY + 1;
+
+/**
+ * Possible types of Viewport Interactions.
+ * Use by FEditorViewportClient to check whether a specific interaction is handled by an ITF-based interaction, or if it
+ * should be handled by legacy logic
+ */
+static FName Zoom = TEXT("Zoom");
+static FName Orbit = TEXT("Orbit");
+static FName CameraDrag= TEXT("CameraDrag");
+static FName CameraTranslate = TEXT("CameraTranslate");
+static FName CameraRotate = TEXT("CameraRotate");
+static FName FOV = TEXT("FOV");
+static FName Commands = TEXT("Commands");
+static FName SingleClick = TEXT("SingleClick");
+
+} // namespace UE::Editor::ViewportInteractions
+
+/**
+ * Base class used to implement viewport interactions.
+ * Has some helper functions to check if Shift, Alt, Ctrl modifier keys and Mouse Button keys are down.
+ * Tick can be used for code which needs to affect things like camera movements in Editor Viewport Client
+ */
+UCLASS(Transient, Abstract)
+class EDITORINTERACTIVETOOLSFRAMEWORK_API UViewportInteraction : public UObject
+{
+	GENERATED_BODY()
+
+public:
+	UViewportInteraction();
+
+	bool IsShiftDown() const;
+	bool IsAltDown() const;
+	bool IsCtrlDown() const;
+
+	bool IsLeftMouseButtonDown() const;
+	bool IsMiddleMouseButtonDown() const;
+	bool IsRightMouseButtonDown() const;
+
+	bool IsAnyMouseButtonDown() const;
+
+	bool IsMouseLooking() const;
+
+	void SetEnabled(bool bInEnabled);
+	bool IsEnabled() const;
+
+	virtual void Tick(float InDeltaTime) const
+	{
+	}
+
+	TArray<TObjectPtr<UInputBehavior>> GetInputBehaviors()
+	{
+		return InputBehaviors;
+	}
+
+	FName GetInteractionName() const;
+
+	/**
+	 * Initialization code which can be expanded by derived classes
+	 */
+	virtual void Initialize(UViewportInteractionsBehaviorSource* const InViewportInteractionsBehaviorSource);
+
+	UViewportInteractionsBehaviorSource* GetViewportInteractionsBehaviorSource() const;
+
+	//~ Begin UObject
+	virtual void BeginDestroy() override;
+	//~ Edn UObject
+
+protected:
+	FEditorViewportClient* GetEditorViewportClient() const;
+
+	/**
+	 * Define what should happen to UInputBehavior(s) when an "observed" command chord changes
+	 */
+	virtual void OnCommandChordChanged()
+	{
+	}
+
+	/**
+	 * Return a list of commands used by a key-based UInputBehavior.
+	 * Used for its initialization, or its refresh, whenever a command chord changes
+	 */
+	virtual TArray<TSharedPtr<FUICommandInfo>> GetCommands() const
+	{
+		return {};
+	}
+
+	void RegisterInputBehavior(UInputBehavior* InBehavior);
+	void RegisterInputBehaviors(TArray<UInputBehavior*> InBehaviors);
+
+	UPROPERTY()
+	TWeakObjectPtr<UViewportInteractionsBehaviorSource> ViewportInteractionsBehaviorSource;
+
+	bool bEnabled = true;
+
+	FName ToolType;
+
+	/**
+	 * Name for this interaction, used by Viewport Interactions Behavior Source to identify it.
+	 * If none is set, static class name will be used.
+	 */
+	TOptional<FName> InteractionName;
+
+private:
+	void OnUserDefinedChordChanged(const FUICommandInfo& InCommandInfo);
+
+	UPROPERTY()
+	TArray<TObjectPtr<UInputBehavior>> InputBehaviors;
+
+	FDelegateHandle OnChordChangedDelegateHandle;
+};
diff --git a/Engine/Source/Editor/Experimental/EditorInteractiveToolsFramework/Public/ViewportInteractions/ViewportInteractionsBehaviorSource.h b/Engine/Source/Editor/Experimental/EditorInteractiveToolsFramework/Public/ViewportInteractions/ViewportInteractionsBehaviorSource.h
new file mode 100644
--- /dev/null
+++ b/Engine/Source/Editor/Experimental/EditorInteractiveToolsFramework/Public/ViewportInteractions/ViewportInteractionsBehaviorSource.h
@@ -0,0 +1,269 @@
+// Copyright Epic Games, Inc. All Rights Reserved.
+
+#pragma once
+
+#include "BaseBehaviors/BehaviorTargetInterfaces.h"
+#include "InputBehaviorSet.h"
+#include "Tools/EdModeInteractiveToolsContext.h"
+#include "ViewportInteraction.h"
+
+#include "ViewportInteractionsBehaviorSource.generated.h"
+
+class FCameraControllerUserImpulseData;
+class FEditorCameraController;
+class FEditorViewportClient;
+class FUICommandInfo;
+class FViewportCameraMover;
+class FViewportClickHandler;
+class UEditorInteractiveToolsContext;
+class UInputBehaviorSet;
+class UViewportInteraction;
+
+namespace UE::Editor::ViewportInteractions
+{
+
+DECLARE_MULTICAST_DELEGATE(FOnEditorViewportInteractionsToggleDelegate)
+
+/**
+ * Returns true if ITF-based input tools should be used.
+ */
+EDITORINTERACTIVETOOLSFRAMEWORK_API bool UseEditorViewportInteractions();
+
+/**
+ * Helps hiding unwanted logs. For testing.
+ */
+EDITORINTERACTIVETOOLSFRAMEWORK_API bool IsVerbose();
+
+/**
+ * Toggles ITF Viewport Interactions On/Off
+ */
+EDITORINTERACTIVETOOLSFRAMEWORK_API void ToggleEditorViewportInteractions(bool bInEnable);
+
+EDITORINTERACTIVETOOLSFRAMEWORK_API FOnEditorViewportInteractionsToggleDelegate& OnEditorViewportInteractionsActivated();
+EDITORINTERACTIVETOOLSFRAMEWORK_API FOnEditorViewportInteractionsToggleDelegate& OnEditorViewportInteractionsDeactivated();
+
+/**
+ * Checks whether the specified Command should be triggered by the specified Key
+ */
+bool CommandMatchesKey(const TSharedPtr<FUICommandInfo>& InCommandInfo, const FKey& InKeyID);
+
+EDITORINTERACTIVETOOLSFRAMEWORK_API void LOG(const TCHAR* InMessage);
+
+} // namespace UE::Editor::ViewportInteractions
+
+DECLARE_MULTICAST_DELEGATE_OneParam(FOnMouseLookingChanged, bool /* bIsMouseLooking*/)
+
+	/**
+	 * This class hosts a list of ITF Viewport Interactions and a UBehaviorSet containing the UInteractionBehavior
+	 * instances required by those Viewport Interactions.
+	 * It also acts as a behavior target, registering the status of Shift, Alt, Ctrl modifier keys. That and other
+	 * pieces of information can be accessed externally, so that e.g. Viewport Interactions can know about modifier
+	 * states without having to implement themselves what would be duplicate logic to handle those inputs. This class
+	 * can also be used to know whether the Viewport camera is being moved using the mouse,or to mark it as such.
+	 */
+	UCLASS(Transient, MinimalAPI)
+class UViewportInteractionsBehaviorSource final
+	: public UObject
+	, public IInputBehaviorSource
+	, public IModifierToggleBehaviorTarget
+{
+	GENERATED_BODY()
+
+public:
+	//~ Begin IInputBehaviorSource
+	virtual const UInputBehaviorSet* GetInputBehaviors() const override
+	{
+		return BehaviorSet;
+	}
+	//~ End IClickDragBehaviorTarget
+
+	//~ Begin IModifierToggleBehaviorTarget
+	virtual void OnUpdateModifierState(int InModifierID, bool bInIsOn) override;
+	virtual void OnForceEndCapture() override;
+	//~ End IModifierToggleBehaviorTarget
+
+	/**
+	 * Will initialize the required behaviors and the Behavior Set
+	 */
+	EDITORINTERACTIVETOOLSFRAMEWORK_API void Initialize(const UEditorInteractiveToolsContext* InInteractiveToolsContext);
+
+	/**
+	 * Register this Input Behavior Source to the InputRouter
+	 */
+	EDITORINTERACTIVETOOLSFRAMEWORK_API void RegisterBehaviorSources();
+
+	/**
+	 * Deregister this Input Behavior Source from the InputRouter
+	 */
+	EDITORINTERACTIVETOOLSFRAMEWORK_API void DeregisterBehaviorSources();
+
+	/**
+	 * Call Tick on those Viewport Interactions which require it (e.g. Camera Mover)
+	 */
+	EDITORINTERACTIVETOOLSFRAMEWORK_API void Tick(float InDeltaTime) const;
+
+	/**
+	 * Add the specified interactions to this Behavior Source
+	 */
+	EDITORINTERACTIVETOOLSFRAMEWORK_API void AddInteractions(
+		const TArray<const UClass*>& InInteractions, bool bInReregister = false
+	);
+
+	/**
+	 *
+	 * Add the specified interaction to this Behavior Source
+	 */
+	EDITORINTERACTIVETOOLSFRAMEWORK_API UViewportInteraction* AddInteraction(
+		const UClass* InInteractionClass, bool bInReregister = false
+	);
+
+	/**
+	 * De-registers behaviors, removes viewport interactions.
+	 * Behaviors added by Initialize will be kept around.
+	 */
+	EDITORINTERACTIVETOOLSFRAMEWORK_API void Reset();
+
+	/**
+	 * Returns the first Viewport Interaction of type T
+	 */
+	template<class T>
+	T* GetTypedViewportInteraction() const
+	{
+		T* Ret = nullptr;
+		for (const TPair<FName, TObjectPtr<UViewportInteraction>>& Pair : ViewportInteractions)
+		{
+			if (UViewportInteraction* Interaction = Pair.Value)
+			{
+				if (Interaction->IsA<T>())
+				{
+					Ret = Cast<T>(Interaction);
+				}
+			}
+		}
+
+		return Ret;
+	}
+
+	/**
+	 * Returns all Viewport Interactions of type T
+	 */
+	template<class T>
+	TArray<T*> GetTypedViewportInteractions() const
+	{
+		TArray<T*> Return;
+
+		for (const TPair<FName, TObjectPtr<UViewportInteraction>>& Pair : ViewportInteractions)
+		{
+			if (UViewportInteraction* Interaction = Pair.Value)
+			{
+				if (Interaction->IsA<T>())
+				{
+					Return.AddUnique(Cast<T>(Interaction));
+				}
+			}
+		}
+
+		return Return;
+	}
+
+	/**
+	 * Sets current mouse cursor override
+	 */
+	void SetMouseCursorOverride(EMouseCursor::Type InMouseCursor);
+
+	/**
+	 * Use the templated version in case there is some class specific initialization required by the newly created interaction,
+	 * since this method will return the newly instantiated interaction
+	 */
+	template<class T>
+	T* AddInteraction(bool bInReregister = false)
+	{
+		return static_cast<T*>(AddInteraction(T::StaticClass(), bInReregister));
+	}
+
+	bool IsShiftDown() const
+	{
+		return bIsShiftDown;
+	}
+
+	bool IsAltDown() const
+	{
+		return bIsAltDown;
+	}
+
+	bool IsCtrlDown() const
+	{
+		return bIsCtrlDown;
+	}
+
+	bool IsLeftMouseButtonDown() const
+	{
+		return bIsLeftMouseButtonDown;
+	}
+
+	bool IsMiddleMouseButtonDown() const
+	{
+		return bIsMiddleMouseButtonDown;
+	}
+
+	bool IsRightMouseButtonDown() const
+	{
+		return bIsRightMouseButtonDown;
+	}
+
+	bool IsMouseLooking() const;
+
+	void SetIsMouseLooking(bool bInIsLooking);
+
+	FEditorViewportClient* GetEditorViewportClient() const;
+
+	/**
+	 * Marks the specified interaction type as active (meaning a viewport interaction of that type is currently active)
+	 */
+	void SetViewportInteractionActive(FName InViewportInteraction, bool bInActive);
+
+	/**
+	 * Checks if there is any active viewport interaction of the specified type
+	 */
+	bool EDITORINTERACTIVETOOLSFRAMEWORK_API IsViewportInteractionActive(FName InViewportInteraction);
+
+	FOnMouseLookingChanged& OnMouseLookingStateChanged()
+	{
+		return OnMouseLookingChangedDelegate;
+	}
+
+private:
+	/**
+	 * All Input Behaviors from registered Viewport Interactions should be added to this Behavior Set (done with ::AddInteraction() )
+	 */
+	UPROPERTY(Transient)
+	TObjectPtr<UInputBehaviorSet> BehaviorSet;
+
+	/**
+	 * A collection of ITF-based Viewport Interactions
+	 * Key is the interaction name (class name when no Interaction name is specified)
+	 */
+	UPROPERTY(Transient)
+	TMap<FName, TObjectPtr<UViewportInteraction>> ViewportInteractions;
+
+	bool bIsShiftDown;
+	bool bIsCtrlDown;
+	bool bIsAltDown;
+
+	bool bIsLeftMouseButtonDown;
+	bool bIsMiddleMouseButtonDown;
+	bool bIsRightMouseButtonDown;
+
+	EMouseCursor::Type CursorOverride = EMouseCursor::None;
+
+	UPROPERTY(Transient)
+	TWeakObjectPtr<const UEditorInteractiveToolsContext> EditorInteractiveToolsContextWeak;
+
+	/**
+	 * Keeping track of active Viewport Interactions, so that FEditorViewportClient can disable legacy behaviors
+	 * accordingly e.g. if viewport interactions for zooming are enabled, skip the legacy ones doing the same
+	 */
+	TMap<FName, bool> ViewportInteractionsStatusMap;
+
+	FOnMouseLookingChanged OnMouseLookingChangedDelegate;
+};
diff --git a/Engine/Source/Editor/Experimental/EditorInteractiveToolsFramework/Public/ViewportInteractions/ViewportInteractionsUtilities.h b/Engine/Source/Editor/Experimental/EditorInteractiveToolsFramework/Public/ViewportInteractions/ViewportInteractionsUtilities.h
new file mode 100644
--- /dev/null
+++ b/Engine/Source/Editor/Experimental/EditorInteractiveToolsFramework/Public/ViewportInteractions/ViewportInteractionsUtilities.h
@@ -0,0 +1,16 @@
+﻿// Copyright Epic Games, Inc. All Rights Reserved.
+
+#pragma once
+
+#include "HAL/Platform.h"
+
+class UViewportInteractionsBehaviorSource;
+
+namespace UE::Editor::ViewportInteractions
+{
+/**
+ * Convenience function to add "standard" movement interactions to the specified Viewport Interactions Behavior Source
+ * e.g. Movement, Rotation, View Angle, Orbit, Panning, FOV, etc
+ */
+EDITORINTERACTIVETOOLSFRAMEWORK_API void AddDefaultCameraMovementInteractions(UViewportInteractionsBehaviorSource* InInteractionsBehaviorSource);
+}
diff --git a/Engine/Source/Editor/Experimental/EditorInteractiveToolsFramework/Public/ViewportInteractions/ViewportMoveYawInteraction.h b/Engine/Source/Editor/Experimental/EditorInteractiveToolsFramework/Public/ViewportInteractions/ViewportMoveYawInteraction.h
new file mode 100644
--- /dev/null
+++ b/Engine/Source/Editor/Experimental/EditorInteractiveToolsFramework/Public/ViewportInteractions/ViewportMoveYawInteraction.h
@@ -0,0 +1,34 @@
+// Copyright Epic Games, Inc. All Rights Reserved.
+
+#pragma once
+
+#include "ViewportDragInteraction.h"
+
+#include "ViewportMoveYawInteraction.generated.h"
+
+class UViewportInteractionsBehaviorSource;
+
+/**
+ * Implements the LMB + Drag Horizontal plane move + Yaw interaction
+ */
+UCLASS(Transient)
+class EDITORINTERACTIVETOOLSFRAMEWORK_API UViewportMoveYawInteraction
+	: public UViewportDragInteraction
+{
+	GENERATED_BODY()
+
+public:
+	UViewportMoveYawInteraction();
+
+	//~ Begin IClickDragBehaviorTarget
+	virtual FInputRayHit CanBeginClickDragSequence(const FInputDeviceRay& InPressPos) override;
+	virtual void OnClickPress(const FInputDeviceRay& PressPos) override;
+	virtual void OnClickRelease(const FInputDeviceRay& ReleasePos) override;
+	//~ End IClickDragBehaviorTarget
+
+	//~ Begin UViewportDragInteraction
+	virtual void OnDragDelta(float InMouseDeltaX, float InMouseDeltaY) override;
+	virtual FDeviceButtonState GetActiveMouseButtonState(const FInputDeviceState& Input) override;
+	virtual bool CanBeActivated(const FInputDeviceState& InInputDeviceState = FInputDeviceState()) const override;
+	//~ End UViewportDragInteraction
+};
diff --git a/Engine/Source/Editor/Experimental/EditorInteractiveToolsFramework/Public/ViewportInteractions/ViewportOrbitInteraction.h b/Engine/Source/Editor/Experimental/EditorInteractiveToolsFramework/Public/ViewportInteractions/ViewportOrbitInteraction.h
new file mode 100644
--- /dev/null
+++ b/Engine/Source/Editor/Experimental/EditorInteractiveToolsFramework/Public/ViewportInteractions/ViewportOrbitInteraction.h
@@ -0,0 +1,31 @@
+// Copyright Epic Games, Inc. All Rights Reserved.
+
+#pragma once
+
+#include "ViewportDragInteraction.h"
+
+#include "ViewportOrbitInteraction.generated.h"
+
+/**
+ * Triggers orbiting around current viewport focus
+ */
+UCLASS(Transient)
+class EDITORINTERACTIVETOOLSFRAMEWORK_API UViewportOrbitInteraction
+	: public UViewportDragInteraction
+{
+	GENERATED_BODY()
+
+public:
+	UViewportOrbitInteraction();
+
+	//~ Begin IClickDragBehaviorTarget
+	virtual FInputRayHit CanBeginClickDragSequence(const FInputDeviceRay& InPressPos) override;
+	virtual void OnClickPress(const FInputDeviceRay& InPressPos) override;
+	virtual void OnClickRelease(const FInputDeviceRay& InReleasePos) override;
+	//~ End IClickDragBehaviorTarget
+
+	//~ Begin UViewportDragInteraction
+	virtual void OnDragDelta(float InMouseDeltaX, float InMouseDeltaY) override;
+	virtual bool CanBeActivated(const FInputDeviceState& InInputDeviceState = FInputDeviceState()) const override;
+	//~ End UViewportDragInteraction
+};
diff --git a/Engine/Source/Editor/Experimental/EditorInteractiveToolsFramework/Public/ViewportInteractions/ViewportOrthoPanInteraction.h b/Engine/Source/Editor/Experimental/EditorInteractiveToolsFramework/Public/ViewportInteractions/ViewportOrthoPanInteraction.h
new file mode 100644
--- /dev/null
+++ b/Engine/Source/Editor/Experimental/EditorInteractiveToolsFramework/Public/ViewportInteractions/ViewportOrthoPanInteraction.h
@@ -0,0 +1,32 @@
+// Copyright Epic Games, Inc. All Rights Reserved.
+
+#pragma once
+
+#include "ViewportDragInteraction.h"
+
+#include "ViewportOrthoPanInteraction.generated.h"
+
+/**
+ * Implements panning for orthographic views. Holding Alt or Shift inverts the panning direction.
+ */
+UCLASS()
+class EDITORINTERACTIVETOOLSFRAMEWORK_API UViewportOrthoPanInteraction : public UViewportDragInteraction
+{
+	GENERATED_BODY()
+
+public:
+	UViewportOrthoPanInteraction();
+
+	//~ Begin IClickDragBehaviorTarget
+	virtual FInputRayHit CanBeginClickDragSequence(const FInputDeviceRay& InPressPos) override;
+	virtual void OnClickPress(const FInputDeviceRay& PressPos) override;
+	virtual void OnClickRelease(const FInputDeviceRay& ReleasePos) override;
+	//~ End IClickDragBehaviorTarget
+
+	//~ Begin UViewportDragInteraction
+	virtual void OnDragDelta(float InMouseDeltaX, float InMouseDeltaY) override;
+	virtual FDeviceButtonState GetActiveMouseButtonState(const FInputDeviceState& Input) override;
+	//~ End UViewportDragInteraction
+
+	virtual bool CanBeActivated(const FInputDeviceState& InInputDeviceState = FInputDeviceState()) const override;
+};
diff --git a/Engine/Source/Editor/Experimental/EditorInteractiveToolsFramework/Public/ViewportInteractions/ViewportPanInteraction.h b/Engine/Source/Editor/Experimental/EditorInteractiveToolsFramework/Public/ViewportInteractions/ViewportPanInteraction.h
new file mode 100644
--- /dev/null
+++ b/Engine/Source/Editor/Experimental/EditorInteractiveToolsFramework/Public/ViewportInteractions/ViewportPanInteraction.h
@@ -0,0 +1,36 @@
+// Copyright Epic Games, Inc. All Rights Reserved.
+
+#pragma once
+
+#include "ViewportDragInteraction.h"
+
+#include "ViewportPanInteraction.generated.h"
+
+class UViewportInteractionsBehaviorSource;
+
+/**
+ * Implements the LMB + Drag Horizontal plane move + Yaw interaction
+ */
+UCLASS(Transient)
+class EDITORINTERACTIVETOOLSFRAMEWORK_API UViewportPanInteraction : public UViewportDragInteraction
+{
+	GENERATED_BODY()
+
+public:
+	UViewportPanInteraction();
+
+	//~ Begin IClickDragBehaviorTarget
+	virtual FInputRayHit CanBeginClickDragSequence(const FInputDeviceRay& InPressPos) override;
+	virtual void OnClickPress(const FInputDeviceRay& PressPos) override;
+	virtual void OnClickRelease(const FInputDeviceRay& ReleasePos) override;
+	//~ End IClickDragBehaviorTarget
+
+	//~ Begin UViewportDragInteraction
+	virtual void OnDragDelta(float InMouseDeltaX, float InMouseDeltaY) override;
+	virtual FDeviceButtonState GetActiveMouseButtonState(const FInputDeviceState& Input) override;
+	//~ End UViewportDragInteraction
+
+private:
+	// Pan direction can be inverted e.g. with a key modifier
+	bool bMultiButtonPanInteraction = false;
+};
diff --git a/Engine/Source/Editor/Experimental/EditorInteractiveToolsFramework/Public/ViewportInteractions/ViewportViewAngleInteraction.h b/Engine/Source/Editor/Experimental/EditorInteractiveToolsFramework/Public/ViewportInteractions/ViewportViewAngleInteraction.h
new file mode 100644
--- /dev/null
+++ b/Engine/Source/Editor/Experimental/EditorInteractiveToolsFramework/Public/ViewportInteractions/ViewportViewAngleInteraction.h
@@ -0,0 +1,35 @@
+// Copyright Epic Games, Inc. All Rights Reserved.
+
+#pragma once
+
+#include "ViewportDragInteraction.h"
+
+#include "ViewportViewAngleInteraction.generated.h"
+
+/**
+ * Implements the RMB + Drag interaction used to change the current Camera View angle
+ */
+UCLASS(Transient)
+class EDITORINTERACTIVETOOLSFRAMEWORK_API UViewportViewAngleInteraction
+	: public UViewportDragInteraction
+{
+	GENERATED_BODY()
+
+public:
+	UViewportViewAngleInteraction();
+
+	//~ Begin IClickDragBehaviorTarget
+	virtual FInputRayHit CanBeginClickDragSequence(const FInputDeviceRay& InPressPos) override;
+	virtual void OnClickPress(const FInputDeviceRay& PressPos) override;
+	virtual void OnClickRelease(const FInputDeviceRay& ReleasePos) override;
+	//~ End IClickDragBehaviorTarget
+
+	//~ Begin UViewportDragInteraction
+	virtual void OnDragDelta(float InMouseDeltaX, float InMouseDeltaY) override;
+	virtual FDeviceButtonState GetActiveMouseButtonState(const FInputDeviceState& Input) override;
+	//~ End UViewportDragInteraction
+
+	//~ Begin IModifierToggleBehaviorTarget
+	virtual void OnUpdateModifierState(int ModifierID, bool bIsOn) override;
+	//~ End IModifierToggleBehaviorTarget
+};
diff --git a/Engine/Source/Editor/Experimental/EditorInteractiveToolsFramework/Public/ViewportInteractions/ViewportZoomInteraction.h b/Engine/Source/Editor/Experimental/EditorInteractiveToolsFramework/Public/ViewportInteractions/ViewportZoomInteraction.h
new file mode 100644
--- /dev/null
+++ b/Engine/Source/Editor/Experimental/EditorInteractiveToolsFramework/Public/ViewportInteractions/ViewportZoomInteraction.h
@@ -0,0 +1,56 @@
+// Copyright Epic Games, Inc. All Rights Reserved.
+
+#pragma once
+
+#include "BaseBehaviors/MouseWheelBehavior.h"
+#include "ViewportInteraction.h"
+
+#include "ViewportZoomInteraction.generated.h"
+
+/**
+ * Implements a viewport interaction to change the current zoom level, using either keyboard inputs or the mouse wheel
+ */
+UCLASS(Transient)
+class EDITORINTERACTIVETOOLSFRAMEWORK_API UViewportZoomInteraction
+	: public UViewportInteraction
+	, public IMouseWheelBehaviorTarget
+	, public IKeyInputBehaviorTarget
+{
+	GENERATED_BODY()
+
+public:
+	UViewportZoomInteraction();
+
+	//~ Begin IMouseWheelBehaviorTarget
+	virtual FInputRayHit ShouldRespondToMouseWheel(const FInputDeviceRay& InCurrentPos) override;
+	virtual void OnMouseWheelScrollUp(const FInputDeviceRay& InCurrentPos) override;
+	virtual void OnMouseWheelScrollDown(const FInputDeviceRay& InCurrentPos) override;
+	//~ End IMouseWheelBehaviorTarget
+
+	//~ Begin IKeyInputBehaviorTarget
+	virtual void OnKeyPressed(const FKey& InKeyID) override;
+	virtual void OnKeyReleased(const FKey& InKeyID) override;
+	virtual void OnForceEndCapture() override;
+	//~ End IKeyInputBehaviorTarget
+
+	//~ Begin UViewportInteractionBase
+	virtual void Tick(float InDeltaTime) const override;
+	//~ End UViewportInteractionBase
+
+private:
+	void OrthographicZoom(bool bInZoomIn, float InZoomMultiplier = 1.0f, bool bInForceZoomToCenter = false) const;
+	void PerspectiveZoom(bool bInZoomIn) const;
+	void MouseWheelZoom(bool bInZoomIn) const;
+
+	void OnClickPress(const FInputDeviceRay& InPressPos);
+	void OnClickRelease(const FInputDeviceRay& InReleasePos);
+	void OnClickDrag(const FInputDeviceRay& InDragPos);
+
+	int16 ZoomDirection = 0;
+
+	// TODO: make them customizable
+	FKey ZoomIn = EKeys::Add;
+	FKey ZoomOut = EKeys::Subtract;
+
+	float LastZoomDragY;
+};
diff --git a/Engine/Source/Editor/UnrealEd/Private/EditorViewportClient.cpp b/Engine/Source/Editor/UnrealEd/Private/EditorViewportClient.cpp
--- a/Engine/Source/Editor/UnrealEd/Private/EditorViewportClient.cpp
+++ b/Engine/Source/Editor/UnrealEd/Private/EditorViewportClient.cpp
@@ -83,6 +83,10 @@
 #include "LandscapeSettings.h"
 #include "TransformGizmoEditorSettings.h"
 #include "Misc/AxisDisplayInfo.h"
+#include "Tools/EdModeInteractiveToolsContext.h"
+#include "ViewportClientNavigationHelper.h"
+#include "ViewportInteractions/ViewportInteractionsBehaviorSource.h"
+#include "ViewportInteractions/ViewportInteraction.h"
 
 #define LOCTEXT_NAMESPACE "EditorViewportClient"
 
@@ -128,6 +132,23 @@ static bool GetDefaultLowDPIPreviewValue()
 	return CVarEditorViewportHighDPIPtr->GetInt() == 0;
 }
 
+/**
+ * Returns true if any viewport interaction of the specified type is active
+ */
+bool IsViewportInteractionActive(const FEditorModeTools* InModeTools, FName InViewportInteraction)
+{
+	if (InModeTools)
+	{
+		if (UViewportInteractionsBehaviorSource* InputToolsBehaviorSource =
+				InModeTools->GetInteractiveToolsContext()->GetViewportInteractionsBehaviorSource())
+		{
+			return InputToolsBehaviorSource->IsViewportInteractionActive(InViewportInteraction);
+		}
+	}
+
+	return false;
+}
+
 float ComputeOrthoZoomFactor(const float ViewportWidth)
 {
 	float Ret = 1.0f;
@@ -469,6 +490,7 @@ FEditorViewportClient::FEditorViewportClient(FEditorModeTools* InModeTools, FPre
 	, bHasMouseMovedSinceClick(false)
 	, CameraController(new FEditorCameraController())
 	, CameraUserImpulseData(new FCameraControllerUserImpulseData())
+	, ViewportNavigationHelper(new FViewportClientNavigationHelper())
 	, TimeForForceRedraw(0.0)
 	, FlightCameraSpeedScale(1.0f)
 	, bUseControllingActorViewInfo(false)
@@ -1908,14 +1930,18 @@ void FEditorViewportClient::UpdateCameraMovement( float DeltaTime )
 		const bool bIsUsingTrackpad = FSlateApplication::Get().IsUsingTrackpad();
 
 		// Certain keys are only available while the flight camera input mode is active
-		const bool bUsingFlightInput = IsFlightCameraInputModeActive() || bIsUsingTrackpad;
+		const bool bUsingFlightInput = IsFlightCameraInputModeActive() || bIsUsingTrackpad || ViewportNavigationHelper->bIsMouseLooking;
 
 		// Is the current press unmodified?
 		const bool bUnmodifiedPress = !IsAltPressed() && !IsShiftPressed() && !IsCtrlPressed() && !IsCmdPressed();
 
 		// Do we want to use the regular arrow keys for flight input?
 		// Because the arrow keys are used for things like nudging actors, we'll only do this while the press is unmodified
-		const bool bRemapArrowKeys = bUnmodifiedPress;
+		const bool bRemapArrowKeys = bUnmodifiedPress || IsViewportInteractionActive(ModeTools.Get(), UE::Editor::ViewportInteractions::CameraTranslate);
+
+		const bool bITFCameraTranslateActive = IsViewportInteractionActive(ModeTools.Get(),UE::Editor::ViewportInteractions::CameraTranslate);
+		const bool bITFCameraRotateActive = IsViewportInteractionActive(ModeTools.Get(),UE::Editor::ViewportInteractions::CameraRotate);
+		const bool bITFFOVActive = IsViewportInteractionActive(ModeTools.Get(),UE::Editor::ViewportInteractions::FOV);
 
 		// Do we want to remap the various WASD keys for flight input?
 		const bool bRemapWASDKeys =
@@ -1980,91 +2006,110 @@ void FEditorViewportClient::UpdateCameraMovement( float DeltaTime )
 			CameraController->GetConfig().bForceRotationalPhysics = false;
 		}
 
-		// Forward/back
-		if( ( bRemapWASDKeys && bForwardKeyState ) ||
-			( bRemapArrowKeys && Viewport->KeyState( EKeys::Up ) ) ||
-			( bUnmodifiedPress && bUseNumpadCameraControl && Viewport->KeyState(EKeys::NumPadEight) ) )
+		// Check if we need to use legacy translation input
+		if (!bITFCameraTranslateActive)
 		{
-			CameraUserImpulseData->MoveForwardBackwardImpulse += 1.0f;
-		}
-		if( (bRemapWASDKeys && bBackwardKeyState) ||
-			( bRemapArrowKeys && Viewport->KeyState( EKeys::Down ) ) ||
-			( bUnmodifiedPress && bUseNumpadCameraControl && Viewport->KeyState( EKeys::NumPadTwo ) ) )
-		{
-			CameraUserImpulseData->MoveForwardBackwardImpulse -= 1.0f;
-		}
+			// Forward/back
+			if( ( bRemapWASDKeys && bForwardKeyState ) ||
+				( bRemapArrowKeys && Viewport->KeyState( EKeys::Up ) ) ||
+				( bUnmodifiedPress && bUseNumpadCameraControl && Viewport->KeyState(EKeys::NumPadEight) ) )
+			{
+				CameraUserImpulseData->MoveForwardBackwardImpulse += 1.0f;
+			}
+			if( (bRemapWASDKeys && bBackwardKeyState) ||
+				( bRemapArrowKeys && Viewport->KeyState( EKeys::Down ) ) ||
+				( bUnmodifiedPress && bUseNumpadCameraControl && Viewport->KeyState( EKeys::NumPadTwo ) ) )
+			{
+				CameraUserImpulseData->MoveForwardBackwardImpulse -= 1.0f;
+			}
 
-		// Right/left
-		if (( bRemapWASDKeys && bRightKeyState) ||
-			( bRemapArrowKeys && Viewport->KeyState( EKeys::Right ) ) ||
-			( bUnmodifiedPress && bUseNumpadCameraControl && Viewport->KeyState( EKeys::NumPadSix ) ) )
-		{
-			CameraUserImpulseData->MoveRightLeftImpulse += 1.0f;
-		}
-		if( ( bRemapWASDKeys && bLeftKeyState) ||
-			( bRemapArrowKeys && Viewport->KeyState( EKeys::Left ) ) ||
-			( bUnmodifiedPress && bUseNumpadCameraControl && Viewport->KeyState( EKeys::NumPadFour ) ) )
-		{
-			CameraUserImpulseData->MoveRightLeftImpulse -= 1.0f;
-		}
+			// Right/left
+			if (( bRemapWASDKeys && bRightKeyState) ||
+				( bRemapArrowKeys && Viewport->KeyState( EKeys::Right ) ) ||
+				( bUnmodifiedPress && bUseNumpadCameraControl && Viewport->KeyState( EKeys::NumPadSix ) ) )
+			{
+				CameraUserImpulseData->MoveRightLeftImpulse += 1.0f;
+			}
+			if( ( bRemapWASDKeys && bLeftKeyState) ||
+				( bRemapArrowKeys && Viewport->KeyState( EKeys::Left ) ) ||
+				( bUnmodifiedPress && bUseNumpadCameraControl && Viewport->KeyState( EKeys::NumPadFour ) ) )
+			{
+				CameraUserImpulseData->MoveRightLeftImpulse -= 1.0f;
+			}
 
-		// World up/down
-		if( ( bRemapWASDKeys && bWorldUpKeyState) ||
-			( bUnmodifiedPress && Viewport->KeyState( EKeys::PageUp ) ) ||
-			( bUnmodifiedPress && bUseNumpadCameraControl && ( Viewport->KeyState( EKeys::NumPadNine ) || Viewport->KeyState( EKeys::Add ) ) ) )
-		{
-			CameraUserImpulseData->MoveWorldUpDownImpulse += 1.0f;
-		}
-		if( ( bRemapWASDKeys && bWorldDownKeyState) ||
-			( bUnmodifiedPress && Viewport->KeyState( EKeys::PageDown ) ) ||
-			( bUnmodifiedPress && bUseNumpadCameraControl && ( Viewport->KeyState( EKeys::NumPadSeven ) || Viewport->KeyState( EKeys::Subtract ) ) ) )
-		{
-			CameraUserImpulseData->MoveWorldUpDownImpulse -= 1.0f;
-		}
+			// World up/down
+			if( ( bRemapWASDKeys && bWorldUpKeyState) ||
+				( bUnmodifiedPress && Viewport->KeyState( EKeys::PageUp ) ) ||
+				( bUnmodifiedPress && bUseNumpadCameraControl && ( Viewport->KeyState( EKeys::NumPadNine ) || Viewport->KeyState( EKeys::Add ) ) ) )
+			{
+				CameraUserImpulseData->MoveWorldUpDownImpulse += 1.0f;
+			}
+			if( ( bRemapWASDKeys && bWorldDownKeyState) ||
+				( bUnmodifiedPress && Viewport->KeyState( EKeys::PageDown ) ) ||
+				( bUnmodifiedPress && bUseNumpadCameraControl && ( Viewport->KeyState( EKeys::NumPadSeven ) || Viewport->KeyState( EKeys::Subtract ) ) ) )
+			{
+				CameraUserImpulseData->MoveWorldUpDownImpulse -= 1.0f;
+			}
 
-		// Local up/down
-		if( ( bRemapWASDKeys && bLocalUpKeyState) )
-		{
-			CameraUserImpulseData->MoveLocalUpDownImpulse += 1.0f;
-		}
-		if( ( bRemapWASDKeys && bLocalDownKeyState) )
-		{
-			CameraUserImpulseData->MoveLocalUpDownImpulse -= 1.0f;
+			// Local up/down
+			if( ( bRemapWASDKeys && bLocalUpKeyState) )
+			{
+				CameraUserImpulseData->MoveLocalUpDownImpulse += 1.0f;
+			}
+			if( ( bRemapWASDKeys && bLocalDownKeyState) )
+			{
+				CameraUserImpulseData->MoveLocalUpDownImpulse -= 1.0f;
+			}
 		}
 
-		// Zoom FOV out/in
-		if( ( bRemapWASDKeys && bZoomOutKeyState) ||
-			( bUnmodifiedPress && bUseNumpadCameraControl && Viewport->KeyState( EKeys::NumPadOne ) ) )
+		// Check if we need to use legacy FOV input
+		if (!bITFFOVActive)
 		{
-			CameraUserImpulseData->ZoomOutInImpulse += 1.0f;
-		}
-		if( ( bRemapWASDKeys && bZoomInKeyState) ||
-			( bUnmodifiedPress && bUseNumpadCameraControl && Viewport->KeyState( EKeys::NumPadThree ) ) )
-		{
-			CameraUserImpulseData->ZoomOutInImpulse -= 1.0f;
+			// Zoom FOV out/in
+			if( ( bRemapWASDKeys && bZoomOutKeyState) ||
+				( bUnmodifiedPress && bUseNumpadCameraControl && Viewport->KeyState( EKeys::NumPadOne ) ) )
+			{
+				CameraUserImpulseData->ZoomOutInImpulse += 1.0f;
+			}
+			if( ( bRemapWASDKeys && bZoomInKeyState) ||
+				( bUnmodifiedPress && bUseNumpadCameraControl && Viewport->KeyState( EKeys::NumPadThree ) ) )
+			{
+				CameraUserImpulseData->ZoomOutInImpulse -= 1.0f;
+			}
 		}
 
-		// Rotate up/down
-		if (bRemapWASDKeys && bRotateUpKeyState)
-		{
-			CameraUserImpulseData->RotatePitchImpulse += 1.0f * CameraController->GetConfig().RotationMultiplier;
-			CameraController->GetConfig().bForceRotationalPhysics = true;
-		}
-		if (bRemapWASDKeys && bRotateDownKeyState)
+		// Check if we need to use legacy rotation input
+		if (!bITFCameraRotateActive)
 		{
-			CameraUserImpulseData->RotatePitchImpulse -= 1.0f * CameraController->GetConfig().RotationMultiplier;
-			CameraController->GetConfig().bForceRotationalPhysics = true;
-		}
+			// Rotate up/down
+			if (bRemapWASDKeys && bRotateUpKeyState)
+			{
+				CameraUserImpulseData->RotatePitchImpulse += 1.0f * CameraController->GetConfig().RotationMultiplier;
+				CameraController->GetConfig().bForceRotationalPhysics = true;
+			}
+			if (bRemapWASDKeys && bRotateDownKeyState)
+			{
+				CameraUserImpulseData->RotatePitchImpulse -= 1.0f * CameraController->GetConfig().RotationMultiplier;
+				CameraController->GetConfig().bForceRotationalPhysics = true;
+			}
 
-		// Rotate left/right
-		if (bRemapWASDKeys && bRotateLeftKeyState)
-		{
-			CameraUserImpulseData->RotateYawImpulse -= 1.0f * CameraController->GetConfig().RotationMultiplier;
-			CameraController->GetConfig().bForceRotationalPhysics = true;
+			// Rotate left/right
+			if (bRemapWASDKeys && bRotateLeftKeyState)
+			{
+				CameraUserImpulseData->RotateYawImpulse -= 1.0f * CameraController->GetConfig().RotationMultiplier;
+				CameraController->GetConfig().bForceRotationalPhysics = true;
+			}
+			if (bRemapWASDKeys && bRotateRightKeyState)
+			{
+				CameraUserImpulseData->RotateYawImpulse += 1.0f * CameraController->GetConfig().RotationMultiplier;
+				CameraController->GetConfig().bForceRotationalPhysics = true;
+			}
 		}
-		if (bRemapWASDKeys && bRotateRightKeyState)
+
+		// Add camera inputs coming from ITF, if any
+		if (bITFCameraTranslateActive || bITFCameraRotateActive || bITFFOVActive)
 		{
-			CameraUserImpulseData->RotateYawImpulse += 1.0f * CameraController->GetConfig().RotationMultiplier;
+			ViewportNavigationHelper->ConsumeImpulseData(*CameraUserImpulseData);
 			CameraController->GetConfig().bForceRotationalPhysics = true;
 		}
 
@@ -2111,17 +2156,17 @@ void FEditorViewportClient::UpdateCameraMovement( float DeltaTime )
 		// Check whether the camera is being moved by the mouse or keyboard
 		bool bHasMovement = GetDefault<ULevelEditorViewportSettings>()->bUseLegacyCameraMovementNotifications;
 
-		if ((*CameraUserImpulseData).RotateYawVelocityModifier != 0.0f ||
-			(*CameraUserImpulseData).RotatePitchVelocityModifier != 0.0f ||
-			(*CameraUserImpulseData).RotateRollVelocityModifier != 0.0f ||
-			(*CameraUserImpulseData).MoveForwardBackwardImpulse != 0.0f ||
-			(*CameraUserImpulseData).MoveRightLeftImpulse != 0.0f ||
-			(*CameraUserImpulseData).MoveWorldUpDownImpulse != 0.0f ||
-			(*CameraUserImpulseData).MoveLocalUpDownImpulse != 0.0f ||
-			(*CameraUserImpulseData).ZoomOutInImpulse != 0.0f ||
-			(*CameraUserImpulseData).RotateYawImpulse != 0.0f ||
-			(*CameraUserImpulseData).RotatePitchImpulse != 0.0f ||
-			(*CameraUserImpulseData).RotateRollImpulse != 0.0f
+		if (CameraUserImpulseData->RotateYawVelocityModifier != 0.0f ||
+			CameraUserImpulseData->RotatePitchVelocityModifier != 0.0f ||
+			CameraUserImpulseData->RotateRollVelocityModifier != 0.0f ||
+			CameraUserImpulseData->MoveForwardBackwardImpulse != 0.0f ||
+			CameraUserImpulseData->MoveRightLeftImpulse != 0.0f ||
+			CameraUserImpulseData->MoveWorldUpDownImpulse != 0.0f ||
+			CameraUserImpulseData->MoveLocalUpDownImpulse != 0.0f ||
+			CameraUserImpulseData->ZoomOutInImpulse != 0.0f ||
+			CameraUserImpulseData->RotateYawImpulse != 0.0f ||
+			CameraUserImpulseData->RotatePitchImpulse != 0.0f ||
+			CameraUserImpulseData->RotateRollImpulse != 0.0f
 			)
 		{
 			bHasMovement = true;
@@ -2433,7 +2478,8 @@ void FEditorViewportClient::UpdateMouseDelta()
 
 	if( Viewport )
 	{
-		if( !DragDelta.IsNearlyZero() )
+		// If mouse delta tracker has delta OR there is transform delta coming from ITF Viewport Interaction, handle it
+		if (!DragDelta.IsNearlyZero() || ViewportNavigationHelper->HasTransformDelta())
 		{
 			const bool LeftMouseButtonDown = Viewport->KeyState(EKeys::LeftMouseButton);
 			const bool MiddleMouseButtonDown = Viewport->KeyState(EKeys::MiddleMouseButton);
@@ -2446,11 +2492,18 @@ void FEditorViewportClient::UpdateMouseDelta()
 			ModeTools->PreConvertMouseMovement(this);
 
 			// Convert the movement delta into drag/rotation deltas
-			FVector Drag;
-			FRotator Rot;
-			FVector Scale;
+			FVector Drag = FVector::ZeroVector;
+			FRotator Rot = FRotator::ZeroRotator;
+			FVector Scale = FVector::ZeroVector;
 			EAxisList::Type CurrentAxis = Widget->GetCurrentAxis();
-			if ( IsOrtho() && ( LeftMouseButtonDown || bIsUsingTrackpad ) && RightMouseButtonDown )
+
+			// Try to handle ITF Viewport inputs before legacy (currently zoom and camera drag ones)
+			if (IsViewportInteractionActive(ModeTools.Get(), UE::Editor::ViewportInteractions::Zoom) ||
+				IsViewportInteractionActive(ModeTools.Get(), UE::Editor::ViewportInteractions::CameraDrag))
+			{
+				ViewportNavigationHelper->ConsumeTransformDelta(Drag, Rot);
+			}
+			else if ( IsOrtho() && ( LeftMouseButtonDown || bIsUsingTrackpad ) && RightMouseButtonDown )
 			{
 				bWidgetAxisControlledByDrag = false;
 				Widget->SetCurrentAxis( EAxisList::None );
@@ -2472,7 +2525,8 @@ void FEditorViewportClient::UpdateMouseDelta()
 				}
 				//if Absolute Translation, and not just moving the camera around
 				if (IsUsingAbsoluteTranslation(false))
-				{ 
+				{
+					// TODO: yet to be checked out for ITF Parity
 					MouseDeltaTracker->AbsoluteTranslationConvertMouseToDragRot(DragStartView, this, Drag, Rot, Scale);
 				}
 				else
@@ -2505,7 +2559,21 @@ void FEditorViewportClient::UpdateMouseDelta()
 
 void FEditorViewportClient::PeformDefaultCameraMovement(FVector& Drag, FRotator& Rot, FVector& Scale)
 {
-	FVector DragDelta = MouseDeltaTracker->GetDelta();
+	// Check if there is an active ITF Orbit interaction, and if there's an orbit delta it might have generated 
+	const bool bUseITFOrbitDelta = IsViewportInteractionActive(ModeTools.Get(), UE::Editor::ViewportInteractions::Orbit)
+					  && !ViewportNavigationHelper->OrbitDelta.IsNearlyZero();
+
+	FVector DragDelta;
+
+	if (bUseITFOrbitDelta)
+	{
+		ViewportNavigationHelper->ConsumeOrbitDelta(DragDelta);
+	}
+	else // no ITF delta, try to use the one from mouse delta tracker
+	{
+		DragDelta = MouseDeltaTracker->GetDelta();
+	}
+
 	if (ShouldOrbitCamera())
 	{
 		bool bHasMovement = !DragDelta.IsNearlyZero();
@@ -3182,6 +3250,7 @@ bool FEditorViewportClient::Internal_InputKey(const FInputKeyEventArgs& EventArg
 	const EInputEvent& Event = EventArgs.Event;
 	const FViewport* InViewport = EventArgs.Viewport;
 
+	bool bRouteToToolsContext = !bIsTracking || UE::Editor::ViewportInteractions::UseEditorViewportInteractions();
 	// Let the current mode have a look at the input before reacting to it. 
 	// Note that bIsTracking tells us whether the viewport client is capturing mouse behavior to fly around,
 	// move objects, etc. In this case we don't want to pass the input the input router, because it is already
@@ -3190,7 +3259,12 @@ bool FEditorViewportClient::Internal_InputKey(const FInputKeyEventArgs& EventArg
 	// for now. In the long term, only old FEdModes will get this input, while the new UEdModes won't, but
 	// this is not yet split apart. Also in the long term, the viewport behaviors will be inside the
 	// input router and we won't need this bIsTracking flag to begin with.
-	if (ModeTools->InputKey(this, Viewport, Key, Event, /*bRouteToToolsContext*/ !bIsTracking))
+	if (ModeTools->InputKey(this, Viewport, Key, Event, bRouteToToolsContext))
+	{
+		return true;
+	}
+
+	if (GetViewportNavigationHelper()->bIsMouseLooking)
 	{
 		return true;
 	}
@@ -3268,9 +3342,10 @@ bool FEditorViewportClient::Internal_InputKey(const FInputKeyEventArgs& EventArg
 		return true;
 	}
 
+	const bool bIsITFSingleClickActive = IsViewportInteractionActive(ModeTools.Get(), UE::Editor::ViewportInteractions::SingleClick);
 
 	// If we are tracking and no mouse button is down and this input event released the mouse button stop tracking and process any clicks if necessary
-	if ( bIsTracking && !InputState.IsAnyMouseButtonDown() && InputState.IsMouseButtonEvent() )
+	if ( bIsTracking && !InputState.IsAnyMouseButtonDown() && InputState.IsMouseButtonEvent() && !bIsITFSingleClickActive)
 	{
 		// Handle possible mouse click viewport
 		ProcessClickInViewport( InputState, *View );
@@ -4851,6 +4926,12 @@ bool FEditorViewportClient::Internal_InputAxis(FViewport* InViewport, FInputDevi
 		return true;
 	}
 
+	// No need to process anything else here, since some ITF Viewport interaction is handling the camera
+	if (GetViewportNavigationHelper()->bIsMouseLooking)
+	{
+		return true;
+	}
+
 	const bool bMouseButtonDown = InViewport->KeyState( EKeys::LeftMouseButton ) || InViewport->KeyState( EKeys::MiddleMouseButton ) || InViewport->KeyState( EKeys::RightMouseButton );
 	const bool bLightMoveDown = InViewport->KeyState(EKeys::L);
 
@@ -5104,7 +5185,7 @@ void FEditorViewportClient::ConvertMovementToDragRot(const FVector& InDelta,
 				if (MiddleMouseButtonDown && !ViewportSettings->bWorldSpaceVerticalPan)
 				{
 					// Pan relative to the camera
-					InDragDelta = FQuat(ViewRotation) * FVector( 0.0f, InDelta.X, InDelta.Y) * (bInvert ? -1.0f : 1.0f);
+					InDragDelta = FQuat(ViewRotation) * FVector( 0.0f, InDelta.X, InDelta.Y) * (bInvert ? 1.0f : -1.0f);
 				}
 				else
 				{
@@ -6024,16 +6105,16 @@ FViewportCursorLocation FEditorViewportClient::GetCursorWorldLocationFromMousePo
 
 void FEditorViewportClient::CapturedMouseMove( FViewport* InViewport, int32 InMouseX, int32 InMouseY )
 {
-	UpdateRequiredCursorVisibility();
-	ApplyRequiredCursorVisibility();
-
-	CapturedMouseMoves.Add(FIntPoint(InMouseX, InMouseY));
-
 	// Let the current editor mode know about the mouse movement.
 	if (ModeTools->CapturedMouseMove(this, InViewport, InMouseX, InMouseY))
 	{
 		return;
 	}
+
+	UpdateRequiredCursorVisibility();
+	ApplyRequiredCursorVisibility();
+
+	CapturedMouseMoves.Add(FIntPoint(InMouseX, InMouseY));
 }
 
 void FEditorViewportClient::ProcessAccumulatedPointerInput(FViewport* InViewport)
diff --git a/Engine/Source/Editor/UnrealEd/Private/SEditorViewport.cpp b/Engine/Source/Editor/UnrealEd/Private/SEditorViewport.cpp
--- a/Engine/Source/Editor/UnrealEd/Private/SEditorViewport.cpp
+++ b/Engine/Source/Editor/UnrealEd/Private/SEditorViewport.cpp
@@ -37,6 +37,7 @@
 #endif
 #include "LevelEditorViewport.h"
 #include "SEditorViewportGridPanel.h"
+#include "ViewportInteractions/ViewportInteractionsBehaviorSource.h"
 
 #define LOCTEXT_NAMESPACE "EditorViewport"
 
@@ -58,6 +59,9 @@ SEditorViewport::~SEditorViewport()
 	Client.Reset();
 
 	check( SceneViewport.IsUnique() );
+
+	// This might be removed once ITF viewport interactions become the default, and legacy ones are removed/deactivated - responds to dedicated CVar
+	UE::Editor::ViewportInteractions::OnEditorViewportInteractionsActivated().RemoveAll(this);
 }
 
 void SEditorViewport::Construct( const FArguments& InArgs )
@@ -291,6 +295,14 @@ void SEditorViewport::Construct( const FArguments& InArgs )
 	// The following RequestUpdateDPIScale() call marks the DPI value for refresh.
 	// This ensures CachedDPIScale value can be properly retrieved at the right time, once the Widget is fully setup and its parent window can be accessed.
 	GetViewportClient()->RequestUpdateDPIScale();
+
+	// Register the Editor-specific viewport interactions
+	RegisterViewportInteractions();
+
+	// This is used while working with both legacy and ITF interactions. Once ITF is shipped as default AND CVar removed, it might not be needed anymore
+	UE::Editor::ViewportInteractions::OnEditorViewportInteractionsActivated().AddRaw(
+		this, &SEditorViewport::RegisterViewportInteractions
+	);
 }
 
 FReply SEditorViewport::OnKeyDown( const FGeometry& MyGeometry, const FKeyEvent& InKeyEvent )
@@ -1150,6 +1162,17 @@ EVisibility SEditorViewport::GetActiveBorderVisibility() const
 	return BaseVisibility;
 }
 
+void SEditorViewport::RegisterViewportInteractions()
+{
+	if (UModeManagerInteractiveToolsContext* InteractiveToolsContext = GetModeManagerInteractiveToolsContext())
+	{
+		if (UViewportInteractionsBehaviorSource* InteractionsBehaviorSource = InteractiveToolsContext->GetViewportInteractionsBehaviorSource())
+		{
+			OnRegisterViewportInteractions(InteractionsBehaviorSource);
+		}
+	}
+}
+
 ///////////////////////////////////////////////////////////////////////////////
 // begin feature level control functions block
 ///////////////////////////////////////////////////////////////////////////////
diff --git a/Engine/Source/Editor/UnrealEd/Private/Tools/EdModeInteractiveToolsContext.cpp b/Engine/Source/Editor/UnrealEd/Private/Tools/EdModeInteractiveToolsContext.cpp
--- a/Engine/Source/Editor/UnrealEd/Private/Tools/EdModeInteractiveToolsContext.cpp
+++ b/Engine/Source/Editor/UnrealEd/Private/Tools/EdModeInteractiveToolsContext.cpp
@@ -41,6 +41,7 @@
 #include "EditorDragTools/DragToolsBehaviourSource.h"
 #include "EditorDragTools/EditorDragToolBehaviorTarget.h"
 #include "UObject/ObjectSaveContext.h"
+#include "ViewportInteractions/ViewportInteractionsBehaviorSource.h"
 
 //#define ENABLE_DEBUG_PRINTING
 
@@ -885,6 +886,11 @@ void UModeManagerInteractiveToolsContext::Tick(FEditorViewportClient* ViewportCl
 	{
 		EdModeContext->Tick(ViewportClient, DeltaTime);
 	}
+
+	if (ViewportInteractionsBehaviorSource)
+	{
+		ViewportInteractionsBehaviorSource->Tick(DeltaTime);
+	}
 }
 
 void UModeManagerInteractiveToolsContext::Render(const FSceneView* View, FViewport* Viewport, FPrimitiveDrawInterface* PDI)
@@ -1269,14 +1275,20 @@ void UModeManagerInteractiveToolsContext::Initialize(IToolsContextQueriesAPI* Qu
 
 	GetParentEditorModeManager()->OnEditorModeIDChanged().AddUObject(this, &UModeManagerInteractiveToolsContext::OnEditorModeChanged);
 
-	// Delegates used to handle SelectionMode.EnableITFTools CVAR changes
+	// Delegates used to handle DragTools.EnableITFTools CVAR changes
 	UE::Editor::DragTools::OnEditorDragToolsActivated().AddUObject(this, &UModeManagerInteractiveToolsContext::RegisterDragTools);
 	UE::Editor::DragTools::OnEditorDragToolsDeactivated().AddUObject(this, &UModeManagerInteractiveToolsContext::UnregisterDragTools);
+
+	// Delegates used to handle InputTools.EnableITFTools CVAR changes
+	UE::Editor::ViewportInteractions::OnEditorViewportInteractionsActivated().AddUObject(this, &UModeManagerInteractiveToolsContext::RegisterViewportInteractions);
+	UE::Editor::ViewportInteractions::OnEditorViewportInteractionsDeactivated().AddUObject(this, &UModeManagerInteractiveToolsContext::UnregisterViewportInteractions);
 }
 
 void UModeManagerInteractiveToolsContext::Shutdown()
 {
 	UnregisterDragTools();
+	UnregisterViewportInteractions();
+
 	GetParentEditorModeManager()->OnEditorModeIDChanged().RemoveAll(this);
 	UE::Editor::DragTools::OnEditorDragToolsActivated().RemoveAll(this);
 	UE::Editor::DragTools::OnEditorDragToolsDeactivated().RemoveAll(this);
@@ -1341,9 +1353,25 @@ bool UModeManagerInteractiveToolsContext::OnChildEdModeDeactivated(UEdModeIntera
 
 void UModeManagerInteractiveToolsContext::SetBypassAltModifier(bool bInBypassAltModifier)
 {
+	if (bInBypassAltModifier
+		&& (UE::Editor::ViewportInteractions::UseEditorViewportInteractions() || UE::Editor::DragTools::UseEditorDragTools()))
+	{
+		return;
+	}
+
 	bBypassAltModifier = bInBypassAltModifier;
 }
 
+UViewportInteractionsBehaviorSource* UModeManagerInteractiveToolsContext::GetViewportInteractionsBehaviorSource()
+{
+	if (!ViewportInteractionsBehaviorSource && UsesViewportInteractions() && UE::Editor::ViewportInteractions::UseEditorViewportInteractions())
+	{
+		RegisterViewportInteractions();
+	}
+
+	return ViewportInteractionsBehaviorSource.Get();
+}
+
 void UModeManagerInteractiveToolsContext::RegisterDragTools()
 {
 	FEditorModeTools* EditorModeTools = GetParentEditorModeManager();
@@ -1445,6 +1473,31 @@ void UModeManagerInteractiveToolsContext::RegisterDragTools()
 	}
 }
 
+void UModeManagerInteractiveToolsContext::RegisterViewportInteractions()
+{
+	FEditorModeTools* EditorModeTools = GetParentEditorModeManager();
+	if (!EditorModeTools)
+	{
+		return;
+	}
+
+	// This Mode Manager does not require Input Tools
+	if (!UsesViewportInteractions())
+	{
+		return;
+	}
+
+	// Not initialized yet
+	if (!ViewportInteractionsBehaviorSource)
+	{
+		SetBypassAltModifier(false);
+
+		ViewportInteractionsBehaviorSource = NewObject<UViewportInteractionsBehaviorSource>();
+		ViewportInteractionsBehaviorSource->Initialize(this);
+		ViewportInteractionsBehaviorSource->RegisterBehaviorSources();
+	}
+}
+
 void UModeManagerInteractiveToolsContext::UnregisterDragTools()
 {
 	if (DragToolsBehaviorSource)
@@ -1453,25 +1506,46 @@ void UModeManagerInteractiveToolsContext::UnregisterDragTools()
 		SetBypassAltModifier(true);
 		DragToolsBehaviorSource->DeregisterSource();
 		DragToolsBehaviorSource = nullptr;
-		
+
 		UToolMenus::UnregisterOwner(this);
 	}
 }
 
-void UModeManagerInteractiveToolsContext::OnEditorModeChanged(const FEditorModeID& InModeID, bool bInIsEnteringMode)
+void UModeManagerInteractiveToolsContext::UnregisterViewportInteractions()
 {
-	if (!UE::Editor::DragTools::UseEditorDragTools())
+	if (ViewportInteractionsBehaviorSource)
 	{
-		return;
+		// Revert to previous state, with Alt Modifier not being processed
+		SetBypassAltModifier(true);
+		ViewportInteractionsBehaviorSource->DeregisterBehaviorSources();
+		ViewportInteractionsBehaviorSource = nullptr;
 	}
+}
 
-	if (bInIsEnteringMode)
+void UModeManagerInteractiveToolsContext::OnEditorModeChanged(const FEditorModeID& InModeID, bool bInIsEnteringMode)
+{
+	if (UE::Editor::DragTools::UseEditorDragTools())
 	{
-		RegisterDragTools();
+		if (bInIsEnteringMode)
+		{
+			RegisterDragTools();
+		}
+		else
+		{
+			UnregisterDragTools();
+		}
 	}
-	else
+
+	if (UE::Editor::ViewportInteractions::UseEditorViewportInteractions())
 	{
-		UnregisterDragTools();
+		if (bInIsEnteringMode)
+		{
+			RegisterViewportInteractions();
+		}
+		else
+		{
+			UnregisterViewportInteractions();
+		}
 	}
 }
 
diff --git a/Engine/Source/Editor/UnrealEd/Private/ViewportClientNavigationHelper.cpp b/Engine/Source/Editor/UnrealEd/Private/ViewportClientNavigationHelper.cpp
new file mode 100644
--- /dev/null
+++ b/Engine/Source/Editor/UnrealEd/Private/ViewportClientNavigationHelper.cpp
@@ -0,0 +1,74 @@
+﻿// Copyright Epic Games, Inc. All Rights Reserved.
+
+#include "ViewportClientNavigationHelper.h"
+
+FViewportClientNavigationHelper::FViewportClientNavigationHelper()
+	: bIsMouseLooking(false)
+{
+	ResetTransformDelta();
+	ResetImpulseData();
+}
+
+bool FViewportClientNavigationHelper::HasTransformDelta() const
+{
+	return !RotationDelta.IsNearlyZero() || !LocationDelta.IsNearlyZero() || !OrbitDelta.IsNearlyZero();
+}
+
+void FViewportClientNavigationHelper::ConsumeTransformDelta(FVector& OutLocationDelta, FRotator& OutRotationDelta)
+{
+	OutLocationDelta = LocationDelta;
+	OutRotationDelta = RotationDelta;
+
+	// Reset after use
+	ResetTransformDelta();
+}
+
+void FViewportClientNavigationHelper::ResetTransformDelta()
+{
+	LocationDelta = FVector::ZeroVector;
+	RotationDelta = FRotator::ZeroRotator;
+}
+
+void FViewportClientNavigationHelper::ConsumeImpulseData(FCameraControllerUserImpulseData& OutImpulseData)
+{
+	// Translation
+	OutImpulseData.MoveForwardBackwardImpulse += ImpulseDataDelta.MoveForwardBackwardImpulse;
+	OutImpulseData.MoveRightLeftImpulse += ImpulseDataDelta.MoveRightLeftImpulse;
+	OutImpulseData.MoveWorldUpDownImpulse += ImpulseDataDelta.MoveWorldUpDownImpulse;
+	OutImpulseData.MoveLocalUpDownImpulse += ImpulseDataDelta.MoveLocalUpDownImpulse;
+
+	// Rotation
+	OutImpulseData.RotateYawImpulse += ImpulseDataDelta.RotateYawImpulse;
+	OutImpulseData.RotatePitchImpulse += ImpulseDataDelta.RotatePitchImpulse;
+	OutImpulseData.RotateRollImpulse += ImpulseDataDelta.RotateRollImpulse;
+	OutImpulseData.RotateYawVelocityModifier += ImpulseDataDelta.RotateYawVelocityModifier;
+	OutImpulseData.RotatePitchVelocityModifier += ImpulseDataDelta.RotatePitchVelocityModifier;
+	OutImpulseData.RotateRollVelocityModifier += ImpulseDataDelta.RotateRollVelocityModifier;
+
+	// FOV
+	OutImpulseData.ZoomOutInImpulse += ImpulseDataDelta.ZoomOutInImpulse;
+
+	// Reset after use
+	ResetImpulseData();
+}
+
+void FViewportClientNavigationHelper::ConsumeOrbitDelta(FVector& OutOrbitDelta)
+{
+	OutOrbitDelta = OrbitDelta;
+	OrbitDelta = FVector::ZeroVector;
+}
+
+void FViewportClientNavigationHelper::ResetImpulseData()
+{
+	ImpulseDataDelta.MoveForwardBackwardImpulse = 0.0f;
+	ImpulseDataDelta.MoveRightLeftImpulse = 0.0f;
+	ImpulseDataDelta.MoveWorldUpDownImpulse = 0.0f;
+	ImpulseDataDelta.MoveLocalUpDownImpulse = 0.0f;
+	ImpulseDataDelta.RotateYawImpulse = 0.0f;
+	ImpulseDataDelta.RotatePitchImpulse = 0.0f;
+	ImpulseDataDelta.RotateRollImpulse = 0.0f;
+	ImpulseDataDelta.RotateYawVelocityModifier = 0.0f;
+	ImpulseDataDelta.RotatePitchVelocityModifier = 0.0f;
+	ImpulseDataDelta.RotateRollVelocityModifier = 0.0f;
+	ImpulseDataDelta.ZoomOutInImpulse = 0.0f;
+}
diff --git a/Engine/Source/Editor/UnrealEd/Public/EditorViewportClient.h b/Engine/Source/Editor/UnrealEd/Public/EditorViewportClient.h
--- a/Engine/Source/Editor/UnrealEd/Public/EditorViewportClient.h
+++ b/Engine/Source/Editor/UnrealEd/Public/EditorViewportClient.h
@@ -40,6 +40,7 @@ class UActorFactory;
 class UTypedElementViewportInteraction;
 enum class EViewStatusForScreenPercentage;
 struct FGizmoState;
+struct FViewportClientNavigationHelper;
 
 /** Delegate called by FEditorViewportClient to check its visibility */
 DECLARE_DELEGATE_RetVal( bool, FViewportStateGetter );
@@ -1598,6 +1599,11 @@ public:
 	inline bool GetIsCurrentLevelEditingFocus() const								{ return bIsCurrentLevelEditingFocus; }
 	inline void SetIsCurrentLevelEditingFocus(bool bInIsCurrentLevelEditingFocus)	{ bIsCurrentLevelEditingFocus = bInIsCurrentLevelEditingFocus; }
 
+	FViewportClientNavigationHelper* GetViewportNavigationHelper() const
+	{
+		return ViewportNavigationHelper;
+	}
+
 protected:
 	/** Invalidates the viewport widget (if valid) to register its active timer */
 	UNREALED_API void InvalidateViewportWidget();
@@ -1954,6 +1960,12 @@ protected:
 	/** Current cached impulse state */
 	FCameraControllerUserImpulseData* CameraUserImpulseData;
 
+	/**
+	 * Can be used to add location, rotation, and FOV deltas to the viewport camera.
+	 * Should be used by ITF-based Viewport Interactions to control the camera.
+	 */
+	FViewportClientNavigationHelper* ViewportNavigationHelper;
+
 	/** When we have LOD locking, it's slow to force redraw of other viewports, so we delay invalidates to reduce the number of redraws */
 	double TimeForForceRedraw;
 
diff --git a/Engine/Source/Editor/UnrealEd/Public/SEditorViewport.h b/Engine/Source/Editor/UnrealEd/Public/SEditorViewport.h
--- a/Engine/Source/Editor/UnrealEd/Public/SEditorViewport.h
+++ b/Engine/Source/Editor/UnrealEd/Public/SEditorViewport.h
@@ -16,6 +16,8 @@ class FSceneViewport;
 class FUICommandList;
 class IPreviewProfileController;
 class SViewport;
+class UModeManagerInteractiveToolsContext;
+class UViewportInteractionsBehaviorSource;
 struct FSlateBrush;
 struct FToolMenuContext;
 
@@ -271,6 +273,23 @@ protected:
 	/** Gets the world this viewport is for */
 	UNREALED_API virtual UWorld* GetWorld() const;
 
+	/**
+	 * Override this function to add the desired viewport interactions for this editor.
+	 * Requires virtual function GetModeManagerInteractiveToolsContext to return a valid UModeManagerInteractiveToolsContext
+	 */
+	UNREALED_API virtual void OnRegisterViewportInteractions(UViewportInteractionsBehaviorSource* InBehaviorSource)
+	{
+	}
+
+	/**
+	 * Override to return the Mode Manager Interactive Tools Context for your specific SEditorViewport. Required by
+	 * OnRegisterViewportInteractions in order to obtain a valid UViewportInteractionsBehaviorSource
+	 */
+	UNREALED_API virtual UModeManagerInteractiveToolsContext* GetModeManagerInteractiveToolsContext()
+	{
+		return nullptr;
+	}
+
 	/**
 	 * Called when surface snapping has been enabled/disabled
 	 */
@@ -336,6 +355,12 @@ private:
 
 	/** Gets the visibility of the active viewport border */
 	UNREALED_API EVisibility GetActiveBorderVisibility() const;
+
+	/**
+	 * Will register ITF viewport interactions, if OnRegisterViewportInteractions is overridden in your SEditorViewport
+	 */
+	void RegisterViewportInteractions();
+
 private:
 	/** The handle to the active timer */
 	TWeakPtr<FActiveTimerHandle> ActiveTimerHandle;
diff --git a/Engine/Source/Editor/UnrealEd/Public/Tools/EdModeInteractiveToolsContext.h b/Engine/Source/Editor/UnrealEd/Public/Tools/EdModeInteractiveToolsContext.h
--- a/Engine/Source/Editor/UnrealEd/Public/Tools/EdModeInteractiveToolsContext.h
+++ b/Engine/Source/Editor/UnrealEd/Public/Tools/EdModeInteractiveToolsContext.h
@@ -48,6 +48,7 @@ class UMaterialInterface;
 class UObject;
 class USelection;
 class UTypedElementSelectionSet;
+class UViewportInteractionsBehaviorSource;
 struct FToolBuilderState;
 
 /**
@@ -367,6 +368,14 @@ public:
 		bUsesDragTools = bInEnabled;
 	}
 
+	/**
+	 * Enable ITF Input tools support (e.g. Viewport Interactions as camera movement and selection)
+	 */
+	void SetInputToolsEnabled(bool bInEnabled)
+	{
+		bUsesViewportInteractions = bInEnabled;
+	}
+
 	/**
 	 * Are Drag Tools supported? e.g. Marquee Select
 	 * Enable/Disable support by using SetDragToolsEnabled(...)
@@ -374,7 +383,20 @@ public:
 	bool UsesDragTools() const
 	{
 		return bUsesDragTools;
-	};
+	}
+
+	/**
+	 * Are Viewport Interactions supported? e.g. camera movements, commands, clicking, etc
+	 */
+	bool UsesViewportInteractions() const
+	{
+		return bUsesViewportInteractions;
+	}
+
+	/**
+	 * Return the current Input Tools Behavior Source - nullptr if not enabled
+	 */
+	UE_API UViewportInteractionsBehaviorSource* GetViewportInteractionsBehaviorSource();
 
 protected:
 	/**
@@ -382,11 +404,21 @@ protected:
 	 */
 	UE_API void RegisterDragTools();
 
+	/**
+	 * Activating Viewport Interactions if current Mode Manager requires them
+	 */
+	void RegisterViewportInteractions();
+
 	/**
 	 * Deactivating Drag Tools
 	 */
 	UE_API void UnregisterDragTools();
 
+	/**
+	 * Deactivating Viewport Interactions
+	 */
+	void UnregisterViewportInteractions();
+
 	/**
 	 * Responding to Mode Changes in order to deactivate/activate Drag Tools accordingly
 	 */
@@ -416,7 +448,11 @@ protected:
 	UPROPERTY(Transient)
 	TObjectPtr<UDragToolsBehaviorSource> DragToolsBehaviorSource;
 
+	UPROPERTY(Transient)
+	TObjectPtr<UViewportInteractionsBehaviorSource> ViewportInteractionsBehaviorSource;
+
 	bool bUsesDragTools = false;
+	bool bUsesViewportInteractions = false;
 };
 
 
diff --git a/Engine/Source/Editor/UnrealEd/Public/ViewportClientNavigationHelper.h b/Engine/Source/Editor/UnrealEd/Public/ViewportClientNavigationHelper.h
new file mode 100644
--- /dev/null
+++ b/Engine/Source/Editor/UnrealEd/Public/ViewportClientNavigationHelper.h
@@ -0,0 +1,78 @@
+// Copyright Epic Games, Inc. All Rights Reserved.
+
+#pragma once
+
+#include "CameraController.h"
+
+/**
+ * Helper class used by FEditorViewportClient to allow external inputs to move the camera.
+ * Used by ITF Viewport interactions.
+ * When deltas from this struct are used, they should be reset since they are used and updated on tick.
+ * Please use ConsumeTransformDelta(...) and ConsumeImpulseData(...) to do that.
+ */
+struct FViewportClientNavigationHelper
+{
+public:
+	FViewportClientNavigationHelper();
+
+	/**
+	 * A location delta which can be used to move a Viewport camera
+	 */
+	FVector LocationDelta;
+
+	/**
+	 * A rotation delta which can be used to rotate a Viewport camera
+	 */
+	FRotator RotationDelta;
+
+	/**
+	 * A delta which is used by Editor Viewport Client to compute orbit movement values
+	 */
+	FVector OrbitDelta;
+
+	/**
+	 * Impulse data used to control rotation, translation and FOV of a camera
+	 * See FEditorViewportClient::UpdateCameraMovement
+	 */
+	FCameraControllerUserImpulseData ImpulseDataDelta;
+
+	/**
+	 * Check whether non-zero location or rotation deltas are available. Also checks for Orbit delta.
+	 * Useful to skip certain logic in FEditorViewportClient::UpdateMouseDelta() or for other purposes
+	 */
+	bool HasTransformDelta() const;
+
+	/**
+	 * Retrieve current location and rotation deltas, and zero them out
+	 * @param OutLocationDelta location delta will be copied to OutLocationDelta
+	 * @param OutRotationDelta rotation delta will be copied to OutRotationDelta
+	 */
+	void ConsumeTransformDelta(FVector& OutLocationDelta, FRotator& OutRotationDelta);
+
+	/**
+	 * Retrieve current impulse deltas and zero them out
+	 * @param OutImpulseData current impulse data will be copied to OutImpulseData before being reset
+	 */
+	void ConsumeImpulseData(FCameraControllerUserImpulseData& OutImpulseData);
+
+	/**
+	 * Retrieve current delta used to generate orbit movement
+	 * @param OutOrbitDelta
+	 */
+	void ConsumeOrbitDelta(FVector& OutOrbitDelta);
+
+	/**
+	 * Sets Location and Rotation deltas to zero
+	 */
+	void ResetTransformDelta();
+
+	/**
+	 * Sets all impulse data to zero
+	 */
+	void ResetImpulseData();
+
+	/**
+	 * This flag is true when the user is controlling a viewport camera using the mouse
+	 */
+	bool bIsMouseLooking;
+};
diff --git a/Engine/Source/Runtime/InteractiveToolsFramework/Private/BaseBehaviors/MouseWheelModifierInputBehavior.cpp b/Engine/Source/Runtime/InteractiveToolsFramework/Private/BaseBehaviors/MouseWheelModifierInputBehavior.cpp
new file mode 100644
--- /dev/null
+++ b/Engine/Source/Runtime/InteractiveToolsFramework/Private/BaseBehaviors/MouseWheelModifierInputBehavior.cpp
@@ -0,0 +1,20 @@
+// Copyright Epic Games, Inc. All Rights Reserved.
+
+#include "BaseBehaviors/MouseWheelModifierInputBehavior.h"
+
+FInputCaptureRequest UMouseWheelModifierInputBehavior::WantsCapture(const FInputDeviceState& InputState)
+{
+	if (InputState.Mouse.WheelDelta != 0 && (ModifierCheckFunc == nullptr || ModifierCheckFunc(InputState)))
+	{
+		if (InputState.Mouse.WheelDelta > 0)
+		{
+			Target->OnMouseWheelScrollUp(InputState);
+		}
+		else
+		{
+			Target->OnMouseWheelScrollDown(InputState);
+		}
+	}
+
+	return FInputCaptureRequest::Ignore();
+}
diff --git a/Engine/Source/Runtime/InteractiveToolsFramework/Private/InputRouter.cpp b/Engine/Source/Runtime/InteractiveToolsFramework/Private/InputRouter.cpp
--- a/Engine/Source/Runtime/InteractiveToolsFramework/Private/InputRouter.cpp
+++ b/Engine/Source/Runtime/InteractiveToolsFramework/Private/InputRouter.cpp
@@ -273,6 +273,11 @@ void UInputRouter::CheckForMouseCaptures(const FInputDeviceState& Input)
 
 void UInputRouter::HandleCapturedMouseInput(const FInputDeviceState& Input)
 {
+	// Calling this to make sure behaviors treating mouse buttons as modifiers are always updated,
+	// even if there is some drag-based behavior currently capturing
+	TArray<FInputCaptureRequest> CaptureRequests;
+	ActiveInputBehaviors->CollectWantsCapture(Input, CaptureRequests);
+
 	if (ActiveLeftCapture == nullptr)
 	{
 		return;
diff --git a/Engine/Source/Runtime/InteractiveToolsFramework/Public/BaseBehaviors/KeyAsModifierInputBehavior.h b/Engine/Source/Runtime/InteractiveToolsFramework/Public/BaseBehaviors/KeyAsModifierInputBehavior.h
--- a/Engine/Source/Runtime/InteractiveToolsFramework/Public/BaseBehaviors/KeyAsModifierInputBehavior.h
+++ b/Engine/Source/Runtime/InteractiveToolsFramework/Public/BaseBehaviors/KeyAsModifierInputBehavior.h
@@ -70,3 +70,17 @@ protected:
 	/** Modifier set for this behavior, internally initialized with check on ModifierKey */
 	FInputBehaviorModifierStates Modifiers;
 };
+
+/**
+ * A version of UKeyAsModifierInputBehavior for mouse buttons
+ */
+UCLASS(MinimalAPI)
+class UMouseButtonAsModifierInputBehavior : public UKeyAsModifierInputBehavior
+{
+	GENERATED_BODY()
+
+	virtual EInputDevices GetSupportedDevices() override
+	{
+		return EInputDevices::Mouse;
+	}
+};
diff --git a/Engine/Source/Runtime/InteractiveToolsFramework/Public/BaseBehaviors/MouseWheelModifierInputBehavior.h b/Engine/Source/Runtime/InteractiveToolsFramework/Public/BaseBehaviors/MouseWheelModifierInputBehavior.h
new file mode 100644
--- /dev/null
+++ b/Engine/Source/Runtime/InteractiveToolsFramework/Public/BaseBehaviors/MouseWheelModifierInputBehavior.h
@@ -0,0 +1,18 @@
+// Copyright Epic Games, Inc. All Rights Reserved.
+
+#pragma once
+
+#include "MouseWheelBehavior.h"
+#include "MouseWheelModifierInputBehavior.generated.h"
+
+/**
+ * 
+ */
+UCLASS()
+class INTERACTIVETOOLSFRAMEWORK_API UMouseWheelModifierInputBehavior : public UMouseWheelInputBehavior
+{
+	GENERATED_BODY()
+
+public:
+	virtual FInputCaptureRequest WantsCapture(const FInputDeviceState& InputState) override;
+};
