Fix for cook on the fly shader compilation, resulting in empty shader maps for materials discovered during the cook.
- Enables ODSC when using cook-on-the-fly
- Can be disabled using Cook.OnTheFly.ShaderCompilationDisabled, and Cook.OnTheFly.AllowIncompleteShadermaps ini settings

#tests Tested cook-on-the-fly via quick launch of QAGame, EngineTests (various levels) and Lyra

diff --git a/Engine/Source/Editor/UnrealEd/Private/CookOnTheFlyServer.cpp b/Engine/Source/Editor/UnrealEd/Private/CookOnTheFlyServer.cpp
--- a/Engine/Source/Editor/UnrealEd/Private/CookOnTheFlyServer.cpp
+++ b/Engine/Source/Editor/UnrealEd/Private/CookOnTheFlyServer.cpp
@@ -239,6 +239,20 @@ static FAutoConsoleVariableRef CVarCookTimeCVarControl(
 		"NOTE: CURRENTLY any non-zero mode will also enable redirecting a ShaderPlatform to the cooking platform"),
 	ECVF_Default);
 
+bool bDisableShaderCompilationDuringCookOnTheFly = false;
+static FAutoConsoleVariableRef CVarDisableShaderCompilationDuringCookOnTheFly(
+	TEXT("Cook.OnTheFly.ShaderCompilationDisabled"),
+	bDisableShaderCompilationDuringCookOnTheFly,
+	TEXT("Controls whether shader compilation is disabled during cook on the fly\n"),
+	ECVF_Default);
+
+bool bAllowIncompleteShaderMapsDuringCookOnTheFly = true;
+static FAutoConsoleVariableRef CVarAllowEmptyShaderMapsDuringCookOnTheFly(
+	TEXT("Cook.OnTheFly.AllowIncompleteShadermaps"),
+	bAllowIncompleteShaderMapsDuringCookOnTheFly,
+	TEXT("Controls whether incomplete shader maps are processed during cook on the fly\n"),
+	ECVF_Default);
+
 ////////////////////////////////////////////////////////////////
 /// Cook on the fly server
 ///////////////////////////////////////////////////////////////
@@ -610,11 +624,8 @@ bool UCookOnTheFlyServer::StartCookOnTheFly(FCookOnTheFlyStartupOptions InCookOn
 {
 	using namespace UE::Cook;
 
-	if (!IsCookingInEditor())
-	{
-		GShaderCompilingManager->SkipShaderCompilation(true);
-		GShaderCompilingManager->SetAllowForIncompleteShaderMaps(true);
-	}
+	GShaderCompilingManager->SkipShaderCompilation(bDisableShaderCompilationDuringCookOnTheFly);
+	GShaderCompilingManager->SetAllowForIncompleteShaderMaps(bAllowIncompleteShaderMapsDuringCookOnTheFly);
 
 	LLM_SCOPE_BYTAG(Cooker);
 #if WITH_COTF
@@ -5601,9 +5612,12 @@ void UCookOnTheFlyServer::ShutdownCookOnTheFly()
 
 		ShutdownCookSession();
 
-		if (!IsCookingInEditor())
+		if (bDisableShaderCompilationDuringCookOnTheFly)
 		{
 			GShaderCompilingManager->SkipShaderCompilation(false);
+		}
+		if (bAllowIncompleteShaderMapsDuringCookOnTheFly)
+		{
 			GShaderCompilingManager->SetAllowForIncompleteShaderMaps(false);
 		}
 	}
@@ -11523,8 +11537,8 @@ void UCookOnTheFlyServer::StartCookAsCookWorker()
 	LoadBeginCookIncrementalFlags(BeginContext);
 	if (IsDirectorCookOnTheFly())
 	{
-		GShaderCompilingManager->SkipShaderCompilation(true);
-		GShaderCompilingManager->SetAllowForIncompleteShaderMaps(true);
+		GShaderCompilingManager->SkipShaderCompilation(bDisableShaderCompilationDuringCookOnTheFly);
+		GShaderCompilingManager->SetAllowForIncompleteShaderMaps(bAllowIncompleteShaderMapsDuringCookOnTheFly);
 	}
 	CookWorkerClient->DoneWithInitialSettings();
 
@@ -11586,9 +11600,16 @@ void UCookOnTheFlyServer::CookAsCookWorkerFinished()
 
 	if (IsDirectorCookOnTheFly())
 	{
-		GShaderCompilingManager->SkipShaderCompilation(false);
-		GShaderCompilingManager->SetAllowForIncompleteShaderMaps(false);
+		if (bDisableShaderCompilationDuringCookOnTheFly)
+		{
+			GShaderCompilingManager->SkipShaderCompilation(false);
+		}
+		if (bAllowIncompleteShaderMapsDuringCookOnTheFly)
+		{
+			GShaderCompilingManager->SetAllowForIncompleteShaderMaps(false);
+		}
 	}
+
 	LogCookWorkerStats();
 	BroadcastCookFinished();
 	CookWorkerClient->FlushLogs();
