Assign a loaded path to packages generated at cook time. The loaded path will match their actual loaded path in a cooked game. This ensures that functionallity that uses the loaded path at runtime will be the same as what happens at cook time.

diff --git a/Engine/Source/Editor/UnrealEd/Private/Cooker/CookGenerationHelper.cpp b/Engine/Source/Editor/UnrealEd/Private/Cooker/CookGenerationHelper.cpp
--- a/Engine/Source/Editor/UnrealEd/Private/Cooker/CookGenerationHelper.cpp
+++ b/Engine/Source/Editor/UnrealEd/Private/Cooker/CookGenerationHelper.cpp
@@ -1594,6 +1594,8 @@ UPackage* FGenerationHelper::TryCreateGeneratedPackage(FCookGenerationInfo& Gene
 		// PackageTracker know the referencer of the package.
 		UE_TRACK_REFERENCING_PACKAGE_SCOPED(GetOwner().GetPackageName(), PackageAccessTrackingOps::NAME_PostLoad);
 		GeneratedPackage = CreatePackage(*GeneratedPackageName);
+		GeneratedPackage->SetLoadedPath(FPackagePath::FromPackageNameChecked(GeneratedPackageName));
+
 	}
 	GeneratedPackage->SetSavedHash(GeneratedInfo.PackageHash);
 	GeneratedPackage->SetPersistentGuid(LocalOwnerPackage->GetPersistentGuid());
