[] Fix for Bake to Anim Sequence/Create Linked Anim Sequence not remembring the last folder used

diff --git a/Engine/Source/Editor/UnrealEd/Private/AnimationEditorUtils.cpp b/Engine/Source/Editor/UnrealEd/Private/AnimationEditorUtils.cpp
--- a/Engine/Source/Editor/UnrealEd/Private/AnimationEditorUtils.cpp
+++ b/Engine/Source/Editor/UnrealEd/Private/AnimationEditorUtils.cpp
@@ -65,9 +65,17 @@ void SCreateAnimationAssetDlg::Construct(const FArguments& InArgs)
 	AssetPath = FText::FromString(FPackageName::GetLongPackagePath(InArgs._DefaultAssetPath.ToString()));
 	AssetName = FText::FromString(FPackageName::GetLongPackageAssetName(InArgs._DefaultAssetPath.ToString()));
 
-	if (AssetPath.IsEmpty())
+	const bool bPrioritizeLastPathUsed = InArgs._PrioritizeLastPathUsed;
+	
+	if (AssetPath.IsEmpty() || (bPrioritizeLastPathUsed && !LastUsedAssetPath.IsEmpty()))
 	{
 		AssetPath = LastUsedAssetPath;
+
+		// If the AssetPath is still empty, fall back to the Game dir
+		if (AssetPath.IsEmpty())
+		{
+			AssetPath = FText::FromString(TEXT("/Game"));
+		}
 	}
 	else
 	{
@@ -496,7 +504,8 @@ namespace AnimationEditorUtils
 					// set the unique asset as a default name
 					TSharedRef<SCreateAnimationAssetDlg> NewAnimDlg =
 						SNew(SCreateAnimationAssetDlg)
-						.DefaultAssetPath(FText::FromString(PackageName));
+						.DefaultAssetPath(FText::FromString(PackageName))
+						.PrioritizeLastPathUsed(true);
 
 					// show a dialog to determine a new asset name
 					if (NewAnimDlg->ShowModal() == EAppReturnType::Cancel)
diff --git a/Engine/Source/Editor/UnrealEd/Public/AnimationEditorUtils.h b/Engine/Source/Editor/UnrealEd/Public/AnimationEditorUtils.h
--- a/Engine/Source/Editor/UnrealEd/Public/AnimationEditorUtils.h
+++ b/Engine/Source/Editor/UnrealEd/Public/AnimationEditorUtils.h
@@ -34,9 +34,11 @@ class SCreateAnimationAssetDlg : public SWindow
 {
 public:
 	SLATE_BEGIN_ARGS(SCreateAnimationAssetDlg)
+		: _PrioritizeLastPathUsed(false)
 	{
 	}
 		SLATE_ARGUMENT(FText, DefaultAssetPath)
+		SLATE_ARGUMENT(bool, PrioritizeLastPathUsed)
 	SLATE_END_ARGS()
 
 		SCreateAnimationAssetDlg()
