DataTableEditorUtils AddRow calls HandleDataTableChanged

FTableRowBase receives callbacks to OnDataTableChanged only through property edits on existing row elements within a DataTable. Adding this call within DataTableEditorUtils::AddRow makes it convenient for the FTableRowBase types respond to the add operation themselves, performing any default initialization or validation.

PR #12213: DataTableEditorUtils AddRow calls HandleDataTableChanged

diff --git a/Engine/Source/Editor/UnrealEd/Private/DataTableEditorUtils.cpp b/Engine/Source/Editor/UnrealEd/Private/DataTableEditorUtils.cpp
--- a/Engine/Source/Editor/UnrealEd/Private/DataTableEditorUtils.cpp
+++ b/Engine/Source/Editor/UnrealEd/Private/DataTableEditorUtils.cpp
@@ -498,6 +498,7 @@ uint8* FDataTableEditorUtils::AddRow(UDataTable* DataTable, FName RowName)
 
 	// Add to row map
 	DataTable->AddRowInternal(RowName, RowData);
+	DataTable->HandleDataTableChanged(RowName);
 	BroadcastPostChange(DataTable, EDataTableChangeInfo::RowList);
 	return RowData;
 }
