[UE] Add support for shadowed subobject references to T3D import/export.

Notes:
- Preserves default subobject type changes across import/export operations in the editor when the outer type supports override serialization (OS) and when exported references are not explicitly overridden at export time.

diff --git a/Engine/Source/Editor/UnrealEd/Private/EditorObject.cpp b/Engine/Source/Editor/UnrealEd/Private/EditorObject.cpp
--- a/Engine/Source/Editor/UnrealEd/Private/EditorObject.cpp
+++ b/Engine/Source/Editor/UnrealEd/Private/EditorObject.cpp
@@ -512,11 +512,41 @@ static const TCHAR* ImportProperties(
 				TemplateClass = BaseTemplate->GetClass();
 			}
 
+			const bool bOverrideT3DSupported = FOverridableSerializationLogic::HasCapabilities(FOverridableSerializationLogic::ECapabilities::T3DSerialization);
+
 			// because the outer won't be a default object
 
 			checkSlow(TemplateClass != NULL);
 			if (bRedefiningSubobject)
 			{
+				// When override T3D is supported, during the initial definition phase (the 'else' case below), the exported
+				// subobject may have been found to not be type-compatible with the default subobject that was instanced for
+				// the outer object at construction time (e.g. due to schema changes after exporting). In that case, we will
+				// have created a unique instance to serve as the destination for the redefinition here. This means that any
+				// exported instance data will not be imported to the instanced default subobject, and its data will be lost.
+				//
+				// Note: Referencing properties are also NOT imported in this case unless explicitly overridden at export time,
+				// so outside of that scenario, inner references (by design) are not remapped to the unique instance on import.
+				//
+				// @todo - IDO support will be needed to capture the unique instance and preserve its imported instance data.
+				if (bOverrideT3DSupported)
+				{
+					FSoftObjectPath ExportedObjectPath(ExportedObjectFullName);
+					checkf(ObjectRemapper, TEXT("Object remapper is expected to be non-NULL when importing T3D with override support."));
+					if (!ensureMsgf(bExportPath, TEXT("Subobject redefinition is missing the object's export path while importing it.")))
+					{
+						// Fall back to the resolved subobject, however its path can differ from the exported path that it was mapped
+						// to if e.g. there was a name collision within the outer scope, so it's not as reliable as the exported path.
+						ExportedObjectPath = FSoftObjectPath(BaseTemplate);
+					}
+
+					UObject* const* RemappedObject = ObjectRemapper->Find(ExportedObjectPath);
+					if (RemappedObject)
+					{
+						BaseTemplate = *RemappedObject;
+					}
+				}
+
 				// since we're redefining an object in the same text block, only need to import properties again
 				SourceText = ImportObjectProperties( (uint8*)BaseTemplate, SourceText, TemplateClass, SubobjectRoot, BaseTemplate,
 													Warn, Depth + 1, ContextSupplier ? ContextSupplier->CurrentLine : 0, &InstanceGraph, ObjectRemapper );
@@ -532,7 +562,6 @@ static const TCHAR* ImportProperties(
 			{
 				FName OverriddenOperationName = NAME_None;
 				const bool bParsed = FParse::Value(Str,TEXT("OverriddenOperation="), OverriddenOperationName);
-				const bool bOverrideT3DSupported = FOverridableSerializationLogic::HasCapabilities(FOverridableSerializationLogic::ECapabilities::T3DSerialization);
 				TOptional<EOverriddenPropertyOperation> OverriddenOperation = bParsed && bOverrideT3DSupported ? GetOverriddenOperationFromName(OverriddenOperationName) : TOptional<EOverriddenPropertyOperation>();
 
 				UObject* Archetype = NULL;
@@ -629,7 +658,7 @@ static const TCHAR* ImportProperties(
 					Archetype = BaseTemplate;
 				}
 
-				UObject* OldComponent = NULL;
+				UObject* OldComponent = nullptr;
 				if (ComponentTemplate)
 				{
 					bool bIsOkToReuse = ComponentTemplate->GetClass() == TemplateClass
@@ -639,10 +668,43 @@ static const TCHAR* ImportProperties(
 
 					if (!bIsOkToReuse)
 					{
-						UE_LOG(LogEditorObject, Log, TEXT("Could not reuse component instance %s, name clash?"), *ComponentTemplate->GetFullName());
-						ComponentTemplate->Rename(nullptr, nullptr, REN_DontCreateRedirectors); // just abandon the existing component, we are going to create
-						OldComponent = ComponentTemplate;
-						ComponentTemplate = NULL;
+						if (bOverrideT3DSupported
+							&& !ComponentTemplate->GetClass()->IsChildOf(TemplateClass))
+						{
+							// If we found an existing instanced subobject above that is not type-compatible,
+							// create a new object to serve as the destination for importing instance data.
+							// Otherwise, we allow instance data to be imported into the existing subobject.
+							// 
+							// This differs from the non-override path, which moves the instanced subobject
+							// aside and constructs a new type-compatible subobject in its place (see below).
+							// That path will remain in place for legacy support and backwards-compatibility.
+							// 
+							// When exporting override T3D, all reference property values will be "shadowed" on
+							// import if they were not otherwise explicitly overridden at export time. Shadowed
+							// reference values are skipped at import time, so it means we intentionally sidestep
+							// the reference replacement pass below, and instead defer replacement to occur only
+							// when we import object properties and find an explicitly-overridden reference value.
+							// In that case the override will be resolved through the export path remapping table.
+							// 
+							// This effectively matches what happens on level load when a serialized export blocks
+							// deferred subobject instancing due to e.g. a schema change that occurs while unloaded.
+							// In that case archetype lookup will fail due to type incompatibility, and the exported
+							// subobject is moved aside to allow a default instance of the updated type to be created.
+							//
+							// @todo - IDO support will be needed to capture the new instance that will result
+							// in order to avoid data loss. If the new instance is not otherwise referenced by
+							// an explicitly-overridden property at import time, then it will become unreachable.
+							TemplateName = MakeUniqueObjectName(SubobjectOuter, TemplateClass, TemplateName);
+						}
+						else
+						{
+							UE_LOG(LogEditorObject, Log, TEXT("Could not reuse component instance %s, name clash?"), *ComponentTemplate->GetFullName());
+							ComponentTemplate->Rename(nullptr, nullptr, REN_DontCreateRedirectors); // just abandon the existing component, we are going to create
+							OldComponent = ComponentTemplate;
+						}
+
+						// Forces a new instance to be created below
+						ComponentTemplate = nullptr;
 					}
 				}
 
@@ -848,6 +910,13 @@ const TCHAR* ImportObjectProperties( FImportObjectParams& InParams )
 	if (bShouldUseDedicatedInstancingGraph)
 	{
 		InstanceGraph.SetDestinationRoot(InParams.SubobjectOuter);
+
+		// When using override T3D, always map exported subobject paths so they can be resolved at redefinition time
+		TMap<FSoftObjectPath, UObject*> ExportedPathToObjectCreated;
+		if (!InParams.ObjectRemapper)
+		{
+			InParams.ObjectRemapper = &ExportedPathToObjectCreated;
+		}
 	}
 	else if ( InParams.SubobjectRoot && InParams.SubobjectRoot != UObject::StaticClass()->GetDefaultObject() )
 	{
diff --git a/Engine/Source/Runtime/CoreUObject/Private/UObject/CoreNative.cpp b/Engine/Source/Runtime/CoreUObject/Private/UObject/CoreNative.cpp
--- a/Engine/Source/Runtime/CoreUObject/Private/UObject/CoreNative.cpp
+++ b/Engine/Source/Runtime/CoreUObject/Private/UObject/CoreNative.cpp
@@ -286,14 +286,29 @@ UObject* FObjectInstancingGraph::GetInstancedSubobject( UObject* SourceSubobject
 								// we need to check the object's type against the type that we're about to instance below. If there's a mismatch, rename
 								// the existing object out of the way, as subobject recycling will assert in that case. Any serialized data will be lost.
 								UObject* ExistingObject = StaticFindObjectFast(nullptr, SubobjectOuter, SubobjectName);
-								if (ExistingObject
-									&& !ExistingObject->GetClass()->IsChildOf(SourceSubobject->GetClass()))
+								if (ExistingObject)
 								{
-									// Keep the existing base name and scope in case some other property was overridden and serialized this reference.
-									// If nothing else winds up referencing this object after load, the export will be garbage collected (which would
-									// happen whether or not we renamed the object below - that is being done in order to free up the name for instancing).
-									FName NewName = MakeUniqueObjectName(SubobjectOuter, ExistingObject->GetClass(), SubobjectName);
-									ExistingObject->Rename(*NewName.ToString(), nullptr, REN_DoNotDirty | REN_NonTransactional | REN_DontCreateRedirectors);
+									if (!ExistingObject->HasAnyFlags(RF_NeedLoad)
+										&& SourceSubobject->IsA(ExistingObject->GetClass()))
+									{
+										if (SourceSubobject->GetClass() == ExistingObject->GetClass())
+										{
+											InstancedSubobject = ExistingObject;
+										}
+										else
+										{
+											// Use the serialized instance data to construct a new object of the correct source type.
+											SourceSubobject = ExistingObject;
+										}
+									}
+									else
+									{
+										// Keep the existing base name and scope in case some other property was overridden and serialized this reference.
+										// If nothing else winds up referencing this object after load, the export will be garbage collected (which would
+										// happen whether or not we renamed the object below - that is being done in order to free up the name for instancing).
+										FName NewName = MakeUniqueObjectName(SubobjectOuter, ExistingObject->GetClass(), SubobjectName);
+										ExistingObject->Rename(*NewName.ToString(), nullptr, REN_DoNotDirty | REN_NonTransactional | REN_DontCreateRedirectors);
+									}
 								}
 							}
 
diff --git a/Engine/Source/Runtime/CoreUObject/Private/UObject/OverriddenPropertySet.cpp b/Engine/Source/Runtime/CoreUObject/Private/UObject/OverriddenPropertySet.cpp
--- a/Engine/Source/Runtime/CoreUObject/Private/UObject/OverriddenPropertySet.cpp
+++ b/Engine/Source/Runtime/CoreUObject/Private/UObject/OverriddenPropertySet.cpp
@@ -180,10 +180,8 @@ EOverriddenPropertyOperation FOverridableSerializationLogic::GetOverriddenProper
 	checkf(OverriddenPortTextPropertyPath, TEXT("Expecting an overridden port text path"));
 
 	const FArchiveSerializedPropertyChain CurrentPropertyChain = OverriddenPortTextPropertyPath->ToSerializedPropertyChain();
-	const EOverriddenPropertyOperation Operation = GetOverriddenPropertyOperation(PortFlags,  &CurrentPropertyChain, nullptr, DataPtr, DefaultValue);
-
-	// For now lets not support subobject shadow serialization until the copy and paste support loose property or placeholder
-	return Operation != EOverriddenPropertyOperation::SubObjectsShadowing ? Operation : EOverriddenPropertyOperation::None;
+	const EOverriddenPropertyOperation Operation = GetOverriddenPropertyOperation(PortFlags, &CurrentPropertyChain, nullptr, DataPtr, DefaultValue);
+	return Operation;
 }
 
 FPropertyVisitorPath* FOverridableSerializationLogic::GetOverriddenPortTextPropertyPath()
diff --git a/Engine/Source/Runtime/CoreUObject/Private/UObject/Property.cpp b/Engine/Source/Runtime/CoreUObject/Private/UObject/Property.cpp
--- a/Engine/Source/Runtime/CoreUObject/Private/UObject/Property.cpp
+++ b/Engine/Source/Runtime/CoreUObject/Private/UObject/Property.cpp
@@ -1935,6 +1935,13 @@ const TCHAR* FProperty::ImportSingleProperty( const TCHAR* Str, void* DestData,
 				return Result;
 			};
 
+			// Avoid shadowed reference values (not directly imported) - callers should advance
+			// @todo - IDOs will need to capture the referenced instance to help avoid data loss
+			if (Operation.IsSet() && Operation.GetValue() == EOverriddenPropertyOperation::SubObjectsShadowing)
+			{
+				return Str;
+			}
+
 			// Handle modifying overridable operations on arrays
 			if (Operation.IsSet() && Operation.GetValue() != EOverriddenPropertyOperation::Replace)
 			{
