Removed FString copy when building a single-element TConstArrayView<FString>

#rnx

diff --git a/Engine/Source/Editor/UnrealEd/Public/IProjectExternalContentInterface.h b/Engine/Source/Editor/UnrealEd/Public/IProjectExternalContentInterface.h
--- a/Engine/Source/Editor/UnrealEd/Public/IProjectExternalContentInterface.h
+++ b/Engine/Source/Editor/UnrealEd/Public/IProjectExternalContentInterface.h
@@ -59,7 +59,6 @@ public:
 
 	void RemoveExternalContent(const FString& VersePath, FRemoveExternalContentComplete CompleteCallback = FRemoveExternalContentComplete())
 	{
-		TArray<FString, TInlineAllocator<1>> VersePaths = { VersePath };
-		RemoveExternalContent(VersePaths, MoveTemp(CompleteCallback));
+		RemoveExternalContent(MakeArrayView(&VersePath, 1), MoveTemp(CompleteCallback));
 	}
 };
