[UnrealEd] Allow edmode cursor override even when interacting/dragging

diff --git a/Engine/Source/Editor/UnrealEd/Private/EditorViewportClient.cpp b/Engine/Source/Editor/UnrealEd/Private/EditorViewportClient.cpp
--- a/Engine/Source/Editor/UnrealEd/Private/EditorViewportClient.cpp
+++ b/Engine/Source/Editor/UnrealEd/Private/EditorViewportClient.cpp
@@ -83,10 +83,6 @@
 #include "LandscapeSettings.h"
 #include "TransformGizmoEditorSettings.h"
 #include "Misc/AxisDisplayInfo.h"
-#include "Tools/EdModeInteractiveToolsContext.h"
-#include "ViewportClientNavigationHelper.h"
-#include "ViewportInteractions/ViewportInteractionsBehaviorSource.h"
-#include "ViewportInteractions/ViewportInteraction.h"
 
 #define LOCTEXT_NAMESPACE "EditorViewportClient"
 
@@ -132,23 +128,6 @@ static bool GetDefaultLowDPIPreviewValue()
 	return CVarEditorViewportHighDPIPtr->GetInt() == 0;
 }
 
-/**
- * Returns true if any viewport interaction of the specified type is active
- */
-bool IsViewportInteractionActive(const FEditorModeTools* InModeTools, FName InViewportInteraction)
-{
-	if (InModeTools)
-	{
-		if (UViewportInteractionsBehaviorSource* InputToolsBehaviorSource =
-				InModeTools->GetInteractiveToolsContext()->GetViewportInteractionsBehaviorSource())
-		{
-			return InputToolsBehaviorSource->IsViewportInteractionActive(InViewportInteraction);
-		}
-	}
-
-	return false;
-}
-
 float ComputeOrthoZoomFactor(const float ViewportWidth)
 {
 	float Ret = 1.0f;
@@ -490,7 +469,6 @@ FEditorViewportClient::FEditorViewportClient(FEditorModeTools* InModeTools, FPre
 	, bHasMouseMovedSinceClick(false)
 	, CameraController(new FEditorCameraController())
 	, CameraUserImpulseData(new FCameraControllerUserImpulseData())
-	, ViewportNavigationHelper(new FViewportClientNavigationHelper())
 	, TimeForForceRedraw(0.0)
 	, FlightCameraSpeedScale(1.0f)
 	, bUseControllingActorViewInfo(false)
@@ -1796,15 +1774,15 @@ EMouseCursor::Type FEditorViewportClient::GetCursor(FViewport* InViewport,int32
 		MouseCursor = EMouseCursor::CardinalCross;
 	}
 	// Don't select widget axes by mouse over while they're being controlled by a mouse drag.
-	else if( InViewport->IsCursorVisible() && !bWidgetAxisControlledByDrag && !ModeTools->HasOngoingTransform())
+	else if (InViewport->IsCursorVisible() && !bWidgetAxisControlledByDrag)
 	{
 		// allow editor modes to override cursor
 		EMouseCursor::Type EditorModeCursor = EMouseCursor::Default;
-		if(ModeTools->GetCursor(EditorModeCursor))
+		if (ModeTools->GetCursor(EditorModeCursor))
 		{
 			MouseCursor = EditorModeCursor;
 		}
-		else
+		else if (ModeTools->HasOngoingTransform())
 		{
 			HHitProxy* HitProxy = InViewport->GetHitProxy(X,Y);
 
@@ -1930,18 +1908,14 @@ void FEditorViewportClient::UpdateCameraMovement( float DeltaTime )
 		const bool bIsUsingTrackpad = FSlateApplication::Get().IsUsingTrackpad();
 
 		// Certain keys are only available while the flight camera input mode is active
-		const bool bUsingFlightInput = IsFlightCameraInputModeActive() || bIsUsingTrackpad || ViewportNavigationHelper->bIsMouseLooking;
+		const bool bUsingFlightInput = IsFlightCameraInputModeActive() || bIsUsingTrackpad;
 
 		// Is the current press unmodified?
 		const bool bUnmodifiedPress = !IsAltPressed() && !IsShiftPressed() && !IsCtrlPressed() && !IsCmdPressed();
 
 		// Do we want to use the regular arrow keys for flight input?
 		// Because the arrow keys are used for things like nudging actors, we'll only do this while the press is unmodified
-		const bool bRemapArrowKeys = bUnmodifiedPress || IsViewportInteractionActive(ModeTools.Get(), UE::Editor::ViewportInteractions::CameraTranslate);
-
-		const bool bITFCameraTranslateActive = IsViewportInteractionActive(ModeTools.Get(),UE::Editor::ViewportInteractions::CameraTranslate);
-		const bool bITFCameraRotateActive = IsViewportInteractionActive(ModeTools.Get(),UE::Editor::ViewportInteractions::CameraRotate);
-		const bool bITFFOVActive = IsViewportInteractionActive(ModeTools.Get(),UE::Editor::ViewportInteractions::FOV);
+		const bool bRemapArrowKeys = bUnmodifiedPress;
 
 		// Do we want to remap the various WASD keys for flight input?
 		const bool bRemapWASDKeys =
@@ -2006,110 +1980,91 @@ void FEditorViewportClient::UpdateCameraMovement( float DeltaTime )
 			CameraController->GetConfig().bForceRotationalPhysics = false;
 		}
 
-		// Check if we need to use legacy translation input
-		if (!bITFCameraTranslateActive)
+		// Forward/back
+		if( ( bRemapWASDKeys && bForwardKeyState ) ||
+			( bRemapArrowKeys && Viewport->KeyState( EKeys::Up ) ) ||
+			( bUnmodifiedPress && bUseNumpadCameraControl && Viewport->KeyState(EKeys::NumPadEight) ) )
 		{
-			// Forward/back
-			if( ( bRemapWASDKeys && bForwardKeyState ) ||
-				( bRemapArrowKeys && Viewport->KeyState( EKeys::Up ) ) ||
-				( bUnmodifiedPress && bUseNumpadCameraControl && Viewport->KeyState(EKeys::NumPadEight) ) )
-			{
-				CameraUserImpulseData->MoveForwardBackwardImpulse += 1.0f;
-			}
-			if( (bRemapWASDKeys && bBackwardKeyState) ||
-				( bRemapArrowKeys && Viewport->KeyState( EKeys::Down ) ) ||
-				( bUnmodifiedPress && bUseNumpadCameraControl && Viewport->KeyState( EKeys::NumPadTwo ) ) )
-			{
-				CameraUserImpulseData->MoveForwardBackwardImpulse -= 1.0f;
-			}
-
-			// Right/left
-			if (( bRemapWASDKeys && bRightKeyState) ||
-				( bRemapArrowKeys && Viewport->KeyState( EKeys::Right ) ) ||
-				( bUnmodifiedPress && bUseNumpadCameraControl && Viewport->KeyState( EKeys::NumPadSix ) ) )
-			{
-				CameraUserImpulseData->MoveRightLeftImpulse += 1.0f;
-			}
-			if( ( bRemapWASDKeys && bLeftKeyState) ||
-				( bRemapArrowKeys && Viewport->KeyState( EKeys::Left ) ) ||
-				( bUnmodifiedPress && bUseNumpadCameraControl && Viewport->KeyState( EKeys::NumPadFour ) ) )
-			{
-				CameraUserImpulseData->MoveRightLeftImpulse -= 1.0f;
-			}
+			CameraUserImpulseData->MoveForwardBackwardImpulse += 1.0f;
+		}
+		if( (bRemapWASDKeys && bBackwardKeyState) ||
+			( bRemapArrowKeys && Viewport->KeyState( EKeys::Down ) ) ||
+			( bUnmodifiedPress && bUseNumpadCameraControl && Viewport->KeyState( EKeys::NumPadTwo ) ) )
+		{
+			CameraUserImpulseData->MoveForwardBackwardImpulse -= 1.0f;
+		}
 
-			// World up/down
-			if( ( bRemapWASDKeys && bWorldUpKeyState) ||
-				( bUnmodifiedPress && Viewport->KeyState( EKeys::PageUp ) ) ||
-				( bUnmodifiedPress && bUseNumpadCameraControl && ( Viewport->KeyState( EKeys::NumPadNine ) || Viewport->KeyState( EKeys::Add ) ) ) )
-			{
-				CameraUserImpulseData->MoveWorldUpDownImpulse += 1.0f;
-			}
-			if( ( bRemapWASDKeys && bWorldDownKeyState) ||
-				( bUnmodifiedPress && Viewport->KeyState( EKeys::PageDown ) ) ||
-				( bUnmodifiedPress && bUseNumpadCameraControl && ( Viewport->KeyState( EKeys::NumPadSeven ) || Viewport->KeyState( EKeys::Subtract ) ) ) )
-			{
-				CameraUserImpulseData->MoveWorldUpDownImpulse -= 1.0f;
-			}
+		// Right/left
+		if (( bRemapWASDKeys && bRightKeyState) ||
+			( bRemapArrowKeys && Viewport->KeyState( EKeys::Right ) ) ||
+			( bUnmodifiedPress && bUseNumpadCameraControl && Viewport->KeyState( EKeys::NumPadSix ) ) )
+		{
+			CameraUserImpulseData->MoveRightLeftImpulse += 1.0f;
+		}
+		if( ( bRemapWASDKeys && bLeftKeyState) ||
+			( bRemapArrowKeys && Viewport->KeyState( EKeys::Left ) ) ||
+			( bUnmodifiedPress && bUseNumpadCameraControl && Viewport->KeyState( EKeys::NumPadFour ) ) )
+		{
+			CameraUserImpulseData->MoveRightLeftImpulse -= 1.0f;
+		}
 
-			// Local up/down
-			if( ( bRemapWASDKeys && bLocalUpKeyState) )
-			{
-				CameraUserImpulseData->MoveLocalUpDownImpulse += 1.0f;
-			}
-			if( ( bRemapWASDKeys && bLocalDownKeyState) )
-			{
-				CameraUserImpulseData->MoveLocalUpDownImpulse -= 1.0f;
-			}
+		// World up/down
+		if( ( bRemapWASDKeys && bWorldUpKeyState) ||
+			( bUnmodifiedPress && Viewport->KeyState( EKeys::PageUp ) ) ||
+			( bUnmodifiedPress && bUseNumpadCameraControl && ( Viewport->KeyState( EKeys::NumPadNine ) || Viewport->KeyState( EKeys::Add ) ) ) )
+		{
+			CameraUserImpulseData->MoveWorldUpDownImpulse += 1.0f;
+		}
+		if( ( bRemapWASDKeys && bWorldDownKeyState) ||
+			( bUnmodifiedPress && Viewport->KeyState( EKeys::PageDown ) ) ||
+			( bUnmodifiedPress && bUseNumpadCameraControl && ( Viewport->KeyState( EKeys::NumPadSeven ) || Viewport->KeyState( EKeys::Subtract ) ) ) )
+		{
+			CameraUserImpulseData->MoveWorldUpDownImpulse -= 1.0f;
 		}
 
-		// Check if we need to use legacy FOV input
-		if (!bITFFOVActive)
+		// Local up/down
+		if( ( bRemapWASDKeys && bLocalUpKeyState) )
 		{
-			// Zoom FOV out/in
-			if( ( bRemapWASDKeys && bZoomOutKeyState) ||
-				( bUnmodifiedPress && bUseNumpadCameraControl && Viewport->KeyState( EKeys::NumPadOne ) ) )
-			{
-				CameraUserImpulseData->ZoomOutInImpulse += 1.0f;
-			}
-			if( ( bRemapWASDKeys && bZoomInKeyState) ||
-				( bUnmodifiedPress && bUseNumpadCameraControl && Viewport->KeyState( EKeys::NumPadThree ) ) )
-			{
-				CameraUserImpulseData->ZoomOutInImpulse -= 1.0f;
-			}
+			CameraUserImpulseData->MoveLocalUpDownImpulse += 1.0f;
+		}
+		if( ( bRemapWASDKeys && bLocalDownKeyState) )
+		{
+			CameraUserImpulseData->MoveLocalUpDownImpulse -= 1.0f;
 		}
 
-		// Check if we need to use legacy rotation input
-		if (!bITFCameraRotateActive)
+		// Zoom FOV out/in
+		if( ( bRemapWASDKeys && bZoomOutKeyState) ||
+			( bUnmodifiedPress && bUseNumpadCameraControl && Viewport->KeyState( EKeys::NumPadOne ) ) )
 		{
-			// Rotate up/down
-			if (bRemapWASDKeys && bRotateUpKeyState)
-			{
-				CameraUserImpulseData->RotatePitchImpulse += 1.0f * CameraController->GetConfig().RotationMultiplier;
-				CameraController->GetConfig().bForceRotationalPhysics = true;
-			}
-			if (bRemapWASDKeys && bRotateDownKeyState)
-			{
-				CameraUserImpulseData->RotatePitchImpulse -= 1.0f * CameraController->GetConfig().RotationMultiplier;
-				CameraController->GetConfig().bForceRotationalPhysics = true;
-			}
+			CameraUserImpulseData->ZoomOutInImpulse += 1.0f;
+		}
+		if( ( bRemapWASDKeys && bZoomInKeyState) ||
+			( bUnmodifiedPress && bUseNumpadCameraControl && Viewport->KeyState( EKeys::NumPadThree ) ) )
+		{
+			CameraUserImpulseData->ZoomOutInImpulse -= 1.0f;
+		}
 
-			// Rotate left/right
-			if (bRemapWASDKeys && bRotateLeftKeyState)
-			{
-				CameraUserImpulseData->RotateYawImpulse -= 1.0f * CameraController->GetConfig().RotationMultiplier;
-				CameraController->GetConfig().bForceRotationalPhysics = true;
-			}
-			if (bRemapWASDKeys && bRotateRightKeyState)
-			{
-				CameraUserImpulseData->RotateYawImpulse += 1.0f * CameraController->GetConfig().RotationMultiplier;
-				CameraController->GetConfig().bForceRotationalPhysics = true;
-			}
+		// Rotate up/down
+		if (bRemapWASDKeys && bRotateUpKeyState)
+		{
+			CameraUserImpulseData->RotatePitchImpulse += 1.0f * CameraController->GetConfig().RotationMultiplier;
+			CameraController->GetConfig().bForceRotationalPhysics = true;
+		}
+		if (bRemapWASDKeys && bRotateDownKeyState)
+		{
+			CameraUserImpulseData->RotatePitchImpulse -= 1.0f * CameraController->GetConfig().RotationMultiplier;
+			CameraController->GetConfig().bForceRotationalPhysics = true;
 		}
 
-		// Add camera inputs coming from ITF, if any
-		if (bITFCameraTranslateActive || bITFCameraRotateActive || bITFFOVActive)
+		// Rotate left/right
+		if (bRemapWASDKeys && bRotateLeftKeyState)
+		{
+			CameraUserImpulseData->RotateYawImpulse -= 1.0f * CameraController->GetConfig().RotationMultiplier;
+			CameraController->GetConfig().bForceRotationalPhysics = true;
+		}
+		if (bRemapWASDKeys && bRotateRightKeyState)
 		{
-			ViewportNavigationHelper->ConsumeImpulseData(*CameraUserImpulseData);
+			CameraUserImpulseData->RotateYawImpulse += 1.0f * CameraController->GetConfig().RotationMultiplier;
 			CameraController->GetConfig().bForceRotationalPhysics = true;
 		}
 
@@ -2156,17 +2111,17 @@ void FEditorViewportClient::UpdateCameraMovement( float DeltaTime )
 		// Check whether the camera is being moved by the mouse or keyboard
 		bool bHasMovement = GetDefault<ULevelEditorViewportSettings>()->bUseLegacyCameraMovementNotifications;
 
-		if (CameraUserImpulseData->RotateYawVelocityModifier != 0.0f ||
-			CameraUserImpulseData->RotatePitchVelocityModifier != 0.0f ||
-			CameraUserImpulseData->RotateRollVelocityModifier != 0.0f ||
-			CameraUserImpulseData->MoveForwardBackwardImpulse != 0.0f ||
-			CameraUserImpulseData->MoveRightLeftImpulse != 0.0f ||
-			CameraUserImpulseData->MoveWorldUpDownImpulse != 0.0f ||
-			CameraUserImpulseData->MoveLocalUpDownImpulse != 0.0f ||
-			CameraUserImpulseData->ZoomOutInImpulse != 0.0f ||
-			CameraUserImpulseData->RotateYawImpulse != 0.0f ||
-			CameraUserImpulseData->RotatePitchImpulse != 0.0f ||
-			CameraUserImpulseData->RotateRollImpulse != 0.0f
+		if ((*CameraUserImpulseData).RotateYawVelocityModifier != 0.0f ||
+			(*CameraUserImpulseData).RotatePitchVelocityModifier != 0.0f ||
+			(*CameraUserImpulseData).RotateRollVelocityModifier != 0.0f ||
+			(*CameraUserImpulseData).MoveForwardBackwardImpulse != 0.0f ||
+			(*CameraUserImpulseData).MoveRightLeftImpulse != 0.0f ||
+			(*CameraUserImpulseData).MoveWorldUpDownImpulse != 0.0f ||
+			(*CameraUserImpulseData).MoveLocalUpDownImpulse != 0.0f ||
+			(*CameraUserImpulseData).ZoomOutInImpulse != 0.0f ||
+			(*CameraUserImpulseData).RotateYawImpulse != 0.0f ||
+			(*CameraUserImpulseData).RotatePitchImpulse != 0.0f ||
+			(*CameraUserImpulseData).RotateRollImpulse != 0.0f
 			)
 		{
 			bHasMovement = true;
@@ -2478,8 +2433,7 @@ void FEditorViewportClient::UpdateMouseDelta()
 
 	if( Viewport )
 	{
-		// If mouse delta tracker has delta OR there is transform delta coming from ITF Viewport Interaction, handle it
-		if (!DragDelta.IsNearlyZero() || ViewportNavigationHelper->HasTransformDelta())
+		if( !DragDelta.IsNearlyZero() )
 		{
 			const bool LeftMouseButtonDown = Viewport->KeyState(EKeys::LeftMouseButton);
 			const bool MiddleMouseButtonDown = Viewport->KeyState(EKeys::MiddleMouseButton);
@@ -2492,18 +2446,11 @@ void FEditorViewportClient::UpdateMouseDelta()
 			ModeTools->PreConvertMouseMovement(this);
 
 			// Convert the movement delta into drag/rotation deltas
-			FVector Drag = FVector::ZeroVector;
-			FRotator Rot = FRotator::ZeroRotator;
-			FVector Scale = FVector::ZeroVector;
+			FVector Drag;
+			FRotator Rot;
+			FVector Scale;
 			EAxisList::Type CurrentAxis = Widget->GetCurrentAxis();
-
-			// Try to handle ITF Viewport inputs before legacy (currently zoom and camera drag ones)
-			if (IsViewportInteractionActive(ModeTools.Get(), UE::Editor::ViewportInteractions::Zoom) ||
-				IsViewportInteractionActive(ModeTools.Get(), UE::Editor::ViewportInteractions::CameraDrag))
-			{
-				ViewportNavigationHelper->ConsumeTransformDelta(Drag, Rot);
-			}
-			else if ( IsOrtho() && ( LeftMouseButtonDown || bIsUsingTrackpad ) && RightMouseButtonDown )
+			if ( IsOrtho() && ( LeftMouseButtonDown || bIsUsingTrackpad ) && RightMouseButtonDown )
 			{
 				bWidgetAxisControlledByDrag = false;
 				Widget->SetCurrentAxis( EAxisList::None );
@@ -2525,8 +2472,7 @@ void FEditorViewportClient::UpdateMouseDelta()
 				}
 				//if Absolute Translation, and not just moving the camera around
 				if (IsUsingAbsoluteTranslation(false))
-				{
-					// TODO: yet to be checked out for ITF Parity
+				{ 
 					MouseDeltaTracker->AbsoluteTranslationConvertMouseToDragRot(DragStartView, this, Drag, Rot, Scale);
 				}
 				else
@@ -2559,21 +2505,7 @@ void FEditorViewportClient::UpdateMouseDelta()
 
 void FEditorViewportClient::PeformDefaultCameraMovement(FVector& Drag, FRotator& Rot, FVector& Scale)
 {
-	// Check if there is an active ITF Orbit interaction, and if there's an orbit delta it might have generated 
-	const bool bUseITFOrbitDelta = IsViewportInteractionActive(ModeTools.Get(), UE::Editor::ViewportInteractions::Orbit)
-					  && !ViewportNavigationHelper->OrbitDelta.IsNearlyZero();
-
-	FVector DragDelta;
-
-	if (bUseITFOrbitDelta)
-	{
-		ViewportNavigationHelper->ConsumeOrbitDelta(DragDelta);
-	}
-	else // no ITF delta, try to use the one from mouse delta tracker
-	{
-		DragDelta = MouseDeltaTracker->GetDelta();
-	}
-
+	FVector DragDelta = MouseDeltaTracker->GetDelta();
 	if (ShouldOrbitCamera())
 	{
 		bool bHasMovement = !DragDelta.IsNearlyZero();
@@ -3250,7 +3182,6 @@ bool FEditorViewportClient::Internal_InputKey(const FInputKeyEventArgs& EventArg
 	const EInputEvent& Event = EventArgs.Event;
 	const FViewport* InViewport = EventArgs.Viewport;
 
-	bool bRouteToToolsContext = !bIsTracking || UE::Editor::ViewportInteractions::UseEditorViewportInteractions();
 	// Let the current mode have a look at the input before reacting to it. 
 	// Note that bIsTracking tells us whether the viewport client is capturing mouse behavior to fly around,
 	// move objects, etc. In this case we don't want to pass the input the input router, because it is already
@@ -3259,12 +3190,7 @@ bool FEditorViewportClient::Internal_InputKey(const FInputKeyEventArgs& EventArg
 	// for now. In the long term, only old FEdModes will get this input, while the new UEdModes won't, but
 	// this is not yet split apart. Also in the long term, the viewport behaviors will be inside the
 	// input router and we won't need this bIsTracking flag to begin with.
-	if (ModeTools->InputKey(this, Viewport, Key, Event, bRouteToToolsContext))
-	{
-		return true;
-	}
-
-	if (GetViewportNavigationHelper()->bIsMouseLooking)
+	if (ModeTools->InputKey(this, Viewport, Key, Event, /*bRouteToToolsContext*/ !bIsTracking))
 	{
 		return true;
 	}
@@ -3342,10 +3268,9 @@ bool FEditorViewportClient::Internal_InputKey(const FInputKeyEventArgs& EventArg
 		return true;
 	}
 
-	const bool bIsITFSingleClickActive = IsViewportInteractionActive(ModeTools.Get(), UE::Editor::ViewportInteractions::SingleClick);
 
 	// If we are tracking and no mouse button is down and this input event released the mouse button stop tracking and process any clicks if necessary
-	if ( bIsTracking && !InputState.IsAnyMouseButtonDown() && InputState.IsMouseButtonEvent() && !bIsITFSingleClickActive)
+	if ( bIsTracking && !InputState.IsAnyMouseButtonDown() && InputState.IsMouseButtonEvent() )
 	{
 		// Handle possible mouse click viewport
 		ProcessClickInViewport( InputState, *View );
@@ -4926,12 +4851,6 @@ bool FEditorViewportClient::Internal_InputAxis(FViewport* InViewport, FInputDevi
 		return true;
 	}
 
-	// No need to process anything else here, since some ITF Viewport interaction is handling the camera
-	if (GetViewportNavigationHelper()->bIsMouseLooking)
-	{
-		return true;
-	}
-
 	const bool bMouseButtonDown = InViewport->KeyState( EKeys::LeftMouseButton ) || InViewport->KeyState( EKeys::MiddleMouseButton ) || InViewport->KeyState( EKeys::RightMouseButton );
 	const bool bLightMoveDown = InViewport->KeyState(EKeys::L);
 
@@ -5185,7 +5104,7 @@ void FEditorViewportClient::ConvertMovementToDragRot(const FVector& InDelta,
 				if (MiddleMouseButtonDown && !ViewportSettings->bWorldSpaceVerticalPan)
 				{
 					// Pan relative to the camera
-					InDragDelta = FQuat(ViewRotation) * FVector( 0.0f, InDelta.X, InDelta.Y) * (bInvert ? 1.0f : -1.0f);
+					InDragDelta = FQuat(ViewRotation) * FVector( 0.0f, InDelta.X, InDelta.Y) * (bInvert ? -1.0f : 1.0f);
 				}
 				else
 				{
@@ -6105,16 +6024,16 @@ FViewportCursorLocation FEditorViewportClient::GetCursorWorldLocationFromMousePo
 
 void FEditorViewportClient::CapturedMouseMove( FViewport* InViewport, int32 InMouseX, int32 InMouseY )
 {
+	UpdateRequiredCursorVisibility();
+	ApplyRequiredCursorVisibility();
+
+	CapturedMouseMoves.Add(FIntPoint(InMouseX, InMouseY));
+
 	// Let the current editor mode know about the mouse movement.
 	if (ModeTools->CapturedMouseMove(this, InViewport, InMouseX, InMouseY))
 	{
 		return;
 	}
-
-	UpdateRequiredCursorVisibility();
-	ApplyRequiredCursorVisibility();
-
-	CapturedMouseMoves.Add(FIntPoint(InMouseX, InMouseY));
 }
 
 void FEditorViewportClient::ProcessAccumulatedPointerInput(FViewport* InViewport)
