Additional logging and error handling for content workers that seemingly get stuck inside FinishInternal function.

We should be able to detect pre-determined set of scopes through UAT by watching the logs and decide (assume) whether or not the process is stuck.

diff --git a/Engine/Source/Editor/UnrealEd/Classes/CookOnTheSide/CookLog.h b/Engine/Source/Editor/UnrealEd/Classes/CookOnTheSide/CookLog.h
--- a/Engine/Source/Editor/UnrealEd/Classes/CookOnTheSide/CookLog.h
+++ b/Engine/Source/Editor/UnrealEd/Classes/CookOnTheSide/CookLog.h
@@ -6,4 +6,6 @@
 
 DECLARE_LOG_CATEGORY_EXTERN(LogCook, Log, All);
 DECLARE_LOG_CATEGORY_EXTERN(LogCookStats, Log, All);
+// LogCookStatus: Logs should be hidden by default but enabled on the build machines so they can be forwarded to stdout
+DECLARE_LOG_CATEGORY_EXTERN(LogCookStatus, Warning, All);
 extern FName LogCookName;
diff --git a/Engine/Source/Editor/UnrealEd/Private/CookOnTheFlyServer.cpp b/Engine/Source/Editor/UnrealEd/Private/CookOnTheFlyServer.cpp
--- a/Engine/Source/Editor/UnrealEd/Private/CookOnTheFlyServer.cpp
+++ b/Engine/Source/Editor/UnrealEd/Private/CookOnTheFlyServer.cpp
@@ -9762,6 +9762,8 @@ void UCookOnTheFlyServer::CookByTheBookFinished()
 
 void UCookOnTheFlyServer::CookByTheBookFinishedInternal()
 {
+	UE_LOG(LogCookStatus, Display, TEXT("FinishedInternal: Start"));
+
 	using namespace UE::Cook;
 
 	check(IsInGameThread());
@@ -9804,8 +9806,6 @@ void UCookOnTheFlyServer::CookByTheBookFinishedInternal()
 		}
 	};
 
-	UE_LOG(LogCook, Display, TEXT("Finishing up..."));
-
 	{
 		UE_SCOPED_COOKTIMER(TickCookableObjects);
 		const double CurrentTime = FPlatformTime::Seconds();
@@ -9813,9 +9813,13 @@ void UCookOnTheFlyServer::CookByTheBookFinishedInternal()
 		LastCookableObjectTickTime = CurrentTime;
 	}
 
+	UE_LOG(LogCookStatus, Display, TEXT("Waiting for AsyncFileWrites"));
 	UPackage::WaitForAsyncFileWrites();
+
+	UE_LOG(LogCookStatus, Display, TEXT("Waiting for BuildDefinitions"));
 	BuildDefinitions->Wait();
 	
+	UE_LOG(LogCookStatus, Display, TEXT("Waiting for DDC"));
 	GetDerivedDataCacheRef().WaitForQuiescence(true);
 	
 	bool bSaveAssetRegistry = !FParse::Param(FCommandLine::Get(), TEXT("SkipSaveAssetRegistry"));
@@ -9879,6 +9883,8 @@ void UCookOnTheFlyServer::CookByTheBookFinishedInternal()
 		}
 	}
 
+	UE_LOG(LogCookStatus, Display, TEXT("Collecting cook information for all platforms"));
+
 	for (const ITargetPlatform* TargetPlatform : PlatformManager->GetSessionPlatforms())
 	{
 		if (FindOrCreateSaveContext(TargetPlatform).PackageWriterCapabilities.bReadOnly)
@@ -10089,6 +10095,8 @@ void UCookOnTheFlyServer::CookByTheBookFinishedInternal()
 		}
 	}
 
+	UE_LOG(LogCookStatus, Display, TEXT("Shutting down shader library cooker & shader compilers"));
+
 	FString ActualLibraryName = GenerateShaderCodeLibraryName(LibraryName, IsCookFlagSet(ECookInitializationFlags::LegacyIterativeSharedBuild));
 	FShaderLibraryCooker::EndCookingLibrary(ActualLibraryName);
 	FShaderLibraryCooker::Shutdown();
@@ -10096,6 +10104,8 @@ void UCookOnTheFlyServer::CookByTheBookFinishedInternal()
 
 	if (CookByTheBookOptions->bGenerateDependenciesForMaps)
 	{
+		UE_LOG(LogCookStatus, Display, TEXT("Generating dependencies for maps"));
+
 		UE_SCOPED_HIERARCHICAL_COOKTIMER(GenerateMapDependencies);
 		for (const ITargetPlatform* Platform : PlatformManager->GetSessionPlatforms())
 		{
@@ -10109,9 +10119,16 @@ void UCookOnTheFlyServer::CookByTheBookFinishedInternal()
 		}
 	}
 
-	GenerateCachedEditorThumbnails();
+
+	if (FParse::Param(FCommandLine::Get(), TEXT("CachedEditorThumbnails")))
+	{
+		UE_LOG(LogCookStatus, Display, TEXT("Generating editor thumbnails"));
+		GenerateCachedEditorThumbnails();
+	}
 
 	FinalizePackageStore();
+
+	UE_LOG(LogCookStatus, Display, TEXT("FinishedInternal: Done"));
 }
 
 void UCookOnTheFlyServer::ShutdownCookSession()
@@ -10283,11 +10300,6 @@ void UCookOnTheFlyServer::WriteMapDependencyGraph(const ITargetPlatform* TargetP
 
 void UCookOnTheFlyServer::GenerateCachedEditorThumbnails()
 {
-	if (!FParse::Param(FCommandLine::Get(), TEXT("CachedEditorThumbnails")))
-	{
-		return;
-	}
-
 	UE_SCOPED_HIERARCHICAL_COOKTIMER(GenerateCachedEditorThumbnails);
 	for (const ITargetPlatform* Platform : PlatformManager->GetSessionPlatforms())
 	{
diff --git a/Engine/Source/Editor/UnrealEd/Private/Cooker/CookLog.cpp b/Engine/Source/Editor/UnrealEd/Private/Cooker/CookLog.cpp
--- a/Engine/Source/Editor/UnrealEd/Private/Cooker/CookLog.cpp
+++ b/Engine/Source/Editor/UnrealEd/Private/Cooker/CookLog.cpp
@@ -25,6 +25,7 @@
 
 DEFINE_LOG_CATEGORY(LogCook);
 DEFINE_LOG_CATEGORY(LogCookStats);
+DEFINE_LOG_CATEGORY(LogCookStatus);
 DEFINE_LOG_CATEGORY(LogCookList);
 
 FName LogCookName(TEXT("LogCook"));
