IncrementalCook: Enable IncrementalCook in DLC and slightly improve the handling of the packages queued for DLC cook.
* Remove the explicit suppression of IncrementalCook in DLC.
* Change BasedOnReleaseCookedPackages from a separate TMap of PackageNames into a bitfield EWhereCooked on the existing omnibus FPackageData, to slightly reduce memory and improve performance.
* Add packages to BasedOnReleaseCookedPackages (now implemented as EWhereCooked) robustly and uniformly from the 3  settings that can affect it: the BasedOnRelease's list of packages, the commandline ExtraReleaseVersionAssets, and the AssetManager's ModifyDLCBasePackages.
* Add The BasedOnReleaseVersion as a CookSetting so that if it changes then the DLC will fully recook. This is useful for projects that do not want to have to implement the importing of CookDependencies from the basegame cook with the DLC cook. Projects that currently do frequent plugin cooks against changed basedonrelease versions should keep the same versionname and will need to import the CookDependencies from the basegame to ensure robust evaluation of build dependencies for the DLC packages.
* Fix Generated packages to not cause a warning when reading packages from the basegame cook's AssetRegistry and bVerifyPackagesExist is on.

diff --git a/Engine/Source/Editor/UnrealEd/Classes/CookOnTheSide/CookOnTheFlyServer.h b/Engine/Source/Editor/UnrealEd/Classes/CookOnTheSide/CookOnTheFlyServer.h
--- a/Engine/Source/Editor/UnrealEd/Classes/CookOnTheSide/CookOnTheFlyServer.h
+++ b/Engine/Source/Editor/UnrealEd/Classes/CookOnTheSide/CookOnTheFlyServer.h
@@ -1324,7 +1324,7 @@ private:
 	 * Populate a map suitable for saving as an ini with the current value of the Cook settings that need to be
 	 * tested for compatibility.
 	 */
-	TMap<FName, FString> CalculateCookSettingStrings() const;
+	TMap<FName, FString> CalculateCookSettingStrings(const FCookByTheBookStartupOptions* StartupOptions) const;
 	/** Get the path to the CookSettings metadata file for the given platform and current output directory. */
 	FString GetCookSettingsFileName(const ITargetPlatform* TargetPlatform) const;
 
diff --git a/Engine/Source/Editor/UnrealEd/Private/CookOnTheFlyServer.cpp b/Engine/Source/Editor/UnrealEd/Private/CookOnTheFlyServer.cpp
--- a/Engine/Source/Editor/UnrealEd/Private/CookOnTheFlyServer.cpp
+++ b/Engine/Source/Editor/UnrealEd/Private/CookOnTheFlyServer.cpp
@@ -7457,7 +7457,7 @@ static FName ExecutableHashName(TEXT("ExecutableHash"));
 static FName ExecutableHashInvalidModuleName(TEXT("ExecutableHashInvalidModule"));
 
 
-TMap<FName, FString> UCookOnTheFlyServer::CalculateCookSettingStrings() const
+TMap<FName, FString> UCookOnTheFlyServer::CalculateCookSettingStrings(const FCookByTheBookStartupOptions* StartupOptions) const
 {
 	TMap<FName, FString> CookSettingStrings;
 	const FName NAME_CookMode(TEXT("CookMode"));
@@ -7522,6 +7522,11 @@ TMap<FName, FString> UCookOnTheFlyServer::CalculateCookSettingStrings() const
 	{
 		CookSettingStrings.Add(NAME_CookMode, TEXT("CookByTheBook"));
 		CookSettingStrings.Add(FName(TEXT("DLCName")), CookByTheBookOptions->DlcName);
+		if (!CookByTheBookOptions->DlcName.IsEmpty())
+		{
+			CookSettingStrings.Add(FName(TEXT("BasedOnReleaseVersion")),
+				StartupOptions ? StartupOptions->BasedOnReleaseVersion : FString());
+		}
 	}
 	else
 	{
@@ -9898,7 +9903,6 @@ void UCookOnTheFlyServer::CookByTheBookFinishedInternal()
 		TArray<FPackageData*> IgnorePackageDatas;
 
 		FString PlatformNameString = TargetPlatform->PlatformName();
-		FName PlatformName(*PlatformNameString);
 
 		TSet<FName> CookedPackageNames;
 		TSet<FName> IgnorePackageNames;
@@ -9909,36 +9913,19 @@ void UCookOnTheFlyServer::CookByTheBookFinishedInternal()
 
 			if (IsCookingDLC())
 			{
-				TMap<FName, FPackageData*> CookedPackagesMap;
-				CookedPackagesMap.Reserve(CookedPackageDatas.Num());
-				for (FPackageData* PackageData : CookedPackageDatas)
-				{
-					CookedPackagesMap.Add(PackageData->GetFileName(), PackageData);
-				}
-				// remove the previous release cooked packages from the new asset registry, add to ignore list
-				UE_SCOPED_HIERARCHICAL_COOKTIMER(RemovingOldManifestEntries);
-
-				const TArray<FName>* PreviousReleaseCookedPackages = CookByTheBookOptions->BasedOnReleaseCookedPackages.Find(PlatformName);
-				if (PreviousReleaseCookedPackages)
+				// We need to add the packages cooked (or incrementally skipped) in this session into the new
+				// AssetRegistry, but need to not add the cooked-elsewhere (e.g. basegame cooked) packages into it. So
+				// we remove them from CookedPackageDatas and add them to IgnorePackageDatas.
+				for (TArray<FPackageData*>::TIterator Iter(CookedPackageDatas); Iter; ++Iter)
 				{
-					for (FName PreviousReleaseCookedPackage : *PreviousReleaseCookedPackages)
+					FPackageData* PackageData = *Iter;
+					FPackagePlatformData& PackagePlatformData = PackageData->FindOrAddPlatformData(TargetPlatform);
+					if (PackagePlatformData.GetWhereCooked() != EWhereCooked::ThisSession)
 					{
-						FPackageData* PackageData;
-						if (!CookedPackagesMap.RemoveAndCopyValue(PreviousReleaseCookedPackage, PackageData))
-						{
-							PackageData = PackageDatas->FindPackageDataByFileName(PreviousReleaseCookedPackage);
-						}
-						if (PackageData)
-						{
-							IgnorePackageDatas.Add(PackageData);
-						}
+						Iter.RemoveCurrentSwap();
+						IgnorePackageDatas.Add(PackageData);
 					}
 				}
-				CookedPackageDatas.Reset();
-				for (TPair<FName, FPackageData*>& Pair : CookedPackagesMap)
-				{
-					CookedPackageDatas.Add(Pair.Value);
-				}
 			}
 
 			CookedPackageNames.Reserve(CookedPackageDatas.Num());
@@ -10613,15 +10600,6 @@ void FBeginCookConfigSettings::LoadLocal(FBeginCookContext& BeginContext)
 			TEXT("INCREMENTAL COOK DEPENDENCIES: Disabled. The cooker is running as ODSC which forces use of legacyiterative."));
 		bCookIncremental = false;
 	}
-	else if (bCookIncremental && BeginContext.COTFS.IsCookingDLC())
-	{
-		// INCREMENTALCOOK_TODO: implement incremental cook for DLC
-		UE_LOG(LogCook, Display, 
-			TEXT("INCREMENTAL COOK DEPENDENCIES: Disabled. Incremental Cook Dependencies are disabled because the cooker is cooking DLC, ")
-			TEXT("and we do not yet support Incremental Cook dependency evaluation in DLC. ")
-			TEXT("Falling back to legacy iterative cook mode."));
-		bCookIncremental = false;
-	}
 	else if (bCookIncremental && !BeginContext.COTFS.IsUsingZenStore())
 	{
 		// Incremental cook uses TargetDomain storage of dependencies which is only implemented in ZenStore
@@ -10810,7 +10788,7 @@ void UCookOnTheFlyServer::LoadBeginCookIncrementalFlagsLocal(FBeginCookContext&
 				UE_LOG(LogCook, Display, TEXT("Local cook is newer then shared cooked build, legacyiteratively cooking from local build."));
 			}
 		}
-		PlatformContext.CurrentCookSettings = CalculateCookSettingStrings();
+		PlatformContext.CurrentCookSettings = CalculateCookSettingStrings(BeginContext.StartupOptions);
 		PlatformContext.bHasMemoryResults = PlatformData->bIsSandboxInitialized;
 
 		if (bIsDiffOnly)
@@ -12045,6 +12023,8 @@ void UCookOnTheFlyServer::GenerateInitialRequests(FBeginCookContext& BeginContex
 
 void UCookOnTheFlyServer::RecordDLCPackagesFromBaseGame(FBeginCookContext& BeginContext)
 {
+	using namespace UE::Cook;
+
 	if (!IsCookingDLC())
 	{
 		return;
@@ -12061,7 +12041,7 @@ void UCookOnTheFlyServer::RecordDLCPackagesFromBaseGame(FBeginCookContext& Begin
 	check(!BasedOnReleaseVersion.IsEmpty());
 
 	auto ReadDevelopmentAssetRegistry = [this, &BasedOnReleaseVersion, bVerifyPackagesExist, bReevaluateUncookedPackages]
-	(TArray<UE::Cook::FConstructPackageData>& OutPackageList, const FString& InPlatformName)
+	(TArray<FConstructPackageData>& OutPackageList, const FString& InPlatformName)
 	{
 		TArray<FString> AttemptedNames;
 		FString OriginalSandboxRegistryFilename = GetBasedOnReleaseVersionAssetRegistryPath(BasedOnReleaseVersion, InPlatformName) / TEXT("Metadata") / GetDevelopmentAssetRegistryFilename();
@@ -12105,7 +12085,7 @@ void UCookOnTheFlyServer::RecordDLCPackagesFromBaseGame(FBeginCookContext& Begin
 		return bSucceeded;
 	};
 
-	TArray<UE::Cook::FConstructPackageData> OverridePackageList;
+	TArray<FConstructPackageData> OverridePackageList;
 	FString DevelopmentAssetRegistryPlatformOverride;
 	const bool bUsingDevRegistryOverride = FParse::Value(FCommandLine::Get(), TEXT("DevelopmentAssetRegistryPlatformOverride="), DevelopmentAssetRegistryPlatformOverride);
 	if (bUsingDevRegistryOverride)
@@ -12125,7 +12105,7 @@ void UCookOnTheFlyServer::RecordDLCPackagesFromBaseGame(FBeginCookContext& Begin
 	bool bFirstAddExistingPackageDatas = true;
 	for (const ITargetPlatform* TargetPlatform : BeginContext.TargetPlatforms)
 	{
-		TArray<UE::Cook::FConstructPackageData> PackageList;
+		TArray<FConstructPackageData> PackageList;
 		FString PlatformNameString = TargetPlatform->PlatformName();
 		FName PlatformName(*PlatformNameString);
 
@@ -12141,27 +12121,29 @@ void UCookOnTheFlyServer::RecordDLCPackagesFromBaseGame(FBeginCookContext& Begin
 			}
 		}
 
-		TArray<UE::Cook::FConstructPackageData>& ActivePackageList = OverridePackageList.Num() > 0 ? OverridePackageList : PackageList;
+		int32 UnusedBaseGameNumForOtherPlatforms = 0;
+		int32& NumForStatusReporting = bFirstAddExistingPackageDatas ? this->PackageDataFromBaseGameNum : UnusedBaseGameNumForOtherPlatforms;
+		TArray<FConstructPackageData>& ActivePackageList = OverridePackageList.Num() > 0 ? OverridePackageList : PackageList;
 		if (ActivePackageList.Num() > 0)
 		{
 			PackageDatas->AddExistingPackageDatasForPlatform(ActivePackageList, TargetPlatform,
-				bFirstAddExistingPackageDatas, PackageDataFromBaseGameNum);
-		}
-
-		TArray<FName>& PlatformBasedPackages = CookByTheBookOptions->BasedOnReleaseCookedPackages.FindOrAdd(PlatformName);
-		PlatformBasedPackages.Reset(ActivePackageList.Num());
-		for (UE::Cook::FConstructPackageData& PackageData : ActivePackageList)
-		{
-			PlatformBasedPackages.Add(PackageData.NormalizedFileName);
+				bFirstAddExistingPackageDatas, NumForStatusReporting);
 		}
 
 		{
 			// allow game or plugins to modify if certain packages from the base game should be recooked.
+			TArray<FName> PlatformBasedPackages;
+			PlatformBasedPackages.Reset(ActivePackageList.Num());
+			for (UE::Cook::FConstructPackageData& PackageData : ActivePackageList)
+			{
+				PlatformBasedPackages.Add(PackageData.NormalizedFileName);
+			}
+
 			TSet<FName> PackagesToClearCookResults;
 			UAssetManager::Get().ModifyDLCBasePackages(TargetPlatform, PlatformBasedPackages, PackagesToClearCookResults);
 			if (PackagesToClearCookResults.Num())
 			{
-				PackageDatas->ClearCookResultsForPackages(PackagesToClearCookResults);
+				PackageDatas->ClearCookResultsForPackages(PackagesToClearCookResults, TargetPlatform, NumForStatusReporting);
 			}
 		}
 		bFirstAddExistingPackageDatas = false;
@@ -12183,10 +12165,20 @@ void UCookOnTheFlyServer::RecordDLCPackagesFromBaseGame(FBeginCookContext& Begin
 		
 		for (const FString& AssetPath : OutAssetPaths)
 		{
-			if (UE::Cook::FPackageData* PackageData = PackageDatas->TryAddPackageDataByFileName(FName(*AssetPath)))
+			if (FPackageData* PackageData = PackageDatas->TryAddPackageDataByFileName(FName(*AssetPath)))
 			{
-				PackageData->SetPlatformsCooked(BeginContext.TargetPlatforms, UE::Cook::ECookResult::Succeeded, /*bWasCookedThisSession=*/false);
-				++PackageDataFromBaseGameNum;
+				for (const ITargetPlatform* TargetPlatform : BeginContext.TargetPlatforms)
+				{
+					{
+						FPackagePlatformData& PlatformData = PackageData->FindOrAddPlatformData(TargetPlatform);
+						if (!PlatformData.IsCookAttempted() && TargetPlatform == BeginContext.TargetPlatforms[0])
+						{
+							++PackageDataFromBaseGameNum;
+						}
+						PlatformData.SetWhereCooked(EWhereCooked::ExtraReleaseVersionAssets);
+					}
+					PackageData->SetPlatformCooked(TargetPlatform, ECookResult::Succeeded, /*bWasCookedThisSession=*/false);
+				}
 			}
 			else
 			{
@@ -12709,8 +12701,6 @@ bool UCookOnTheFlyServer::GetAllPackageFilenamesFromAssetRegistry(const FString&
 					return;
 				}
 
-				// TODO ICookPackageSplitter: Need to handle GeneratedPackages that exist in the cooked AssetRegistry we are 
-				// reading, but do not exist in WorkspaceDomain and so are not found when we look them up here.
 				FName PackageFileName = FPackageDatas::LookupFileNameOnDisk(PackageName, true /* bRequireExists */);
 				if (!PackageFileName.IsNone())
 				{
@@ -12718,7 +12708,8 @@ bool UCookOnTheFlyServer::GetAllPackageFilenamesFromAssetRegistry(const FString&
 					return;
 				}
 
-				if (bVerifyPackagesExist)
+				bool bGeneratedPackage = !!(RegistryData.PackageFlags & PKG_CookGenerated);
+				if (bVerifyPackagesExist && !bGeneratedPackage)
 				{
 					UE_LOG(LogCook, Warning, TEXT("Could not resolve package %s from %s"),
 						*PackageName.ToString(), *AssetRegistryPath);
diff --git a/Engine/Source/Editor/UnrealEd/Private/Cooker/CookPackageData.cpp b/Engine/Source/Editor/UnrealEd/Private/Cooker/CookPackageData.cpp
--- a/Engine/Source/Editor/UnrealEd/Private/Cooker/CookPackageData.cpp
+++ b/Engine/Source/Editor/UnrealEd/Private/Cooker/CookPackageData.cpp
@@ -60,7 +60,7 @@ static FAutoConsoleVariableRef CVarPollAsyncPeriod(
 FPackagePlatformData::FPackagePlatformData()
 	: Reachability((uint8)EReachability::None), ReachabilityVisitedByCluster((uint8)EReachability::None)
 	, bSaveTimedOut(0), bCookable(1), bExplorable(1), bExplorableOverride(0) , IncrementallyUnmodified(0)
-	, bIncrementallySkipped(0), bRegisteredForCachedObjectsInOuter(0), bCommitted(0)
+	, bIncrementallySkipped(0), WhereCooked((uint8)EWhereCooked::ThisSession), bRegisteredForCachedObjectsInOuter(0), bCommitted(0)
 	, CookResults((uint8)ECookResult::NotAttempted)
 {
 }
@@ -2676,6 +2676,8 @@ void FPackageDatas::AddExistingPackageDatasForPlatform(TConstArrayView<FConstruc
 			PackageData = NewPackageData;
 		}
 		PackageData->SetPlatformCooked(TargetPlatform, ECookResult::Succeeded, /*bWasCookedThisSession=*/false);
+		FPackagePlatformData& PlatformData = PackageData->FindOrAddPlatformData(TargetPlatform);
+		PlatformData.SetWhereCooked(EWhereCooked::BaseGame);
 	}
 	OutPackageDataFromBaseGameNum += ExistingPackages.Num();
 }
@@ -2802,22 +2804,33 @@ void FPackageDatas::ClearCookedPlatforms()
 	});
 }
 
-void FPackageDatas::ClearCookResultsForPackages(const TSet<FName>& InPackages)
+void FPackageDatas::ClearCookResultsForPackages(const TSet<FName>& InPackages, const ITargetPlatform* TargetPlatform,
+	int32& InOutNumBaseGamePackages)
 {
 	int32 AffectedPackagesCount = 0;
-	LockAndEnumeratePackageDatas([InPackages, &AffectedPackagesCount](FPackageData* PackageData)
+	LockAndEnumeratePackageDatas([InPackages, &AffectedPackagesCount, TargetPlatform](FPackageData* PackageData)
 		{
 			const FName& PackageName = PackageData->GetPackageName();
 			if (InPackages.Contains(PackageName))
 			{
-				PackageData->ClearCookResults();
-				AffectedPackagesCount++;
+				bool bCookAttempted;
+				{
+					FPackagePlatformData& PlatformData = PackageData->FindOrAddPlatformData(TargetPlatform);
+					PlatformData.SetWhereCooked(EWhereCooked::ThisSession);
+					bCookAttempted = PlatformData.IsCookAttempted();
+				}
+				if (bCookAttempted)
+				{
+					PackageData->ClearCookResults(TargetPlatform);
+					AffectedPackagesCount++;
+				}
 			}
 		});
 
 	UE_LOG(LogCook, Display,
 		TEXT("Cleared the cook results of %d packages because ClearCookResultsForPackages requested them to be recooked."),
 		AffectedPackagesCount);
+	InOutNumBaseGamePackages -= AffectedPackagesCount;
 }
 
 void FPackageDatas::OnRemoveSessionPlatform(const ITargetPlatform* TargetPlatform)
diff --git a/Engine/Source/Editor/UnrealEd/Private/Cooker/CookPackageData.h b/Engine/Source/Editor/UnrealEd/Private/Cooker/CookPackageData.h
--- a/Engine/Source/Editor/UnrealEd/Private/Cooker/CookPackageData.h
+++ b/Engine/Source/Editor/UnrealEd/Private/Cooker/CookPackageData.h
@@ -167,7 +167,6 @@ enum class ESendFlags : uint8
 };
 ENUM_CLASS_FLAGS(ESendFlags);
 
-
 /**
  * Reachability can be set for multiple properties. Reachability of each property is initially assigned to a package
  * from initial requests, by e.g. the AssetManager requests in StartCookByTheBook, and it transitively is assigned to
@@ -270,6 +269,15 @@ struct FPackagePlatformData
 	/** Restore the value to its unset state (e.g. beginning a new cook session). */
 	void ClearIncrementallyUnmodified();
 
+	/**
+	 * Report which session cooked this package, if it was cooked. For packages cooked in the current session, or
+	 * or for packages not yet cooked, this is EWhereCooked::CurrentSession. For packages encountered during DLC
+	 * cooks but cooked in a different session, this result gives the category of that other session's relationship
+	 * to the current session, e.g. the BaseGame cook.
+	 */
+	EWhereCooked GetWhereCooked() const;
+	void SetWhereCooked(EWhereCooked Value);
+
 	/**
 	 * The package was skipped in the current incremental cook. This might not be equal to IncrementallyUnmodified,
 	 * depending on packagewriter).
@@ -308,6 +316,7 @@ private:
 	uint32 bExplorableOverride : 1;
 	uint32 IncrementallyUnmodified : 2;
 	uint32 bIncrementallySkipped : 1;
+	uint32 WhereCooked : (int)EWhereCooked::NumBits;
 	uint32 bRegisteredForCachedObjectsInOuter : 1;
 	uint32 bReportedToDirector : 1;
 	uint32 bCommitted : 1;
@@ -1412,8 +1421,9 @@ public:
 	void Clear();
 	/** Set all platforms to not cooked in all PackageDatas. Used to e.g. invalidate previous cooks. */
 	void ClearCookedPlatforms();
-	/** Set all platforms to not cooked in all PackageDatas that are in the given set */
-	void ClearCookResultsForPackages(const TSet<FName>& InPackages);
+	/** Set PackageDatas that are in the given set to not cooked for the given platform. */
+	void ClearCookResultsForPackages(const TSet<FName>& InPackages, const ITargetPlatform* TargetPlatform,
+		int32& InOutNumBaseGamePackages);
 	/** Remove all data about the given platform from all PackageDatas and other memory used by *this. */
 	void OnRemoveSessionPlatform(const ITargetPlatform* TargetPlatform);
 
@@ -1610,6 +1620,16 @@ inline void FPackagePlatformData::SetIncrementallySkipped(bool bValue)
 	bIncrementallySkipped = (uint32)bValue;
 }
 
+inline EWhereCooked FPackagePlatformData::GetWhereCooked() const
+{
+	return static_cast<EWhereCooked>(WhereCooked);
+}
+
+inline void FPackagePlatformData::SetWhereCooked(EWhereCooked Value)
+{
+	WhereCooked = static_cast<uint32>(Value);
+}
+
 inline ECookResult FPackagePlatformData::GetCookResults() const
 {
 	return (ECookResult)CookResults;
diff --git a/Engine/Source/Editor/UnrealEd/Private/Cooker/CookTypes.cpp b/Engine/Source/Editor/UnrealEd/Private/Cooker/CookTypes.cpp
--- a/Engine/Source/Editor/UnrealEd/Private/Cooker/CookTypes.cpp
+++ b/Engine/Source/Editor/UnrealEd/Private/Cooker/CookTypes.cpp
@@ -956,7 +956,6 @@ FCbWriter& operator<<(FCbWriter& Writer, const UE::Cook::FCookByTheBookOptions&
 
 	Writer << "DlcName" << Value.DlcName;
 	Writer << "CreateReleaseVersion" << Value.CreateReleaseVersion;
-	Writer << "BasedOnReleaseCookedPackages" << Value.BasedOnReleaseCookedPackages;
 	Writer << "SourceToLocalizedPackageVariants" << Value.SourceToLocalizedPackageVariants;
 	Writer << "AllCulturesToCook" << Value.AllCulturesToCook;
 
@@ -987,7 +986,6 @@ bool LoadFromCompactBinary(FCbFieldView Field, UE::Cook::FCookByTheBookOptions&
 
 	bOk = LoadFromCompactBinary(Field["DlcName"], OutValue.DlcName) & bOk;
 	bOk = LoadFromCompactBinary(Field["CreateReleaseVersion"], OutValue.CreateReleaseVersion) & bOk;
-	bOk = LoadFromCompactBinary(Field["BasedOnReleaseCookedPackages"], OutValue.BasedOnReleaseCookedPackages) & bOk;
 	bOk = LoadFromCompactBinary(Field["SourceToLocalizedPackageVariants"], OutValue.SourceToLocalizedPackageVariants) & bOk;
 	bOk = LoadFromCompactBinary(Field["AllCulturesToCook"], OutValue.AllCulturesToCook) & bOk;
 
diff --git a/Engine/Source/Editor/UnrealEd/Private/Cooker/CookTypes.h b/Engine/Source/Editor/UnrealEd/Private/Cooker/CookTypes.h
--- a/Engine/Source/Editor/UnrealEd/Private/Cooker/CookTypes.h
+++ b/Engine/Source/Editor/UnrealEd/Private/Cooker/CookTypes.h
@@ -389,6 +389,29 @@ namespace UE::Cook
 		TArray<TSharedRef<IPlugin>> EnabledPluginRefPtrs;
 	};
 
+	/**
+	 * Category for in which session a package was cooked. When cooking a project normally, every cooked package will be
+	 * cooked (or incrementally skipped) in the current session. But for DLC packages some packages may have been
+	 * previously cooked in the basegame or other DLCs. Whether a package was incrementally skipped does not affect this
+	 * category. If not cooked for a different basegame or plugin, incrementally skipped packages are still marked as
+	 * cooked in the current session.
+	 */
+	enum class EWhereCooked : uint8
+	{
+		/** Cooked in the current session, the default. */
+		ThisSession,
+		/** Current session is a DLC cook based on a basegame cook, and the package was cooked in the basegame. */
+		BaseGame,
+		/**
+		 * We were instructed by commandline argument to treat these packages as already cooked, without being told
+		 * in which other session they were cooked.
+		 */
+		ExtraReleaseVersionAssets,
+
+		Count,
+		NumBits = FPlatformMath::ConstExprCeilLogTwo(EWhereCooked::Count),
+	};
+
 	/* Thread Local Storage access to identify which thread is the SchedulerThread for cooking. */
 	void InitializeTls();
 	bool IsSchedulerThread();
@@ -511,12 +534,6 @@ public:
 	/** Create a release from this manifest and store it in the releases directory for this cgame */
 	FString							CreateReleaseVersion;
 
-	/**
-	 * If we are based on a release version of the game this is the set of packages which were cooked in that release.
-	 * Map from platform name to list of uncooked package filenames.
-	 */
-	TMap<FName, TArray<FName>>		BasedOnReleaseCookedPackages;
-
 	/**
 	 * Mapping from source packages to their localized variants (based on the culture list in
 	 * FCookByTheBookStartupOptions)
