World Folders - When performing list update during blocking loading, defer update to the next tick

diff --git a/Engine/Source/Editor/UnrealEd/Private/WorldFolders.cpp b/Engine/Source/Editor/UnrealEd/Private/WorldFolders.cpp
--- a/Engine/Source/Editor/UnrealEd/Private/WorldFolders.cpp
+++ b/Engine/Source/Editor/UnrealEd/Private/WorldFolders.cpp
@@ -12,6 +12,8 @@
 #include "LevelInstance/LevelInstanceSubsystem.h"
 #include "LevelInstance/LevelInstanceInterface.h"
 #include "UObject/Package.h"
+#include "UObject/WeakObjectPtr.h"
+#include "Editor.h"
 
 #include UE_INLINE_GENERATED_CPP_BY_NAME(WorldFolders)
 
@@ -32,12 +34,40 @@ void UWorldFolders::Initialize(UWorld* InWorld)
 	PersistentFolders = MakeUnique<FWorldPersistentFolders>(*this);
 	TransientFolders = MakeUnique<FWorldTransientFolders>(*this);
 
+	bQueuedRebuild = false;
 	RebuildList();
-
+	
 	LoadState();
 }
 
 void UWorldFolders::RebuildList()
+{
+	if (GetWorld()->GetIsInBlockTillLevelStreamingCompleted())
+	{
+		if (!bQueuedRebuild)
+		{
+			bQueuedRebuild = true;
+
+			GEditor->GetTimerManager()->SetTimerForNextTick([WeakThis = TWeakObjectPtr<UWorldFolders>(this)]()
+			{
+				if (WeakThis.IsValid())
+				{
+					// Test if the rebuild is still queued...
+					if (WeakThis->bQueuedRebuild)
+					{
+						WeakThis->RebuildListInternal();
+					}
+				}
+			});
+		}
+	}
+	else
+	{
+		RebuildListInternal();
+	}
+}
+
+void UWorldFolders::RebuildListInternal()
 {
 	TRACE_CPUPROFILER_EVENT_SCOPE(UWorldFolders::RebuildList);
 
@@ -86,6 +116,8 @@ void UWorldFolders::RebuildList()
 			}
 		}
 	}
+
+	bQueuedRebuild = false;
 }
 
 UWorld* UWorldFolders::GetWorld() const
diff --git a/Engine/Source/Editor/UnrealEd/Public/WorldFolders.h b/Engine/Source/Editor/UnrealEd/Public/WorldFolders.h
--- a/Engine/Source/Editor/UnrealEd/Public/WorldFolders.h
+++ b/Engine/Source/Editor/UnrealEd/Public/WorldFolders.h
@@ -110,6 +110,8 @@ private:
 	UNREALED_API FString GetWorldStateFilename() const;
 	UNREALED_API void LoadState();
 
+	void RebuildListInternal();
+
 	TUniquePtr<FWorldPersistentFolders> PersistentFolders;
 	TUniquePtr<FWorldTransientFolders> TransientFolders;
 
@@ -122,6 +124,8 @@ private:
 
 	TArray<FActorPlacementFolder> CurrentFolderStack;
 
+	bool bQueuedRebuild;
+
 	friend class FWorldFoldersImplementation;
 	friend class FWorldPersistentFolders;
 	friend class FWorldTransientFolders;
