Update AssetRegistryGenerator to support chunked DLC cooks.
The pak command looks for the pakchunkList in the DLC Name's metadata folder.

diff --git a/Engine/Source/Editor/UnrealEd/Private/Commandlets/AssetRegistryGenerator.cpp b/Engine/Source/Editor/UnrealEd/Private/Commandlets/AssetRegistryGenerator.cpp
--- a/Engine/Source/Editor/UnrealEd/Private/Commandlets/AssetRegistryGenerator.cpp
+++ b/Engine/Source/Editor/UnrealEd/Private/Commandlets/AssetRegistryGenerator.cpp
@@ -2583,7 +2583,20 @@ void FAssetRegistryGenerator::AddPackageToChunk(FChunkPackageSet& ThisPackageSet
 
 FString FAssetRegistryGenerator::GetChunkManifestDirectoryForPlatform(const FString& Platform, UE::Cook::FCookSandbox& InSandboxFile) const
 {
-	return InSandboxFile.GetSandboxDirectory(Platform) / FApp::GetProjectName() / TEXT("Metadata") / TEXT("ChunkManifest");
+	static FString DLCName = []()
+		{
+			FString DLCNameReturn;
+			FParse::Value(FCommandLine::Get(), TEXT("DLCNAME="), DLCNameReturn);
+			return DLCNameReturn;
+		}();
+	if (DLCName.IsEmpty())
+	{
+		return InSandboxFile.GetSandboxDirectory(Platform) / FApp::GetProjectName() / TEXT("Metadata") / TEXT("ChunkManifest");
+	}
+	else
+	{
+		return InSandboxFile.GetSandboxDirectory(Platform) / FApp::GetProjectName() / TEXT("Plugins") / DLCName / TEXT("Metadata") / TEXT("ChunkManifest");
+	}
 }
 
 void FAssetRegistryGenerator::FixupPackageDependenciesForChunks(UE::Cook::FCookSandbox& InSandboxFile)
