[UnrealEd] Add comment to CookWorkerServer about [Cookworker X] log anonymization, in case the format ever changes, and the anonymization has to be updated.

diff --git a/Engine/Source/Editor/UnrealEd/Private/Cooker/CookWorkerServer.cpp b/Engine/Source/Editor/UnrealEd/Private/Cooker/CookWorkerServer.cpp
--- a/Engine/Source/Editor/UnrealEd/Private/Cooker/CookWorkerServer.cpp
+++ b/Engine/Source/Editor/UnrealEd/Private/Cooker/CookWorkerServer.cpp
@@ -205,6 +205,9 @@ void FCookWorkerServer::SendCrashDiagnostics()
 				Verbosity = ELogVerbosity::Error;
 			}
 			bLoggedErrorMessage |= Verbosity == ELogVerbosity::Error;
+			// NOTE:  If you change the format of [CookWorker %d] you must also change
+			// the corresponding hash aononymization in HashedIssueHandler.cs, and generally
+			// search for [CookWorker as it's cut and pasted around a few spots. 
 			FMsg::Logf(__FILE__, __LINE__, Category, Verbosity, TEXT("[CookWorker %d]: %.*s"),
 				ProfileId, Message.Len(), Message.GetData());
 		}
