ExternalActorsCommandlet:
- Added option 'disable' to completely disable external actors from the specified map, if the map supports it.
- Added option 'force' to ignore if the map supports it and forcibly disable external actors, potentially losing some data.
- Added option `report` that will simply report what would be done by the commandlet without doing anything.
- Added option 'listmaps' that will parse all worlds using the asset registry and report all potential maps that can have external actors disabled.

diff --git a/Engine/Source/Editor/UnrealEd/Private/Commandlets/ExternalActorsCommandlet.cpp b/Engine/Source/Editor/UnrealEd/Private/Commandlets/ExternalActorsCommandlet.cpp
--- a/Engine/Source/Editor/UnrealEd/Private/Commandlets/ExternalActorsCommandlet.cpp
+++ b/Engine/Source/Editor/UnrealEd/Private/Commandlets/ExternalActorsCommandlet.cpp
@@ -1,6 +1,8 @@
 // Copyright Epic Games, Inc. All Rights Reserved.
 
 #include "Commandlets/ExternalActorsCommandlet.h"
+#include "AssetRegistry/AssetData.h"
+#include "AssetRegistry/AssetRegistryModule.h"
 #include "HAL/FileManager.h"
 #include "Misc/Paths.h"
 #include "Editor.h"
@@ -10,10 +12,12 @@
 #include "UObject/UObjectHash.h"
 #include "AssetRegistry/AssetRegistryModule.h"
 #include "PackageHelperFunctions.h"
+#include "WorldPartition/WorldPartition.h"
+#include "WorldPartition/DataLayer/DataLayerManager.h"
 
 #include UE_INLINE_GENERATED_CPP_BY_NAME(ExternalActorsCommandlet)
 
-DEFINE_LOG_CATEGORY_STATIC(LogExternalActorsCommandlet, All, All);
+DEFINE_LOG_CATEGORY_STATIC(LogExternalActorsCommandlet, Log, All);
 
 UExternalActorsCommandlet::UExternalActorsCommandlet(const FObjectInitializer& ObjectInitializer)
 : Super(ObjectInitializer)
@@ -51,6 +55,50 @@ int32 UExternalActorsCommandlet::Main(const FString& Params)
 	}
 
 	const bool bRepair = Switches.Contains(TEXT("repair"));
+	const bool bDisable = Switches.Contains(TEXT("disable"));
+	const bool bForce = Switches.Contains(TEXT("force"));
+	const bool bReport = Switches.Contains(TEXT("report"));
+	const bool bListMaps = Switches.Contains(TEXT("listmaps"));
+
+	if (bDisable)
+	{
+		if (bRepair)
+		{
+			UE_LOG(LogExternalActorsCommandlet, Error, TEXT("Bad parameters: -repair cannot be used with -disable"));
+			return 1;
+		}
+	}
+	else if (bForce || bReport || bListMaps)
+	{
+		UE_LOG(LogExternalActorsCommandlet, Error, TEXT("Bad parameters: -force, -report and -listmaps can only be used along with -disable"));
+		return 1;
+	}
+
+	if (bDisable && bListMaps)
+	{
+		FAssetRegistryModule& AssetRegistryModule = FModuleManager::LoadModuleChecked<FAssetRegistryModule>(AssetRegistryConstants::ModuleName);
+
+		UE_LOG(LogExternalActorsCommandlet, Display, TEXT("Waiting for asset registry..."));
+		AssetRegistryModule.Get().SearchAllAssets(true);
+
+		TArray<FAssetData> WorldAssets;
+		AssetRegistryModule.Get().GetAssetsByClass(UWorld::StaticClass()->GetClassPathName(), WorldAssets, true);
+		WorldAssets.Sort();
+
+		for (const FAssetData& WorldAsset : WorldAssets)
+		{
+			if (ULevel::GetIsLevelUsingExternalActorsFromAsset(WorldAsset) && (!ULevel::GetIsLevelPartitionedFromAsset(WorldAsset) || ULevel::GetIsStreamingDisabledFromAsset(WorldAsset)))
+			{
+				UE_LOG(LogExternalActorsCommandlet, Display, TEXT("Level '%s' is a potential candidate to have external actors disabled."), *WorldAsset.GetSoftObjectPath().ToString());
+			}
+			else
+			{
+				UE_LOG(LogExternalActorsCommandlet, Verbose, TEXT("Level '%s' is not a potential candidate to have external actors disabled."), *WorldAsset.GetSoftObjectPath().ToString());
+			}
+		}
+
+		return 0;
+	}
 
 	if (!FPackageName::SearchForPackageOnDisk(Tokens[0], &Tokens[0]))
 	{
@@ -66,114 +114,212 @@ int32 UExternalActorsCommandlet::Main(const FString& Params)
 		return 1;
 	}
 
-	// Validate external actors
-	FString ExternalActorsPath = ULevel::GetExternalActorsPath(Tokens[0]);
-	FString ExternalActorsFilePath = FPackageName::LongPackageNameToFilename(ExternalActorsPath);
-
-	// Look for duplicated actor GUIDs
-	TMap<FGuid, AActor*> ActorGuids;
-
 	TSet<UPackage*> PackagesToSave;
 	TArray<FString> PackagesToDelete;	
-	if (IFileManager::Get().DirectoryExists(*ExternalActorsFilePath))
+
+	if (bDisable)
 	{
-		bool bResult = IFileManager::Get().IterateDirectoryRecursively(*ExternalActorsFilePath, [this, bRepair, &PackagesToSave, &PackagesToDelete, &ActorGuids](const TCHAR* FilenameOrDirectory, bool bIsDirectory)
+		if (!MainWorld->PersistentLevel->IsUsingExternalActors())
+		{
+			UE_LOG(LogExternalActorsCommandlet, Error, TEXT("Cannot disable external actors for partitioned level '%s' (already disabled)"), *Tokens[0]);
+			return 1;
+		}
+
+		if (UWorldPartition* WorldPartition = MainWorld->GetWorldPartition())
+		{
+			int32 NumErrors = 0;
+
+			if (WorldPartition->IsStreamingEnabled())
+			{
+				UE_LOG(LogExternalActorsCommandlet, Error, TEXT("Cannot disable external actors for partitioned level '%s' with streaming enabled"), *Tokens[0]);
+				NumErrors++;
+
+				if (bForce)
+				{
+					WorldPartition->SetEnableStreaming(false);
+				}
+			}
+
+			if (UDataLayerManager* DataLayerManager = WorldPartition->GetDataLayerManager())
+			{
+				if (DataLayerManager->GetDataLayerInstances().Num())
+				{
+					UE_LOG(LogExternalActorsCommandlet, Error, TEXT("Cannot disable external actors for partitioned level '%s' with data layers"), *Tokens[0]);
+					NumErrors++;
+				}
+			}
+
+			if (!NumErrors)
+			{
+				UE_LOG(LogExternalActorsCommandlet, Display, TEXT("External actors for partitioned level '%s' can be disabled"), *Tokens[0]);
+			}
+			else if (bForce)
+			{
+				UE_LOG(LogExternalActorsCommandlet, Display, TEXT("External actors for partitioned level '%s' will be forcibly disabled"), *Tokens[0]);
+			}
+			else
+			{
+				return 1;
+			}			
+
+			if (!bReport)
+			{
+				WorldPartition->Initialize(MainWorld, FTransform::Identity);
+
+				if (!UWorldPartition::RemoveWorldPartition(WorldPartition->GetWorld()->GetWorldSettings()))
+				{
+					UE_LOG(LogExternalActorsCommandlet, Error, TEXT("Error disabling external actors for partitioned level '%s'"), *Tokens[0]);
+					return 1;
+				}
+
+				WorldPartition->Uninitialize();
+			}
+		}
+
+		if (!bReport)
 		{
-			if (!bIsDirectory)
+			for (AActor* Actor : MainWorld->PersistentLevel->Actors)
+			{
+				if (IsValid(Actor) && Actor->IsMainPackageActor())
+				{
+					PackagesToDelete.Add(Actor->GetPackage()->GetLoadedPath().GetLocalFullPath());
+					Actor->SetPackageExternal(false);
+				}
+			}
+
+			if (MainWorld->PersistentLevel->IsUsingExternalObjects())
 			{
-				FString Filename(FilenameOrDirectory);
-				if (Filename.EndsWith(FPackageName::GetAssetPackageExtension()))
+				ForEachObjectWithOuter(MainWorld->PersistentLevel, [MainWorld, &PackagesToDelete](UObject* Object)
 				{
-					AActor* MainPackageActor = nullptr;
-					AActor* PotentialMainPackageActor = nullptr;
+					if (Object->IsPackageExternal() && !Object->IsA<AActor>())
+					{
+						PackagesToDelete.Add(Object->GetPackage()->GetLoadedPath().GetLocalFullPath());
+						FExternalPackageHelper::SetPackagingMode(Object, MainWorld->PersistentLevel, false);
+					}
+					return true;
+				}, true, RF_MirroredGarbage);
+			}
+
+			if (PackagesToDelete.Num())
+			{
+				PackagesToSave.Add(MainWorld->GetPackage());
+			}
 
-					const FString PackageName = FPackageName::FilenameToLongPackageName(*Filename);
-					if (UPackage* Package = LoadPackage(nullptr, *Filename, LOAD_None, nullptr, nullptr))
+			MainWorld->PersistentLevel->SetUseActorFolders(false);
+			MainWorld->PersistentLevel->SetUseExternalActors(false);
+		}
+	}
+	else
+	{
+		// Validate external actors
+		FString ExternalActorsPath = ULevel::GetExternalActorsPath(Tokens[0]);
+		FString ExternalActorsFilePath = FPackageName::LongPackageNameToFilename(ExternalActorsPath);
+
+		// Look for duplicated actor GUIDs
+		TMap<FGuid, AActor*> ActorGuids;
+
+		if (IFileManager::Get().DirectoryExists(*ExternalActorsFilePath))
+		{
+			bool bResult = IFileManager::Get().IterateDirectoryRecursively(*ExternalActorsFilePath, [this, bRepair, &PackagesToSave, &PackagesToDelete, &ActorGuids](const TCHAR* FilenameOrDirectory, bool bIsDirectory)
+			{
+				if (!bIsDirectory)
+				{
+					FString Filename(FilenameOrDirectory);
+					if (Filename.EndsWith(FPackageName::GetAssetPackageExtension()))
 					{
-						ForEachObjectWithPackage(Package, [&MainPackageActor, &PotentialMainPackageActor](UObject* Object)
+						AActor* MainPackageActor = nullptr;
+						AActor* PotentialMainPackageActor = nullptr;
+
+						const FString PackageName = FPackageName::FilenameToLongPackageName(*Filename);
+						if (UPackage* Package = LoadPackage(nullptr, *Filename, LOAD_None, nullptr, nullptr))
 						{
-							if (AActor* Actor = Cast<AActor>(Object))
+							ForEachObjectWithPackage(Package, [&MainPackageActor, &PotentialMainPackageActor](UObject* Object)
 							{
-								if (Actor->IsMainPackageActor())
-								{
-									MainPackageActor = Actor;
-									PotentialMainPackageActor = nullptr;
-								}
-								else if (!MainPackageActor)
+								if (AActor* Actor = Cast<AActor>(Object))
 								{
-									if (!Actor->IsChildActor())
+									if (Actor->IsMainPackageActor())
 									{
-										PotentialMainPackageActor = Actor;
+										MainPackageActor = Actor;
+										PotentialMainPackageActor = nullptr;
+									}
+									else if (!MainPackageActor)
+									{
+										if (!Actor->IsChildActor())
+										{
+											PotentialMainPackageActor = Actor;
+										}
 									}
 								}
-							}
-							return true;
-						});
-					}
-
-					if (!MainPackageActor)
-					{
-						UE_LOG(LogExternalActorsCommandlet, Error, TEXT("Missing main actor for file '%s'"), *Filename);
+								return true;
+							});
+						}
 
-						if (bRepair)
+						if (!MainPackageActor)
 						{
-							if (PotentialMainPackageActor)
+							UE_LOG(LogExternalActorsCommandlet, Error, TEXT("Missing main actor for file '%s'"), *Filename);
+
+							if (bRepair)
 							{
-								PotentialMainPackageActor->SetPackageExternal(false);
-								PotentialMainPackageActor->SetPackageExternal(true);
+								if (PotentialMainPackageActor)
+								{
+									PotentialMainPackageActor->SetPackageExternal(false);
+									PotentialMainPackageActor->SetPackageExternal(true);
 								
-								UPackage* PackageToSave = PotentialMainPackageActor->GetPackage();
-								PackagesToSave.Add(PackageToSave);
+									UPackage* PackageToSave = PotentialMainPackageActor->GetPackage();
+									PackagesToSave.Add(PackageToSave);
 
-								MainPackageActor = PotentialMainPackageActor;
-							}
+									MainPackageActor = PotentialMainPackageActor;
+								}
 
-							PackagesToDelete.Add(Filename);
+								PackagesToDelete.Add(Filename);
+							}
 						}
-					}
 
-					if (MainPackageActor)
-					{
-						if (ActorGuids.Contains(MainPackageActor->GetActorGuid()))
+						if (MainPackageActor)
 						{
-							UE_LOG(LogExternalActorsCommandlet, Error, TEXT("Duplicated actor guid for file '%s'"), *Filename);
-
-							if (bRepair)
+							if (ActorGuids.Contains(MainPackageActor->GetActorGuid()))
 							{
-								FSetActorGuid SetActorGuid(MainPackageActor, FGuid::NewGuid());
+								UE_LOG(LogExternalActorsCommandlet, Error, TEXT("Duplicated actor guid for file '%s'"), *Filename);
+
+								if (bRepair)
+								{
+									FSetActorGuid SetActorGuid(MainPackageActor, FGuid::NewGuid());
 
-								UPackage* PackageToSave = MainPackageActor->GetPackage();
-								PackagesToSave.Add(PackageToSave);
+									UPackage* PackageToSave = MainPackageActor->GetPackage();
+									PackagesToSave.Add(PackageToSave);
+								}
+							}
+							else
+							{
+								ActorGuids.Add(MainPackageActor->GetActorGuid(), MainPackageActor);
 							}
-						}
-						else
-						{
-							ActorGuids.Add(MainPackageActor->GetActorGuid(), MainPackageActor);
 						}
 					}
-				}
-				else
-				{
-					UE_LOG(LogExternalActorsCommandlet, Error, TEXT("Invalid actor file '%s'"), *Filename);
-
-					if (bRepair)
+					else
 					{
-						PackagesToDelete.Add(Filename);
+						UE_LOG(LogExternalActorsCommandlet, Error, TEXT("Invalid actor file '%s'"), *Filename);
+
+						if (bRepair)
+						{
+							PackagesToDelete.Add(Filename);
+						}
 					}
 				}
-			}
-			return true;
-		});
+				return true;
+			});
+		}
 	}
 
-	for (UPackage* PackageToSave: PackagesToSave)
+	FSavePackageArgs SaveArgs;
+	SaveArgs.TopLevelFlags = RF_Standalone;
+
+	for (UPackage* PackageToSave : PackagesToSave)
 	{
-		FSavePackageArgs SaveArgs;
-		SaveArgs.TopLevelFlags = RF_Standalone;
-		FString PackageFileName = SourceControlHelpers::PackageFilename(PackageToSave);
+		const FString PackageFileName = SourceControlHelpers::PackageFilename(PackageToSave);
+		UE_LOG(LogExternalActorsCommandlet, Display, TEXT("Saving package '%s'"), *PackageFileName);
 
 		if (PackageHelper.Checkout(PackageToSave))
-		{
+		{				
 			if (UPackage::SavePackage(PackageToSave, nullptr, *PackageFileName, SaveArgs))
 			{
 				PackageHelper.AddToSourceControl(PackageToSave);
@@ -183,8 +329,12 @@ int32 UExternalActorsCommandlet::Main(const FString& Params)
 
 	CollectGarbage(RF_NoFlags);
 
+	PackagesToDelete.Sort();
 	for (const FString& PackageToDelete : PackagesToDelete)
 	{
+		const FString PackageFileName = SourceControlHelpers::PackageFilename(PackageToDelete);
+		UE_LOG(LogExternalActorsCommandlet, Display, TEXT("Deleting package '%s'"), *PackageFileName);
+
 		PackageHelper.Delete(*PackageToDelete);
 	}
 
