Fix an issue where ITF fly camera movements were invalidating hit proxies every frame.

diff --git a/Engine/Source/Editor/UnrealEd/Private/EditorViewportClient.cpp b/Engine/Source/Editor/UnrealEd/Private/EditorViewportClient.cpp
--- a/Engine/Source/Editor/UnrealEd/Private/EditorViewportClient.cpp
+++ b/Engine/Source/Editor/UnrealEd/Private/EditorViewportClient.cpp
@@ -2212,7 +2212,9 @@ void FEditorViewportClient::UpdateCameraMovement( float DeltaTime )
 			const bool bInvalidateChildViews=true;
 
 			// When flying the camera around the hit proxies dont need to be invalidated since we are flying around and not clicking on anything
-			const bool bInvalidateHitProxies=!IsFlightCameraActive();
+			const bool bInvalidateHitProxies= UE::Editor::ViewportInteractions::UseEditorViewportInteractions()
+            	? !ViewportNavigationHelper->bIsMouseLooking
+            	: !IsFlightCameraActive();
 			Invalidate(bInvalidateChildViews,bInvalidateHitProxies);
 
 			// Update the FOV
