[UnrealEd] Allow edmode cursor override even when interacting/dragging

diff --git a/Engine/Source/Editor/UnrealEd/Private/EditorViewportClient.cpp b/Engine/Source/Editor/UnrealEd/Private/EditorViewportClient.cpp
--- a/Engine/Source/Editor/UnrealEd/Private/EditorViewportClient.cpp
+++ b/Engine/Source/Editor/UnrealEd/Private/EditorViewportClient.cpp
@@ -1796,15 +1796,15 @@ EMouseCursor::Type FEditorViewportClient::GetCursor(FViewport* InViewport,int32
 		MouseCursor = EMouseCursor::CardinalCross;
 	}
 	// Don't select widget axes by mouse over while they're being controlled by a mouse drag.
-	else if( InViewport->IsCursorVisible() && !bWidgetAxisControlledByDrag && !ModeTools->HasOngoingTransform())
+	else if (InViewport->IsCursorVisible() && !bWidgetAxisControlledByDrag)
 	{
 		// allow editor modes to override cursor
 		EMouseCursor::Type EditorModeCursor = EMouseCursor::Default;
-		if(ModeTools->GetCursor(EditorModeCursor))
+		if (ModeTools->GetCursor(EditorModeCursor))
 		{
 			MouseCursor = EditorModeCursor;
 		}
-		else
+		else if (!ModeTools->HasOngoingTransform())
 		{
 			HHitProxy* HitProxy = InViewport->GetHitProxy(X,Y);
 
