Fix a regression where the middle mouse button pan invert behavior was flipped.

#rnx

diff --git a/Engine/Source/Editor/UnrealEd/Private/EditorViewportClient.cpp b/Engine/Source/Editor/UnrealEd/Private/EditorViewportClient.cpp
--- a/Engine/Source/Editor/UnrealEd/Private/EditorViewportClient.cpp
+++ b/Engine/Source/Editor/UnrealEd/Private/EditorViewportClient.cpp
@@ -5187,7 +5187,7 @@ void FEditorViewportClient::ConvertMovementToDragRot(const FVector& InDelta,
 				if (MiddleMouseButtonDown && !ViewportSettings->bWorldSpaceVerticalPan)
 				{
 					// Pan relative to the camera
-					InDragDelta = FQuat(ViewRotation) * FVector( 0.0f, InDelta.X, InDelta.Y) * (bInvert ? 1.0f : -1.0f);
+					InDragDelta = FQuat(ViewRotation) * FVector( 0.0f, InDelta.X, InDelta.Y) * (bInvert ? -1.0f : 1.0f);
 				}
 				else
 				{
