Add support for direct manipulation of editor viewport orthographic near & far planes. Add a new Editor Viewport Depth Bar that allows for visual control over those planes. The depth bar supports displaying custom widgets that are automatically arranged by depth. This is used in the Level Viewport Client to display selected objects along the depth bar.

These features are hidden behind the "r.Editor.ManualOrthoDepth 1" CVar.

diff --git a/Engine/Source/Editor/EditorStyle/Private/StarshipStyle.cpp b/Engine/Source/Editor/EditorStyle/Private/StarshipStyle.cpp
--- a/Engine/Source/Editor/EditorStyle/Private/StarshipStyle.cpp
+++ b/Engine/Source/Editor/EditorStyle/Private/StarshipStyle.cpp
@@ -3242,6 +3242,17 @@ void FStarshipEditorStyle::FStyle::SetupViewportStyles()
 		Set("VRViewportActorPreview.Attached", new IMAGE_BRUSH("Common/ScreenAttach_VR", Icon64x64));
 		Set("VRViewportActorPreview.Detached", new IMAGE_BRUSH("Common/ScreenDetach_VR", Icon64x64));
 	}
+	
+	// Viewport Depth Bar
+	{
+		Set("EditorViewport.DepthBar.Background", new FSlateRoundedBoxBrush(FStyleColors::Panel, 4.0f));
+		Set("EditorViewport.DepthBar.Slice", new FSlateColorBrush(FStyleColors::Hover));
+		Set("EditorViewport.DepthBar.SliceHovered", new FSlateColorBrush(FStyleColors::Hover2));
+		Set("EditorViewport.DepthBar.SliceTop", new FSlateRoundedBoxBrush(FStyleColors::Hover2, FVector4(2.0f, 2.0f, 0.0f, 0.0f)));
+		Set("EditorViewport.DepthBar.SliceBottom", new FSlateRoundedBoxBrush(FStyleColors::Hover2, FVector4(0.0f, 0.0f, 2.0f, 2.0f)));
+		
+		Set("EditorViewport.DepthBar.SelectionHighlight", new FSlateColorBrush(FSlateColor(COLOR("EBA20AFF"))));
+	}
 }
 
 void FStarshipEditorStyle::FStyle::SetupMenuBarStyles()
diff --git a/Engine/Source/Editor/LevelEditor/Private/SLevelViewport.cpp b/Engine/Source/Editor/LevelEditor/Private/SLevelViewport.cpp
--- a/Engine/Source/Editor/LevelEditor/Private/SLevelViewport.cpp
+++ b/Engine/Source/Editor/LevelEditor/Private/SLevelViewport.cpp
@@ -103,6 +103,7 @@
 #include "LevelViewportLayout.h"
 #include "EditorViewportTabContent.h"
 #include "EditorViewportCommands.h"
+#include "SEditorViewportDepthBar.h"
 #include "FunctionalUIScreenshotTest.h"
 #include "GPUSkinCache.h"
 #include "GPUSkinCacheVisualizationMenuCommands.h"
@@ -110,6 +111,7 @@
 #include "ShowFlagMenuCommands.h"
 #include "SkeletalRenderPublic.h"
 #include "ViewportToolBarContext.h"
+#include "Elements/Interfaces/TypedElementObjectInterface.h"
 #include "Subsystems/PanelExtensionSubsystem.h"
 #include "ViewportToolbar/LevelEditorSubmenus.h"
 #include "ViewportToolbar/LevelViewportContext.h"
@@ -3596,6 +3598,54 @@ void SLevelViewport::OnElementSelectionChanged(const UTypedElementSelectionSet*
 	// Request preview update. It's possible that the actor currently selected has forbidden default preview.
 	// However, we need to show the preview widget when some of its child components selected.
 	bNeedToUpdatePreviews = true;
+	
+	if (!DepthBar)
+	{
+		return;
+	}
+	
+	DepthBar->ClearChildren();
+	FTypedElementListConstRef ElementList = SelectionSet->GetElementList();
+	for (int32 Index = 0; Index < ElementList->Num(); ++Index)
+	{
+		FTypedElementHandle Handle = ElementList->GetElementHandleAt(Index);
+		
+		if (TTypedElement<ITypedElementWorldInterface> WorldElement = ElementList->GetElement<ITypedElementWorldInterface>(Handle))
+		{
+			TAttribute<FBoxSphereBounds> Bounds = TAttribute<FBoxSphereBounds>::CreateLambda([SelectionSet, Handle]
+			{
+				FBoxSphereBounds Bounds(ForceInit);
+				if (TTypedElement<ITypedElementWorldInterface> WorldElement = SelectionSet->GetElementList()->GetElement<ITypedElementWorldInterface>(Handle))
+				{
+					WorldElement.GetBounds(Bounds);
+				}
+				
+				return Bounds;
+			});
+				
+			TAttribute<FText> DisplayName;
+				
+			if (TTypedElement<ITypedElementObjectInterface> ObjectElement = ElementList->GetElement<ITypedElementObjectInterface>(Handle))
+			{
+				if (AActor* Actor = ObjectElement.GetObjectAs<AActor>())
+				{
+					DisplayName.BindWeakLambda(Actor, [Actor]
+					{
+						return FText::FromString(Actor->GetActorNameOrLabel());
+					});
+				}
+			}
+				
+			DepthBar->AddSlot()
+			.Bounds(Bounds)
+			[
+				SNew(SImage)
+				.Image(FAppStyle::Get().GetBrush("EditorViewport.DepthBar.SelectionHighlight"))
+				.ToolTipText(DisplayName)
+				.DesiredSizeOverride(FVector2d(4, 4))
+			];
+		}
+	}
 }
 
 void SLevelViewport::PreviewSelectedCameraActors(const bool bPreviewInDesktopViewport)
diff --git a/Engine/Source/Editor/UnrealEd/Private/EditorViewportClient.cpp b/Engine/Source/Editor/UnrealEd/Private/EditorViewportClient.cpp
--- a/Engine/Source/Editor/UnrealEd/Private/EditorViewportClient.cpp
+++ b/Engine/Source/Editor/UnrealEd/Private/EditorViewportClient.cpp
@@ -126,6 +126,12 @@ static TAutoConsoleVariable<int32> CVarEditorViewGPUMirrorTest(
 	TEXT("difficult, due to lack of interactivity, and multiple scene renders happening per frame.\n"),
 	ECVF_RenderThreadSafe);
 
+static TAutoConsoleVariable<int32> CVarEditorViewportManualOrthographicDepthControl(
+	TEXT("r.Editor.ManualOrthoDepth"),
+	0,
+	TEXT("Enable or disable the new orthographic depth tools."),
+	ECVF_RenderThreadSafe);
+
 static bool GetDefaultLowDPIPreviewValue()
 {
 	static auto CVarEditorViewportHighDPIPtr = IConsoleManager::Get().FindConsoleVariable(TEXT("r.Editor.Viewport.HighDPI"));
@@ -783,6 +789,31 @@ float FEditorViewportClient::GetOrthoUnitsPerPixel(const FViewport* InViewport)
 	return (GetOrthoZoom() / (SizeX * 15.f)) * ComputeOrthoZoomFactor(SizeX);
 }
 
+FVector FEditorViewportClient::GetForwardVector() const
+{
+	switch (GetViewportType())
+	{
+	case LVT_Perspective:
+		return ViewTransformPerspective.GetRotation().Vector();
+	case LVT_OrthoFreelook:
+		return ViewTransformOrthographic.GetRotation().Vector();
+	case LVT_OrthoTop:
+		return FVector::DownVector;
+	case LVT_OrthoBottom:
+		return FVector::UpVector;
+	case LVT_OrthoLeft:
+		return FVector::LeftVector;
+	case LVT_OrthoRight:
+		return FVector::RightVector;
+	case LVT_OrthoFront:
+		return FVector::BackwardVector;
+	case LVT_OrthoBack:
+		return FVector::ForwardVector;
+	default:
+		return FVector::Zero();
+	}
+}
+
 void FEditorViewportClient::SetViewLocationForOrbiting(const FVector& LookAtPoint, float DistanceToCamera )
 {
 	FMatrix Matrix = FTranslationMatrix(-GetViewLocation());
@@ -919,7 +950,13 @@ void FEditorViewportClient::FocusViewportOnBox( const FBox& BoundingBox, bool bI
 			// For ortho viewports just set the camera position to the center of the bounding volume.
 			//SetViewLocation( Position );
 			ViewTransform.TransitionToLocation(Position, EditorViewportWidget, bInstant);
-
+			
+			if (SupportsManualOrthographicDepthControl())
+			{
+				// Shift the orthographic near view plane back so that it is behind the object
+				OrthographicNearPlane = -Radius;
+			}
+			
 			if (!(Viewport->KeyState(EKeys::LeftControl) || Viewport->KeyState(EKeys::RightControl)))
 			{
 				/**
@@ -933,7 +970,7 @@ void FEditorViewportClient::FocusViewportOnBox( const FBox& BoundingBox, bool bI
 				{
 					uint32 MinAxisSize = (AspectToUse > 1.0f) ? ViewportSize.Y : ViewportSize.X;
 					float Zoom = Radius / (MinAxisSize / 2.0f);
-					NewOrthoZoom = Zoom * (ViewportSize.X * 15.0f);
+					NewOrthoZoom = Zoom * (MinAxisSize * 15.0f);
 				}
 
 				NewOrthoZoom = FMath::Clamp<float>(NewOrthoZoom, GetMinimumOrthoZoom(), MAX_ORTHOZOOM);
@@ -1300,7 +1337,38 @@ FSceneView* FEditorViewportClient::CalcSceneView(FSceneViewFamily* ViewFamily, c
 
 			FMatrix::FReal ZScale = 0.5f / UE_OLD_WORLD_MAX;
 			FMatrix::FReal ZOffset = UE_OLD_WORLD_MAX;
-			if(!ViewFamily->EngineShowFlags.Wireframe)
+			
+			if (SupportsManualOrthographicDepthControl())
+			{
+				FMinimalViewInfo CalculatePlanesViewInfo;
+				CalculatePlanesViewInfo.Rotation = ViewInitOptions.ViewRotationMatrix.Rotator();
+				CalculatePlanesViewInfo.AspectRatio = AspectRatio;
+				CalculatePlanesViewInfo.bConstrainAspectRatio = false;
+
+				CalculatePlanesViewInfo.ProjectionMode = ECameraProjectionMode::Orthographic;				
+				CalculatePlanesViewInfo.OrthoWidth = OrthoWidth;
+				CalculatePlanesViewInfo.bAutoCalculateOrthoPlanes = true;
+				CalculatePlanesViewInfo.AutoPlaneShift = 0.0f;
+				CalculatePlanesViewInfo.bUpdateOrthoPlanes = true;
+				CalculatePlanesViewInfo.bUseCameraHeightAsViewTarget = false;
+				CalculatePlanesViewInfo.OrthoNearClipPlane = -UE_OLD_WORLD_MAX;
+				CalculatePlanesViewInfo.OrthoFarClipPlane = UE_OLD_WORLD_MAX;
+
+				if (OrthographicNearPlane.IsSet())
+				{
+					CalculatePlanesViewInfo.OrthoNearClipPlane = OrthographicNearPlane.GetValue();
+				}
+				
+				if (OrthographicFarPlane.IsSet())
+				{
+					CalculatePlanesViewInfo.OrthoFarClipPlane = OrthographicFarPlane.GetValue();
+				}
+				
+				ViewInitOptions.UpdateOrthoPlanes(CalculatePlanesViewInfo);	
+				ZScale = 1.0f / (CalculatePlanesViewInfo.OrthoFarClipPlane - CalculatePlanesViewInfo.OrthoNearClipPlane);
+				ZOffset = -CalculatePlanesViewInfo.OrthoNearClipPlane;
+			}
+			else if(!ViewFamily->EngineShowFlags.Wireframe)
 			{
 				FMinimalViewInfo CalculatePlanesViewInfo;
 				CalculatePlanesViewInfo.Rotation = ViewInitOptions.ViewRotationMatrix.Rotator();
@@ -1322,8 +1390,8 @@ FSceneView* FEditorViewportClient::CalcSceneView(FSceneViewFamily* ViewFamily, c
 					ZScale = 1.0f / (CalculatePlanesViewInfo.OrthoFarClipPlane - CalculatePlanesViewInfo.OrthoNearClipPlane);
 					ZOffset = -CalculatePlanesViewInfo.OrthoNearClipPlane;
 				}
-			}			
-
+			}
+			
 			ViewInitOptions.ProjectionMatrix = FReversedZOrthoMatrix(
 				OrthoWidth,
 				OrthoHeight,
@@ -5598,6 +5666,31 @@ void FEditorViewportClient::OverrideFarClipPlane(const float InFarPlane)
 	FarPlane = InFarPlane;
 }
 
+bool FEditorViewportClient::SupportsManualOrthographicDepthControl() const
+{
+	return CVarEditorViewportManualOrthographicDepthControl.GetValueOnAnyThread() > 0;
+}
+
+TOptional<double> FEditorViewportClient::GetOrthographicNearPlaneOverride() const
+{
+	return OrthographicNearPlane;
+}
+
+void FEditorViewportClient::SetOrthographicNearPlaneOverride(const TOptional<double>& InValue)
+{
+	OrthographicNearPlane = InValue;
+}
+
+TOptional<double> FEditorViewportClient::GetOrthographicFarPlaneOverride() const
+{
+	return OrthographicFarPlane;
+}
+
+void FEditorViewportClient::SetOrthographicFarPlaneOverride(const TOptional<double>& InValue)
+{
+	OrthographicFarPlane = InValue;
+}
+
 float FEditorViewportClient::GetSceneDepthAtLocation(int32 X, int32 Y)
 {
 	// #todo: in the future we will just sample the depth buffer
diff --git a/Engine/Source/Editor/UnrealEd/Private/SEditorViewport.cpp b/Engine/Source/Editor/UnrealEd/Private/SEditorViewport.cpp
--- a/Engine/Source/Editor/UnrealEd/Private/SEditorViewport.cpp
+++ b/Engine/Source/Editor/UnrealEd/Private/SEditorViewport.cpp
@@ -36,8 +36,10 @@
 	#include "RenderGraph.h"
 #endif
 #include "LevelEditorViewport.h"
+#include "SEditorViewportDepthBar.h"
 #include "SEditorViewportGridPanel.h"
 #include "ViewportInteractions/ViewportInteractionsBehaviorSource.h"
+#include "Widgets/Input/SButton.h"
 
 #define LOCTEXT_NAMESPACE "EditorViewport"
 
@@ -283,6 +285,41 @@ void SEditorViewport::Construct( const FArguments& InArgs )
 			];
 	}
 	// clang-format on
+	
+	ViewportOverlay->AddSlot()
+	.HAlign(HAlign_Right)	
+	.VAlign(VAlign_Fill)
+	[
+		SNew(SBox)
+		.Padding(0.0f, 100.0f, 6.0f, 100.0f)
+		.Visibility(this, &SEditorViewport::OnGetDepthBarVisibility)
+		[
+			SNew(SVerticalBox)
+			+ SVerticalBox::Slot()
+			.FillHeight(1.0f)
+			[
+				SAssignNew(DepthBar, SEditorViewportDepthBar, GetViewportClient())
+			]
+			+ SVerticalBox::Slot()
+			.Padding(0.0f, 6.0f, 0.0f, 0.0f)
+			.AutoHeight()
+			[
+				SNew(SButton)
+				.OnClicked(this, &SEditorViewport::DismissDepthBar)
+				.ToolTipText(LOCTEXT("ResetViewportDepth", "Reset Viewport Depth"))
+				.ContentPadding(FMargin(4.0f))
+				.NormalPaddingOverride(FMargin(0.0f))
+				.PressedPaddingOverride(FMargin(0.0f))
+				.HAlign(HAlign_Center)
+				.VAlign(VAlign_Center)
+				[
+					SNew(SImage)
+					.Image(FAppStyle::GetBrush("Cross"))
+					.DesiredSizeOverride(FVector2D(12.0f, 12.0f))
+				]
+			]
+		]
+	];
 
 	
 	PopulateViewportOverlays(ViewportOverlay.ToSharedRef());
@@ -332,6 +369,12 @@ FReply SEditorViewport::OnFocusReceived( const FGeometry& MyGeometry, const FFoc
 void SEditorViewport::Tick( const FGeometry& AllottedGeometry, const double InCurrentTime, const float InDeltaTime )
 {
 	LastTickTime = FPlatformTime::Seconds();
+	
+	if (!bShowDepthBar && Client->SupportsManualOrthographicDepthControl() && Client->IsOrtho() &&
+		(Client->GetOrthographicNearPlaneOverride().IsSet() || Client->GetOrthographicFarPlaneOverride().IsSet()))
+	{
+		bShowDepthBar = true;
+	}
 }
 
 void SEditorViewport::BindCommands()
@@ -689,6 +732,23 @@ EVisibility SEditorViewport::OnGetViewportContentVisibility() const
 	return bIsViewportUIHidden ? EVisibility::Collapsed : EVisibility::SelfHitTestInvisible;
 }
 
+EVisibility SEditorViewport::OnGetDepthBarVisibility() const
+{
+	if (Client->SupportsManualOrthographicDepthControl() && Client->IsOrtho() && bShowDepthBar)
+	{
+		return EVisibility::Visible;
+	}
+	return EVisibility::Collapsed;
+}
+
+FReply SEditorViewport::DismissDepthBar()
+{
+	Client->SetOrthographicFarPlaneOverride(TOptional<double>());
+	Client->SetOrthographicNearPlaneOverride(TOptional<double>());
+	bShowDepthBar = false;
+	return FReply::Handled();
+}
+
 void SEditorViewport::OnToggleRealtime()
 {
 	if (Client->IsRealtime())
diff --git a/Engine/Source/Editor/UnrealEd/Private/SEditorViewportDepthBar.cpp b/Engine/Source/Editor/UnrealEd/Private/SEditorViewportDepthBar.cpp
new file mode 100644
--- /dev/null
+++ b/Engine/Source/Editor/UnrealEd/Private/SEditorViewportDepthBar.cpp
@@ -0,0 +1,477 @@
+﻿// Copyright Epic Games, Inc. All Rights Reserved.
+
+
+#include "SEditorViewportDepthBar.h"
+#include "EditorViewportClient.h"
+#include "Misc/CoreMiscDefines.h"
+#include "EngineUtils.h"
+#include "Elements/Interfaces/TypedElementWorldInterface.h"
+#include "Math/UnitConversion.h"
+#include "WorldPartition/WorldPartition.h"
+
+#define LOCTEXT_NAMESPACE "EditorViewportDepthBar"
+
+namespace UE::Editor::ViewportDepthBar::Private
+{
+	constexpr float IndicatorExtraHover = 2.0f;
+	constexpr float MinIndicatorDifference = 2.0f;
+	
+	FText GetDistanceDisplayText(double Value)
+	{
+		if (FUnitConversion::Settings().ShouldDisplayUnits())
+		{
+			const TArray<EUnit>& AllDistanceUnits = FUnitConversion::Settings().GetDisplayUnits(EUnitType::Distance);
+			const EUnit BaseUnit = AllDistanceUnits.Num() ? AllDistanceUnits[0] : EUnit::Centimeters;
+			const FNumericUnit DisplayUnit = FUnitConversion::QuantizeUnitsToBestFit(Value, BaseUnit);
+			return FText::Format(LOCTEXT("DistanceFormat", "{0}{1}"), DisplayUnit.Value, FText::FromString(FUnitConversion::GetUnitDisplayString(DisplayUnit.Units)));
+		}
+		return FText::Format(LOCTEXT("RawDistanceFormat", "{0}"), Value);
+	}
+}
+
+void SEditorViewportDepthBar::Construct(const FArguments& InArgs, const TSharedPtr<FEditorViewportClient>& Client)
+{
+	WeakClient = Client;
+	
+	BackgroundBrush = FAppStyle::Get().GetBrush("EditorViewport.DepthBar.Background");
+	SliceBrush = FAppStyle::Get().GetBrush("EditorViewport.DepthBar.Slice");
+	HoveredSliceBrush = FAppStyle::Get().GetBrush("EditorViewport.DepthBar.SliceHovered");
+	HoveredSliceTopBrush = FAppStyle::Get().GetBrush("EditorViewport.DepthBar.SliceTop");
+	HoveredSliceBottomBrush = FAppStyle::Get().GetBrush("EditorViewport.DepthBar.SliceBottom");
+	
+	Children.AddSlots(MoveTemp(const_cast<TArray<FSlot::FSlotArguments>&>(InArgs._Slots)));
+	
+	SetToolTipText(TAttribute<FText>::CreateSP(this, &SEditorViewportDepthBar::GetTooltipText));
+}
+
+void SEditorViewportDepthBar::Tick(const FGeometry& AllottedGeometry, const double InCurrentTime, const float InDeltaTime)
+{
+	if (TSharedPtr<FEditorViewportClient> Client = WeakClient.Pin())
+	{
+		UWorld* World = Client->GetWorld();
+		if (UWorldPartition* WorldPartition = World->GetWorldPartition())
+		{
+			WorldBounds = WorldPartition->GetEditorWorldBounds();
+		}
+		else
+		{
+			// TODO: Not this every frame, good lord
+			WorldBounds = FBox(EForceInit::ForceInitToZero);
+			for (TActorIterator<AActor> It(Client->GetWorld()); It; ++It)
+			{
+				AActor* Actor = *It;
+				if (IsValid(Actor))
+				{
+					WorldBounds += Actor->GetComponentsBoundingBox(true);
+				}
+			}	
+		}
+	}
+}
+
+FCursorReply SEditorViewportDepthBar::OnCursorQuery(const FGeometry& MyGeometry, const FPointerEvent& CursorEvent) const
+{
+	using namespace UE::Editor::ViewportDepthBar::Private;
+	
+	switch (DragTarget)
+	{
+		case EDragTarget::Far:
+		case EDragTarget::Near:
+			return FCursorReply::Cursor(EMouseCursor::ResizeUpDown);
+		case EDragTarget::Slice:
+			return FCursorReply::Cursor(EMouseCursor::GrabHandClosed);
+		default:
+			switch (HoverTarget)
+			{
+				case EDragTarget::Far:
+				case EDragTarget::Near:
+					return FCursorReply::Cursor(EMouseCursor::ResizeUpDown);
+				case EDragTarget::Slice:
+					return FCursorReply::Cursor(EMouseCursor::GrabHand);
+				default:
+					return FCursorReply::Unhandled();	
+			}
+	}
+}
+
+SEditorViewportDepthBar::FSlot& SEditorViewportDepthBar::GetSlot(int32 Index)
+{
+	return Children[Index];
+}
+
+const SEditorViewportDepthBar::FSlot& SEditorViewportDepthBar::GetSlot(int32 Index) const
+{
+	return Children[Index];
+}
+
+void SEditorViewportDepthBar::ClearChildren()
+{
+	Children.Empty();
+}
+
+FReply SEditorViewportDepthBar::OnMouseButtonDown(const FGeometry& MyGeometry, const FPointerEvent& MouseEvent)
+{
+	using namespace UE::Editor::ViewportDepthBar::Private;
+
+	if (HoverTarget != EDragTarget::None)
+	{
+		if (const FBarPositions* Positions = GetBarPositions(MyGeometry).GetPtrOrNull())
+		{
+			const FVector2f MousePosition = MyGeometry.AbsoluteToLocal(MouseEvent.GetScreenSpacePosition());
+		
+			DragTarget = HoverTarget;
+			if (DragTarget == EDragTarget::Slice)
+			{
+				DragOffset = MousePosition.Y - Positions->BarFarPlane;
+			}
+			
+			return FReply::Handled().CaptureMouse(AsShared());
+		}
+	}
+
+	return FReply::Unhandled();
+}
+
+FReply SEditorViewportDepthBar::OnMouseMove(const FGeometry& MyGeometry, const FPointerEvent& MouseEvent)
+{
+	using namespace UE::Editor::ViewportDepthBar::Private;
+
+	if (DragTarget != EDragTarget::None)
+	{
+		TSharedPtr<FEditorViewportClient> Client = WeakClient.Pin();
+		const FBarPositions* Positions = GetBarPositions(MyGeometry).GetPtrOrNull();
+		if (Client && Positions)
+		{
+			const FVector2f MousePosition = MyGeometry.AbsoluteToLocal(MouseEvent.GetScreenSpacePosition());
+			
+			if (Client->IsOrtho())
+			{
+				if (DragTarget == EDragTarget::Far)
+				{
+					const float ClampedBarPosition = FMath::Clamp(MousePosition.Y, Positions->BarMin, Positions->BarNearPlane - MinIndicatorDifference);
+					if (ClampedBarPosition == Positions->BarMin)
+					{
+						Client->SetOrthographicFarPlaneOverride(TOptional<double>());
+					}
+					else
+					{
+						Client->SetOrthographicFarPlaneOverride(Positions->BarToRealPosition(ClampedBarPosition));
+					}
+					
+					Client->Invalidate();
+					return FReply::Handled();
+				}
+				if (DragTarget == EDragTarget::Near)
+				{
+					const float ClampedBarPosition = FMath::Clamp(MousePosition.Y, Positions->BarFarPlane + MinIndicatorDifference, Positions->BarMax);
+					if (ClampedBarPosition == Positions->BarMax)
+					{
+						Client->SetOrthographicNearPlaneOverride(TOptional<double>());
+					}
+					else
+					{
+						Client->SetOrthographicNearPlaneOverride(Positions->BarToRealPosition(ClampedBarPosition));
+					}
+					
+					Client->Invalidate();
+					return FReply::Handled();
+				}
+				if (DragTarget == EDragTarget::Slice)
+				{
+					const float RawPosition = MousePosition.Y - DragOffset;
+					const float ClampedBarPosition = FMath::Clamp(RawPosition, Positions->BarMin, Positions->BarMax - (Positions->BarNearPlane - Positions->BarFarPlane));
+					
+					Client->SetOrthographicFarPlaneOverride(Positions->BarToRealPosition(ClampedBarPosition));
+					Client->SetOrthographicNearPlaneOverride(Positions->BarToRealPosition(ClampedBarPosition + (Positions->BarNearPlane - Positions->BarFarPlane)));
+					
+					Client->Invalidate();
+					return FReply::Handled();
+				}
+			}
+		}
+	}
+	else if (const FBarPositions* Positions = GetBarPositions(MyGeometry).GetPtrOrNull())
+	{
+		const FVector2f MousePosition = MyGeometry.AbsoluteToLocal(MouseEvent.GetScreenSpacePosition());
+		
+		const float HoverZone = SliderHeadHeight * 0.5f + IndicatorExtraHover;
+		
+		if (MousePosition.Y >= Positions->BarFarPlane - HoverZone && MousePosition.Y < Positions->BarFarPlane + HoverZone)
+		{
+			HoverTarget = EDragTarget::Far; 
+		}
+		else if (MousePosition.Y >= Positions->BarNearPlane - HoverZone && MousePosition.Y < Positions->BarNearPlane + HoverZone)
+		{
+			HoverTarget = EDragTarget::Near;
+		}
+		else if (MousePosition.Y > Positions->BarFarPlane && MousePosition.Y < Positions->BarNearPlane)
+		{
+			HoverTarget = EDragTarget::Slice;
+		}
+		else
+		{
+			HoverTarget = EDragTarget::None;
+		}
+	}
+	return FReply::Unhandled();
+}
+
+FReply SEditorViewportDepthBar::OnMouseButtonUp(const FGeometry& MyGeometry, const FPointerEvent& MouseEvent)
+{
+	if (DragTarget != EDragTarget::None)
+	{
+		DragTarget = EDragTarget::None;
+		HoverTarget = EDragTarget::None;
+		return FReply::Handled().ReleaseMouseCapture();
+	}
+	return FReply::Unhandled();
+}
+
+void SEditorViewportDepthBar::OnMouseLeave(const FPointerEvent& MouseEvent)
+{
+	SPanel::OnMouseLeave(MouseEvent);
+	
+	if (DragTarget == EDragTarget::None)
+	{
+		HoverTarget = EDragTarget::None;
+	}
+}
+
+void SEditorViewportDepthBar::OnArrangeChildren(const FGeometry& AllottedGeometry, FArrangedChildren& ArrangedChildren) const
+{
+	if (FBarPositions* BarPositions = GetBarPositions(AllottedGeometry).GetPtrOrNull())
+	{
+		FMargin Inset = Padding.Get();
+	
+		for (int32 ChildIndex = 0; ChildIndex < Children.Num(); ++ChildIndex)
+		{
+			const FSlot& CurrentChild = GetSlot(ChildIndex);
+			
+			if (CurrentChild.GetBounds().IsSet())
+			{
+				FBoxSphereBounds Bounds = CurrentChild.GetBounds().Get();
+				const FSphere Sphere = Bounds.GetSphere();
+				
+				const FVector Max = Sphere.Center + BarPositions->DepthSpace.Forward * Sphere.W;
+				const FVector Min = Sphere.Center - BarPositions->DepthSpace.Forward * Sphere.W;
+					
+				const float BarMax = BarPositions->WorldToBarPosition(Max);
+				const float BarMin = BarPositions->WorldToBarPosition(Min);
+				const float BarSize = FMath::Max(BarMin - BarMax, 1.0f);
+				
+				FVector2f Offset(Inset.Left, BarMax);
+				FVector2f BoundSize(AllottedGeometry.GetLocalSize().X - Inset.Left - Inset.Right, BarSize);
+				FVector2f ChildDesiredSize = CurrentChild.GetWidget()->GetDesiredSize();
+				
+				if (ChildDesiredSize.X > BoundSize.X)
+				{
+					Offset.X -= (ChildDesiredSize.X - BoundSize.X) * 0.5f;
+					BoundSize.X = ChildDesiredSize.X;
+				}
+				if (ChildDesiredSize.Y > BoundSize.Y)
+				{
+					Offset.Y -= (ChildDesiredSize.Y - BoundSize.Y) * 0.5f;
+					BoundSize.Y = ChildDesiredSize.Y;
+				}
+				
+				ArrangedChildren.AddWidget(AllottedGeometry.MakeChild(CurrentChild.GetWidget(), BoundSize, FSlateLayoutTransform(Offset)));
+			}
+		}
+	}	
+}
+
+FChildren* SEditorViewportDepthBar::GetChildren()
+{
+	return &Children;
+}
+
+int32 SEditorViewportDepthBar::OnPaint(const FPaintArgs& Args, const FGeometry& AllottedGeometry, const FSlateRect& MyCullingRect, FSlateWindowElementList& OutDrawElements, int32 LayerId, const FWidgetStyle& InWidgetStyle, bool bParentEnabled) const
+{
+	using namespace UE::Editor::ViewportDepthBar::Private;
+
+	FSlateDrawElement::MakeBox(OutDrawElements,
+		LayerId++,
+		AllottedGeometry.ToPaintGeometry(),
+		BackgroundBrush,
+		ESlateDrawEffect::None,
+		BackgroundBrush->GetTint(InWidgetStyle)
+	);
+
+	FMargin Inset = Padding.Get();
+ 
+	if (const FBarPositions* BarPositions = GetBarPositions(AllottedGeometry).GetPtrOrNull())
+	{
+		const FSlateBrush* CurrentSliceBrush = HoverTarget == EDragTarget::Slice ? HoveredSliceBrush : SliceBrush;
+
+		if (HoverTarget == EDragTarget::Slice)
+		{
+			LayerId = SPanel::OnPaint(Args, AllottedGeometry, MyCullingRect, OutDrawElements, LayerId, InWidgetStyle, bParentEnabled) + 1;
+		}
+
+		FSlateDrawElement::MakeBox(OutDrawElements,
+			LayerId,
+			AllottedGeometry.MakeChild(
+				FVector2f(AllottedGeometry.GetLocalSize().X - Inset.Left - Inset.Right - SliderMargin * 2.0f, BarPositions->BarNearPlane - BarPositions->BarFarPlane),
+				FSlateLayoutTransform(FVector2f(Inset.Left + SliderMargin, BarPositions->BarFarPlane))
+			).ToPaintGeometry(),
+			CurrentSliceBrush,
+			ESlateDrawEffect::None,
+			CurrentSliceBrush->GetTint(InWidgetStyle).CopyWithNewOpacity(0.7f)
+		);
+		
+		if (HoverTarget != EDragTarget::Slice)
+		{
+			LayerId = SPanel::OnPaint(Args, AllottedGeometry, MyCullingRect, OutDrawElements, LayerId, InWidgetStyle, bParentEnabled) + 1;
+		}
+		
+		FSlateDrawElement::MakeBox(OutDrawElements,
+			++LayerId,
+			AllottedGeometry.MakeChild(
+				FVector2f(AllottedGeometry.GetLocalSize().X - Inset.Left - Inset.Right - SliderHeadMargin * 2.0f, SliderHeadHeight),
+				FSlateLayoutTransform(FVector2f(Inset.Left + SliderHeadMargin, BarPositions->BarFarPlane - SliderHeadHeight * 0.5f))
+			).ToPaintGeometry(),
+			HoveredSliceTopBrush,
+			ESlateDrawEffect::None,
+			HoveredSliceTopBrush->GetTint(InWidgetStyle)
+		);
+		
+		FSlateDrawElement::MakeBox(OutDrawElements,
+			++LayerId,
+			AllottedGeometry.MakeChild(
+				FVector2f(AllottedGeometry.GetLocalSize().X - Inset.Left - Inset.Right - SliderHeadMargin * 2.0f, SliderHeadHeight),
+				FSlateLayoutTransform(FVector2f(Inset.Left + SliderHeadMargin, BarPositions->BarNearPlane - SliderHeadHeight * 0.5f))
+			).ToPaintGeometry(),
+			HoveredSliceBottomBrush,
+			ESlateDrawEffect::None,
+			HoveredSliceBottomBrush->GetTint(InWidgetStyle)
+		);
+	}
+
+	return LayerId;
+}
+
+FVector2D SEditorViewportDepthBar::ComputeDesiredSize(float LayoutScaleMultiplier) const
+{
+	return FVector2D(20.0, 800.0);
+}
+
+TOptional<SEditorViewportDepthBar::FBarPositions> SEditorViewportDepthBar::GetBarPositions(const FGeometry& InGeometry) const
+{
+	TSharedPtr<FEditorViewportClient> Client = WeakClient.Pin();
+	if (!Client)
+	{
+		return TOptional<FBarPositions>();
+	}
+	
+	const FDepthSpace DepthSpace
+	{
+		.Origin = Client->GetViewTransform().GetLocation(),
+		.Forward = Client->GetForwardVector()
+	};
+	
+	const double WorldMaxAxis = DepthSpace.GetAxisPosition(WorldBounds.Max);
+	const double WorldMinAxis = DepthSpace.GetAxisPosition(WorldBounds.Min);
+	
+	const double RealMax = WorldMaxAxis > 0.0 ? WorldMaxAxis : WorldMinAxis;
+	const double RealMin = WorldMinAxis < 0.0 ? WorldMinAxis : WorldMaxAxis;
+	
+	FMargin Inset = Padding.Get();
+	
+	if (Client->IsOrtho())
+	{
+		FBarPositions Positions
+		{
+			.DepthSpace = DepthSpace,
+			.RealMax = RealMax,
+			.RealMin = RealMin,
+			.BarMin = Inset.Top,
+			.BarMax = InGeometry.GetLocalSize().Y - Inset.Bottom
+		};
+		
+		Positions.BarFarPlane = Positions.RelativeToBarPosition(Client->GetOrthographicFarPlaneOverride().IsSet() ? Positions.AxisToRelativePosition(Client->GetOrthographicFarPlaneOverride().GetValue()) : 1.0);
+		Positions.BarNearPlane = Positions.RelativeToBarPosition(Client->GetOrthographicNearPlaneOverride().IsSet() ? Positions.AxisToRelativePosition(Client->GetOrthographicNearPlaneOverride().GetValue()) : 0.0f);
+		
+		return Positions;
+	}
+	
+	return TOptional<FBarPositions>();
+}
+
+double SEditorViewportDepthBar::FDepthSpace::GetAxisPosition(const FVector& Position) const
+{
+	const FVector OriginToPosition = Position - Origin;
+	const FVector ProjectedPosition = OriginToPosition.ProjectOnToNormal(Forward);
+	const double Dot = ProjectedPosition.Dot(Forward);
+	return ProjectedPosition.Length() * (Dot > 0.0 ? 1.0 : -1.0); 
+}
+
+float SEditorViewportDepthBar::FBarPositions::WorldToBarPosition(const FVector& WorldPosition) const
+{
+	const double AxisPosition = DepthSpace.GetAxisPosition(WorldPosition);
+	const double RelativePosition = AxisToRelativePosition(AxisPosition);
+	return RelativeToBarPosition(RelativePosition);
+}
+
+double SEditorViewportDepthBar::FBarPositions::AxisToRelativePosition(double AxisPosition) const
+{
+	return (AxisPosition - RealMin) / (RealMax - RealMin);
+}
+
+float SEditorViewportDepthBar::FBarPositions::RelativeToBarPosition(double RelativePosition) const
+{
+	return FMath::Clamp(1.0f - RelativePosition, 0.0f, 1.0f) * (BarMax - BarMin) + BarMin;
+}
+
+double SEditorViewportDepthBar::FBarPositions::BarToRealPosition(float BarPosition) const
+{
+	const float RelativePosition = 1.0f - (BarPosition - BarMin) / (BarMax - BarMin);
+	return RelativePosition * (RealMax - RealMin) + RealMin;
+}
+
+FText SEditorViewportDepthBar::GetTooltipText() const
+{
+	if (TSharedPtr<FEditorViewportClient> Client = WeakClient.Pin())
+	{
+		switch (HoverTarget)
+		{
+		case EDragTarget::Far:
+			return GetFarPlaneText(Client);
+		case EDragTarget::Near:
+			return GetNearPlaneText(Client);
+		case EDragTarget::Slice:
+			return FText::Format(LOCTEXT("CombinedPlanesFormat", "{0}\n\n{1}"), GetFarPlaneText(Client), GetNearPlaneText(Client));
+		default:
+			break;
+		}
+	}
+	
+	return FText::GetEmpty();
+}
+
+FText SEditorViewportDepthBar::GetFarPlaneText(const TSharedPtr<FEditorViewportClient>& Client)
+{
+	if (Client->IsOrtho())
+	{
+		if (const double* FarPlane = Client->GetOrthographicFarPlaneOverride().GetPtrOrNull())
+		{
+			return FText::Format(LOCTEXT("OrthoFarPlaneFormat", "Far Plane: {0}"), UE::Editor::ViewportDepthBar::Private::GetDistanceDisplayText(*FarPlane));
+		}
+		return LOCTEXT("OrthoFarPlaneDefault", "Infinite Far Plane");
+	}
+	return FText::GetEmpty();
+}
+
+FText SEditorViewportDepthBar::GetNearPlaneText(const TSharedPtr<FEditorViewportClient>& Client)
+{
+	if (Client->IsOrtho())
+	{
+		if (const double* NearPlane = Client->GetOrthographicNearPlaneOverride().GetPtrOrNull())
+		{
+			return FText::Format(LOCTEXT("OrthoNearPlaneFormat", "Near Plane: {0}"), UE::Editor::ViewportDepthBar::Private::GetDistanceDisplayText(*NearPlane));
+		}
+		return LOCTEXT("OrthoNearPlaneDefault", "Infinite Near Plane");
+	}
+	return FText::GetEmpty();
+}
+
+#undef LOCTEXT_NAMESPACE
diff --git a/Engine/Source/Editor/UnrealEd/Public/EditorViewportClient.h b/Engine/Source/Editor/UnrealEd/Public/EditorViewportClient.h
--- a/Engine/Source/Editor/UnrealEd/Public/EditorViewportClient.h
+++ b/Engine/Source/Editor/UnrealEd/Public/EditorViewportClient.h
@@ -568,6 +568,9 @@ public:
 
 	/** @return The number of units per pixel displayed in this viewport */
 	UNREALED_API float GetOrthoUnitsPerPixel(const FViewport* Viewport) const;
+	
+	/** @return The forward vector of the camera view */
+	UNREALED_API FVector GetForwardVector() const;
 
 	/** Get a prettified string representation of the specified unreal units */
 	static UNREALED_API FString UnrealUnitsToSiUnits(float UnrealUnits);
@@ -1321,6 +1324,21 @@ public:
 	/** Override the far clipping plane. Set to a negative value to disable the override. */
 	UNREALED_API void OverrideFarClipPlane(const float InFarPlane);
 
+	/** Accessor for whether orthographic control should be enabled. */
+	bool SupportsManualOrthographicDepthControl() const; 
+
+	/** Get the orthographic near plane override for this viewport. */
+	UNREALED_API TOptional<double> GetOrthographicNearPlaneOverride() const;
+	
+	/** Set the orthographic near plane override for this viewport. */
+	UNREALED_API void SetOrthographicNearPlaneOverride(const TOptional<double>& InValue);
+	
+	/** Get the orthographic far plane override for this viewport. */
+	UNREALED_API TOptional<double> GetOrthographicFarPlaneOverride() const;
+	
+	/** Set the orthographic far plane override for this viewport. */
+	UNREALED_API void SetOrthographicFarPlaneOverride(const TOptional<double>& InValue);
+
 	/** When collision draw mode changes, this function allows hidden objects to be drawn, so hidden colliding objects can be seen */
 	UNREALED_API void UpdateHiddenCollisionDrawing();
 	/** Returns the scene depth at the given viewport X,Y */
@@ -2180,6 +2198,12 @@ private:
 
 	/** If > 0, overrides the view's far clipping plane with a plane at the specified distance. */
 	float FarPlane;
+	
+	/** The location of the near clipping plane in orthographic view. */
+	TOptional<double> OrthographicNearPlane;
+	
+	/** The location of the far clipping plane in orthographic view. */
+	TOptional<double> OrthographicFarPlane;
 
 	/** If true, we are in Game View mode*/
 	bool bInGameViewMode;
diff --git a/Engine/Source/Editor/UnrealEd/Public/SEditorViewport.h b/Engine/Source/Editor/UnrealEd/Public/SEditorViewport.h
--- a/Engine/Source/Editor/UnrealEd/Public/SEditorViewport.h
+++ b/Engine/Source/Editor/UnrealEd/Public/SEditorViewport.h
@@ -20,6 +20,7 @@ class UModeManagerInteractiveToolsContext;
 class UViewportInteractionsBehaviorSource;
 struct FSlateBrush;
 struct FToolMenuContext;
+class SEditorViewportDepthBar;
 
 class SEditorViewport
 	: public SCompoundWidget
@@ -338,6 +339,9 @@ protected:
 	/** The client responsible for setting up the scene */
 	TSharedPtr<FEditorViewportClient> Client;
 	
+	/** The depth bar used by this viewport */
+	TSharedPtr<SEditorViewportDepthBar> DepthBar;
+	
 	/** Commandlist used in the viewport (Maps commands to viewport specific actions) */
 	TSharedPtr<FUICommandList> CommandList;
 
@@ -360,6 +364,12 @@ private:
 	 * Will register ITF viewport interactions, if OnRegisterViewportInteractions is overridden in your SEditorViewport
 	 */
 	void RegisterViewportInteractions();
+	
+	/** @return The visibility of the viewport depth bar. */
+	EVisibility OnGetDepthBarVisibility() const;
+	
+	/** Reset the depth bar back to the default state. */
+	FReply DismissDepthBar();
 
 private:
 	/** The handle to the active timer */
@@ -368,6 +378,9 @@ private:
 	/** Whether the viewport needs to update, even without input or realtime (e.g. inertial camera movement) */
 	bool bInvalidated;
 	
+	/** Whether the depth bar should be displayed */
+	bool bShowDepthBar = false;
+	
 	/** For automatic legacy toolbar upgrades. Registers the child widget to expect in order to allow an automatic upgrade. */
 	TWeakPtr<SWidget> AutoUpgradeWidgetChild = nullptr;
 	/** A flag indicating that the automatic toolbar upgrade has been approved */
diff --git a/Engine/Source/Editor/UnrealEd/Public/SEditorViewportDepthBar.h b/Engine/Source/Editor/UnrealEd/Public/SEditorViewportDepthBar.h
new file mode 100644
--- /dev/null
+++ b/Engine/Source/Editor/UnrealEd/Public/SEditorViewportDepthBar.h
@@ -0,0 +1,155 @@
+﻿// Copyright Epic Games, Inc. All Rights Reserved.
+
+#pragma once
+
+#include "Widgets/SPanel.h"
+#include "SlotBase.h"
+#include "Widgets/DeclarativeSyntaxSupport.h"
+#include "Layout/Children.h"
+
+class FEditorViewportClient;
+
+/**
+ * A bar that allows for visualization and manipulation of a viewport's depth
+ */
+class SEditorViewportDepthBar : public SPanel
+{
+public:
+	class FSlot : public TSlotBase<FSlot>
+	{
+	public:
+	
+		SLATE_SLOT_BEGIN_ARGS(FSlot, TSlotBase<FSlot>)
+			/** Determines where in the bar the slot's contents appear. */
+			SLATE_ATTRIBUTE(FBoxSphereBounds, Bounds)
+		SLATE_SLOT_END_ARGS()
+		
+		void SetBounds(TAttribute<FBoxSphereBounds> InBounds)
+		{
+			Bounds = InBounds;
+		}
+		
+		void Construct(const FChildren& SlotOwner, FSlotArguments&& InArgs)
+		{
+			TSlotBase<FSlot>::Construct(SlotOwner, MoveTemp(InArgs));
+			Bounds = InArgs._Bounds;
+		}
+		
+		const TAttribute<FBoxSphereBounds>& GetBounds() const { return Bounds; }
+		
+	protected:
+		TAttribute<FBoxSphereBounds> Bounds;
+	};
+
+	static FSlot::FSlotArguments Slot()
+	{
+		return FSlot::FSlotArguments(MakeUnique<FSlot>());
+	}
+
+	SLATE_BEGIN_ARGS(SEditorViewportDepthBar)
+	{}
+		SLATE_ATTRIBUTE(TAttribute<FMargin>, Padding)
+		
+		SLATE_SLOT_ARGUMENT(SEditorViewportDepthBar::FSlot, Slots)
+	SLATE_END_ARGS()
+	
+	using FScopedWidgetSlotArguments = TPanelChildren<FSlot>::FScopedWidgetSlotArguments;
+	FScopedWidgetSlotArguments AddSlot()
+	{
+		return InsertSlot(INDEX_NONE);
+	}
+
+	FScopedWidgetSlotArguments InsertSlot(int32 Index = INDEX_NONE)
+	{
+		return FScopedWidgetSlotArguments(MakeUnique<FSlot>(), this->Children, Index);
+	}
+	
+	UNREALED_API FSlot& GetSlot(int32 Index);
+	UNREALED_API const FSlot& GetSlot(int32 Index) const;
+	UNREALED_API void ClearChildren();
+	
+	SEditorViewportDepthBar()
+		: Children(this)
+		, Padding(FMargin(0.0f, 4.0f))
+		, SliderMargin(4.0f)
+		, SliderHeadMargin(2.0f)
+		, SliderHeadHeight(3.0f)
+	{}
+	
+	UNREALED_API void Construct(const FArguments& InArgs, const TSharedPtr<FEditorViewportClient>& Client);
+	
+	UNREALED_API virtual void Tick(const FGeometry& AllottedGeometry, const double InCurrentTime, const float InDeltaTime) override;
+	UNREALED_API virtual FCursorReply OnCursorQuery(const FGeometry& MyGeometry, const FPointerEvent& CursorEvent) const override;
+	UNREALED_API virtual void OnArrangeChildren(const FGeometry& AllottedGeometry, FArrangedChildren& ArrangedChildren) const override;
+	
+	UNREALED_API virtual FReply OnMouseButtonDown(const FGeometry& MyGeometry, const FPointerEvent& MouseEvent) override;
+	UNREALED_API virtual FReply OnMouseMove(const FGeometry& MyGeometry, const FPointerEvent& MouseEvent) override;
+	UNREALED_API virtual FReply OnMouseButtonUp(const FGeometry& MyGeometry, const FPointerEvent& MouseEvent) override;
+	UNREALED_API virtual void OnMouseLeave(const FPointerEvent& MouseEvent) override;
+	
+protected:
+	UNREALED_API virtual FChildren* GetChildren() override;
+	UNREALED_API virtual int32 OnPaint(const FPaintArgs& Args, const FGeometry& AllottedGeometry, const FSlateRect& MyCullingRect, FSlateWindowElementList& OutDrawElements, int32 LayerId, const FWidgetStyle& InWidgetStyle, bool bParentEnabled) const override;
+	UNREALED_API virtual FVector2D ComputeDesiredSize(float LayoutScaleMultiplier) const override;
+	
+private:
+	struct FDepthSpace
+	{
+		FVector Origin;
+		FVector Forward;
+		
+		double GetAxisPosition(const FVector& Position) const;
+	};
+
+	struct FBarPositions
+	{
+		FDepthSpace DepthSpace;
+	
+		double RealMax;
+		double RealMin;
+		
+		float BarMin;
+		float BarMax;
+		
+		float BarFarPlane;
+		float BarNearPlane;
+		
+		float WorldToBarPosition(const FVector& WorldPosition) const;
+		double AxisToRelativePosition(double AxisPosition) const;
+		float RelativeToBarPosition(double RelativePosition) const;
+		double BarToRealPosition(float BarPosition) const;
+	};
+
+	TOptional<FBarPositions> GetBarPositions(const FGeometry& InGeometry) const;
+	
+	FText GetTooltipText() const;
+	static FText GetFarPlaneText(const TSharedPtr<FEditorViewportClient>& Client);
+	static FText GetNearPlaneText(const TSharedPtr<FEditorViewportClient>& Client);
+	
+	FBox WorldBounds;
+	TWeakPtr<FEditorViewportClient> WeakClient;
+	
+	TPanelChildren<FSlot> Children;
+	
+	enum class EDragTarget
+	{
+		None,
+		Near,
+		Far,
+		Slice
+	};
+	EDragTarget HoverTarget = EDragTarget::None;
+	EDragTarget DragTarget = EDragTarget::None;
+	float DragOffset = 0.0f;
+	
+	TAttribute<FMargin> Padding;
+	float SliderMargin;
+	float SliderHeadMargin;
+	float SliderHeadHeight;
+	
+	const FSlateBrush* BackgroundBrush = nullptr;
+	const FSlateBrush* SliceBrush = nullptr;
+	const FSlateBrush* HoveredSliceBrush = nullptr;
+	const FSlateBrush* HoveredSliceTopBrush = nullptr;
+	const FSlateBrush* HoveredSliceBottomBrush = nullptr;
+};
