Fixed compilation errors with UE_LOG_LAZY_CATEGORY_NAMES=1.

diff --git a/Engine/Source/Developer/MessageLog/Private/Model/MessageLogListingModel.cpp b/Engine/Source/Developer/MessageLog/Private/Model/MessageLogListingModel.cpp
--- a/Engine/Source/Developer/MessageLog/Private/Model/MessageLogListingModel.cpp
+++ b/Engine/Source/Developer/MessageLog/Private/Model/MessageLogListingModel.cpp
@@ -65,7 +65,7 @@ void FMessageLogListingModel::AddMessageInternal( const TSharedRef<FTokenizedMes
 			SET_WARN_COLOR(LogColor);
 		}
 
-		FMsg::Logf(__FILE__, __LINE__, *LogName.ToString(), FMessageLog::GetLogVerbosity(NewMessage->GetSeverity()), TEXT("%s"), *NewMessage->ToText().ToString());
+		FMsg::Logf(__FILE__, __LINE__, UE_FNAME_TO_LOG_CATEGORY_NAME(LogName), FMessageLog::GetLogVerbosity(NewMessage->GetSeverity()), TEXT("%s"), *NewMessage->ToText().ToString());
 
 		CLEAR_WARN_COLOR();
 	}
diff --git a/Engine/Source/Editor/UnrealEd/Private/Cooker/CookLog.cpp b/Engine/Source/Editor/UnrealEd/Private/Cooker/CookLog.cpp
--- a/Engine/Source/Editor/UnrealEd/Private/Cooker/CookLog.cpp
+++ b/Engine/Source/Editor/UnrealEd/Private/Cooker/CookLog.cpp
@@ -350,7 +350,7 @@ void FLogHandler::UnMarshalAndLog(const FReplicatedLogData& LogData,
 			SerializedString = &TransformedString;
 		}
 
-		FMsg::Logf(__FILE__, __LINE__, UnStructuredLogData->Category, UnStructuredLogData->Verbosity,
+		FMsg::Logf(__FILE__, __LINE__, UE_FNAME_TO_LOG_CATEGORY_NAME(UnStructuredLogData->Category), UnStructuredLogData->Verbosity,
 			TEXT("%s"), **SerializedString);
 	}
 	else if (const FCbObject* StructuredLogObject = LogData.LogDataVariant.TryGet<FCbObject>())
diff --git a/Engine/Source/Editor/UnrealEd/Private/Cooker/CookWorkerServer.cpp b/Engine/Source/Editor/UnrealEd/Private/Cooker/CookWorkerServer.cpp
--- a/Engine/Source/Editor/UnrealEd/Private/Cooker/CookWorkerServer.cpp
+++ b/Engine/Source/Editor/UnrealEd/Private/Cooker/CookWorkerServer.cpp
@@ -208,7 +208,7 @@ void FCookWorkerServer::SendCrashDiagnostics()
 			// NOTE:  If you change the format of [CookWorker %d] you must also change
 			// the corresponding hash aononymization in HashedIssueHandler.cs, and generally
 			// search for [CookWorker as it's cut and pasted around a few spots. 
-			FMsg::Logf(__FILE__, __LINE__, Category, Verbosity, TEXT("[CookWorker %d]: %.*s"),
+			FMsg::Logf(__FILE__, __LINE__, UE_FNAME_TO_LOG_CATEGORY_NAME(Category), Verbosity, TEXT("[CookWorker %d]: %.*s"),
 				ProfileId, Message.Len(), Message.GetData());
 		}
 	}
diff --git a/Engine/Source/Runtime/Core/Public/Logging/LogCategory.h b/Engine/Source/Runtime/Core/Public/Logging/LogCategory.h
--- a/Engine/Source/Runtime/Core/Public/Logging/LogCategory.h
+++ b/Engine/Source/Runtime/Core/Public/Logging/LogCategory.h
@@ -10,8 +10,10 @@
 
 #if defined(UE_LOG_LAZY_CATEGORY_NAMES) && UE_LOG_LAZY_CATEGORY_NAMES
 using FLogCategoryName = FLazyName;
+#define UE_FNAME_TO_LOG_CATEGORY_NAME(Name) FLazyName(Name)
 #else
 using FLogCategoryName = FName;
+#define UE_FNAME_TO_LOG_CATEGORY_NAME(Name) Name
 #endif
 
 /** Base class for all log categories. */
diff --git a/Engine/Source/Runtime/Core/Public/UObject/NameTypes.h b/Engine/Source/Runtime/Core/Public/UObject/NameTypes.h
--- a/Engine/Source/Runtime/Core/Public/UObject/NameTypes.h
+++ b/Engine/Source/Runtime/Core/Public/UObject/NameTypes.h
@@ -1749,6 +1749,11 @@ public:
 	}
 
 	CORE_API bool UEOpEquals(const FLazyName& Rhs) const;
+
+	friend FORCEINLINE uint32 GetTypeHash(FLazyName Name)
+	{
+		return GetTypeHash(Name.Resolve());
+	}
 };
 
 // Ordering operators for FName are deliberately disabled as there is no correct default.
diff --git a/Engine/Source/Runtime/Engine/Private/VisualLogger/VisualLoggerKismetLibrary.cpp b/Engine/Source/Runtime/Engine/Private/VisualLogger/VisualLoggerKismetLibrary.cpp
--- a/Engine/Source/Runtime/Engine/Private/VisualLogger/VisualLoggerKismetLibrary.cpp
+++ b/Engine/Source/Runtime/Engine/Private/VisualLogger/VisualLoggerKismetLibrary.cpp
@@ -13,7 +13,7 @@
 	}\
 	else\
 	{\
-		FMsg::Logf(__FILE__, __LINE__, CategoryName, DefaultVerbosity, (Format), __VA_ARGS__);\
+		FMsg::Logf(__FILE__, __LINE__, UE_FNAME_TO_LOG_CATEGORY_NAME(CategoryName), DefaultVerbosity, (Format), __VA_ARGS__);\
 	}
 
 void UVisualLoggerKismetLibrary::RedirectVislog(UObject* SourceOwner, UObject* DestinationOwner)
