Add divide-by-zero protection to calculations in the viewport depth bar.

diff --git a/Engine/Source/Editor/UnrealEd/Private/SEditorViewportDepthBar.cpp b/Engine/Source/Editor/UnrealEd/Private/SEditorViewportDepthBar.cpp
--- a/Engine/Source/Editor/UnrealEd/Private/SEditorViewportDepthBar.cpp
+++ b/Engine/Source/Editor/UnrealEd/Private/SEditorViewportDepthBar.cpp
@@ -414,7 +414,8 @@ float SEditorViewportDepthBar::FBarPositions::WorldToBarPosition(const FVector&
 
 double SEditorViewportDepthBar::FBarPositions::AxisToRelativePosition(double AxisPosition) const
 {
-	return (AxisPosition - RealMin) / (RealMax - RealMin);
+	const double Size = RealMax - RealMin;
+	return Size != 0.0 ? ((AxisPosition - RealMin) / Size) : 0.0;
 }
 
 float SEditorViewportDepthBar::FBarPositions::RelativeToBarPosition(double RelativePosition) const
@@ -424,7 +425,8 @@ float SEditorViewportDepthBar::FBarPositions::RelativeToBarPosition(double Relat
 
 double SEditorViewportDepthBar::FBarPositions::BarToRealPosition(float BarPosition) const
 {
-	const float RelativePosition = 1.0f - (BarPosition - BarMin) / (BarMax - BarMin);
+	const float BarSize = BarMax - BarMin;
+	const float RelativePosition = BarSize > 0.0f ? (1.0f - (BarPosition - BarMin) / BarSize) : 0.0f;
 	return RelativePosition * (RealMax - RealMin) + RealMin;
 }
 
