Added UCommandlet::ApplyCommandLineSwitches that will parse all commandline switches and apply the values to the internal properties if they exists.
This removes the need to manually parse commandline switches.

diff --git a/Engine/Source/Editor/UnrealEd/Private/WorldPartition/WorldPartitionResaveActorsBuilder.cpp b/Engine/Source/Editor/UnrealEd/Private/WorldPartition/WorldPartitionResaveActorsBuilder.cpp
--- a/Engine/Source/Editor/UnrealEd/Private/WorldPartition/WorldPartitionResaveActorsBuilder.cpp
+++ b/Engine/Source/Editor/UnrealEd/Private/WorldPartition/WorldPartitionResaveActorsBuilder.cpp
@@ -12,11 +12,9 @@
 #include "Engine/World.h"
 #include "GameFramework/Actor.h"
 #include "Commandlets/Commandlet.h"
-#include "Kismet2/BlueprintEditorUtils.h"
 #include "Misc/CommandLine.h"
 #include "UObject/MetaData.h"
 #include "UObject/SavePackage.h"
-#include "Kismet2/BlueprintEditorUtils.h"
 #include "ReferenceCluster.h"
 #include "ActorFolder.h"
 #include "Misc/FileHelper.h"
@@ -35,41 +33,10 @@ bool UWorldPartitionResaveActorsBuilder::PreRun(UWorld* World, FPackageSourceCon
 {
 	TArray<FString> Tokens, Switches;
 	UCommandlet::ParseCommandLine(*GetBuilderArgs(), Tokens, Switches);
-
-	//@todo_ow: generalize to all builders
-	for (const FString& Switch : Switches)
+	
+	if (!UCommandlet::ApplyCommandLineSwitches(this, Switches))
 	{
-		FString Key, Value;
-		if (!Switch.Split(TEXT("="), &Key, &Value))
-		{
-			Key = Switch;
-		}
-
-		// Lookup property
-		const FProperty* Property = GetClass()->FindPropertyByName(*Key);
-
-		// If we can't find the property, try for properties with the 'b' prefix
-		if (!Property)
-		{
-			Key = TEXT("b") + Key;
-			Property = GetClass()->FindPropertyByName(*Key);
-		}
-
-		if (Property)
-		{
-			// If the property is a bool, treat no values as true
-			if (Property->IsA(FBoolProperty::StaticClass()) && Value.IsEmpty())
-			{
-				Value = TEXT("True");
-			}
-
-			uint8* Container = (uint8*)this;
-			if (!FBlueprintEditorUtils::PropertyValueFromString(Property, Value, Container, nullptr))
-			{
-				UE_LOG(LogWorldPartitionResaveActorsBuilder, Error, TEXT("Cannot set value for '%s': '%s'"), *Key, *Value);
-				return false;
-			}
-		}
+		return false;
 	}
 
 	if (bSwitchActorPackagingSchemeToReduced)
diff --git a/Engine/Source/Runtime/Engine/Classes/Commandlets/Commandlet.h b/Engine/Source/Runtime/Engine/Classes/Commandlets/Commandlet.h
--- a/Engine/Source/Runtime/Engine/Classes/Commandlets/Commandlet.h
+++ b/Engine/Source/Runtime/Engine/Classes/Commandlets/Commandlet.h
@@ -36,7 +36,7 @@
 #include "UObject/Object.h"
 #include "Commandlet.generated.h"
 
-UCLASS(abstract, transient,MinimalAPI)
+UCLASS(abstract, transient, MinimalAPI)
 class UCommandlet : public UObject
 {
 	GENERATED_UCLASS_BODY()
@@ -160,13 +160,23 @@ class UCommandlet : public UObject
 		}
 	}
 
+#if WITH_EDITOR
+	/**
+	 * Apply the switched to internal values if they exists.
+	 *
+	 * @param	Switches	The switches to apply to internal properties.
+	 *
+	 * @return	True if all values were properly applied.
+	 */
+	static ENGINE_API bool ApplyCommandLineSwitches( UObject* Object, const TArray<FString>& Switches );
+#endif //  WITH_EDITOR
+
 	/**
 	 * Allows commandlets to override the default behavior and create a custom engine class for the commandlet. If
 	 * the commandlet implements this function, it should fully initialize the UEngine object as well.  Commandlets
 	 * should indicate that they have implemented this function by assigning the custom UEngine to GEngine.
 	 */
 	virtual void CreateCustomEngine(const FString& Params) {}
-
 };
 
 namespace CommandletHelpers
diff --git a/Engine/Source/Runtime/Engine/Private/Commandlets/Commandlet.cpp b/Engine/Source/Runtime/Engine/Private/Commandlets/Commandlet.cpp
--- a/Engine/Source/Runtime/Engine/Private/Commandlets/Commandlet.cpp
+++ b/Engine/Source/Runtime/Engine/Private/Commandlets/Commandlet.cpp
@@ -17,8 +17,14 @@
 #include "HAL/ThreadManager.h"
 #include "Framework/Application/SlateApplication.h"
 
+#if WITH_EDITOR
+#include "Kismet2/BlueprintEditorUtils.h"
+#endif //  WITH_EDITOR
+
 #include UE_INLINE_GENERATED_CPP_BY_NAME(Commandlet)
 
+DEFINE_LOG_CATEGORY_STATIC(LogCommandlet, Log, All);
+
 UCommandlet::UCommandlet(const FObjectInitializer& ObjectInitializer)
 	: Super(ObjectInitializer)
 {
@@ -31,7 +37,49 @@ UCommandlet::UCommandlet(const FObjectInitializer& ObjectInitializer)
 	UseCommandletResultAsExitCode = false;
 }
 
+#if WITH_EDITOR
+bool UCommandlet::ApplyCommandLineSwitches( UObject* Object, const TArray<FString>& Switches )
+{
+	for (const FString& Switch : Switches)
+	{
+		FString Key;
+		FString Value;
+
+		if (!Switch.Split(TEXT("="), &Key, &Value))
+		{
+			Key = Switch;
+		}
+
+		// Lookup property
+		const FProperty* Property = Object->GetClass()->FindPropertyByName(*Key);
 
+		// If we can't find the property, try for properties with the 'b' prefix
+		if (!Property)
+		{
+			Key = TEXT("b") + Key;
+			Property = Object->GetClass()->FindPropertyByName(*Key);
+		}
+
+		if (Property)
+		{
+			// If the property is a bool, treat no values as true
+			if (Property->IsA(FBoolProperty::StaticClass()) && Value.IsEmpty())
+			{
+				Value = TEXT("True");
+			}
+
+			// Apply the property
+			if (!FBlueprintEditorUtils::PropertyValueFromString(Property, Value, (uint8*)Object, nullptr))
+			{
+				UE_LOG(LogCommandlet, Error, TEXT("Cannot set value for '%s': '%s'"), *Key, *Value);
+				return false;
+			}
+		}
+	}
+
+	return true;
+}
+#endif //  WITH_EDITOR
 
 /* 
 	Tests for Commandlet Utils
