ExternalActorsCommandlet
- Converted manual commandline switches parsing by using UCommandlet::ApplyCommandLineSwitches.
- Added option `-DumpCSVFile=` to allow dumping all potential worlds that can have external actors disabled into a CSV file, along with the number of external files.

diff --git a/Engine/Source/Editor/UnrealEd/Classes/Commandlets/ExternalActorsCommandlet.h b/Engine/Source/Editor/UnrealEd/Classes/Commandlets/ExternalActorsCommandlet.h
--- a/Engine/Source/Editor/UnrealEd/Classes/Commandlets/ExternalActorsCommandlet.h
+++ b/Engine/Source/Editor/UnrealEd/Classes/Commandlets/ExternalActorsCommandlet.h
@@ -20,4 +20,22 @@ public:
 
 protected:
 	UNREALED_API UWorld* LoadWorld(const FString& LevelToLoad);
+
+	UPROPERTY()
+	bool bRepair;
+
+	UPROPERTY()
+	bool bDisable;
+
+	UPROPERTY()
+	bool bForce;
+
+	UPROPERTY()
+	bool bReport;
+
+	UPROPERTY()
+	bool bListMaps;
+
+	UPROPERTY()
+	FString DumpCSVFile;
 };
diff --git a/Engine/Source/Editor/UnrealEd/Private/Commandlets/ExternalActorsCommandlet.cpp b/Engine/Source/Editor/UnrealEd/Private/Commandlets/ExternalActorsCommandlet.cpp
--- a/Engine/Source/Editor/UnrealEd/Private/Commandlets/ExternalActorsCommandlet.cpp
+++ b/Engine/Source/Editor/UnrealEd/Private/Commandlets/ExternalActorsCommandlet.cpp
@@ -46,19 +46,7 @@ int32 UExternalActorsCommandlet::Main(const FString& Params)
 
 	TArray<FString> Tokens, Switches;
 	ParseCommandLine(*Params, Tokens, Switches);
-
-	// Need at least the level to convert
-	if (Tokens.Num() < 1)
-	{
-		UE_LOG(LogExternalActorsCommandlet, Error, TEXT("Bad parameters"));
-		return 1;
-	}
-
-	const bool bRepair = Switches.Contains(TEXT("repair"));
-	const bool bDisable = Switches.Contains(TEXT("disable"));
-	const bool bForce = Switches.Contains(TEXT("force"));
-	const bool bReport = Switches.Contains(TEXT("report"));
-	const bool bListMaps = Switches.Contains(TEXT("listmaps"));
+	ApplyCommandLineSwitches(this, Switches);
 
 	if (bDisable)
 	{
@@ -67,10 +55,12 @@ int32 UExternalActorsCommandlet::Main(const FString& Params)
 			UE_LOG(LogExternalActorsCommandlet, Error, TEXT("Bad parameters: -repair cannot be used with -disable"));
 			return 1;
 		}
+
+		bListMaps |= !DumpCSVFile.IsEmpty();
 	}
-	else if (bForce || bReport || bListMaps)
+	else if (bForce || bReport || bListMaps || DumpCSVFile.Len())
 	{
-		UE_LOG(LogExternalActorsCommandlet, Error, TEXT("Bad parameters: -force, -report and -listmaps can only be used along with -disable"));
+		UE_LOG(LogExternalActorsCommandlet, Error, TEXT("Bad parameters: -force, -report, -listmaps and -dumpcsv can only be used along with -disable"));
 		return 1;
 	}
 
@@ -85,16 +75,71 @@ int32 UExternalActorsCommandlet::Main(const FString& Params)
 		AssetRegistryModule.Get().GetAssetsByClass(UWorld::StaticClass()->GetClassPathName(), WorldAssets, true);
 		WorldAssets.Sort();
 
+		struct FWorldEntry
+		{
+			FString Path;
+			int32 NumExternalActors;
+			FString DateModified;
+
+			inline bool operator<(const FWorldEntry& Other) const
+			{
+				return NumExternalActors > Other.NumExternalActors;
+			}
+		};
+
+		TArray<FWorldEntry> WorldEntries;
 		for (const FAssetData& WorldAsset : WorldAssets)
 		{
-			if (ULevel::GetIsLevelUsingExternalActorsFromAsset(WorldAsset) && (!ULevel::GetIsLevelPartitionedFromAsset(WorldAsset) || ULevel::GetIsStreamingDisabledFromAsset(WorldAsset)))
+			if (ULevel::GetIsLevelUsingExternalActorsFromAsset(WorldAsset))
 			{
-				UE_LOG(LogExternalActorsCommandlet, Display, TEXT("Level '%s' is a potential candidate to have external actors disabled."), *WorldAsset.GetSoftObjectPath().ToString());
+				const FString ExternalActorsPath = ULevel::GetExternalActorsPath(WorldAsset.PackageName.ToString());
+
+				TArray<FAssetData> WorldExternalActors;
+				AssetRegistryModule.Get().GetAssetsByPath(*ExternalActorsPath, WorldExternalActors, true, true);
+
+				if (!ULevel::GetIsLevelPartitionedFromAsset(WorldAsset) || ULevel::GetIsStreamingDisabledFromAsset(WorldAsset))
+				{
+					FString DateModified;
+					WorldAsset.GetTagValue(TEXT("DateModified"), DateModified);
+
+					WorldEntries.Add(
+					{
+						.Path = WorldAsset.PackageName.ToString(),
+						.NumExternalActors = WorldExternalActors.Num(),
+						.DateModified = DateModified
+					});
+
+					UE_LOG(LogExternalActorsCommandlet, Display, TEXT("Level '%s' is a potential candidate to have external actors disabled (%d)."), *WorldAsset.GetSoftObjectPath().ToString(), WorldExternalActors.Num());
+				}
+				else
+				{
+					UE_LOG(LogExternalActorsCommandlet, Verbose, TEXT("Level '%s' is not a potential candidate to have external actors disabled (%d)."), *WorldAsset.GetSoftObjectPath().ToString(), WorldExternalActors.Num());
+				}
 			}
-			else
+		}
+
+		if (DumpCSVFile.Len())
+		{
+			FArchive* LogFile = IFileManager::Get().CreateFileWriter(*DumpCSVFile);
+			if (!LogFile)
 			{
-				UE_LOG(LogExternalActorsCommandlet, Verbose, TEXT("Level '%s' is not a potential candidate to have external actors disabled."), *WorldAsset.GetSoftObjectPath().ToString());
+				UE_LOG(LogExternalActorsCommandlet, Error, TEXT("Cannot create csv file '%s'"), *DumpCSVFile);
+				return 1;
 			}
+
+			FString LineEntry = TEXT("Map,NumOFPA,DateModified") LINE_TERMINATOR;
+			LogFile->Serialize(TCHAR_TO_ANSI(*LineEntry), LineEntry.Len());
+
+			WorldEntries.Sort();
+
+			for (const FWorldEntry& WorldEntry : WorldEntries)
+			{
+				LineEntry = FString::Printf(TEXT("%s,%d,%s" LINE_TERMINATOR), *WorldEntry.Path, WorldEntry. NumExternalActors, *WorldEntry.DateModified);
+				LogFile->Serialize(TCHAR_TO_ANSI(*LineEntry), LineEntry.Len());
+			}
+
+			LogFile->Close();
+			delete LogFile;
 		}
 
 		return 0;
@@ -220,7 +265,7 @@ int32 UExternalActorsCommandlet::Main(const FString& Params)
 
 		if (IFileManager::Get().DirectoryExists(*ExternalActorsFilePath))
 		{
-			bool bResult = IFileManager::Get().IterateDirectoryRecursively(*ExternalActorsFilePath, [this, bRepair, &PackagesToSave, &PackagesToDelete, &ActorGuids](const TCHAR* FilenameOrDirectory, bool bIsDirectory)
+			bool bResult = IFileManager::Get().IterateDirectoryRecursively(*ExternalActorsFilePath, [this, &PackagesToSave, &PackagesToDelete, &ActorGuids](const TCHAR* FilenameOrDirectory, bool bIsDirectory)
 			{
 				if (!bIsDirectory)
 				{
