ExternalActorsCommandlet:  added '-MapListFile=' option to specify a text file containing the list of maps to work on.

diff --git a/Engine/Source/Editor/UnrealEd/Classes/Commandlets/ExternalActorsCommandlet.h b/Engine/Source/Editor/UnrealEd/Classes/Commandlets/ExternalActorsCommandlet.h
--- a/Engine/Source/Editor/UnrealEd/Classes/Commandlets/ExternalActorsCommandlet.h
+++ b/Engine/Source/Editor/UnrealEd/Classes/Commandlets/ExternalActorsCommandlet.h
@@ -36,6 +36,9 @@ protected:
 	UPROPERTY()
 	bool bListMaps;
 
+	UPROPERTY()
+	FString MapListFile;
+
 	UPROPERTY()
 	FString DumpCSVFile;
 };
diff --git a/Engine/Source/Editor/UnrealEd/Private/Commandlets/ExternalActorsCommandlet.cpp b/Engine/Source/Editor/UnrealEd/Private/Commandlets/ExternalActorsCommandlet.cpp
--- a/Engine/Source/Editor/UnrealEd/Private/Commandlets/ExternalActorsCommandlet.cpp
+++ b/Engine/Source/Editor/UnrealEd/Private/Commandlets/ExternalActorsCommandlet.cpp
@@ -4,6 +4,7 @@
 #include "AssetRegistry/AssetData.h"
 #include "AssetRegistry/AssetRegistryModule.h"
 #include "HAL/FileManager.h"
+#include "Misc/FileHelper.h"
 #include "Misc/Paths.h"
 #include "Editor.h"
 #include "Engine/Level.h"
@@ -145,243 +146,274 @@ int32 UExternalActorsCommandlet::Main(const FString& Params)
 		return 0;
 	}
 
-	if (!FPackageName::SearchForPackageOnDisk(Tokens[0], &Tokens[0]))
-	{
-		UE_LOG(LogExternalActorsCommandlet, Error, TEXT("Unknown level '%s'"), *Tokens[0]);
-		return 1;
-	}
+	TArray<FString> MapList;
 
-	// Load world
-	UWorld* MainWorld = LoadWorld(Tokens[0]);
-	if (!MainWorld)
+	if (!MapListFile.IsEmpty())
 	{
-		UE_LOG(LogExternalActorsCommandlet, Error, TEXT("Unknown world '%s'"), *Tokens[0]);
-		return 1;
+		if (!FFileHelper::LoadFileToStringArray(MapList, *MapListFile))
+		{
+			UE_LOG(LogExternalActorsCommandlet, Error, TEXT("Invalid map list filename"));
+			return 1;
+		}
 	}
+	else
+	{
+		if (Tokens.Num() < 1)
+		{
+			UE_LOG(LogExternalActorsCommandlet, Error, TEXT("Missing map name"));
+			return 1;
+		}
 
-	TSet<UPackage*> PackagesToSave;
-	TArray<FString> PackagesToDelete;	
+		MapList.Add(Tokens[0]);
+	}
 
-	if (bDisable)
+	int32 Result = 1;
+	for (const FString& MapName : MapList)
 	{
-		if (!MainWorld->PersistentLevel->IsUsingExternalActors())
+		FString FullMapName;
+		if (!FPackageName::SearchForPackageOnDisk(MapName, &FullMapName))
 		{
-			UE_LOG(LogExternalActorsCommandlet, Error, TEXT("Cannot disable external actors for partitioned level '%s' (already disabled)"), *Tokens[0]);
-			return 1;
+			UE_LOG(LogExternalActorsCommandlet, Error, TEXT("Unknown level '%s'"), *MapName);
+			Result = 0;
+			continue;
 		}
 
-		if (UWorldPartition* WorldPartition = MainWorld->GetWorldPartition())
+		// Load world
+		UWorld* MainWorld = LoadWorld(FullMapName);
+		if (!MainWorld)
 		{
-			int32 NumErrors = 0;
+			UE_LOG(LogExternalActorsCommandlet, Error, TEXT("Unknown world '%s'"), *FullMapName);
+			Result = 0;
+			continue;
+		}
 
-			if (WorldPartition->IsStreamingEnabled())
-			{
-				UE_LOG(LogExternalActorsCommandlet, Error, TEXT("Cannot disable external actors for partitioned level '%s' with streaming enabled"), *Tokens[0]);
-				NumErrors++;
+		TSet<UPackage*> PackagesToSave;
+		TArray<FString> PackagesToDelete;	
 
-				if (bForce)
-				{
-					WorldPartition->SetEnableStreaming(false);
-				}
+		if (bDisable)
+		{
+			if (!MainWorld->PersistentLevel->IsUsingExternalActors())
+			{
+				UE_LOG(LogExternalActorsCommandlet, Error, TEXT("Cannot disable external actors for partitioned level '%s' (already disabled)"), *FullMapName);
+				Result = 0;
+				continue;
 			}
 
-			if (UDataLayerManager* DataLayerManager = WorldPartition->GetDataLayerManager())
+			if (UWorldPartition* WorldPartition = MainWorld->GetWorldPartition())
 			{
-				if (DataLayerManager->GetDataLayerInstances().Num())
+				int32 NumErrors = 0;
+
+				if (WorldPartition->IsStreamingEnabled())
 				{
-					UE_LOG(LogExternalActorsCommandlet, Error, TEXT("Cannot disable external actors for partitioned level '%s' with data layers"), *Tokens[0]);
+					UE_LOG(LogExternalActorsCommandlet, Error, TEXT("Cannot disable external actors for partitioned level '%s' with streaming enabled"), *FullMapName);
 					NumErrors++;
-				}
-			}
-
-			if (!NumErrors)
-			{
-				UE_LOG(LogExternalActorsCommandlet, Display, TEXT("External actors for partitioned level '%s' can be disabled"), *Tokens[0]);
-			}
-			else if (bForce)
-			{
-				UE_LOG(LogExternalActorsCommandlet, Display, TEXT("External actors for partitioned level '%s' will be forcibly disabled"), *Tokens[0]);
-			}
-			else
-			{
-				return 1;
-			}			
 
-			if (!bReport)
-			{
-				WorldPartition->Initialize(MainWorld, FTransform::Identity);
+					if (bForce)
+					{
+						WorldPartition->SetEnableStreaming(false);
+					}
+				}
 
-				if (!UWorldPartition::RemoveWorldPartition(WorldPartition->GetWorld()->GetWorldSettings()))
+				if (UDataLayerManager* DataLayerManager = WorldPartition->GetDataLayerManager())
 				{
-					UE_LOG(LogExternalActorsCommandlet, Error, TEXT("Error disabling external actors for partitioned level '%s'"), *Tokens[0]);
-					return 1;
+					if (DataLayerManager->GetDataLayerInstances().Num())
+					{
+						UE_LOG(LogExternalActorsCommandlet, Error, TEXT("Cannot disable external actors for partitioned level '%s' with data layers"), *FullMapName);
+						NumErrors++;
+					}
 				}
 
-				WorldPartition->Uninitialize();
-			}
-		}
+				if (!NumErrors)
+				{
+					UE_LOG(LogExternalActorsCommandlet, Display, TEXT("External actors for partitioned level '%s' can be disabled"), *FullMapName);
+				}
+				else if (bForce)
+				{
+					UE_LOG(LogExternalActorsCommandlet, Display, TEXT("External actors for partitioned level '%s' will be forcibly disabled"), *FullMapName);
+				}
+				else
+				{
+					Result = 0;
+					continue;
+				}			
 
-		if (!bReport)
-		{
-			for (AActor* Actor : MainWorld->PersistentLevel->Actors)
-			{
-				if (IsValid(Actor) && Actor->IsMainPackageActor())
+				if (!bReport)
 				{
-					PackagesToDelete.Add(Actor->GetPackage()->GetLoadedPath().GetLocalFullPath());
-					Actor->SetPackageExternal(false);
+					WorldPartition->Initialize(MainWorld, FTransform::Identity);
+
+					if (!UWorldPartition::RemoveWorldPartition(WorldPartition->GetWorld()->GetWorldSettings()))
+					{
+						UE_LOG(LogExternalActorsCommandlet, Error, TEXT("Error disabling external actors for partitioned level '%s'"), *FullMapName);
+						Result = 0;
+						continue;
+					}
+
+					WorldPartition->Uninitialize();
 				}
 			}
 
-			if (MainWorld->PersistentLevel->IsUsingExternalObjects())
+			if (!bReport)
 			{
-				ForEachObjectWithOuter(MainWorld->PersistentLevel, [MainWorld, &PackagesToDelete](UObject* Object)
+				for (AActor* Actor : MainWorld->PersistentLevel->Actors)
 				{
-					if (Object->IsPackageExternal() && !Object->IsA<AActor>())
+					if (IsValid(Actor) && Actor->IsMainPackageActor())
 					{
-						PackagesToDelete.Add(Object->GetPackage()->GetLoadedPath().GetLocalFullPath());
-						FExternalPackageHelper::SetPackagingMode(Object, MainWorld->PersistentLevel, false);
+						PackagesToDelete.Add(Actor->GetPackage()->GetLoadedPath().GetLocalFullPath());
+						Actor->SetPackageExternal(false);
 					}
-					return true;
-				}, true, RF_MirroredGarbage);
-			}
+				}
 
-			if (PackagesToDelete.Num())
-			{
-				PackagesToSave.Add(MainWorld->GetPackage());
-			}
+				if (MainWorld->PersistentLevel->IsUsingExternalObjects())
+				{
+					ForEachObjectWithOuter(MainWorld->PersistentLevel, [MainWorld, &PackagesToDelete](UObject* Object)
+					{
+						if (Object->IsPackageExternal() && !Object->IsA<AActor>())
+						{
+							PackagesToDelete.Add(Object->GetPackage()->GetLoadedPath().GetLocalFullPath());
+							FExternalPackageHelper::SetPackagingMode(Object, MainWorld->PersistentLevel, false);
+						}
+						return true;
+					}, true, RF_MirroredGarbage);
+				}
+
+				if (PackagesToDelete.Num())
+				{
+					PackagesToSave.Add(MainWorld->GetPackage());
+				}
 
-			MainWorld->PersistentLevel->SetUseActorFolders(false);
-			MainWorld->PersistentLevel->SetUseExternalActors(false);
+				MainWorld->PersistentLevel->SetUseActorFolders(false);
+				MainWorld->PersistentLevel->SetUseExternalActors(false);
+			}
 		}
-	}
-	else
-	{
-		// Validate external actors
-		FString ExternalActorsPath = ULevel::GetExternalActorsPath(Tokens[0]);
-		FString ExternalActorsFilePath = FPackageName::LongPackageNameToFilename(ExternalActorsPath);
+		else
+		{
+			// Validate external actors
+			FString ExternalActorsPath = ULevel::GetExternalActorsPath(FullMapName);
+			FString ExternalActorsFilePath = FPackageName::LongPackageNameToFilename(ExternalActorsPath);
 
-		// Look for duplicated actor GUIDs
-		TMap<FGuid, AActor*> ActorGuids;
+			// Look for duplicated actor GUIDs
+			TMap<FGuid, AActor*> ActorGuids;
 
-		if (IFileManager::Get().DirectoryExists(*ExternalActorsFilePath))
-		{
-			bool bResult = IFileManager::Get().IterateDirectoryRecursively(*ExternalActorsFilePath, [this, &PackagesToSave, &PackagesToDelete, &ActorGuids](const TCHAR* FilenameOrDirectory, bool bIsDirectory)
+			if (IFileManager::Get().DirectoryExists(*ExternalActorsFilePath))
 			{
-				if (!bIsDirectory)
+				bool bResult = IFileManager::Get().IterateDirectoryRecursively(*ExternalActorsFilePath, [this, &PackagesToSave, &PackagesToDelete, &ActorGuids](const TCHAR* FilenameOrDirectory, bool bIsDirectory)
 				{
-					FString Filename(FilenameOrDirectory);
-					if (Filename.EndsWith(FPackageName::GetAssetPackageExtension()))
+					if (!bIsDirectory)
 					{
-						AActor* MainPackageActor = nullptr;
-						AActor* PotentialMainPackageActor = nullptr;
-
-						const FString PackageName = FPackageName::FilenameToLongPackageName(*Filename);
-						if (UPackage* Package = LoadPackage(nullptr, *Filename, LOAD_None, nullptr, nullptr))
+						FString Filename(FilenameOrDirectory);
+						if (Filename.EndsWith(FPackageName::GetAssetPackageExtension()))
 						{
-							ForEachObjectWithPackage(Package, [&MainPackageActor, &PotentialMainPackageActor](UObject* Object)
+							AActor* MainPackageActor = nullptr;
+							AActor* PotentialMainPackageActor = nullptr;
+
+							const FString PackageName = FPackageName::FilenameToLongPackageName(*Filename);
+							if (UPackage* Package = LoadPackage(nullptr, *Filename, LOAD_None, nullptr, nullptr))
 							{
-								if (AActor* Actor = Cast<AActor>(Object))
+								ForEachObjectWithPackage(Package, [&MainPackageActor, &PotentialMainPackageActor](UObject* Object)
 								{
-									if (Actor->IsMainPackageActor())
+									if (AActor* Actor = Cast<AActor>(Object))
 									{
-										MainPackageActor = Actor;
-										PotentialMainPackageActor = nullptr;
-									}
-									else if (!MainPackageActor)
-									{
-										if (!Actor->IsChildActor())
+										if (Actor->IsMainPackageActor())
+										{
+											MainPackageActor = Actor;
+											PotentialMainPackageActor = nullptr;
+										}
+										else if (!MainPackageActor)
 										{
-											PotentialMainPackageActor = Actor;
+											if (!Actor->IsChildActor())
+											{
+												PotentialMainPackageActor = Actor;
+											}
 										}
 									}
-								}
-								return true;
-							});
-						}
-
-						if (!MainPackageActor)
-						{
-							UE_LOG(LogExternalActorsCommandlet, Error, TEXT("Missing main actor for file '%s'"), *Filename);
+									return true;
+								});
+							}
 
-							if (bRepair)
+							if (!MainPackageActor)
 							{
-								if (PotentialMainPackageActor)
+								UE_LOG(LogExternalActorsCommandlet, Error, TEXT("Missing main actor for file '%s'"), *Filename);
+
+								if (bRepair)
 								{
-									PotentialMainPackageActor->SetPackageExternal(false);
-									PotentialMainPackageActor->SetPackageExternal(true);
+									if (PotentialMainPackageActor)
+									{
+										PotentialMainPackageActor->SetPackageExternal(false);
+										PotentialMainPackageActor->SetPackageExternal(true);
 								
-									UPackage* PackageToSave = PotentialMainPackageActor->GetPackage();
-									PackagesToSave.Add(PackageToSave);
+										UPackage* PackageToSave = PotentialMainPackageActor->GetPackage();
+										PackagesToSave.Add(PackageToSave);
 
-									MainPackageActor = PotentialMainPackageActor;
-								}
+										MainPackageActor = PotentialMainPackageActor;
+									}
 
-								PackagesToDelete.Add(Filename);
+									PackagesToDelete.Add(Filename);
+								}
 							}
-						}
 
-						if (MainPackageActor)
-						{
-							if (ActorGuids.Contains(MainPackageActor->GetActorGuid()))
+							if (MainPackageActor)
 							{
-								UE_LOG(LogExternalActorsCommandlet, Error, TEXT("Duplicated actor guid for file '%s'"), *Filename);
-
-								if (bRepair)
+								if (ActorGuids.Contains(MainPackageActor->GetActorGuid()))
 								{
-									FSetActorGuid SetActorGuid(MainPackageActor, FGuid::NewGuid());
+									UE_LOG(LogExternalActorsCommandlet, Error, TEXT("Duplicated actor guid for file '%s'"), *Filename);
+
+									if (bRepair)
+									{
+										FSetActorGuid SetActorGuid(MainPackageActor, FGuid::NewGuid());
 
-									UPackage* PackageToSave = MainPackageActor->GetPackage();
-									PackagesToSave.Add(PackageToSave);
+										UPackage* PackageToSave = MainPackageActor->GetPackage();
+										PackagesToSave.Add(PackageToSave);
+									}
+								}
+								else
+								{
+									ActorGuids.Add(MainPackageActor->GetActorGuid(), MainPackageActor);
 								}
-							}
-							else
-							{
-								ActorGuids.Add(MainPackageActor->GetActorGuid(), MainPackageActor);
 							}
 						}
-					}
-					else
-					{
-						UE_LOG(LogExternalActorsCommandlet, Error, TEXT("Invalid actor file '%s'"), *Filename);
-
-						if (bRepair)
+						else
 						{
-							PackagesToDelete.Add(Filename);
+							UE_LOG(LogExternalActorsCommandlet, Error, TEXT("Invalid actor file '%s'"), *Filename);
+
+							if (bRepair)
+							{
+								PackagesToDelete.Add(Filename);
+							}
 						}
 					}
-				}
-				return true;
-			});
+					return true;
+				});
+			}
 		}
-	}
 
-	FSavePackageArgs SaveArgs;
-	SaveArgs.TopLevelFlags = RF_Standalone;
+		FSavePackageArgs SaveArgs;
+		SaveArgs.TopLevelFlags = RF_Standalone;
 
-	for (UPackage* PackageToSave : PackagesToSave)
-	{
-		const FString PackageFileName = SourceControlHelpers::PackageFilename(PackageToSave);
-		UE_LOG(LogExternalActorsCommandlet, Display, TEXT("Saving package '%s'"), *PackageFileName);
+		for (UPackage* PackageToSave : PackagesToSave)
+		{
+			const FString PackageFileName = SourceControlHelpers::PackageFilename(PackageToSave);
+			UE_LOG(LogExternalActorsCommandlet, Display, TEXT("Saving package '%s'"), *PackageFileName);
 
-		if (PackageHelper.Checkout(PackageToSave))
-		{				
-			if (UPackage::SavePackage(PackageToSave, nullptr, *PackageFileName, SaveArgs))
-			{
-				PackageHelper.AddToSourceControl(PackageToSave);
+			if (PackageHelper.Checkout(PackageToSave))
+			{				
+				if (UPackage::SavePackage(PackageToSave, nullptr, *PackageFileName, SaveArgs))
+				{
+					PackageHelper.AddToSourceControl(PackageToSave);
+				}
 			}
 		}
-	}
 
-	CollectGarbage(RF_NoFlags);
+		CollectGarbage(RF_NoFlags);
 
-	PackagesToDelete.Sort();
-	for (const FString& PackageToDelete : PackagesToDelete)
-	{
-		const FString PackageFileName = SourceControlHelpers::PackageFilename(PackageToDelete);
-		UE_LOG(LogExternalActorsCommandlet, Display, TEXT("Deleting package '%s'"), *PackageFileName);
+		PackagesToDelete.Sort();
+		for (const FString& PackageToDelete : PackagesToDelete)
+		{
+			const FString PackageFileName = SourceControlHelpers::PackageFilename(PackageToDelete);
+			UE_LOG(LogExternalActorsCommandlet, Display, TEXT("Deleting package '%s'"), *PackageFileName);
 
-		PackageHelper.Delete(*PackageToDelete);
+			PackageHelper.Delete(*PackageToDelete);
+		}
 	}
 
-	return 0;
+	return Result;
 }
