Extra cook delegates:
* Progress update, reporting cooked/pending packages
* Potential stall during saving, reporting save queue size, pending objects, expected status
* Potetnial stall during loading, reporting load queue size

#rnx

diff --git a/Engine/Source/Editor/UnrealEd/Private/CookOnTheFlyServer.cpp b/Engine/Source/Editor/UnrealEd/Private/CookOnTheFlyServer.cpp
--- a/Engine/Source/Editor/UnrealEd/Private/CookOnTheFlyServer.cpp
+++ b/Engine/Source/Editor/UnrealEd/Private/CookOnTheFlyServer.cpp
@@ -1735,10 +1735,12 @@ void UCookOnTheFlyServer::SetSaveBusy(bool bInBusy)
 				ExpectedPackages.Remove(Package);
 			}
 			
+			bool bExpectedDueToSlowBuildOperations = !ExpectedObjects.IsEmpty() && NonExpectedObjects.IsEmpty();
+			
 			FString Message = FString::Printf(TEXT("Cooker has been blocked from saving the current packages for %.0f seconds."),
 				(float)CurrentTime - SaveBusyStartTimeSeconds);
 			ELogVerbosity::Type MessageSeverity = ELogVerbosity::Display;
-			if (ExpectedObjects.IsEmpty() || !NonExpectedObjects.IsEmpty())
+			if (!bExpectedDueToSlowBuildOperations)
 			{
 				MessageSeverity = CookerIdleWarningSeverity;
 			}
@@ -1799,6 +1801,9 @@ void UCookOnTheFlyServer::SetSaveBusy(bool bInBusy)
 			}
 
 			SaveBusyWarnTimeSeconds = CurrentTime + GCookProgressWarnBusyTime;
+			
+			UE::Cook::FDelegates::CookSaveIdle.Broadcast(
+				*this, SaveQueue.Num(), PackageDatas->GetPendingCookedPlatformDataNum(), bExpectedDueToSlowBuildOperations);
 		}
 	}
 }
@@ -1853,6 +1858,8 @@ void UCookOnTheFlyServer::SetLoadBusy(bool bInLoadBusy)
 				UE_LOG(LogCook, Display, TEXT("    <None>"));
 			}
 			LoadBusyWarnTimeSeconds = CurrentTime + GCookProgressWarnBusyTime;
+			
+			UE::Cook::FDelegates::CookLoadIdle.Broadcast(*this, LoadQueue.Num());
 		}
 	}
 }
@@ -1910,6 +1917,8 @@ void UCookOnTheFlyServer::UpdateDisplay(UE::Cook::FTickStackData& StackData, boo
 			TEXT("Cooked packages %d Packages Remain %d Total %d"),
 			CookedPackagesCount, CookPendingCount, CookedPackagesCount + CookPendingCount);
 
+		UE::Cook::FDelegates::CookUpdateDisplay.Broadcast(*this, CookedPackagesCount, CookPendingCount);
+
 		LastCookedPackagesCount = CookedCountWitness;
 		LastCookPendingCount = PendingCountWitness;
 		LastProgressDisplayTime = CurrentTime;
diff --git a/Engine/Source/Runtime/CoreUObject/Private/UObject/ICookInfo.cpp b/Engine/Source/Runtime/CoreUObject/Private/UObject/ICookInfo.cpp
--- a/Engine/Source/Runtime/CoreUObject/Private/UObject/ICookInfo.cpp
+++ b/Engine/Source/Runtime/CoreUObject/Private/UObject/ICookInfo.cpp
@@ -89,6 +89,9 @@ FCookInfoEvent FDelegates::CookFinished;
 FValidateSourcePackage FDelegates::ValidateSourcePackage;
 FPackageBlockedEvent FDelegates::PackageBlocked;
 FCookInfoModifyCookDelegate FDelegates::ModifyCook;
+FCookUpdateDisplayEvent FDelegates::CookUpdateDisplay;
+FCookSaveIdleEvent FDelegates::CookSaveIdle;
+FCookLoadIdleEvent FDelegates::CookLoadIdle;
 
 const TCHAR* GetReferencedSetFilename()
 {
diff --git a/Engine/Source/Runtime/CoreUObject/Public/UObject/ICookInfo.h b/Engine/Source/Runtime/CoreUObject/Public/UObject/ICookInfo.h
--- a/Engine/Source/Runtime/CoreUObject/Public/UObject/ICookInfo.h
+++ b/Engine/Source/Runtime/CoreUObject/Public/UObject/ICookInfo.h
@@ -290,6 +290,9 @@ DECLARE_DELEGATE_RetVal_TwoParams(EDataValidationResult, FValidateSourcePackage,
 DECLARE_MULTICAST_DELEGATE_TwoParams(FPackageBlockedEvent, const UObject* /* Object */, FStringBuilderBase& OutDebugInfo)
 DECLARE_MULTICAST_DELEGATE_TwoParams(FCookInfoModifyCookDelegate, ICookInfo& /* CookInfo */,
 	TArray<FPackageCookRule>& /* InOutPackageCookRules */);
+DECLARE_MULTICAST_DELEGATE_ThreeParams(FCookUpdateDisplayEvent, ICookInfo& /* CookInfo */, int32 /* CookedPackagesCount */, int32 /* CookPendingCount */)
+DECLARE_MULTICAST_DELEGATE_FourParams(FCookSaveIdleEvent, ICookInfo& /* CookInfo */, int32 /* NumPackagesInSaveQueue */, int32 /* NumPendingCookedPlatformData */, bool /* bExpectedDueToSlowBuildOperations */)
+DECLARE_MULTICAST_DELEGATE_TwoParams(FCookLoadIdleEvent, ICookInfo& /* CookInfo */, int32 /* NumPackagesInLoadQueue */)
 
 /** UE::Cook::FDelegates: callbacks for cook events. */
 struct FDelegates
@@ -312,6 +315,12 @@ public:
 	 * (aka nevercooked), or have their reference behavior changed in other ways.
 	 */
 	static COREUOBJECT_API FCookInfoModifyCookDelegate ModifyCook;
+	/** Called during UpdateDisplay to report cooked/pending packages and controlled by the same cook.display.* CVars. for frequency */
+	static COREUOBJECT_API FCookUpdateDisplayEvent CookUpdateDisplay;
+	/** Called during SetSaveBusy when a potential soft lock is detected. */
+	static COREUOBJECT_API FCookSaveIdleEvent CookSaveIdle;
+	/** Called during SetLoadBusy when a potential soft lock is detected. */
+	static COREUOBJECT_API FCookLoadIdleEvent CookLoadIdle;
 };
 
 
