IncrementalCook: DLC: Avoid unnecessary fetch of FIncrementalCookAttachments data for some packages in requestcluster. We currently fetch data for any cookable package the first time we notice it is reachable, but that is unnecessary if the package has already been cooked (so we don't need to evaluate incrementally skippable) and we don't need to explore its dependencies (which we do not need to do for base game packages). These unnecessary fetches currently are cheap, but in an upcoming change will require a relatively expensive message to be sent to the ZenServer process.
#rnx
#incrementalcook

diff --git a/Engine/Source/Editor/UnrealEd/Private/Cooker/CookRequestCluster.cpp b/Engine/Source/Editor/UnrealEd/Private/Cooker/CookRequestCluster.cpp
--- a/Engine/Source/Editor/UnrealEd/Private/Cooker/CookRequestCluster.cpp
+++ b/Engine/Source/Editor/UnrealEd/Private/Cooker/CookRequestCluster.cpp
@@ -1247,8 +1247,8 @@ void FRequestCluster::FGraphSearch::VisitVertex(FVertexData& Vertex)
 	
 	FPackagePlatformData* CookerLoadingPlatform = nullptr;
 	const ITargetPlatform* FirstReachableSessionPlatform = nullptr;
+	bool bAnyPlatformNeedsCook = false;
 	ESuppressCookReason SuppressCookReason = ESuppressCookReason::Invalid;
-	bool bAllReachablesUncookable = true;
 	for (TPair<const ITargetPlatform*, FPackagePlatformData>& Pair :
 		PackageData.GetPlatformDatasConstKeysMutableValues())
 	{
@@ -1270,25 +1270,50 @@ void FRequestCluster::FGraphSearch::VisitVertex(FVertexData& Vertex)
 			{
 				FirstReachableSessionPlatform = TargetPlatform;
 			}
+			ESuppressCookReason SuppressCookReasonForPlatform;
 			if (!PlatformData.IsVisitedByCluster(ClusterReachability))
 			{
 				VisitVertexForPlatform(Vertex, TargetPlatform, ClusterReachability, PlatformData,
-					SuppressCookReason);
-
-				if (Cluster.TraversalExploreIncremental() ||
-					(Cluster.TraversalExploreRuntimeDependencies() && PlatformData.IsExplorable()))
-				{
-					ShouldFetchPlatforms[PlatformIndex] = true;
-					Vertex.GetPlatformData()[PlatformIndex].bExploreRequested = true;
-					// Exploration of any session platform also requires exploration of PlatformAgnosticPlatform
-					Vertex.GetPlatformData()[PlatformAgnosticPlatformIndex].bExploreRequested = true;
-				}
+					SuppressCookReasonForPlatform);
+			}
+			else
+			{
+				// Already visited. Use the existing values of IsCookable and IsExplorable.
+				SuppressCookReasonForPlatform = PlatformData.IsCookable()
+					? ESuppressCookReason::NotSuppressed : ESuppressCookReason::Invalid;
 			}
-			if (PlatformData.IsCookable())
+
+			// We cook the package for all platforms on which it is reachable and not suppressed. We mark that this
+			// unsuppressed reachability has occurred by setting SuppressCookReason=NotSuppressed for the package in
+			// Cluster.SetSuppressReason below.
+			if (SuppressCookReasonForPlatform == ESuppressCookReason::NotSuppressed ||
+				SuppressCookReason == ESuppressCookReason::NotSuppressed)
 			{
-				bAllReachablesUncookable = false;
 				SuppressCookReason = ESuppressCookReason::NotSuppressed;
 			}
+			else if (SuppressCookReason == ESuppressCookReason::Invalid)
+			{
+				SuppressCookReason = SuppressCookReasonForPlatform;
+			}
+
+			bool bPlatformNeedsCook = PlatformData.IsCookable() && !PlatformData.IsCookAttempted();
+			bAnyPlatformNeedsCook |= bPlatformNeedsCook;
+
+			// Other questions that the cluster might need to answer about the package is what runtime dependencies of
+			// the package need to be marked reachable because this package is reachable, and whether we want to mark
+			// the package as incrementally skippable (and therefore mark it already cooked without cooking it in this
+			// session). Both of those questions require expensively and asynchronously fetching dependencies about
+			// the package from the previous cook's oplog. Mark that we need to do that fetch based on whether this
+			// cluster needs to answer those questions about the package.
+			if ((Cluster.TraversalExploreIncremental() && bPlatformNeedsCook)
+				|| (Cluster.TraversalExploreRuntimeDependencies() && PlatformData.IsExplorable()
+					&& PlatformData.GetWhereCooked() == EWhereCooked::ThisSession))
+			{
+				ShouldFetchPlatforms[PlatformIndex] = true;
+				Vertex.GetPlatformData()[PlatformIndex].bExploreRequested = true;
+				// Exploration of any session platform also requires exploration of PlatformAgnosticPlatform
+				Vertex.GetPlatformData()[PlatformAgnosticPlatformIndex].bExploreRequested = true;
+			}
 		}
 	}
 
@@ -1310,10 +1335,16 @@ void FRequestCluster::FGraphSearch::VisitVertex(FVertexData& Vertex)
 
 	if (Cluster.TraversalMarkCookable())
 	{
-		bool bAnyCookable = (FirstReachableSessionPlatform == nullptr) | !bAllReachablesUncookable;
-		if (bAnyCookable != Vertex.IsAnyCookable())
+		// Set bMarkPackageAsCookable=true even if we don't know yet because it was unreachable and so we didn't run
+		// IsCookable. This is important because we initially mark every package in the cluster is cookable until it
+		// becomes reachable and we test cookability. This is important for build dependencies, which might be visited
+		// for incrementally modified calculations on a referencer package before being visited again for
+		// reachability.
+		bool bMarkPackageAsCookable = FirstReachableSessionPlatform == nullptr
+			|| SuppressCookReason == ESuppressCookReason::NotSuppressed;
+		if (bMarkPackageAsCookable != Vertex.IsAnyCookable())
 		{
-			if (!bAnyCookable)
+			if (!bMarkPackageAsCookable)
 			{
 				if (SuppressCookReason == ESuppressCookReason::Invalid)
 				{
@@ -1321,14 +1352,21 @@ void FRequestCluster::FGraphSearch::VisitVertex(FVertexData& Vertex)
 					// we don't have it stored in this->ClusterPackages, then we must have calculated it in
 					// a previous cluster, but we don't store it anywhere. Recalculate it from the
 					// FirstReachableSessionPlatform. FirstReachableSessionPlatform must be non-null, otherwise
-					// bAnyCookable would be true.
+					// bMarkPackageAsCookable would be true.
 					check(FirstReachableSessionPlatform);
 					bool bCookable;
 					bool bExplorable;
 					Cluster.IsRequestCookable(FirstReachableSessionPlatform, PackageData, SuppressCookReason, bCookable,
 						bExplorable);
-					check(!bCookable); // We don't support bCookable changing for a given package and platform
-					check(SuppressCookReason != ESuppressCookReason::Invalid);
+					if (bCookable && PackageData.FindOrAddPlatformData(FirstReachableSessionPlatform).IsCookAttempted())
+					{
+						SuppressCookReason = ESuppressCookReason::AlreadyCooked;
+					}
+					// We don't support bCookable changing for a given package and platform, so if it is suppressed now
+					// it should have either been not cookable before or should be suppressed because it was already
+					// cooked.
+					check(SuppressCookReason != ESuppressCookReason::Invalid
+						&& SuppressCookReason != ESuppressCookReason::NotSuppressed);
 				}
 			}
 			else
@@ -1336,25 +1374,32 @@ void FRequestCluster::FGraphSearch::VisitVertex(FVertexData& Vertex)
 				check(SuppressCookReason == ESuppressCookReason::NotSuppressed);
 			}
 			Cluster.SetSuppressReason(Vertex, SuppressCookReason);
-			Vertex.SetAnyCookable(bAnyCookable);
+			Vertex.SetAnyCookable(bMarkPackageAsCookable);
 		}
 
-		// If any target platform is cookable, then we need to mark the CookerLoadingPlatform as reachable because we will
-		// need to load the package to cook it
-		if (bAnyCookable)
+		// If any reachable target platform is cookable and not already cooked, then we need to mark the
+		// CookerLoadingPlatform as reachable and explore its dependencies because we will need to load the package
+		// to cook it.
+		if (bAnyPlatformNeedsCook)
 		{
 			if (!CookerLoadingPlatform)
 			{
 				CookerLoadingPlatform = &PackageData.FindOrAddPlatformData(CookerLoadingPlatformKey);
 			}
 			CookerLoadingPlatform->AddReachability(EReachability::Runtime);
-		}
-		if (CookerLoadingPlatform && CookerLoadingPlatform->IsReachable(EReachability::Runtime)
-			&& !CookerLoadingPlatform->IsVisitedByCluster(EReachability::Runtime))
-		{
-			CookerLoadingPlatform->SetCookable(true);
-			CookerLoadingPlatform->SetExplorable(true);
-			CookerLoadingPlatform->AddVisitedByCluster(EReachability::Runtime);
+			if (!CookerLoadingPlatform->IsVisitedByCluster(EReachability::Runtime))
+			{
+				// Once reachable due to a SessionPlatform being cookable and reachable, the cooker loading platform
+				// is defined as cookable and reachable.
+				CookerLoadingPlatform->SetCookable(true);
+				CookerLoadingPlatform->SetExplorable(true);
+				CookerLoadingPlatform->AddVisitedByCluster(EReachability::Runtime);
+			}
+			// Fetch and explore the hard imports of the package to mark other packages as expected load, but skip
+			// doing that if the cluster is running in a mode where it is just marking each package as cookable
+			// individually rather than searching dependencies, because we don't need to read the reachability flags
+			// packages in that case (e.g. because we are on a CookWorker being told a list of packages to cook by the
+			// director).
 			if (Cluster.TraversalExploreRuntimeDependencies())
 			{
 				ShouldFetchPlatforms[CookerLoadingPlatformIndex] = true;
@@ -1404,32 +1449,31 @@ void FRequestCluster::FGraphSearch::VisitVertex(FVertexData& Vertex)
 
 void FRequestCluster::FGraphSearch::VisitVertexForPlatform(FVertexData& Vertex, const ITargetPlatform* Platform,
 	EReachability ClusterReachability, FPackagePlatformData& PlatformData,
-	ESuppressCookReason& AccumulatedSuppressCookReason)
+	ESuppressCookReason& OutSuppressCookReason)
 {
 	if (Cluster.TraversalMarkCookable())
 	{
 		FPackageData& PackageData = *Vertex.GetPackageData();
-		ESuppressCookReason SuppressCookReason = ESuppressCookReason::Invalid;
 		bool bCookable;
 		bool bExplorable;
-		Cluster.IsRequestCookable(Platform, PackageData, SuppressCookReason,
+		Cluster.IsRequestCookable(Platform, PackageData, OutSuppressCookReason,
 			bCookable, bExplorable);
 		PlatformData.SetCookable(bCookable);
 		PlatformData.SetExplorable(bExplorable);
 		if (bCookable)
 		{
-			AccumulatedSuppressCookReason = ESuppressCookReason::NotSuppressed;
+			OutSuppressCookReason = ESuppressCookReason::NotSuppressed;
 		}
 		else
 		{
-			check(SuppressCookReason != ESuppressCookReason::Invalid
-				&& SuppressCookReason != ESuppressCookReason::NotSuppressed);
-			if (AccumulatedSuppressCookReason == ESuppressCookReason::Invalid)
-			{
-				AccumulatedSuppressCookReason = SuppressCookReason;
-			}
+			check(OutSuppressCookReason != ESuppressCookReason::Invalid
+				&& OutSuppressCookReason != ESuppressCookReason::NotSuppressed);
 		}
 	}
+	else
+	{
+		OutSuppressCookReason = ESuppressCookReason::Invalid;
+	}
 	PlatformData.AddVisitedByCluster(ClusterReachability);
 }
 
