IncrementalCook: DLC: Read BuildDependencies from BaseGame by querying the basegame's oplog from the ZenServer.
#rnx
#incrementalcook

diff --git a/Engine/Source/Developer/IoStoreUtilities/Internal/ZenStoreHttpClient.h b/Engine/Source/Developer/IoStoreUtilities/Internal/ZenStoreHttpClient.h
--- a/Engine/Source/Developer/IoStoreUtilities/Internal/ZenStoreHttpClient.h
+++ b/Engine/Source/Developer/IoStoreUtilities/Internal/ZenStoreHttpClient.h
@@ -78,6 +78,7 @@ public:
 
 	UE_API TFuture<TIoStatusOr<FCbObject>> GetProjectInfo();
 	UE_API TFuture<TIoStatusOr<FCbObject>> GetOplog(bool bTrimByReferencedset = true);
+	UE_API TFuture<TArray<TIoStatusOr<FCbObject>>> GetOplogEntries(FStringView TargetProject, FStringView TargetOplog, TArray<FString> EntryKeys);
 	UE_API TFuture<TIoStatusOr<FCbObject>> GetFiles();
 	UE_API TFuture<TIoStatusOr<FCbObject>> GetChunkInfos();
 
diff --git a/Engine/Source/Developer/IoStoreUtilities/Private/ZenStoreHttpClient.cpp b/Engine/Source/Developer/IoStoreUtilities/Private/ZenStoreHttpClient.cpp
--- a/Engine/Source/Developer/IoStoreUtilities/Private/ZenStoreHttpClient.cpp
+++ b/Engine/Source/Developer/IoStoreUtilities/Private/ZenStoreHttpClient.cpp
@@ -946,6 +946,57 @@ TFuture<TIoStatusOr<FCbObject>> FZenStoreHttpClient::GetOplog(bool bTrimByRefere
 	});
 }
 
+TFuture<TArray<TIoStatusOr<FCbObject>>> FZenStoreHttpClient::GetOplogEntries(FStringView TargetProject, FStringView TargetOplog, TArray<FString> EntryKeys)
+{
+	FString TargetOplogPath;
+	checkf(TargetProject.IsEmpty() == TargetOplog.IsEmpty(), TEXT("TargetProject and TargetOplog must be either both valid or both empty"));
+	if (TargetProject.IsEmpty())
+	{
+		TargetOplogPath = OplogPath;
+	}
+	else
+	{
+		TargetOplogPath = WriteToString<128>("/prj/", TargetProject, "/oplog/", TargetOplog);
+	}
+
+#if WITH_EDITOR
+	EAsyncExecution ThreadPool = EAsyncExecution::LargeThreadPool;
+#else
+	EAsyncExecution ThreadPool = EAsyncExecution::ThreadPool;
+#endif
+	return Async(ThreadPool, [this, TargetOplogPath = MoveTemp(TargetOplogPath), EntryKeys = MoveTemp(EntryKeys)]
+		{
+			TArray<TIoStatusOr<FCbObject>> Results;
+			Results.Reserve(EntryKeys.Num());
+
+			for (const FString& EntryKey : EntryKeys)
+			{
+				if (EntryKey.IsEmpty())
+				{
+					Results.Add(TIoStatusOr<FCbObject>(FIoStatus(EIoErrorCode::NotFound)));
+					continue;
+				}
+				UE::Zen::FZenScopedRequestPtr Request(RequestPool.Get());
+
+				TStringBuilder<128> Uri;
+				Uri << TargetOplogPath << TEXT("/entries");
+				Uri << TEXT("?") << TEXT("opkey=") << *EntryKey;
+
+				TArray64<uint8> GetBuffer;
+				if (Download(*Request, Uri, &GetBuffer, Zen::EContentType::CbObject) && Request->GetResponseCode() == 200)
+				{
+					FCbObjectView Response(GetBuffer.GetData());
+					Results.Add(TIoStatusOr<FCbObject>(FCbObject::Clone(Response)));
+				}
+				else
+				{
+					Results.Add(TIoStatusOr<FCbObject>(FIoStatus(EIoErrorCode::NotFound)));
+				}
+			}
+			return Results;
+		});
+}
+
 TFuture<TIoStatusOr<FCbObject>> FZenStoreHttpClient::GetFiles()
 {
 #if WITH_EDITOR
diff --git a/Engine/Source/Developer/IoStoreUtilities/Private/ZenStoreWriter.cpp b/Engine/Source/Developer/IoStoreUtilities/Private/ZenStoreWriter.cpp
--- a/Engine/Source/Developer/IoStoreUtilities/Private/ZenStoreWriter.cpp
+++ b/Engine/Source/Developer/IoStoreUtilities/Private/ZenStoreWriter.cpp
@@ -224,7 +224,6 @@ FZenStoreWriter::FZenStoreWriter(
 {
 	StaticInit();
 
-	FString ParentProjectId;
 	GetZenStoreProjectId(ProjectId, ParentProjectId);
 
 	if (FParse::Value(FCommandLine::Get(), TEXT("-ZenStorePlatform="), OplogId) == false)
@@ -1300,6 +1299,155 @@ void FZenStoreWriter::GetOplogAttachments(TArrayView<FName> PackageNames,
 	});
 }
 
+void FZenStoreWriter::GetBaseGameOplogAttachments(TArrayView<FName> PackageNames,
+	TArrayView<FUtf8StringView> AttachmentKeys,
+	TUniqueFunction<void(FName PackageName, FUtf8StringView AttachmentKey, FCbObject&& Attachment)>&& Callback)
+{
+	if (ParentProjectId.IsEmpty())
+	{
+		for (FName PackageName : PackageNames)
+		{
+			for (FUtf8StringView AttachmentKey : AttachmentKeys)
+			{
+				Callback(PackageName, AttachmentKey, FCbObject());
+			}
+		}
+		return;
+	}
+
+	TArray<FUtf8StringView> AttachmentKeysBuffer;
+	for (FUtf8StringView AttachmentKey : AttachmentKeys)
+	{
+		AttachmentKeysBuffer.Emplace(AttachmentKey);
+	}
+	TArray<FName> PackageNamesBuffer(PackageNames);
+	TArray<FString> EntryKeys;
+	EntryKeys.Reserve(PackageNames.Num());
+	for (FName PackageName : PackageNames)
+	{
+		EntryKeys.Add(PackageName.ToString().ToLower());
+	}
+
+	HttpClient->GetOplogEntries(ParentProjectId, OplogId, MoveTemp(EntryKeys)).Next(
+		[HttpClient = this->HttpClient.Get(), AttachmentKeys = MoveTemp(AttachmentKeysBuffer), Callback = MoveTemp(Callback),
+		PackageNames = MoveTemp(PackageNamesBuffer)]
+		(TArray<TIoStatusOr<FCbObject>> Entries) mutable
+		{
+			const int MaximumHashCount = PackageNames.Num() * AttachmentKeys.Num();
+			TArray<FIoHash> AttachmentHashes;
+			AttachmentHashes.Reserve(MaximumHashCount);
+			struct FAttachmentHashParam
+			{
+				FName PackageName;
+				FUtf8StringView AttachmentKey;
+
+				FAttachmentHashParam(const FName& InPackageName, FUtf8StringView InAttachmentKey)
+					: PackageName(InPackageName), AttachmentKey(InAttachmentKey)
+				{
+				}
+			};
+			TMultiMap<FIoHash, FAttachmentHashParam> AttachmentHashParams;
+			AttachmentHashParams.Reserve(MaximumHashCount);
+
+			TArray<FAttachmentHashParam> InvalidAttachmentHashParams;
+			InvalidAttachmentHashParams.Reserve(MaximumHashCount);
+
+			TArray<const UTF8CHAR*, TInlineAllocator<2>> AttachmentIds;
+			for (FUtf8StringView AttachmentKey : AttachmentKeys)
+			{
+				AttachmentIds.Add(UE::FZenStoreHttpClient::FindAttachmentId(AttachmentKey));
+			}
+
+			TMap<FName, TMap<const UTF8CHAR*, FIoHash>> AttachmentsForPackage;
+			for (TIoStatusOr<FCbObject>& EntryStatus : Entries)
+			{
+				if (!EntryStatus.IsOk())
+				{
+					continue;
+				}
+
+				FCbObject QueryResultObj = EntryStatus.ConsumeValueOrDie();
+				FCbObjectView OplogObj = QueryResultObj["entry"].AsObjectView();
+				if (OplogObj["packagestoreentry"])
+				{
+					FPackageStoreEntryResource Entry = FPackageStoreEntryResource::FromCbObject(OplogObj["packagestoreentry"].AsObjectView());
+					const FName PackageName = Entry.PackageName;
+
+					TMap<const UTF8CHAR*, FIoHash>& Attachments = AttachmentsForPackage.FindOrAdd(PackageName);
+
+					for (FCbFieldView Field : OplogObj)
+					{
+						FUtf8StringView FieldName = Field.GetName();
+						if (IsReservedOplogKey(FieldName))
+						{
+							continue;
+						}
+						if (Field.IsHash())
+						{
+							const UTF8CHAR* AttachmentId = UE::FZenStoreHttpClient::FindOrAddAttachmentId(FieldName);
+							Attachments.Add(AttachmentId, Field.AsHash());
+						}
+					}
+				}
+			}
+
+			for (FName PackageName : PackageNames)
+			{
+				TMap<const UTF8CHAR*, FIoHash>* Attachments = AttachmentsForPackage.Find(PackageName);
+
+				for (int32 InputAttachmentIndex = 0; InputAttachmentIndex < AttachmentKeys.Num(); ++InputAttachmentIndex)
+				{
+					FUtf8StringView AttachmentKey = AttachmentKeys[InputAttachmentIndex];
+					const UTF8CHAR* AttachmentId = AttachmentIds[InputAttachmentIndex];
+					FIoHash* AttachmentHash = (Attachments && AttachmentId) ? Attachments->Find(AttachmentId) : nullptr;
+					if (AttachmentHash)
+					{
+						AttachmentHashes.Add(*AttachmentHash);
+						AttachmentHashParams.Emplace(*AttachmentHash, FAttachmentHashParam{ PackageName, AttachmentKey });
+					}
+					else
+					{
+						InvalidAttachmentHashParams.Emplace(PackageName, AttachmentKey);
+					}
+				}
+			}
+
+			// Invoke the callback for all invalid attachment hashes
+			for (FAttachmentHashParam& InvalidAttachmentHashParam : InvalidAttachmentHashParams)
+			{
+				Callback(InvalidAttachmentHashParam.PackageName, InvalidAttachmentHashParam.AttachmentKey, FCbObject());
+			}
+
+			if (AttachmentHashes.IsEmpty())
+			{
+				return;
+			}
+
+			HttpClient->ReadChunksAsync(AttachmentHashes,
+				[Callback = MoveTemp(Callback), AttachmentHashParams = MoveTemp(AttachmentHashParams)]
+				(const FIoHash& RawHash, TIoStatusOr<FIoBuffer> Result)
+				{
+					for (auto It(AttachmentHashParams.CreateConstKeyIterator(RawHash)); It; ++It)
+					{
+						const FAttachmentHashParam& Param = It.Value();
+						if (!Result.IsOk())
+						{
+							Callback(Param.PackageName, Param.AttachmentKey, FCbObject());
+							continue;
+						}
+
+						FIoBuffer Buffer = Result.ConsumeValueOrDie();
+						if (Buffer.DataSize() == 0)
+						{
+							Callback(Param.PackageName, Param.AttachmentKey, FCbObject());
+							continue;
+						}
+						FSharedBuffer SharedBuffer = IoBufferToSharedBuffer(Buffer);
+						Callback(Param.PackageName, Param.AttachmentKey, FCbObject(SharedBuffer));
+					}
+				});
+		});
+}
 IPackageWriter::ECommitStatus FZenStoreWriter::GetCommitStatus(FName PackageName)
 {
 	FReadScopeLock _(EntriesLock);
diff --git a/Engine/Source/Developer/IoStoreUtilities/Public/ZenStoreWriter.h b/Engine/Source/Developer/IoStoreUtilities/Public/ZenStoreWriter.h
--- a/Engine/Source/Developer/IoStoreUtilities/Public/ZenStoreWriter.h
+++ b/Engine/Source/Developer/IoStoreUtilities/Public/ZenStoreWriter.h
@@ -159,6 +159,9 @@ public:
 	IOSTOREUTILITIES_API virtual void GetOplogAttachments(TArrayView<FName> PackageNames,
 		TArrayView<FUtf8StringView> AttachmentKeys,
 		TUniqueFunction<void(FName PackageName, FUtf8StringView AttachmentKey, FCbObject&& Attachment)>&& Callback) override;
+	IOSTOREUTILITIES_API virtual void GetBaseGameOplogAttachments(TArrayView<FName> PackageNames,
+		TArrayView<FUtf8StringView> AttachmentKeys,
+		TUniqueFunction<void(FName PackageName, FUtf8StringView AttachmentKey, FCbObject&& Attachment)>&& Callback) override;
 	IOSTOREUTILITIES_API virtual ECommitStatus GetCommitStatus(FName PackageName) override;
 
 	IOSTOREUTILITIES_API virtual void RemoveCookedPackages(TArrayView<const FName> PackageNamesToRemove) override;
@@ -259,6 +262,7 @@ private:
 	const ITargetPlatform&				TargetPlatform;
 	const FName							TargetPlatformFName;
 	FString								ProjectId;
+	FString								ParentProjectId;
 	FString								OplogId;
 	FString								OutputPath;
 	FString								MetadataDirectoryPath;
diff --git a/Engine/Source/Editor/UnrealEd/Private/Cooker/CookPackageArtifacts.cpp b/Engine/Source/Editor/UnrealEd/Private/Cooker/CookPackageArtifacts.cpp
--- a/Engine/Source/Editor/UnrealEd/Private/Cooker/CookPackageArtifacts.cpp
+++ b/Engine/Source/Editor/UnrealEd/Private/Cooker/CookPackageArtifacts.cpp
@@ -1301,84 +1301,160 @@ FIncrementalCookAttachments FIncrementalCookAttachments::Collect(const UPackage*
 	return Result;
 }
 
-void FIncrementalCookAttachments::Fetch(TArrayView<FName> PackageNames, const ITargetPlatform* TargetPlatform,
-	ICookedPackageWriter* PackageWriter,
-	TUniqueFunction<void(FName PackageName, FIncrementalCookAttachments&& Result)>&& Callback)
+void FIncrementalCookAttachments::Fetch(TArrayView<FPackageIncrementalCookId> PackageIds,
+	const ITargetPlatform* TargetPlatform, ICookedPackageWriter* PackageWriter,
+	TFunction<void(FName PackageName, FIncrementalCookAttachments&& Result)>&& Callback)
 {
 	using namespace UE::TargetDomain;
 
+	if (!TargetPlatform && !GEditorDomainOplog)
+	{
+		for (const FPackageIncrementalCookId& Id : PackageIds)
+		{
+			Callback(Id.PackageName, FIncrementalCookAttachments());
+		}
+		return;
+	}
+
 	struct FInProgressResult
 	{
 		FIncrementalCookAttachments Result;
+		EWhereCooked WhereCooked = EWhereCooked::ThisSession;
 		int ReceivedAttachmentCount = 0;
 	};
 	TMap<FName, FInProgressResult> InProgressResults;
-	TArray<FUtf8StringView> RequestedAttachments = { PackageArtifactsAttachmentKey, BuildDefinitionsAttachmentKey,
-		ImportsCheckerAttachmentKey, LogMessagesAttachmentKey };
+	InProgressResults.Reserve(PackageIds.Num());
 
-	if (TargetPlatform || GEditorDomainOplog)
+	TArray<FName> SessionOplogNames;
+	TArray<FName> BaseGameOplogNames;
+	TArray<FName> EditorDomainNames;
+	if (TargetPlatform)
 	{
-		auto OnOplogAttachment =
-			[PackageWriter, Callback = MoveTemp(Callback), RequestedAttachmentNum = RequestedAttachments.Num(), InProgressResults = MoveTemp(InProgressResults)]
-			(FName PackageName, FUtf8StringView AttachmentKey, FCbObject&& Attachment) mutable
+		check(PackageWriter);
+		for (const FPackageIncrementalCookId& Id : PackageIds)
+		{
+			InProgressResults.FindOrAdd(Id.PackageName).WhereCooked = Id.WhereCooked;
+			switch (Id.WhereCooked)
 			{
-				FInProgressResult& InProgressResult = InProgressResults.FindOrAdd(PackageName);
-				InProgressResult.ReceivedAttachmentCount++;
-				if (AttachmentKey == PackageArtifactsAttachmentKey)
-				{
-					if (PackageWriter)
+			case EWhereCooked::ThisSession:
+				SessionOplogNames.Add(Id.PackageName);
+				break;
+			case EWhereCooked::BaseGame:
+				BaseGameOplogNames.Add(Id.PackageName);
+				break;
+			case EWhereCooked::ExtraReleaseVersionAssets:
+				UE_CALL_ONCE([&]
 					{
-						InProgressResult.Result.CommitStatus = PackageWriter->GetCommitStatus(PackageName);
-					}
-					else
+						UE_LOG(LogCook, Warning, TEXT("CookDependency fetch was requested for package %s with EWhereCooked::ExtraReleaseVersionAssets, which is not handled. ")
+							TEXT("TransitiveBuildDepencies will be missing and false positive incremental skips could result."),
+							*Id.PackageName.ToString(), Id.WhereCooked);
+					});
+				Callback(Id.PackageName, FIncrementalCookAttachments());
+				break;
+			default:
+				UE_CALL_ONCE([&]
 					{
-						InProgressResult.Result.CommitStatus = Attachment ? IPackageWriter::ECommitStatus::Success
-							: IPackageWriter::ECommitStatus::NotCommitted;
-					}
+						UE_LOG(LogCook, Warning, TEXT("CookDependency fetch was requested for package %s with EWhereCooked value %d, which is not handled. ")
+							TEXT("TransitiveBuildDepencies will be missing and false positive incremental skips could result."),
+							*Id.PackageName.ToString(), Id.WhereCooked);
+					});
+				Callback(Id.PackageName, FIncrementalCookAttachments());
+				break;
+			}
+		}
+	}
+	else
+	{
+		for (const FPackageIncrementalCookId& Id : PackageIds)
+		{
+			InProgressResults.FindOrAdd(Id.PackageName).WhereCooked = Id.WhereCooked;
+			EditorDomainNames.Add(Id.PackageName);
+		}
+	}
 
-					if (LoadFromCompactBinary(MoveTemp(Attachment), InProgressResult.Result.Artifacts))
-					{
-						InProgressResult.Result.Artifacts.PackageName = PackageName;
-					}
-				}
-				else if (AttachmentKey == BuildDefinitionsAttachmentKey)
-				{
-					LoadFromCompactBinary(MoveTemp(Attachment), InProgressResult.Result.BuildDefinitions);
-				}
-				else if (AttachmentKey == ImportsCheckerAttachmentKey)
+	FUtf8StringView AttachmentsForIncrementalSkipAndBuildDependenciesArray[] =
+	{ PackageArtifactsAttachmentKey, ImportsCheckerAttachmentKey, BuildDefinitionsAttachmentKey,
+		LogMessagesAttachmentKey };
+	TArrayView<FUtf8StringView> AttachmentsForIncrementalSkipAndBuildDependencies(
+		AttachmentsForIncrementalSkipAndBuildDependenciesArray);
+	FUtf8StringView AttachmentsForBuildDependenciesArray[] =
+	{ PackageArtifactsAttachmentKey, ImportsCheckerAttachmentKey };
+	TArrayView<FUtf8StringView> AttachmentsForBuildDependencies(AttachmentsForBuildDependenciesArray);
+
+	TFunction<void(FName , FUtf8StringView, FCbObject&&)> OnOplogAttachment =
+		[PackageWriter, Callback = MoveTemp(Callback), AttachmentsForIncrementalSkipAndBuildDependenciesNum = AttachmentsForIncrementalSkipAndBuildDependencies.Num(),
+		AttachmentsForBuildDependenciesNum = AttachmentsForBuildDependencies.Num(),
+		InProgressResults = MoveTemp(InProgressResults)]
+		(FName PackageName, FUtf8StringView AttachmentKey, FCbObject&& Attachment) mutable
+		{
+			FInProgressResult& InProgressResult = InProgressResults.FindOrAdd(PackageName);
+			InProgressResult.ReceivedAttachmentCount++;
+			if (AttachmentKey == PackageArtifactsAttachmentKey)
+			{
+				if (PackageWriter && InProgressResult.WhereCooked == EWhereCooked::ThisSession)
 				{
-					LoadFromCompactBinary(Attachment.AsFieldView(), InProgressResult.Result.ImportsCheckerData);
+					InProgressResult.Result.CommitStatus = PackageWriter->GetCommitStatus(PackageName);
 				}
-				else if (AttachmentKey == LogMessagesAttachmentKey)
+				else
 				{
-					FLogMessagesArray LogMessagesArray(InProgressResult.Result.LogMessages);
-					LoadFromCompactBinary(Attachment.AsFieldView(), LogMessagesArray);
+					InProgressResult.Result.CommitStatus = Attachment ? IPackageWriter::ECommitStatus::Success
+						: IPackageWriter::ECommitStatus::NotCommitted;
 				}
 
-				if (InProgressResult.ReceivedAttachmentCount == RequestedAttachmentNum)
+				if (LoadFromCompactBinary(MoveTemp(Attachment), InProgressResult.Result.Artifacts))
 				{
-					Callback(PackageName, MoveTemp(InProgressResult.Result));
-					InProgressResults.Remove(PackageName);
+					InProgressResult.Result.Artifacts.PackageName = PackageName;
 				}
-			};
+			}
+			else if (AttachmentKey == BuildDefinitionsAttachmentKey)
+			{
+				LoadFromCompactBinary(MoveTemp(Attachment), InProgressResult.Result.BuildDefinitions);
+			}
+			else if (AttachmentKey == ImportsCheckerAttachmentKey)
+			{
+				LoadFromCompactBinary(Attachment.AsFieldView(), InProgressResult.Result.ImportsCheckerData);
+			}
+			else if (AttachmentKey == LogMessagesAttachmentKey)
+			{
+				FLogMessagesArray LogMessagesArray(InProgressResult.Result.LogMessages);
+				LoadFromCompactBinary(Attachment.AsFieldView(), LogMessagesArray);
+			}
 
-		if (PackageWriter)
-		{
-			PackageWriter->GetOplogAttachments(PackageNames,
-				RequestedAttachments, MoveTemp(OnOplogAttachment));
-		}
-		else
-		{
-			GEditorDomainOplog->GetOplogAttachments(PackageNames,
-				RequestedAttachments, MoveTemp(OnOplogAttachment));
-		}
+			int32 RequestedAttachmentNum = InProgressResult.WhereCooked == EWhereCooked::ThisSession
+				? AttachmentsForIncrementalSkipAndBuildDependenciesNum : AttachmentsForBuildDependenciesNum;
+			if (InProgressResult.ReceivedAttachmentCount == RequestedAttachmentNum)
+			{
+				Callback(PackageName, MoveTemp(InProgressResult.Result));
+			}
+		};
+	
+
+	if (!SessionOplogNames.IsEmpty())
+	{
+		auto UniqueCallback = [OnOplogAttachment](FName PackageName, FUtf8StringView AttachmentKey, FCbObject&& Attachment)
+			{
+				OnOplogAttachment(PackageName, AttachmentKey, MoveTemp(Attachment));
+			};
+		PackageWriter->GetOplogAttachments(SessionOplogNames,
+			AttachmentsForIncrementalSkipAndBuildDependencies, MoveTemp(UniqueCallback));
 	}
-	else
+	if (!BaseGameOplogNames.IsEmpty())
 	{
-		for (FName PackageName : PackageNames)
-		{
-			Callback(PackageName, FIncrementalCookAttachments());
-		}
+		auto UniqueCallback = [OnOplogAttachment](FName PackageName, FUtf8StringView AttachmentKey, FCbObject&& Attachment)
+			{
+				OnOplogAttachment(PackageName, AttachmentKey, MoveTemp(Attachment));
+			};
+		PackageWriter->GetBaseGameOplogAttachments(BaseGameOplogNames, AttachmentsForBuildDependencies,
+			MoveTemp(UniqueCallback));
+	}
+	if (!EditorDomainNames.IsEmpty())
+	{
+		auto UniqueCallback = [OnOplogAttachment](FName PackageName, FUtf8StringView AttachmentKey, FCbObject&& Attachment)
+			{
+				OnOplogAttachment(PackageName, AttachmentKey, MoveTemp(Attachment));
+			};
+		GEditorDomainOplog->GetOplogAttachments(EditorDomainNames,
+			AttachmentsForIncrementalSkipAndBuildDependencies, MoveTemp(UniqueCallback));
 	}
 }
 
diff --git a/Engine/Source/Editor/UnrealEd/Private/Cooker/CookPackageArtifacts.h b/Engine/Source/Editor/UnrealEd/Private/Cooker/CookPackageArtifacts.h
--- a/Engine/Source/Editor/UnrealEd/Private/Cooker/CookPackageArtifacts.h
+++ b/Engine/Source/Editor/UnrealEd/Private/Cooker/CookPackageArtifacts.h
@@ -7,6 +7,7 @@
 #include "Cooker/CookDependency.h"
 #include "Cooker/CookImportsChecker.h"
 #include "Cooker/CookLogPrivate.h"
+#include "Cooker/CookTypes.h"
 #include "Containers/Array.h"
 #include "Containers/ArrayView.h"
 #include "Containers/Map.h"
@@ -357,9 +358,9 @@ struct FIncrementalCookAttachments
 	void Empty();
 	void AppendCommitAttachments(TArray<IPackageWriter::FCommitAttachmentInfo>& OutAttachments);
 
-	static void Fetch(TArrayView<FName> PackageNames, const ITargetPlatform* TargetPlatform,
+	static void Fetch(TArrayView<FPackageIncrementalCookId> PackageIds, const ITargetPlatform* TargetPlatform,
 		ICookedPackageWriter* PackageWriter,
-		TUniqueFunction<void(FName PackageName, FIncrementalCookAttachments&& Result)>&& Callback);
+		TFunction<void(FName PackageName, FIncrementalCookAttachments&& Result)>&& Callback);
 	static FIncrementalCookAttachments Collect(const UPackage* Package,
 		const ITargetPlatform* TargetPlatform, FBuildResultDependenciesMap&& InResultDependencies,
 		bool bHasSaveResult, TConstArrayView<FName> UntrackedSoftPackageReferences, FGenerationHelper* GenerationHelper,
diff --git a/Engine/Source/Editor/UnrealEd/Private/Cooker/CookRequestCluster.cpp b/Engine/Source/Editor/UnrealEd/Private/Cooker/CookRequestCluster.cpp
--- a/Engine/Source/Editor/UnrealEd/Private/Cooker/CookRequestCluster.cpp
+++ b/Engine/Source/Editor/UnrealEd/Private/Cooker/CookRequestCluster.cpp
@@ -2590,7 +2590,7 @@ void FRequestCluster::FQueryVertexBatch::Reset()
 {
 	for (FPlatformData& PlatformData : PlatformDatas)
 	{
-		PlatformData.PackageNames.Reset();
+		PlatformData.PackageIds.Reset();
 	}
 	Vertices.Reset();
 }
@@ -2614,7 +2614,10 @@ void FRequestCluster::FQueryVertexBatch::Send()
 				if (PlatformData.CompareExchangeAsyncQueryStatus(ExpectedStatus,
 					EAsyncQueryStatus::AsyncRequested))
 				{
-					PlatformDatas[PlatformIndex].PackageNames.Add(Pair.Key);
+					FFetchPlatformData& FetchPlatformData = ThreadSafeOnlyVars.FetchPlatforms[PlatformIndex];
+					FPackagePlatformData& PackagePlatformData = Vertex->GetPackageData()->FindOrAddPlatformData(FetchPlatformData.Platform);
+
+					PlatformDatas[PlatformIndex].PackageIds.Add(FPackageIncrementalCookId{ Pair.Key, PackagePlatformData.GetWhereCooked() });
 					++NumAddedRequests;
 				}
 			}
@@ -2642,7 +2645,7 @@ void FRequestCluster::FQueryVertexBatch::Send()
 	for (int32 PlatformIndex = 0; PlatformIndex < PlatformDatas.Num(); ++PlatformIndex)
 	{
 		FPlatformData& PlatformData = PlatformDatas[PlatformIndex];
-		if (PlatformData.PackageNames.IsEmpty())
+		if (PlatformData.PackageIds.IsEmpty())
 		{
 			continue;
 		}
@@ -2656,12 +2659,12 @@ void FRequestCluster::FQueryVertexBatch::Send()
 															// always use empty
 			)
 		{
-			TUniqueFunction<void(FName PackageName, FIncrementalCookAttachments&& Result)> Callback =
+			TFunction<void(FName PackageName, FIncrementalCookAttachments&& Result)> Callback =
 				[this, PlatformIndex](FName PackageName, FIncrementalCookAttachments&& Attachments)
 			{
 				RecordCacheResults(PackageName, PlatformIndex, MoveTemp(Attachments));
 			};
-			FIncrementalCookAttachments::Fetch(PlatformData.PackageNames, FetchPlatformData.Platform,
+			FIncrementalCookAttachments::Fetch(PlatformData.PackageIds, FetchPlatformData.Platform,
 				FetchPlatformData.Writer, MoveTemp(Callback));
 		}
 		else
@@ -2671,13 +2674,13 @@ void FRequestCluster::FQueryVertexBatch::Send()
 
 			// Don't use a ranged-for, as we are not allowed to access this or this->PackageNames after the
 			// last index, and ranged-for != at the end of the final loop iteration can read from PackageNames
-			int32 NumPackageNames = PlatformData.PackageNames.Num();
-			FName* PackageNamesData = PlatformData.PackageNames.GetData();
-			for (int32 PackageNameIndex = 0; PackageNameIndex < NumPackageNames; ++PackageNameIndex)
+			int32 NumPackages = PlatformData.PackageIds.Num();
+			FPackageIncrementalCookId* PackageIdsData = PlatformData.PackageIds.GetData();
+			for (int32 PackageNameIndex = 0; PackageNameIndex < NumPackages; ++PackageNameIndex)
 			{
-				FName PackageName = PackageNamesData[PackageNameIndex];
+				const FPackageIncrementalCookId& PackageId = PackageIdsData[PackageNameIndex];
 				FIncrementalCookAttachments Attachments;
-				RecordCacheResults(PackageName, PlatformIndex, MoveTemp(Attachments));
+				RecordCacheResults(PackageId.PackageName, PlatformIndex, MoveTemp(Attachments));
 			}
 		}
 	}
diff --git a/Engine/Source/Editor/UnrealEd/Private/Cooker/CookRequestCluster.h b/Engine/Source/Editor/UnrealEd/Private/Cooker/CookRequestCluster.h
--- a/Engine/Source/Editor/UnrealEd/Private/Cooker/CookRequestCluster.h
+++ b/Engine/Source/Editor/UnrealEd/Private/Cooker/CookRequestCluster.h
@@ -263,7 +263,7 @@ private:
 
 		struct FPlatformData
 		{
-			TArray<FName> PackageNames;
+			TArray<FPackageIncrementalCookId> PackageIds;
 		};
 
 		TArray<FPlatformData> PlatformDatas;
@@ -438,7 +438,7 @@ private:
 		/** Calculate and store the vertex's PackageData's cookability for the platform. */
 		void VisitVertexForPlatform(FVertexData& Vertex, const ITargetPlatform* Platform,
 			EReachability ClusterReachability, FPackagePlatformData& PlatformData,
-			ESuppressCookReason& AccumulatedSuppressCookReason);
+			ESuppressCookReason& OutSuppressCookReason);
 		void ResolveTransitiveBuildDependencyCycle();
 
 		/** Queue a vertex for visiting and dependency traversal */
diff --git a/Engine/Source/Editor/UnrealEd/Private/Cooker/CookTypes.h b/Engine/Source/Editor/UnrealEd/Private/Cooker/CookTypes.h
--- a/Engine/Source/Editor/UnrealEd/Private/Cooker/CookTypes.h
+++ b/Engine/Source/Editor/UnrealEd/Private/Cooker/CookTypes.h
@@ -412,6 +412,16 @@ namespace UE::Cook
 		NumBits = FPlatformMath::ConstExprCeilLogTwo(EWhereCooked::Count),
 	};
 
+	/**
+	 * Data identifiying a package for fetching Metadata about the package recorded in an earlier cook
+	 * (e.g. FIncrementalCookAttachments).
+	 */
+	struct FPackageIncrementalCookId
+	{
+		FName PackageName;
+		EWhereCooked WhereCooked;
+	};
+
 	/* Thread Local Storage access to identify which thread is the SchedulerThread for cooking. */
 	void InitializeTls();
 	bool IsSchedulerThread();
diff --git a/Engine/Source/Editor/UnrealEd/Private/Cooker/DiffPackageWriter.h b/Engine/Source/Editor/UnrealEd/Private/Cooker/DiffPackageWriter.h
--- a/Engine/Source/Editor/UnrealEd/Private/Cooker/DiffPackageWriter.h
+++ b/Engine/Source/Editor/UnrealEd/Private/Cooker/DiffPackageWriter.h
@@ -101,6 +101,13 @@ public:
 	{
 		return Inner->GetOplogAttachments(PackageNames, AttachmentKeys, MoveTemp(Callback));
 	}
+	virtual void GetBaseGameOplogAttachments(TArrayView<FName> PackageNames,
+		TArrayView<FUtf8StringView> AttachmentKeys,
+		TUniqueFunction<void(FName PackageName, FUtf8StringView AttachmentKey, FCbObject&& Attachment)>&& Callback) override
+	{
+		return Inner->GetOplogAttachments(PackageNames, AttachmentKeys, MoveTemp(Callback));
+	}
+
 	virtual ECommitStatus GetCommitStatus(FName PackageName) override
 	{
 		return Inner->GetCommitStatus(PackageName);
diff --git a/Engine/Source/Editor/UnrealEd/Private/Cooker/LooseCookedPackageWriter.cpp b/Engine/Source/Editor/UnrealEd/Private/Cooker/LooseCookedPackageWriter.cpp
--- a/Engine/Source/Editor/UnrealEd/Private/Cooker/LooseCookedPackageWriter.cpp
+++ b/Engine/Source/Editor/UnrealEd/Private/Cooker/LooseCookedPackageWriter.cpp
@@ -459,6 +459,20 @@ void FLooseCookedPackageWriter::GetOplogAttachments(TArrayView<FName> PackageNam
 	}
 }
 
+void FLooseCookedPackageWriter::GetBaseGameOplogAttachments(TArrayView<FName> PackageNames,
+	TArrayView<FUtf8StringView> AttachmentKeys,
+	TUniqueFunction<void(FName PackageName, FUtf8StringView AttachmentKey, FCbObject&& Attachment)>&& Callback)
+{
+	/** Oplog attachments are not implemented by FLooseCookedPackageWriter */
+	for (FName PackageName : PackageNames)
+	{
+		for (FUtf8StringView AttachmentKey : AttachmentKeys)
+		{
+			Callback(PackageName, AttachmentKey, FCbObject());
+		}
+	}
+}
+
 IPackageWriter::ECommitStatus FLooseCookedPackageWriter::GetCommitStatus(FName PackageName)
 {
 	// GetCommitStatus is not implemented by FLooseCookedPackageWriter
diff --git a/Engine/Source/Editor/UnrealEd/Private/Cooker/LooseCookedPackageWriter.h b/Engine/Source/Editor/UnrealEd/Private/Cooker/LooseCookedPackageWriter.h
--- a/Engine/Source/Editor/UnrealEd/Private/Cooker/LooseCookedPackageWriter.h
+++ b/Engine/Source/Editor/UnrealEd/Private/Cooker/LooseCookedPackageWriter.h
@@ -61,6 +61,9 @@ public:
 	virtual void GetOplogAttachments(TArrayView<FName> PackageNames,
 		TArrayView<FUtf8StringView> AttachmentKeys,
 		TUniqueFunction<void(FName PackageName, FUtf8StringView AttachmentKey, FCbObject&& Attachment)>&& Callback) override;
+	virtual void GetBaseGameOplogAttachments(TArrayView<FName> PackageNames,
+		TArrayView<FUtf8StringView> AttachmentKeys,
+		TUniqueFunction<void(FName PackageName, FUtf8StringView AttachmentKey, FCbObject&& Attachment)>&& Callback) override;
 	virtual ECommitStatus GetCommitStatus(FName PackageName) override;
 	virtual void RemoveCookedPackages(TArrayView<const FName> PackageNamesToRemove) override;
 	virtual void RemoveCookedPackages() override;
diff --git a/Engine/Source/Runtime/Core/Public/Serialization/PackageWriter.h b/Engine/Source/Runtime/Core/Public/Serialization/PackageWriter.h
--- a/Engine/Source/Runtime/Core/Public/Serialization/PackageWriter.h
+++ b/Engine/Source/Runtime/Core/Public/Serialization/PackageWriter.h
@@ -424,6 +424,16 @@ public:
 		TArrayView<FUtf8StringView> AttachmentKeys,
 		TUniqueFunction<void(FName PackageName, FUtf8StringView AttachmentKey, FCbObject&& Attachment)>&& Callback) = 0;
 
+	/**
+	 * Read Oplog attachments for the given PackageNames in the BaseGame oplog for the current DLC cook. Returns empty objects
+	 * if current cook is not a DLC cook or the basegame oplog is unavailable.
+	 * Callback is called for every cross product pair of PackageNames X AttachmentKeys. The Attachment object
+	 * passed to the callback is empty if not found.
+	 */
+	virtual void GetBaseGameOplogAttachments(TArrayView<FName> PackageNames,
+		TArrayView<FUtf8StringView> AttachmentKeys,
+		TUniqueFunction<void(FName PackageName, FUtf8StringView AttachmentKey, FCbObject&& Attachment)>&& Callback) = 0;
+
 	/**
 	 * Report commit status for the given package. If cooking incrementally, this reports the status
 	 * from a previous cook if it has not yet been committed in the current cook.
