Updated the in-editor localized asset tooling to fully support Verse identifiers

This has been supported in the cooker and asset loader for several months, but the tooling was never updated to match.

The "Asset Localization" menu can now handle both UE and Verse style on-disk culture identifiers, and new localized assets will now use the preferred on-disk form for the project (see FCulture::GetPreferredOnDiskIdentifierType and FCulture::GetPreferredOnDiskIdentifier).

The preferred on-disk form is still UE-style by default, but can be overridden for your project via any Engine INI file, eg)
[Internationalization]
PreferredOnDiskCultureIdentifierType=Verse

The dialogue import tooling has also been updated to respect the preferred on-disk form. When changing your preferred form you should ideally manually migrate your localized content to the new form for consistency.

#redirectsaresafe

diff --git a/Engine/Plugins/Editor/ContentBrowser/ContentBrowserAssetDataSource/Source/ContentBrowserAssetDataSource/Private/AssetFileContextMenu.cpp b/Engine/Plugins/Editor/ContentBrowser/ContentBrowserAssetDataSource/Source/ContentBrowserAssetDataSource/Private/AssetFileContextMenu.cpp
--- a/Engine/Plugins/Editor/ContentBrowser/ContentBrowserAssetDataSource/Source/ContentBrowserAssetDataSource/Private/AssetFileContextMenu.cpp
+++ b/Engine/Plugins/Editor/ContentBrowser/ContentBrowserAssetDataSource/Source/ContentBrowserAssetDataSource/Private/AssetFileContextMenu.cpp
@@ -793,17 +793,22 @@ void FAssetFileContextMenu::MakeAssetLocalizationSubMenu(UToolMenu* Menu)
 				if (FPackageLocalizationUtil::ConvertSourceToLocalized(ObjectPath, CurrentCulture->GetName(), LocalizedObjectPath))
 				{
 					// Does this localized asset already exist?
-					FAssetRegistryModule& AssetRegistryModule = FModuleManager::LoadModuleChecked<FAssetRegistryModule>(TEXT("AssetRegistry"));
-					FAssetData LocalizedAssetData = AssetRegistryModule.Get().GetAssetByObjectPath(FSoftObjectPath(LocalizedObjectPath));
+					// Need to check for both UE and Verse style identifiers, as both may be in use
+					FAssetData LocalizedAssetData = IAssetRegistry::GetChecked().GetAssetByObjectPath(FSoftObjectPath(LocalizedObjectPath));
+					if (!LocalizedAssetData.IsValid())
+					{
+						if (const FString VerseIdentifier = CurrentCulture->GetVerseIdentifier();
+							VerseIdentifier != CurrentCulture->GetName())
+						{
+							verify(FPackageLocalizationUtil::ConvertSourceToLocalized(ObjectPath, VerseIdentifier, LocalizedObjectPath));
+							LocalizedAssetData = IAssetRegistry::GetChecked().GetAssetByObjectPath(FSoftObjectPath(LocalizedObjectPath));
+						}
+					}
 
 					if (LocalizedAssetData.IsValid())
 					{
 						LocalizedAssetsStateForCulture.CurrentAssets.Add(FSoftObjectPath(LocalizedObjectPath));
 					}
-					else
-					{
-						LocalizedAssetsStateForCulture.NewAssets.Add(FSoftObjectPath(LocalizedObjectPath));
-					}
 				}
 			}
 		}
@@ -1016,8 +1021,9 @@ void FAssetFileContextMenu::ExecuteCreateLocalizedAsset(TSet<FSoftObjectPath> In
 			continue;
 		}
 
+		// Always create new assets using the preferred on-disk culture identifier for this project
 		FString LocalizedPackageName;
-		if (!FPackageLocalizationUtil::ConvertSourceToLocalized(SourceAssetObject->GetOutermost()->GetPathName(), InLocalizedAssetsStateForCulture.Culture->GetName(), LocalizedPackageName))
+		if (!FPackageLocalizationUtil::ConvertSourceToLocalized(SourceAssetObject->GetOutermost()->GetPathName(), InLocalizedAssetsStateForCulture.Culture->GetPreferredOnDiskIdentifier(), LocalizedPackageName))
 		{
 			continue;
 		}
diff --git a/Engine/Plugins/Editor/ContentBrowser/ContentBrowserAssetDataSource/Source/ContentBrowserAssetDataSource/Public/AssetFileContextMenu.h b/Engine/Plugins/Editor/ContentBrowser/ContentBrowserAssetDataSource/Source/ContentBrowserAssetDataSource/Public/AssetFileContextMenu.h
--- a/Engine/Plugins/Editor/ContentBrowser/ContentBrowserAssetDataSource/Source/ContentBrowserAssetDataSource/Public/AssetFileContextMenu.h
+++ b/Engine/Plugins/Editor/ContentBrowser/ContentBrowserAssetDataSource/Source/ContentBrowserAssetDataSource/Public/AssetFileContextMenu.h
@@ -35,7 +35,6 @@ private:
 	struct FLocalizedAssetsState
 	{
 		FCulturePtr Culture;
-		TSet<FSoftObjectPath> NewAssets;
 		TSet<FSoftObjectPath> CurrentAssets;
 	};
 
diff --git a/Engine/Source/Editor/LocalizationCommandletExecution/Private/LocalizationCommandletTasks.cpp b/Engine/Source/Editor/LocalizationCommandletExecution/Private/LocalizationCommandletTasks.cpp
--- a/Engine/Source/Editor/LocalizationCommandletExecution/Private/LocalizationCommandletTasks.cpp
+++ b/Engine/Source/Editor/LocalizationCommandletExecution/Private/LocalizationCommandletTasks.cpp
@@ -573,8 +573,20 @@ bool LocalizationCommandletTasks::ReportLoadedAudioAssets(const TArray<ULocaliza
 					continue;
 				}
 
-				DialogueWavePathsToTest.Add(RootAssetPath / TEXT("L10N") / Culture);
-				SoundWavePathsToTest.Add(RootAssetPath / TEXT("L10N") / Culture / Target->Settings.ImportDialogueSettings.ImportedDialogueFolder);
+				FString PreferredCultureIdentifier = Culture;
+				if (FCulture::GetPreferredOnDiskIdentifierType() == ECultureIdentifierType::Verse)
+				{
+					PreferredCultureIdentifier = FCulture::CultureNameToVerseIdentifier(Culture);
+				}
+
+				// Need to check for both UE and Verse style identifiers, as both may be in use
+				DialogueWavePathsToTest.Add(RootAssetPath / TEXT("L10N") / PreferredCultureIdentifier);
+				SoundWavePathsToTest.Add(RootAssetPath / TEXT("L10N") / PreferredCultureIdentifier / Target->Settings.ImportDialogueSettings.ImportedDialogueFolder);
+				if (PreferredCultureIdentifier != Culture)
+				{
+					DialogueWavePathsToTest.Add(RootAssetPath / TEXT("L10N") / Culture);
+					SoundWavePathsToTest.Add(RootAssetPath / TEXT("L10N") / Culture / Target->Settings.ImportDialogueSettings.ImportedDialogueFolder);
+				}
 			}
 		}
 
diff --git a/Engine/Source/Editor/UnrealEd/Private/Commandlets/ImportLocalizedDialogueCommandlet.cpp b/Engine/Source/Editor/UnrealEd/Private/Commandlets/ImportLocalizedDialogueCommandlet.cpp
--- a/Engine/Source/Editor/UnrealEd/Private/Commandlets/ImportLocalizedDialogueCommandlet.cpp
+++ b/Engine/Source/Editor/UnrealEd/Private/Commandlets/ImportLocalizedDialogueCommandlet.cpp
@@ -1,21 +1,23 @@
 // Copyright Epic Games, Inc. All Rights Reserved.
 
 #include "Commandlets/ImportLocalizedDialogueCommandlet.h"
-#include "Modules/ModuleManager.h"
-#include "Misc/PackageName.h"
+
 #include "AssetRegistry/AssetData.h"
-#include "Sound/SoundWave.h"
-#include "Misc/Paths.h"
-#include "Misc/App.h"
-#include "UObject/MetaData.h"
-#include "EditorFramework/AssetImportData.h"
-#include "Sound/DialogueWave.h"
-#include "Utils.h"
 #include "AssetRegistry/AssetRegistryModule.h"
+#include "AudioCompressionSettingsUtils.h"
+#include "AudioEditorModule.h"
+#include "EditorFramework/AssetImportData.h"
 #include "Interfaces/ITargetPlatform.h"
 #include "Interfaces/ITargetPlatformManagerModule.h"
-#include "AudioEditorModule.h"
-#include "AudioCompressionSettingsUtils.h"
+#include "Internationalization/Culture.h"
+#include "Misc/App.h"
+#include "Misc/PackageName.h"
+#include "Misc/Paths.h"
+#include "Modules/ModuleManager.h"
+#include "Sound/DialogueWave.h"
+#include "Sound/SoundWave.h"
+#include "UObject/MetaData.h"
+#include "Utils.h"
 
 #include UE_INLINE_GENERATED_CPP_BY_NAME(ImportLocalizedDialogueCommandlet)
 
@@ -165,37 +167,46 @@ int32 UImportLocalizedDialogueCommandlet::Main(const FString& Params)
 		return -1;
 	}
 
+	// Paths to scan for existing localized dialogue and sound waves
+	TArray<FName> LocalizedDialogueWavePathsToSearch;
+	TArray<FName> LocalizedSoundWavePathsToSearch;
+
 	// Build up the culture specific import info
 	TMap<FString, FCultureImportInfo> CultureImportInfoMap;
 	for (const FString& CultureName : CulturesToGenerate)
 	{
+		FString PreferredCultureIdentifier = CultureName;
+		if (FCulture::GetPreferredOnDiskIdentifierType() == ECultureIdentifierType::Verse)
+		{
+			PreferredCultureIdentifier = FCulture::CultureNameToVerseIdentifier(CultureName);
+		}
+
 		FCultureImportInfo& CultureImportInfo = CultureImportInfoMap.Add(CultureName);
 		CultureImportInfo.Name = CultureName;
 		CultureImportInfo.AudioPath = RawAudioPath / CultureName;
 		CultureImportInfo.ArchiveFileName = SourcePath / CultureName / ArchiveName;
-		CultureImportInfo.LocalizedRootContentPath = RootContentDir / TEXT("L10N") / CultureName;
-		CultureImportInfo.LocalizedRootPackagePath = RootAssetPath / TEXT("L10N") / CultureName;
+		CultureImportInfo.LocalizedRootContentPath = RootContentDir / TEXT("L10N") / PreferredCultureIdentifier;
+		CultureImportInfo.LocalizedRootPackagePath = RootAssetPath / TEXT("L10N") / PreferredCultureIdentifier;
 		CultureImportInfo.LocalizedImportedDialoguePackagePath = CultureImportInfo.LocalizedRootPackagePath / ImportedDialogueFolder;
 		CultureImportInfo.bIsNativeCulture = CultureName == NativeCulture;
+
+		// Need to check for both UE and Verse style identifiers, as both may be in use
+		LocalizedDialogueWavePathsToSearch.Add(*CultureImportInfo.LocalizedRootPackagePath);
+		LocalizedSoundWavePathsToSearch.Add(*CultureImportInfo.LocalizedImportedDialoguePackagePath);
+		if (PreferredCultureIdentifier != CultureName)
+		{
+			LocalizedDialogueWavePathsToSearch.Add(*(RootAssetPath / TEXT("L10N") / CultureName));
+			LocalizedSoundWavePathsToSearch.Add(*(RootAssetPath / TEXT("L10N") / CultureName / ImportedDialogueFolder));
+		}
 	}
 
 	// Find all of the existing localized dialogue and sound waves - we'll keep track of which ones we process so we can delete any that are no longer needed
 	TArray<FAssetData> LocalizedAssetsToPotentiallyDelete;
 	{
-		TArray<FName> LocalizedDialogueWavePathsToSearch;
-		TArray<FName> LocalizedSoundWavePathsToSearch;
-
 		// We always add the source imported dialogue folder to ensure we clean it up correctly if we change the "import native as source" option
 		// This is also why we always add the native culture, even though only one will be in use at any one time
 		LocalizedSoundWavePathsToSearch.Add(*(RootAssetPath / ImportedDialogueFolder));
 
-		for (const auto& CultureImportInfoPair : CultureImportInfoMap)
-		{
-			const FCultureImportInfo& CultureImportInfo = CultureImportInfoPair.Value;
-			LocalizedDialogueWavePathsToSearch.Add(*CultureImportInfo.LocalizedRootPackagePath);
-			LocalizedSoundWavePathsToSearch.Add(*CultureImportInfo.LocalizedImportedDialoguePackagePath);
-		}
-
 		FLocalizedAssetUtil::GetAssetsByPathAndClass(AssetRegistry, LocalizedDialogueWavePathsToSearch, UDialogueWave::StaticClass()->GetClassPathName(), /*bIncludeLocalizedAssets*/true, LocalizedAssetsToPotentiallyDelete);
 		FLocalizedAssetUtil::GetAssetsByPathAndClass(AssetRegistry, LocalizedSoundWavePathsToSearch, USoundWave::StaticClass()->GetClassPathName(), /*bIncludeLocalizedAssets*/true, LocalizedAssetsToPotentiallyDelete);
 	}
diff --git a/Engine/Source/Runtime/Core/Private/Internationalization/Culture.cpp b/Engine/Source/Runtime/Core/Private/Internationalization/Culture.cpp
--- a/Engine/Source/Runtime/Core/Private/Internationalization/Culture.cpp
+++ b/Engine/Source/Runtime/Core/Private/Internationalization/Culture.cpp
@@ -219,6 +219,33 @@ FString FCulture::CultureNameFromVerseIdentifier(const FString& VerseIdentifer)
 	return Result;
 }
 
+ECultureIdentifierType FCulture::GetPreferredOnDiskIdentifierType()
+{
+	checkf(GConfig && GConfig->IsReadyForUse(), TEXT("FCulture::GetPreferredOnDiskIdentifierType called before the config system was initialized!"));
+
+	static const ECultureIdentifierType IdentifierType = 
+		[]()
+		{
+			if (FString IdentifierTypeStr;
+				GConfig->GetString(TEXT("Internationalization"), TEXT("PreferredOnDiskCultureIdentifierType"), IdentifierTypeStr, GEngineIni))
+			{
+				if (IdentifierTypeStr == TEXTVIEW("Verse"))
+				{
+					return ECultureIdentifierType::Verse;
+				}
+			}
+			return ECultureIdentifierType::UE;
+		}();
+	return IdentifierType;
+}
+
+const FString& FCulture::GetPreferredOnDiskIdentifier() const
+{
+	return GetPreferredOnDiskIdentifierType() == ECultureIdentifierType::Verse
+		? GetVerseIdentifier()
+		: GetName();
+}
+
 const FString& FCulture::GetName() const
 {
 	return CachedName;
diff --git a/Engine/Source/Runtime/Core/Public/Internationalization/Culture.h b/Engine/Source/Runtime/Core/Public/Internationalization/Culture.h
--- a/Engine/Source/Runtime/Core/Public/Internationalization/Culture.h
+++ b/Engine/Source/Runtime/Core/Public/Internationalization/Culture.h
@@ -19,6 +19,14 @@ class FLegacyCultureImplementation;
 typedef FLegacyCultureImplementation FCultureImplementation;
 #endif
 
+enum class ECultureIdentifierType : uint8
+{
+	/** UE culture name (eg, "en-US") */
+	UE,
+	/** Verse culture identifier (eg, "en_US") */
+	Verse,
+};
+
 class FCulture
 {
 #if UE_ENABLE_ICU
@@ -60,6 +68,18 @@ public:
 	 */
 	static CORE_API FString CultureNameFromVerseIdentifier(const FString& VerseIdentifer);
 
+	/**
+	 * Get the preferred on-disk identifier type (eg, "en-US" vs "en_US") to use when creating new localized content.
+	 * @note You should still assume both kinds of identifier may be used within a project.
+	 */
+	static CORE_API ECultureIdentifierType GetPreferredOnDiskIdentifierType();
+
+	/**
+	 * Get the preferred on-disk identifier of this culture (eg, "en-US" vs "en_US") to use when creating new localized content.
+	 * @note You should still assume both kinds of identifier may be used within a project.
+	 */
+	CORE_API const FString& GetPreferredOnDiskIdentifier() const;
+
 	/**
 	 * Get the name of this culture, eg, "en-US"
 	 */
