IncrementalCook: DLC: Crash fixes from static analysis.
#rnx

diff --git a/Engine/Source/Developer/IoStoreUtilities/Private/ZenStoreWriter.cpp b/Engine/Source/Developer/IoStoreUtilities/Private/ZenStoreWriter.cpp
--- a/Engine/Source/Developer/IoStoreUtilities/Private/ZenStoreWriter.cpp
+++ b/Engine/Source/Developer/IoStoreUtilities/Private/ZenStoreWriter.cpp
@@ -1315,7 +1315,7 @@ void FZenStoreWriter::GetBaseGameOplogAttachments(TArrayView<FName> PackageNames
 		return;
 	}
 
-	TArray<FUtf8StringView> AttachmentKeysBuffer;
+	TArray<FUtf8String> AttachmentKeysBuffer;
 	for (FUtf8StringView AttachmentKey : AttachmentKeys)
 	{
 		AttachmentKeysBuffer.Emplace(AttachmentKey);
diff --git a/Engine/Source/Editor/UnrealEd/Private/Cooker/CookPackageArtifacts.cpp b/Engine/Source/Editor/UnrealEd/Private/Cooker/CookPackageArtifacts.cpp
--- a/Engine/Source/Editor/UnrealEd/Private/Cooker/CookPackageArtifacts.cpp
+++ b/Engine/Source/Editor/UnrealEd/Private/Cooker/CookPackageArtifacts.cpp
@@ -336,7 +336,7 @@ bool FBuildDependencySet::TryGetModifiedDependencies(FName PackageName, const IT
 				OutModifiedDependencies.Add(FString::Printf(TEXT("%s: %s: Could not get current ClassDigest"),
 					LexToString(Dependency.GetType()), *Dependency.GetDebugIdentifier()));
 			}
-			if (*CurrentValue != Dependency.GetHashValue())
+			else if (*CurrentValue != Dependency.GetHashValue())
 			{
 				OutModifiedDependencies.Add(FString::Printf(TEXT("%s: %s: %s -> %s"),
 					LexToString(Dependency.GetType()), *Dependency.GetDebugIdentifier(),
@@ -1367,7 +1367,9 @@ void FIncrementalCookAttachments::Fetch(TArrayView<FPackageIncrementalCookId> Pa
 	{
 		for (const FPackageIncrementalCookId& Id : PackageIds)
 		{
-			InProgressResults.FindOrAdd(Id.PackageName).WhereCooked = Id.WhereCooked;
+			// EditorDomainNames are called with AttachmentsForIncrementalSkipAndBuildDependencies
+			check(Id.WhereCooked == EWhereCooked::ThisSession);
+			InProgressResults.FindOrAdd(Id.PackageName).WhereCooked = EWhereCooked::ThisSession;
 			EditorDomainNames.Add(Id.PackageName);
 		}
 	}
diff --git a/Engine/Source/Editor/UnrealEd/Private/Cooker/DiffPackageWriter.h b/Engine/Source/Editor/UnrealEd/Private/Cooker/DiffPackageWriter.h
--- a/Engine/Source/Editor/UnrealEd/Private/Cooker/DiffPackageWriter.h
+++ b/Engine/Source/Editor/UnrealEd/Private/Cooker/DiffPackageWriter.h
@@ -105,7 +105,7 @@ public:
 		TArrayView<FUtf8StringView> AttachmentKeys,
 		TUniqueFunction<void(FName PackageName, FUtf8StringView AttachmentKey, FCbObject&& Attachment)>&& Callback) override
 	{
-		return Inner->GetOplogAttachments(PackageNames, AttachmentKeys, MoveTemp(Callback));
+		return Inner->GetBaseGameOplogAttachments(PackageNames, AttachmentKeys, MoveTemp(Callback));
 	}
 
 	virtual ECommitStatus GetCommitStatus(FName PackageName) override
