PackageTools: Add FScopedSlowTasks for long running world operations.

#rnx

diff --git a/Engine/Source/Editor/UnrealEd/Private/PackageTools.cpp b/Engine/Source/Editor/UnrealEd/Private/PackageTools.cpp
--- a/Engine/Source/Editor/UnrealEd/Private/PackageTools.cpp
+++ b/Engine/Source/Editor/UnrealEd/Private/PackageTools.cpp
@@ -844,6 +844,9 @@ UPackageTools::UPackageTools(const FObjectInitializer& ObjectInitializer)
 		TArray<ULevelStreaming*> RemovedStreamingLevels;
 		if (UWorld* CurrentWorldPtr = CurrentWorld.Get())
 		{
+			FScopedSlowTask UnloadingWorldSlowTask(1.f, NSLOCTEXT("UnrealEd", "UnloadingWorld", "Unloading World"));
+			UnloadingWorldSlowTask.ForceRefresh();
+		
 			// Is the current world being reloaded? If so, we just reset the current world and load it again at the end rather than let it go through ReloadPackage 
 			// (which doesn't work for the current world due to some assumptions about worlds, and their lifetimes).
 			// We also need to skip the build data package as that will also be destroyed by the transition.
@@ -975,6 +978,9 @@ UPackageTools::UPackageTools(const FObjectInitializer& ObjectInitializer)
 		// Load the previous world (if needed).
 		if (!WorldNameToReload.IsNone())
 		{
+			FScopedSlowTask ReloadingWorldSlowTask(1.f, NSLOCTEXT("UnrealEd", "ReloadingWorld", "Reloading World"));
+			ReloadingWorldSlowTask.ForceRefresh();
+
 			if (GIsEditor)
 			{
 				UWorld::WorldTypePreLoadMap.FindOrAdd(CurrentWorldPackageName) = EWorldType::Editor;
