Dialogue Import fixes

* Fixed bug where we could miss the scanning for the Verse-style identifier when the preferred on-disk style was UE.
* Fixed a potential crash if an asset failed to load.

#rnx

diff --git a/Engine/Source/Editor/LocalizationCommandletExecution/Private/LocalizationCommandletTasks.cpp b/Engine/Source/Editor/LocalizationCommandletExecution/Private/LocalizationCommandletTasks.cpp
--- a/Engine/Source/Editor/LocalizationCommandletExecution/Private/LocalizationCommandletTasks.cpp
+++ b/Engine/Source/Editor/LocalizationCommandletExecution/Private/LocalizationCommandletTasks.cpp
@@ -573,19 +573,14 @@ bool LocalizationCommandletTasks::ReportLoadedAudioAssets(const TArray<ULocaliza
 					continue;
 				}
 
-				FString PreferredCultureIdentifier = Culture;
-				if (FCulture::GetPreferredOnDiskIdentifierType() == ECultureIdentifierType::Verse)
-				{
-					PreferredCultureIdentifier = FCulture::CultureNameToVerseIdentifier(Culture);
-				}
-
 				// Need to check for both UE and Verse style identifiers, as both may be in use
-				DialogueWavePathsToTest.Add(RootAssetPath / TEXT("L10N") / PreferredCultureIdentifier);
-				SoundWavePathsToTest.Add(RootAssetPath / TEXT("L10N") / PreferredCultureIdentifier / Target->Settings.ImportDialogueSettings.ImportedDialogueFolder);
-				if (PreferredCultureIdentifier != Culture)
+				DialogueWavePathsToTest.Add(RootAssetPath / TEXT("L10N") / Culture);
+				SoundWavePathsToTest.Add(RootAssetPath / TEXT("L10N") / Culture / Target->Settings.ImportDialogueSettings.ImportedDialogueFolder);
+				if (const FString VerseIdentifier = FCulture::CultureNameToVerseIdentifier(Culture);
+					VerseIdentifier != Culture)
 				{
-					DialogueWavePathsToTest.Add(RootAssetPath / TEXT("L10N") / Culture);
-					SoundWavePathsToTest.Add(RootAssetPath / TEXT("L10N") / Culture / Target->Settings.ImportDialogueSettings.ImportedDialogueFolder);
+					DialogueWavePathsToTest.Add(RootAssetPath / TEXT("L10N") / VerseIdentifier);
+					SoundWavePathsToTest.Add(RootAssetPath / TEXT("L10N") / VerseIdentifier / Target->Settings.ImportDialogueSettings.ImportedDialogueFolder);
 				}
 			}
 		}
diff --git a/Engine/Source/Editor/UnrealEd/Private/Commandlets/ImportLocalizedDialogueCommandlet.cpp b/Engine/Source/Editor/UnrealEd/Private/Commandlets/ImportLocalizedDialogueCommandlet.cpp
--- a/Engine/Source/Editor/UnrealEd/Private/Commandlets/ImportLocalizedDialogueCommandlet.cpp
+++ b/Engine/Source/Editor/UnrealEd/Private/Commandlets/ImportLocalizedDialogueCommandlet.cpp
@@ -191,12 +191,13 @@ int32 UImportLocalizedDialogueCommandlet::Main(const FString& Params)
 		CultureImportInfo.bIsNativeCulture = CultureName == NativeCulture;
 
 		// Need to check for both UE and Verse style identifiers, as both may be in use
-		LocalizedDialogueWavePathsToSearch.Add(*CultureImportInfo.LocalizedRootPackagePath);
-		LocalizedSoundWavePathsToSearch.Add(*CultureImportInfo.LocalizedImportedDialoguePackagePath);
-		if (PreferredCultureIdentifier != CultureName)
+		LocalizedDialogueWavePathsToSearch.Add(*(RootAssetPath / TEXT("L10N") / CultureName));
+		LocalizedSoundWavePathsToSearch.Add(*(RootAssetPath / TEXT("L10N") / CultureName / ImportedDialogueFolder));
+		if (const FString VerseIdentifier = FCulture::CultureNameToVerseIdentifier(CultureName);
+			VerseIdentifier != CultureName)
 		{
-			LocalizedDialogueWavePathsToSearch.Add(*(RootAssetPath / TEXT("L10N") / CultureName));
-			LocalizedSoundWavePathsToSearch.Add(*(RootAssetPath / TEXT("L10N") / CultureName / ImportedDialogueFolder));
+			LocalizedDialogueWavePathsToSearch.Add(*(RootAssetPath / TEXT("L10N") / VerseIdentifier));
+			LocalizedSoundWavePathsToSearch.Add(*(RootAssetPath / TEXT("L10N") / VerseIdentifier / ImportedDialogueFolder));
 		}
 	}
 
@@ -218,7 +219,7 @@ int32 UImportLocalizedDialogueCommandlet::Main(const FString& Params)
 		// Verify that the found asset is a dialogue wave
 		if (AssetData.GetClass() != UDialogueWave::StaticClass())
 		{
-			UE_LOG(LogImportLocalizedDialogueCommandlet, Error, TEXT("Asset registry found asset '%s', but the asset with this name is not actually a dialogue wave."), *AssetData.AssetName.ToString());
+			UE_LOG(LogImportLocalizedDialogueCommandlet, Error, TEXT("Asset registry found asset '%s', but it's not actually a dialogue wave."), *AssetData.GetSoftObjectPath().ToString());
 			continue;
 		}
 
@@ -228,7 +229,7 @@ int32 UImportLocalizedDialogueCommandlet::Main(const FString& Params)
 		// Verify that the dialogue wave was loaded
 		if (!DialogueWave)
 		{
-			UE_LOG(LogImportLocalizedDialogueCommandlet, Error, TEXT("Asset registry found asset '%s', but the dialogue wave could not be accessed."), *AssetData.AssetName.ToString());
+			UE_LOG(LogImportLocalizedDialogueCommandlet, Error, TEXT("Asset registry found asset '%s', but the dialogue wave could not be loaded."), *AssetData.GetSoftObjectPath().ToString());
 			continue;
 		}
 
@@ -274,8 +275,14 @@ int32 UImportLocalizedDialogueCommandlet::Main(const FString& Params)
 			continue;
 		}
 
-		// Check the package meta-data to make sure that we only delete packages that we own
 		UObject* LocalizedAsset = LocalizedAssetData.GetAsset();
+		if (!LocalizedAsset)
+		{
+			UE_LOG(LogImportLocalizedDialogueCommandlet, Error, TEXT("Asset registry found localized asset '%s' to delete, but the asset could not be loaded."), *LocalizedAssetData.GetSoftObjectPath().ToString());
+			continue;
+		}
+
+		// Check the package meta-data to make sure that we only delete packages that we own
 		const FString AssetGeneratedByCommandletMetaDataValue = LocalizedAsset->GetPackage()->GetMetaData().GetValue(LocalizedAsset, *GeneratedByCommandletMetaDataKey);
 		if (AssetGeneratedByCommandletMetaDataValue != GeneratedByCommandletMetaDataValue)
 		{
