MultiprocessCook: Fix failure to wait for ECookPhase::Cook to complete on CookWorkers before starting ECookPhase::BuildDependencies, in the case that there was so little work to be done that it finished before the first heartbeat. LastReceivedHeartbeatNumber should be initialized to a value less than any possible HeartbeatNumber.
#rnx

diff --git a/Engine/Source/Editor/UnrealEd/Private/Cooker/CookWorkerServer.h b/Engine/Source/Editor/UnrealEd/Private/Cooker/CookWorkerServer.h
--- a/Engine/Source/Editor/UnrealEd/Private/Cooker/CookWorkerServer.h
+++ b/Engine/Source/Editor/UnrealEd/Private/Cooker/CookWorkerServer.h
@@ -237,7 +237,7 @@ private:
 	FTickState TickState;
 	uint32 CookWorkerProcessId = 0;
 	int32 ProfileId = 0;
-	int32 LastReceivedHeartbeatNumber = 0;
+	int32 LastReceivedHeartbeatNumber = -1;
 	int32 PackagesAssignedFenceMarker = 0;
 	int32 PackagesRetiredFenceMarker = 0;
 	int32 LastAbortHeartbeatNumber = -1;
