UE-300446 - Output summary of diffonly cooks to logs on cook completion

#tests Tested diffonly cooks of Lyra, CitySample

diff --git a/Engine/Source/Editor/UnrealEd/Private/CookOnTheFlyServer.cpp b/Engine/Source/Editor/UnrealEd/Private/CookOnTheFlyServer.cpp
--- a/Engine/Source/Editor/UnrealEd/Private/CookOnTheFlyServer.cpp
+++ b/Engine/Source/Editor/UnrealEd/Private/CookOnTheFlyServer.cpp
@@ -5740,8 +5740,8 @@ public:
 		switch (DiffMode)
 		{
 		case EDiffMode::DiffOnly:
-			CookedPackageWriter = new FDiffPackageWriter(TUniquePtr<ICookedPackageWriter>(CookedPackageWriter),
-				InDeterminismManager);
+			CookedPackageWriter = new FDiffPackageWriter(COTFS, TUniquePtr<ICookedPackageWriter>(CookedPackageWriter),
+				InDeterminismManager, FDiffPackageWriter::EReportingMode::OutputPackageSummary);
 			bCurrentWriterAllowWrite = false;
 			break;
 		case EDiffMode::IncrementalValidate:
diff --git a/Engine/Source/Editor/UnrealEd/Private/Cooker/DiffPackageWriter.cpp b/Engine/Source/Editor/UnrealEd/Private/Cooker/DiffPackageWriter.cpp
--- a/Engine/Source/Editor/UnrealEd/Private/Cooker/DiffPackageWriter.cpp
+++ b/Engine/Source/Editor/UnrealEd/Private/Cooker/DiffPackageWriter.cpp
@@ -21,6 +21,7 @@
 #include "ProfilingDebugging/CookStats.h"
 #include "Serialization/LargeMemoryWriter.h"
 #include "Serialization/JsonWriter.h"
+#include "Serialization/CompactBinarySerialization.h"
 #include "Templates/UniquePtr.h"
 #include "Templates/UnrealTemplate.h"
 #include "UObject/LinkerSave.h"
@@ -340,12 +341,80 @@ protected:
 		return GetSanitizedPaths(Result);
 	}
 };
+
+// Responsible for replicating diff stats collected on workers to the director
+class FDiffPackageWriterMPCollector : public UE::Cook::IMPCollector
+{
+public:
+	FDiffPackageWriterMPCollector(FDiffPackageWriter* InOwner) : Owner(InOwner) {}
+
+	virtual FGuid GetMessageType() const { return MessageType; }
+	virtual const TCHAR* GetDebugName() const { return TEXT("DiffPackageWriterMPCollector"); }
+
+	virtual void ClientTickPackage(UE::Cook::FMPCollectorClientTickPackageContext& Context) override;
+	virtual void ServerReceiveMessage(UE::Cook::FMPCollectorServerMessageContext& Context, FCbObjectView Message) override;
+
+public:
+	static FGuid MessageType;
+
+private:
+	FDiffPackageWriter* Owner;
+};
+
+FGuid FDiffPackageWriterMPCollector::MessageType(0x0E2BA390, 0x2D574257, 0xBBD107D7, 0xB5930C79);
+
+// sends the diff stats for the current package to the director, if we have any 
+void FDiffPackageWriterMPCollector::ClientTickPackage(UE::Cook::FMPCollectorClientTickPackageContext& Context)
+{
+	const FDiffPackageWriter::FPackageDiffStats* FoundStats = Owner->PackageDiffStats.Find(Context.GetPackageName());
+	if (FoundStats != nullptr)
+	{
+		FCbWriter Writer;
+		Writer.BeginObject();
+		Writer << "Package" << Context.GetPackageName();
+		Writer << "AssetClass" << FoundStats->AssetClass;
+		Writer << "TotalDiffCount" << FoundStats->TotalDifferences;
+		Writer << "TotalDiffBytes" << FoundStats->TotalSizeDifferenceBytes;
+		Writer.EndObject();
+		Context.AddMessage(Writer.Save().AsObject());
+	}
+}
+
+// receives diff stats from a client
+void FDiffPackageWriterMPCollector::ServerReceiveMessage(UE::Cook::FMPCollectorServerMessageContext& Context, FCbObjectView Message)
+{
+	FDiffPackageWriter::FPackageDiffStats ReceivedStats;
+	FName ReceivedPackageName;
+	bool bReadOK = LoadFromCompactBinary(Message["Package"], ReceivedPackageName);
+	if (bReadOK)
+	{
+		bReadOK = LoadFromCompactBinary(Message["AssetClass"], ReceivedStats.AssetClass);
+	}
+	if (bReadOK)
+	{
+		bReadOK = LoadFromCompactBinary(Message["TotalDiffCount"], ReceivedStats.TotalDifferences);
+	}
+	if (bReadOK)
+	{
+		bReadOK = LoadFromCompactBinary(Message["TotalDiffBytes"], ReceivedStats.TotalSizeDifferenceBytes);
+	}
+	if (bReadOK)
+	{
+		Owner->PackageDiffStats.FindOrAdd(ReceivedPackageName, ReceivedStats);		// Duplicates can be received here, we ignore them
+	}
+	else
+	{
+		UE_LOG(LogCook, Error, TEXT("Failed to read diff status from MP collector"));
+	}
+}
+
 } // namespace UE::DiffWriter
 
 
-FDiffPackageWriter::FDiffPackageWriter(TUniquePtr<ICookedPackageWriter>&& InInner,
-	UE::Cook::FDeterminismManager* InDeterminismManager)
+FDiffPackageWriter::FDiffPackageWriter(UCookOnTheFlyServer& InCOTFS, TUniquePtr<ICookedPackageWriter>&& InInner,
+	UE::Cook::FDeterminismManager* InDeterminismManager, EReportingMode ReportMode)
 	: Inner(MoveTemp(InInner))
+	, COTFS(InCOTFS)
 {
 	AccumulatorGlobals.Reset(new UE::DiffWriter::FAccumulatorGlobals(Inner.Get()));
 
@@ -373,7 +442,6 @@ FDiffPackageWriter::FDiffPackageWriter(TUniquePtr<ICookedPackageWriter>&& InInne
 		AccumulatorGlobals->DetailRecorder = MakeUnique<UE::DiffWriter::FJsonDetailRecorder>(JsonFile);
 	}
 
-
 	ParseCmds();
 
 	Indent = FCString::Spc(FOutputDeviceHelper::FormatLogLine(ELogVerbosity::Warning,
@@ -382,6 +450,22 @@ FDiffPackageWriter::FDiffPackageWriter(TUniquePtr<ICookedPackageWriter>&& InInne
 
 	check(InDeterminismManager);
 	DeterminismManager = InDeterminismManager;
+
+	if (ReportMode == EReportingMode::OutputPackageSummary)
+	{
+		DiffStatCollector = new UE::DiffWriter::FDiffPackageWriterMPCollector(this);
+		COTFS.RegisterCollector(DiffStatCollector);
+		bCollectDiffStats = true;
+	}
+}
+
+FDiffPackageWriter::~FDiffPackageWriter()
+{
+	if (DiffStatCollector != nullptr)
+	{
+		COTFS.UnregisterCollector(DiffStatCollector);
+		DiffStatCollector.SafeRelease();
+	}
 }
 
 void FDiffPackageWriter::ParseCmds()
@@ -524,6 +608,16 @@ void FDiffPackageWriter::WritePackageData(const FPackageInfo& Info, FLargeMemory
 		{
 			bIsDifferent = false;
 		}
+
+		if (bCollectDiffStats)
+		{
+			FPackageDiffStats& DiffStats = PackageDiffStats.Add(Info.PackageName);
+			DiffStats = {
+				.AssetClass = Accumulator.GetAssetClass(),
+				.TotalSizeDifferenceBytes = 0,
+				.TotalDifferences = 0
+			};
+		}
 	}
 	else
 	{
@@ -533,9 +627,23 @@ void FDiffPackageWriter::WritePackageData(const FPackageInfo& Info, FLargeMemory
 
 		Accumulator.OnSecondSaveComplete(LocalInfo.HeaderSize);
 
-		TMap<FName, FArchiveDiffStats> PackageDiffStats;
+		TMap<FName, FArchiveDiffStats> DiffStats;
 		Accumulator.SetDeterminismManager(*DeterminismManager);
-		Accumulator.CompareWithPrevious(CallstackCutoffString, PackageDiffStats);
+		Accumulator.CompareWithPrevious(CallstackCutoffString, DiffStats);
+
+		if (bCollectDiffStats)
+		{
+			FPackageDiffStats* Stats = PackageDiffStats.Find(Info.PackageName);
+			if (Stats)
+			{
+				for (const TPair<FName, FArchiveDiffStats>& DiffStat : DiffStats)
+				{
+					Stats->TotalSizeDifferenceBytes += FMath::Abs(DiffStat.Value.DiffSize);
+					Stats->TotalDifferences += DiffStat.Value.NumDiffs;
+				}
+			}
+		}
+		
 
 		//COOK_STAT(FSavePackageStats::NumberOfDifferentPackages++);
 		//COOK_STAT(FSavePackageStats::MergeStats(PackageDiffStats));
@@ -544,6 +652,59 @@ void FDiffPackageWriter::WritePackageData(const FPackageInfo& Info, FLargeMemory
 	Inner->WritePackageData(LocalInfo, ExportsArchive, FileRegions);
 }
 
+TMap<FName, FDiffPackageWriter::FPackageDiffAssetClassSummary> FDiffPackageWriter::GetAssetClassDifferenceSummaries() const
+{
+	TMap<FName, FPackageDiffAssetClassSummary> AllSummaries;
+	for (const TPair<FName, FPackageDiffStats>& PackageDiffStat : PackageDiffStats)
+	{
+		if (PackageDiffStat.Value.TotalDifferences != 0 || PackageDiffStat.Value.TotalSizeDifferenceBytes != 0)
+		{
+			FPackageDiffAssetClassSummary& Summary = AllSummaries.FindOrAdd(PackageDiffStat.Value.AssetClass);
+			Summary.TotalPackageCount++;
+			Summary.TotalPackagesWithDifferences += PackageDiffStat.Value.TotalDifferences > 0 ? 1 : 0;
+			Summary.TotalDifferenceCount += PackageDiffStat.Value.TotalDifferences;
+		}
+	}
+
+	// Sort by package differences + class name
+	AllSummaries.KeySort([](const FName& A, const FName& B) {
+		return A.LexicalLess(B);
+	});
+	AllSummaries.ValueStableSort([](const FPackageDiffAssetClassSummary& A, const FPackageDiffAssetClassSummary& B) {
+		return A.TotalPackagesWithDifferences > B.TotalPackagesWithDifferences;
+	});
+
+	return AllSummaries;
+}
+
+void FDiffPackageWriter::LogAssetClassDiffSummary() const
+{
+	const TMap<FName, FDiffPackageWriter::FPackageDiffAssetClassSummary> AllSummaries = GetAssetClassDifferenceSummaries();
+	if (AllSummaries.Num() > 0)
+	{
+		UE_LOG(LogDiff, Display, TEXT("Differences detected per asset class:"));
+		UE_LOG(LogDiff, Display, TEXT("\tAsset Class | Changed Packages | Avg. Differences per Package"));
+	}
+	for (const TPair<FName, FPackageDiffAssetClassSummary>& AssetClassDifferences : AllSummaries)
+	{
+		const FPackageDiffAssetClassSummary& Stats = AssetClassDifferences.Value;
+		const float AverageDifferenceCount = Stats.TotalDifferenceCount / static_cast<float>(Stats.TotalPackagesWithDifferences);
+		UE_LOG(LogDiff, Display, TEXT("\t%s | %d of %d | %.1f"), 
+			*AssetClassDifferences.Key.ToString(), Stats.TotalPackagesWithDifferences, Stats.TotalPackageCount, AverageDifferenceCount);
+	}
+}
+
+void FDiffPackageWriter::EndCook(const FCookInfo& Info)
+{
+	Inner->EndCook(Info);
+
+	const bool bShouldOutputClassDummary = COTFS.GetCookMode() != ECookMode::CookWorker;	// only the director or single process cook should output summary here
+	if (bShouldOutputClassDummary)
+	{
+		LogAssetClassDiffSummary();
+	}
+}
+
 bool FDiffPackageWriter::IsPackageDiffAllowed() const
 {
 	if (!CompareDenyListClasses.IsEmpty() && Package != nullptr)
diff --git a/Engine/Source/Editor/UnrealEd/Private/Cooker/DiffPackageWriter.h b/Engine/Source/Editor/UnrealEd/Private/Cooker/DiffPackageWriter.h
--- a/Engine/Source/Editor/UnrealEd/Private/Cooker/DiffPackageWriter.h
+++ b/Engine/Source/Editor/UnrealEd/Private/Cooker/DiffPackageWriter.h
@@ -7,13 +7,21 @@
 #include "Cooker/DiffWriterArchive.h"
 #include "Serialization/PackageWriter.h"
 
-namespace UE::DiffWriter { struct FAccumulatorGlobals; }
+class UCookOnTheFlyServer;
+namespace UE::DiffWriter { struct FAccumulatorGlobals; class FDiffPackageWriterMPCollector; }
 
 /** A CookedPackageWriter that diffs output from the current cook with the file that was saved in the previous cook. */
 class FDiffPackageWriter : public ICookedPackageWriter
 {
 public:
-	FDiffPackageWriter(TUniquePtr<ICookedPackageWriter>&& InInner, UE::Cook::FDeterminismManager* InDeterminismManager);
+	enum class EReportingMode {
+		OutputPackageSummary,
+		NoPackageSummary
+	};
+
+	FDiffPackageWriter(UCookOnTheFlyServer& InCOTFS, TUniquePtr<ICookedPackageWriter>&& InInner, 
+		UE::Cook::FDeterminismManager* InDeterminismManager, EReportingMode ReportMode);
+	virtual ~FDiffPackageWriter();
 
 	virtual void SetCooker(UE::PackageWriter::Private::ICookerInterface* CookerInterface) override
 	{
@@ -83,10 +91,7 @@ public:
 	{
 		Inner->BeginCook(Info);
 	}
-	virtual void EndCook(const FCookInfo& Info) override
-	{
-		Inner->EndCook(Info);
-	}
+	virtual void EndCook(const FCookInfo& Info) override;
 	virtual TUniquePtr<FAssetRegistryState> LoadPreviousAssetRegistry() override
 	{
 		return Inner->LoadPreviousAssetRegistry();
@@ -144,6 +149,22 @@ public:
 	}
 
 protected:
+	// Stores differences found for a particular package
+	struct FPackageDiffStats
+	{
+		FName AssetClass;					// Top-level asset class for this package
+		int64 TotalSizeDifferenceBytes;		// Total size difference between old + new asset
+		int64 TotalDifferences;				// Total count of data differences in this asset 
+	};
+
+	// Accumulated diff stats for an entire asset class
+	struct FPackageDiffAssetClassSummary
+	{
+		int64 TotalPackageCount;				// Total packages saved of this asset class
+		int64 TotalPackagesWithDifferences;		// Total packages with detected differences
+		int64 TotalDifferenceCount;				// Total detected differences across all packages
+	};
+
 	void ParseCmds();
 	void ParseDumpObjList(FString InParams);
 	void ParseDumpObjects(FString InParams);
@@ -156,6 +177,8 @@ protected:
 	FString ResolveText(FStringView Message);
 	UE::DiffWriter::FAccumulator& ConstructAccumulator(FName PackageName, UObject* Asset, uint16 MultiOutputIndex);
 	bool IsPackageDiffAllowed() const;
+	void LogAssetClassDiffSummary() const;
+	TMap<FName, FPackageDiffAssetClassSummary> GetAssetClassDifferenceSummaries() const;
 
 	/** 
 	 * When BeginPackage is called, CommitPackage must be called as well. But sometimes we don't want or don't need to do all this and we just want to
@@ -165,12 +188,16 @@ protected:
 	*/
 	void BeginPackageDiff(const FBeginPackageInfo& Info);
 
+	TMap<FName, FPackageDiffStats> PackageDiffStats;
+	TRefCountPtr<UE::DiffWriter::FDiffPackageWriterMPCollector> DiffStatCollector;
 	TRefCountPtr<UE::DiffWriter::FAccumulator> Accumulators[2];
 	FBeginPackageInfo BeginInfo;
 	TSet<FTopLevelAssetPath> CompareDenyListClasses;
 	TUniquePtr<ICookedPackageWriter> Inner;
 	TUniquePtr<UE::DiffWriter::FAccumulatorGlobals> AccumulatorGlobals;
 	UE::Cook::FDeterminismManager* DeterminismManager = nullptr;
+	UCookOnTheFlyServer& COTFS;
+
 	/** Only non-null between BeginPackage and CommitPackage. */
 	UPackage* Package = nullptr;
 	const TCHAR* Indent = nullptr;
@@ -187,4 +214,7 @@ protected:
 	bool bDumpObjList = false;
 	bool bDumpObjects = false;
 	bool bDumpObjectsSorted = false;
+	bool bCollectDiffStats = false;
+
+	friend class UE::DiffWriter::FDiffPackageWriterMPCollector;
 };
diff --git a/Engine/Source/Editor/UnrealEd/Private/Cooker/IncrementalValidatePackageWriter.cpp b/Engine/Source/Editor/UnrealEd/Private/Cooker/IncrementalValidatePackageWriter.cpp
--- a/Engine/Source/Editor/UnrealEd/Private/Cooker/IncrementalValidatePackageWriter.cpp
+++ b/Engine/Source/Editor/UnrealEd/Private/Cooker/IncrementalValidatePackageWriter.cpp
@@ -320,9 +320,8 @@ void FIncrementalValidateMPCollector::ReadAndSyncPackageStatus(FCbObjectView Mes
 FIncrementalValidatePackageWriter::FIncrementalValidatePackageWriter(UCookOnTheFlyServer& InCOTFS,
 	TUniquePtr<ICookedPackageWriter>&& InInner, EPhase InPhase, const FString& ResolvedMetadataPath,
 	UE::Cook::FDeterminismManager* InDeterminismManager)
-	: FDiffPackageWriter(MoveTemp(InInner), InDeterminismManager)
+	: FDiffPackageWriter(InCOTFS, MoveTemp(InInner), InDeterminismManager, FDiffPackageWriter::EReportingMode::NoPackageSummary)
 	, MetadataPath(ResolvedMetadataPath)
-	, COTFS(InCOTFS)
 	, Phase(InPhase)
 {
 	COTFS.RegisterCollector(new FIncrementalValidateMPCollector(this));
diff --git a/Engine/Source/Editor/UnrealEd/Private/Cooker/IncrementalValidatePackageWriter.h b/Engine/Source/Editor/UnrealEd/Private/Cooker/IncrementalValidatePackageWriter.h
--- a/Engine/Source/Editor/UnrealEd/Private/Cooker/IncrementalValidatePackageWriter.h
+++ b/Engine/Source/Editor/UnrealEd/Private/Cooker/IncrementalValidatePackageWriter.h
@@ -4,7 +4,6 @@
 
 #include "DiffPackageWriter.h"
 
-class UCookOnTheFlyServer;
 class FIncrementalValidateMPCollector; 
 namespace UE::Cook { struct FWorkerId; }
 
@@ -126,7 +125,6 @@ protected:
 	TMap<FTopLevelAssetPath, TMap<EPackageStatus, TArray<FName>>> ClassStatusSummary; // Only populated on director
 	TSet<FName> PackageIgnoreList;
 	FString MetadataPath;
-	UCookOnTheFlyServer& COTFS;
 	int32 LoggingSoftMaximum = -1;
 	EPhase Phase = EPhase::AllInOnePhase;
 	ESaveAction SaveAction = ESaveAction::IgnoreResults;
