Small fixes to logdiff summary
- Uninitialised members of FPackageDiffAssetClassSummary
- Potential divide by zero in LogAssetClassDiffSummary

diff --git a/Engine/Source/Editor/UnrealEd/Private/Cooker/DiffPackageWriter.cpp b/Engine/Source/Editor/UnrealEd/Private/Cooker/DiffPackageWriter.cpp
--- a/Engine/Source/Editor/UnrealEd/Private/Cooker/DiffPackageWriter.cpp
+++ b/Engine/Source/Editor/UnrealEd/Private/Cooker/DiffPackageWriter.cpp
@@ -688,9 +688,12 @@ void FDiffPackageWriter::LogAssetClassDiffSummary() const
 	for (const TPair<FName, FPackageDiffAssetClassSummary>& AssetClassDifferences : AllSummaries)
 	{
 		const FPackageDiffAssetClassSummary& Stats = AssetClassDifferences.Value;
-		const float AverageDifferenceCount = Stats.TotalDifferenceCount / static_cast<float>(Stats.TotalPackagesWithDifferences);
-		UE_LOG(LogDiff, Display, TEXT("\t%s | %d of %d | %.1f"), 
-			*AssetClassDifferences.Key.ToString(), Stats.TotalPackagesWithDifferences, Stats.TotalPackageCount, AverageDifferenceCount);
+		if (Stats.TotalPackagesWithDifferences > 0)
+		{
+			const float AverageDifferenceCount = Stats.TotalDifferenceCount / static_cast<float>(Stats.TotalPackagesWithDifferences);
+			UE_LOG(LogDiff, Display, TEXT("\t%s | %d of %d | %.1f"),
+				*AssetClassDifferences.Key.ToString(), Stats.TotalPackagesWithDifferences, Stats.TotalPackageCount, AverageDifferenceCount);
+		}
 	}
 }
 
diff --git a/Engine/Source/Editor/UnrealEd/Private/Cooker/DiffPackageWriter.h b/Engine/Source/Editor/UnrealEd/Private/Cooker/DiffPackageWriter.h
--- a/Engine/Source/Editor/UnrealEd/Private/Cooker/DiffPackageWriter.h
+++ b/Engine/Source/Editor/UnrealEd/Private/Cooker/DiffPackageWriter.h
@@ -152,17 +152,17 @@ protected:
 	// Stores differences found for a particular package
 	struct FPackageDiffStats
 	{
-		FName AssetClass;					// Top-level asset class for this package
-		int64 TotalSizeDifferenceBytes;		// Total size difference between old + new asset
-		int64 TotalDifferences;				// Total count of data differences in this asset 
+		FName AssetClass;						// Top-level asset class for this package
+		int64 TotalSizeDifferenceBytes = 0;		// Total size difference between old + new asset
+		int64 TotalDifferences = 0;				// Total count of data differences in this asset 
 	};
 
 	// Accumulated diff stats for an entire asset class
 	struct FPackageDiffAssetClassSummary
 	{
-		int64 TotalPackageCount;				// Total packages saved of this asset class
-		int64 TotalPackagesWithDifferences;		// Total packages with detected differences
-		int64 TotalDifferenceCount;				// Total detected differences across all packages
+		int64 TotalPackageCount = 0;				// Total packages saved of this asset class
+		int64 TotalPackagesWithDifferences = 0;		// Total packages with detected differences
+		int64 TotalDifferenceCount = 0;				// Total detected differences across all packages
 	};
 
 	void ParseCmds();
