Cooker: GeneratedPackages: Allow the CookPackageSplitter to specify that a generated package can be mounted independently and should receive its chunk assignment from the AssetManager rather than copying from its generator. WorldPartition will use this for assignment of streaming cells for content bundles to the chunk of the plugin owning the content bundle.
#rnx

diff --git a/Engine/Source/Editor/UnrealEd/Private/Commandlets/AssetRegistryGenerator.cpp b/Engine/Source/Editor/UnrealEd/Private/Commandlets/AssetRegistryGenerator.cpp
--- a/Engine/Source/Editor/UnrealEd/Private/Commandlets/AssetRegistryGenerator.cpp
+++ b/Engine/Source/Editor/UnrealEd/Private/Commandlets/AssetRegistryGenerator.cpp
@@ -923,36 +923,40 @@ void FAssetRegistryGenerator::CalculateChunkIdsAndAssignToManifest(const FName&
 	}
 	else
 	{
-		FName PackageNameThatDefinesChunks = PackageFName;
-
 		UAssetManager& AssetManager = UAssetManager::Get();
 
-		TArray<int32> PackageChunkIDs;
-		
-		for (int32 ChunkID : AssetManager.GetEncryptedChunkIDsForPackage(PackageFName))
+		// Handle encrypted packages. Encrypted packages specified by GetEncryptedChunkIDsForPackage
+		// use those specific chunks and nothing else.
+		TArray<int32> EncryptedPackageChunkIDs = AssetManager.GetEncryptedChunkIDsForPackage(PackageFName).Array();
+		if (EncryptedPackageChunkIDs.Num() != 0)
 		{
-			PackageChunkIDs.Add(ChunkID);
+			TargetChunks.Append(EncryptedPackageChunkIDs);
 		}
-
-		// We only want to override the package name 
-		if (PackageChunkIDs.Num() == 0)
+		else
 		{
-			// Generated packages use the chunks defined by their Generator
-			FName GeneratorName = GetGeneratorPackage(PackageFName, this->State);
+			// Handle propagation of chunk assignments of Generated packages.
+			// Generated packages by default use the chunks assigned to their Generator (unless encrypted, see above).
+			// But generated packages for which the generator said it does not want chunk assignment from the generator
+			// get their chunk assignment separately from the AssetManager.
+			FName PackageNameThatDefinesChunks = PackageFName;
+			UE::AssetRegistry::EDependencyProperty DependencyPropertyOfGenerator;
+			FName GeneratorName = GetGeneratorPackage(PackageFName, this->State, DependencyPropertyOfGenerator);
 			if (!GeneratorName.IsNone())
 			{
-				PackageNameThatDefinesChunks = GeneratorName;
+				// CookGenerationHelper sets the Game flag on the dependency property if and only if it wants
+				// the generated package to receive chunks from the generator,
+				// see FGenerationHelper::FinishGeneratorPlatformSave.
+				if (EnumHasAnyFlags(DependencyPropertyOfGenerator, UE::AssetRegistry::EDependencyProperty::Game))
+				{
+					PackageNameThatDefinesChunks = GeneratorName;
+				}
 			}
 
-			PackageChunkIDs = GetExplicitChunkIDs(PackageNameThatDefinesChunks);
+			TArray<int32> PackageChunkIDs = GetExplicitChunkIDs(PackageNameThatDefinesChunks);
 			ExistingChunkIDs = GetExistingPackageChunkAssignments(PackageNameThatDefinesChunks, StartupPackages);
 			PackageChunkIDs.Append(ExistingChunkIDs);
 			AssetManager.GetPackageChunkIds(PackageNameThatDefinesChunks, TargetPlatform, PackageChunkIDs, TargetChunks);
 		}
-		else
-		{
-			TargetChunks.Append(PackageChunkIDs);
-		}
 	}
 
 	// Add the package to the manifest for every chunk the AssetManager found it should belong to
@@ -1425,7 +1429,8 @@ void FAssetRegistryGenerator::ComputePackageDifferences(const FComputeDifference
 			// If it's a generated package, exclude it from the results list and do not remove it.
 			// It will be evaluated for identical/modified/removed only if the generator package is cooked,
 			// during the generator's process step.
-			FName GeneratorName = GetGeneratorPackage(PackageName, PreviousState);
+			UE::AssetRegistry::EDependencyProperty UnusedGeneratorProperty;
+			FName GeneratorName = GetGeneratorPackage(PackageName, PreviousState, UnusedGeneratorProperty);
 			if (!GeneratorName.IsNone())
 			{
 				// Keep track of all generators and their list of generated
@@ -1636,8 +1641,10 @@ FAssetPackageData FAssetRegistryGenerator::CopyAssetPackageDataForIncrementalCoo
 	return Result;
 }
 
-FName FAssetRegistryGenerator::GetGeneratorPackage(FName PackageName, const FAssetRegistryState& InState)
+FName FAssetRegistryGenerator::GetGeneratorPackage(FName PackageName, const FAssetRegistryState& InState,
+	UE::AssetRegistry::EDependencyProperty& OutDependencyPropertyOfGenerator)
 {
+	OutDependencyPropertyOfGenerator = UE::AssetRegistry::EDependencyProperty::None;
 	bool bGenerated = false;
 	InState.EnumerateAssetsByPackageName(PackageName, [&bGenerated](const FAssetData* Asset)
 		{
@@ -1649,10 +1656,15 @@ FName FAssetRegistryGenerator::GetGeneratorPackage(FName PackageName, const FAss
 		return NAME_None;
 	}
 
-	TArray<FAssetIdentifier> Referencers;
+	TArray<FAssetDependency> Referencers;
 	InState.GetReferencers(FAssetIdentifier(PackageName), Referencers, UE::AssetRegistry::EDependencyCategory::Package);
-	FName GeneratorName = Referencers.Num() == 1 ? Referencers[0].PackageName : NAME_None;
-	return GeneratorName;
+	if (Referencers.Num() != 1)
+	{
+		return NAME_None;
+	}
+	FAssetDependency& Dependency = Referencers[0];
+	OutDependencyPropertyOfGenerator = Dependency.Properties;
+	return Dependency.AssetId.PackageName;
 }
 
 void FAssetRegistryGenerator::ComputePackageRemovals(const FAssetRegistryState& PreviousState, TArray<FName>& OutRemovedPackages,
diff --git a/Engine/Source/Editor/UnrealEd/Private/Cooker/CookGenerationHelper.cpp b/Engine/Source/Editor/UnrealEd/Private/Cooker/CookGenerationHelper.cpp
--- a/Engine/Source/Editor/UnrealEd/Private/Cooker/CookGenerationHelper.cpp
+++ b/Engine/Source/Editor/UnrealEd/Private/Cooker/CookGenerationHelper.cpp
@@ -1084,6 +1084,8 @@ bool FGenerationHelper::TryGenerateList()
 		// Copy hash from package splitter so it can partcipate in the package hash
 		GeneratedInfo->GenerationHash = SplitterData.GenerationHash;
 
+		GeneratedInfo->SetCopyChunkAssignmentFromGenerator(SplitterData.bCopyChunkAssignmentFromGenerator);
+
 		// Create the Hash from the GenerationHash and Dependencies
 		GeneratedInfo->CreatePackageHash();
 
@@ -1624,7 +1626,21 @@ void FGenerationHelper::FinishGeneratorPlatformSave(FPackageData& PackageData, b
 		FAssetDependency& Dependency = OutPackageDependencies.Emplace_GetRef();
 		Dependency.AssetId = FAssetIdentifier(GeneratedInfo.PackageData->GetPackageName());
 		Dependency.Category = UE::AssetRegistry::EDependencyCategory::Package;
-		Dependency.Properties = UE::AssetRegistry::EDependencyProperty::Game;
+		// The dependencies we write on the generated packages should be soft, not build, and
+		// game by default but optionally editor-only.
+		// Soft because they're not necessarily hard imports and the generator's package does not need to load them
+		// before it loads. The generator may have added them as imports, but we don't need to record that in the AR.
+		// Not build because they they don't affect whether the generator package needs to be recooked.
+		// Game by default because usually they are needed at runtime (and therefore mounted in the same chunk),
+		// but some Generator classes may want to override that behavior for some generated packages
+		// e.g. world partition may want streaming cells for a content bundle to be chunked with the content bundle).
+		// See the comment on bCopyChunkAssignmentFromGenerator in ICookPackageSplitter, and see the use of this
+		// dependency in FAssetRegistryGenerator::CalculateChunkIdsAndAssignToManifest.
+		Dependency.Properties = UE::AssetRegistry::EDependencyProperty::None;
+		if (GeneratedInfo.IsCopyChunkAssignmentFromGenerator())
+		{
+			Dependency.Properties |= UE::AssetRegistry::EDependencyProperty::Game;
+		}
 	}
 
 	if (bFirstPlatform)
@@ -2552,6 +2568,7 @@ FCookGenerationInfo::FCookGenerationInfo(FGenerationHelper& GenerationHelper, FP
 	: PackageData(&InPackageData)
 	, bCreateAsMap(false), bHasCreatedPackage(false), bTakenOverCachedCookedPlatformData(false)
 	, bIssuedUndeclaredMovedObjectsWarning(false), bGenerator(bInGenerator), bHasCalledPopulate(false)
+	, bCopyChunkAssignmentFromGenerator(true)
 {
 	// The Info for the generator is created in the constructor of the GenerationHelper, and PlatformDatas is not yet readable.
 	// The GenerationHelper constructor will set the PlatformDatas on this Info later on in its constructor.
diff --git a/Engine/Source/Editor/UnrealEd/Private/Cooker/CookGenerationHelper.h b/Engine/Source/Editor/UnrealEd/Private/Cooker/CookGenerationHelper.h
--- a/Engine/Source/Editor/UnrealEd/Private/Cooker/CookGenerationHelper.h
+++ b/Engine/Source/Editor/UnrealEd/Private/Cooker/CookGenerationHelper.h
@@ -96,6 +96,8 @@ public:
 	void SetIsGenerator(bool bValue);
 	bool HasCalledPopulate() const;
 	void SetHasCalledPopulate(bool bValue);
+	bool IsCopyChunkAssignmentFromGenerator() const;
+	void SetCopyChunkAssignmentFromGenerator(bool bValue);
 
 	/**
 	 * Steal the list of cached objects to call BeginCacheForCookedPlatformData on from the PackageData,
@@ -178,6 +180,7 @@ private:
 	bool bIssuedUndeclaredMovedObjectsWarning : 1;
 	bool bGenerator : 1;
 	bool bHasCalledPopulate : 1;
+	bool bCopyChunkAssignmentFromGenerator : 1;
 
 	friend UE::Cook::FGenerationHelper;
 };
@@ -693,6 +696,16 @@ inline void FCookGenerationInfo::SetHasCalledPopulate(bool bValue)
 	bHasCalledPopulate = bValue;
 }
 
+inline bool FCookGenerationInfo::IsCopyChunkAssignmentFromGenerator() const
+{
+	return bCopyChunkAssignmentFromGenerator;
+}
+
+inline void FCookGenerationInfo::SetCopyChunkAssignmentFromGenerator(bool bValue)
+{
+	bCopyChunkAssignmentFromGenerator = bValue;
+}
+
 inline TConstArrayView<FAssetDependency> FCookGenerationInfo::GetDependencies() const
 {
 	return PackageDependencies;
diff --git a/Engine/Source/Editor/UnrealEd/Public/Commandlets/AssetRegistryGenerator.h b/Engine/Source/Editor/UnrealEd/Public/Commandlets/AssetRegistryGenerator.h
--- a/Engine/Source/Editor/UnrealEd/Public/Commandlets/AssetRegistryGenerator.h
+++ b/Engine/Source/Editor/UnrealEd/Public/Commandlets/AssetRegistryGenerator.h
@@ -470,8 +470,12 @@ private:
 	/** Initialize ChunkIdPakchunkIndexMapping and PakchunkIndexChunkIdMapping. */
 	void InitializeChunkIdPakchunkIndexMapping();
 
-	/** If InState records PackageName is generated, return the name of the Generator, otherwise return NAME_None. */
-	static FName GetGeneratorPackage(FName PackageName, const FAssetRegistryState& InState);
+	/**
+	 * If InState records PackageName is generated, return the name of the Generator and the property flags of the Generator's
+	 * AssetRegistry dependency upon the generated package, otherwise return NAME_None and EDependencyProperty::None.
+	 */
+	static FName GetGeneratorPackage(FName PackageName, const FAssetRegistryState& InState,
+		UE::AssetRegistry::EDependencyProperty& OutDependencyPropertyOfGenerator);
 
 	friend class UE::Cook::FAssetRegistryReporterRemote;
 };
diff --git a/Engine/Source/Editor/UnrealEd/Public/CookPackageSplitter.h b/Engine/Source/Editor/UnrealEd/Public/CookPackageSplitter.h
--- a/Engine/Source/Editor/UnrealEd/Public/CookPackageSplitter.h
+++ b/Engine/Source/Editor/UnrealEd/Public/CookPackageSplitter.h
@@ -164,6 +164,20 @@ public:
 		/* GetGenerateList must specify true if the package will be a map (.umap, contains a UWorld or ULevel), else false */
 		void SetCreateAsMap(bool bInCreateAsMap) { bCreateAsMap = bInCreateAsMap; }
 		const TOptional<bool>& GetCreateAsMap() const { return bCreateAsMap; }
+
+		/**
+		 * This flag affects both the presence of a runtime dependency and the calculation of chunk assignments.
+		 * The runtime dependency is only relevant if the project's runtime or tools is reading data from the
+		 * generated AssetRegistry, for e.g. validation that all dependencies are staged.
+		 * The calculation of chunk assignments is only relevant if the project is using streaming manifests
+		 * (multiple pak and iostore files in the staged project).
+		 * If true, the generator's chunk assignment will be copied onto this generated package's chunk assignment, AND
+		 * the generated package will be marked as a runtime dependency of the generator.
+		 * If false, the generator's chunk assignment will instead be read from UAssetManager::GetPackageChunkIds for the
+		 * name of the generated package, and the generated package will be marked as an editoronly dependency of the
+		 * generator.
+		 */
+		bool bCopyChunkAssignmentFromGenerator = true;
 	private:
 		TOptional<bool> bCreateAsMap;
 	};
