When generating a chunk'd DLC cook do not write a pakchunk response file if the files in that chunk do not exist.
This was already done for a normal cook. Ignore making a file anyway even if there are pakChunkFile rules.

Cherry-picking shelve 44724238 by [REDACTED_PERSON].

#tests Local runs of Horde cook on Windows/IOS/Android

diff --git a/Engine/Source/Editor/UnrealEd/Private/Commandlets/AssetRegistryGenerator.cpp b/Engine/Source/Editor/UnrealEd/Private/Commandlets/AssetRegistryGenerator.cpp
--- a/Engine/Source/Editor/UnrealEd/Private/Commandlets/AssetRegistryGenerator.cpp
+++ b/Engine/Source/Editor/UnrealEd/Private/Commandlets/AssetRegistryGenerator.cpp
@@ -638,6 +638,13 @@ bool FAssetRegistryGenerator::GenerateStreamingInstallManifest(int64 InOverrideC
 
 	const TMap<int32, FString> PakChunkIdToStringOverride = GetPakChunkIdToStringOverrideMap();
 
+	static bool bIsDLCCook = []()
+		{
+			FString DLCNameReturn;
+			FParse::Value(FCommandLine::Get(), TEXT("DLCNAME="), DLCNameReturn);
+			return !DLCNameReturn.IsEmpty();
+		}();
+
 	// generate per-chunk pak list files
 	FDefaultPakFileRules DefaultPakFileRules;
 	bool bSucceeded = true;
@@ -687,6 +694,12 @@ bool FAssetRegistryGenerator::GenerateStreamingInstallManifest(int64 InOverrideC
 		// Do not create any files if the chunk is empty and is not referenced by rules applied during staging
 		if (ChunkFilenames.IsEmpty())
 		{
+			// During a chunk'd DLC cook ignore DefaultPakFile Rules. Those are handled during main cooks.
+			if (bIsDLCCook)
+			{
+				continue;
+			}
+
 			DefaultPakFileRules.InitializeFromConfig(TargetPlatform);
 			if (!DefaultPakFileRules.IsChunkReferenced(PakchunkIndex))
 			{
