Cooker: Bugfix for DLC IncrementalCook fallout: Fix the requestcluster code that was adding a null TargetPlatform accidentally to PackageDatas in response to the PlatformAgnosticPlatformIndex.
#rnx
#incrementalcook

diff --git a/Engine/Source/Editor/UnrealEd/Private/Cooker/CookRequestCluster.cpp b/Engine/Source/Editor/UnrealEd/Private/Cooker/CookRequestCluster.cpp
--- a/Engine/Source/Editor/UnrealEd/Private/Cooker/CookRequestCluster.cpp
+++ b/Engine/Source/Editor/UnrealEd/Private/Cooker/CookRequestCluster.cpp
@@ -2615,9 +2615,14 @@ void FRequestCluster::FQueryVertexBatch::Send()
 					EAsyncQueryStatus::AsyncRequested))
 				{
 					FFetchPlatformData& FetchPlatformData = ThreadSafeOnlyVars.FetchPlatforms[PlatformIndex];
-					FPackagePlatformData& PackagePlatformData = Vertex->GetPackageData()->FindOrAddPlatformData(FetchPlatformData.Platform);
+					EWhereCooked WhereCooked = EWhereCooked::ThisSession;
+					if (FetchPlatformData.Platform)
+					{
+						FPackagePlatformData& PackagePlatformData = Vertex->GetPackageData()->FindOrAddPlatformData(FetchPlatformData.Platform);
+						WhereCooked = PackagePlatformData.GetWhereCooked();
+					}
 
-					PlatformDatas[PlatformIndex].PackageIds.Add(FPackageIncrementalCookId{ Pair.Key, PackagePlatformData.GetWhereCooked() });
+					PlatformDatas[PlatformIndex].PackageIds.Add(FPackageIncrementalCookId{ Pair.Key, WhereCooked });
 					++NumAddedRequests;
 				}
 			}
