Cooker: Add asserts to enforce the contract that FPackageData does not have a PlatformData with TargetPlatform==nullptr.
#rnx
#incrementalcook

diff --git a/Engine/Source/Editor/UnrealEd/Private/Cooker/CookPackageData.cpp b/Engine/Source/Editor/UnrealEd/Private/Cooker/CookPackageData.cpp
--- a/Engine/Source/Editor/UnrealEd/Private/Cooker/CookPackageData.cpp
+++ b/Engine/Source/Editor/UnrealEd/Private/Cooker/CookPackageData.cpp
@@ -402,6 +402,7 @@ void FPackageData::SetPlatformCooked(
 	ECookResult CookResult, 
 	const bool bInWasCookedThisSession)
 {
+	check(TargetPlatform != nullptr);
 	bWasCookedThisSession |= bInWasCookedThisSession && (CookResult == ECookResult::Succeeded);
 
 	bool bNewCookAttemptedValue = (CookResult != ECookResult::NotAttempted);
@@ -448,6 +449,7 @@ void FPackageData::SetPlatformCooked(
 
 void FPackageData::SetPlatformCommitted(const ITargetPlatform* TargetPlatform)
 {
+	check(TargetPlatform != nullptr);
 	FPackagePlatformData& Value = PlatformDatas.FindOrAdd(TargetPlatform);
 	Value.SetCommitted(true);
 }
@@ -527,6 +529,7 @@ FPackageData::GetPlatformDatasConstKeysMutableValues()
 
 FPackagePlatformData& FPackageData::FindOrAddPlatformData(const ITargetPlatform* TargetPlatform)
 {
+	check(TargetPlatform != nullptr);
 	return PlatformDatas.FindOrAdd(TargetPlatform);
 }
 
