MPCook: Workaround assertion failure "Package %s was assigned to worker, but at sendmessage time it has no platforms needing commit....". Downgrade that assertion to a warning and handle it instead. The occurence seems to come from retraction, so add an earlier version of the warning during retraction to try to diagnose.
#rnx

diff --git a/Engine/Source/Editor/UnrealEd/Private/Cooker/CookDirector.cpp b/Engine/Source/Editor/UnrealEd/Private/Cooker/CookDirector.cpp
--- a/Engine/Source/Editor/UnrealEd/Private/Cooker/CookDirector.cpp
+++ b/Engine/Source/Editor/UnrealEd/Private/Cooker/CookDirector.cpp
@@ -2262,6 +2262,8 @@ FCookDirector::FRetractionHandler::ReassignPackages(const FWorkerId& FromWorker,
 	TConstArrayView<FPackageData*> Packages)
 {
 	TArray<TRefCountPtr<FCookWorkerServer>> LocalRemoteWorkers = Director.CopyRemoteWorkers();
+	ECookPhase CookPhase = Director.COTFS.GetCookPhase();
+	TConstArrayView<const ITargetPlatform*> SessionPlatforms = Director.COTFS.PlatformManager->GetSessionPlatforms();
 
 	TArray<FWorkerId> WorkersRequiredByConstraint;
 	TArray<FPackageData*> AssignmentPackages;
@@ -2272,6 +2274,38 @@ FCookDirector::FRetractionHandler::ReassignPackages(const FWorkerId& FromWorker,
 		{
 			continue;
 		}
+
+		// Handle and log diagnostics for an unexpected edge case we have seen: retracted packages already have been
+		// committed. Demote them to idle now rather than reassigning.
+		int32 NumNeedCommitPlatforms = 0;
+		for (const ITargetPlatform* SessionPlatform : SessionPlatforms)
+		{
+			FPackagePlatformData* PlatformData = PackageData->FindPlatformData(SessionPlatform);
+			if (PlatformData && PlatformData->NeedsCommit(SessionPlatform, CookPhase))
+			{
+				++NumNeedCommitPlatforms;
+			}
+		}
+		if (NumNeedCommitPlatforms == 0)
+		{
+			TStringBuilder<256> PlatformDataText;
+			for (const TPair<const ITargetPlatform*, FPackagePlatformData>& PlatformPair : PackageData->GetPlatformDatas())
+			{
+				PlatformDataText.Appendf(TEXT("{ %s: Reachable=%s, Committed=%s }, "),
+					PlatformPair.Key == CookerLoadingPlatformKey ? TEXT("CookerLoadingPlatform") : *PlatformPair.Key->PlatformName(),
+					PlatformPair.Value.IsReachable(CookPhase == ECookPhase::Cook
+						? EReachability::Runtime : EReachability::Build) ? TEXT("true") : TEXT("false"),
+					PlatformPair.Value.IsCommitted() ? TEXT("true") : TEXT("false"));
+			}
+			UE_LOG(LogCook, Warning,
+				TEXT("Package %s was retracted from %s, but already has no platforms needing commit. Demoting the package to idle instead of reassinging. State = %s. CookPhase = %s. WorkerAssignment = %s. [ %s ]"),
+				*PackageData->GetPackageName().ToString(), *Director.GetDisplayName(FromWorker), LexToString(PackageData->GetState()),
+				LexToString(CookPhase), *Director.GetDisplayName(PackageData->GetWorkerAssignment()), *PlatformDataText);
+			PackageData->SetWorkerAssignment(FWorkerId::Invalid());
+			Director.COTFS.DemoteToIdle(*PackageData, ESendFlags::QueueAddAndRemove, ESuppressCookReason::AlreadyCooked);
+			continue;
+		}
+
 		FWorkerId WorkerConstraint = PackageData->GetWorkerAssignmentConstraint();
 		if (WorkerConstraint.IsValid())
 		{
diff --git a/Engine/Source/Editor/UnrealEd/Private/Cooker/CookWorkerServer.cpp b/Engine/Source/Editor/UnrealEd/Private/Cooker/CookWorkerServer.cpp
--- a/Engine/Source/Editor/UnrealEd/Private/Cooker/CookWorkerServer.cpp
+++ b/Engine/Source/Editor/UnrealEd/Private/Cooker/CookWorkerServer.cpp
@@ -706,15 +706,7 @@ void FCookWorkerServer::SendPendingPackages()
 
 	for (FPackageData* PackageData : PackagesToAssign)
 	{
-		FAssignPackageData& AssignData = AssignDatas.Emplace_GetRef();
-		AssignData.ConstructData = PackageData->CreateConstructData();
-		AssignData.ParentGenerator = PackageData->GetParentGenerator();
-		AssignData.DoesGeneratedRequireGenerator = PackageData->DoesGeneratedRequireGenerator();
-		AssignData.Reachability = Reachability;
-		AssignData.Instigator = PackageData->GetInstigator(Reachability);
-		AssignData.Urgency = PackageData->GetUrgency();
 		SessionPlatformNeedsCommit.Init(false, OrderedSessionPlatforms.Num());
-
 		int32 PlatformIndex = 0;
 		int32 NumNeedCommitPlatforms = 0;
 		for (const ITargetPlatform* SessionPlatform : OrderedSessionPlatforms)
@@ -739,11 +731,29 @@ void FCookWorkerServer::SendPendingPackages()
 						? EReachability::Runtime : EReachability::Build) ? TEXT("true") : TEXT("false"),
 					PlatformPair.Value.IsCommitted() ? TEXT("true") : TEXT("false"));
 			}
-			checkf(false,
-				TEXT("Package %s was assigned to worker, but at sendmessage time it has no platforms needing commit. State = %s. CookPhase = %s. [ %s ]"),
+			UE_LOG(LogCook, Warning, TEXT("Package %s was assigned to worker, but at sendmessage time it has no platforms needing commit. State = %s. CookPhase = %s. [ %s ]"),
 				*PackageData->GetPackageName().ToString(), LexToString(PackageData->GetState()),
 					LexToString(CookPhase), *PlatformDataText);
+			// SendPendingMessages is called from the communication thread, so we cannot demote this package to idle
+			// immediately; that has to be done on the scheduler thread. To avoid adding special case code to handle
+			// this edge case, instead add a received message for it.
+			FPackageResultsMessage DemoteMessage;
+			FPackageRemoteResult& DemoteResult = DemoteMessage.Results.Emplace_GetRef();
+			DemoteResult.SetPackageName(PackageData->GetPackageName());
+			DemoteResult.SetSuppressCookReason(ESuppressCookReason::AlreadyCooked);
+			DemoteResult.SetPlatforms(OrderedSessionPlatforms);
+			ReceiveMessages.Add(MarshalToCompactBinaryTCP(DemoteMessage));
+			continue;
 		}
+
+		FAssignPackageData& AssignData = AssignDatas.Emplace_GetRef();
+		AssignData.ConstructData = PackageData->CreateConstructData();
+		AssignData.ParentGenerator = PackageData->GetParentGenerator();
+		AssignData.DoesGeneratedRequireGenerator = PackageData->DoesGeneratedRequireGenerator();
+		AssignData.Reachability = Reachability;
+		AssignData.Instigator = PackageData->GetInstigator(Reachability);
+		AssignData.Urgency = PackageData->GetUrgency();
+
 		AssignData.NeedCommitPlatforms = FDiscoveredPlatformSet(SessionPlatformNeedsCommit);
 		FAssignPackageExtraData* ExtraData = PackagesToAssignExtraDatas.Find(PackageData);
 		if (ExtraData)
