Cooker: RemapTargetPlatforms: Enforce the contract that FPackageData does not have a PlatformData with TargetPlatform==nullptr.
#rnx
#incrementalcook

diff --git a/Engine/Source/Editor/UnrealEd/Private/Cooker/CookPackageData.cpp b/Engine/Source/Editor/UnrealEd/Private/Cooker/CookPackageData.cpp
--- a/Engine/Source/Editor/UnrealEd/Private/Cooker/CookPackageData.cpp
+++ b/Engine/Source/Editor/UnrealEd/Private/Cooker/CookPackageData.cpp
@@ -1507,7 +1507,9 @@ void FPackageData::RemapTargetPlatforms(const TMap<ITargetPlatform*, ITargetPlat
 	NewPlatformDatas.Reserve(PlatformDatas.Num());
 	for (TPair<const ITargetPlatform*, FPackagePlatformData>& ExistingPair : PlatformDatas)
 	{
-		ITargetPlatform* NewKey = Remap[ExistingPair.Key];
+		ITargetPlatform*const* NewKeyPtr = Remap.Find(ExistingPair.Key);
+		ITargetPlatform* NewKey = NewKeyPtr ? *NewKeyPtr : const_cast<ITargetPlatform*>(ExistingPair.Key);
+		check(NewKey);
 		NewPlatformDatas.FindOrAdd(NewKey) = MoveTemp(ExistingPair.Value);
 
 		if (DiscoveredDependencies)
