Cooker: CookOnTheFly Handling: Fix typo that was causing unnecessary demotion in (currently unused) ITargetPlatform remapping code.
#rnx

diff --git a/Engine/Source/Editor/UnrealEd/Private/Cooker/CookPackageData.cpp b/Engine/Source/Editor/UnrealEd/Private/Cooker/CookPackageData.cpp
--- a/Engine/Source/Editor/UnrealEd/Private/Cooker/CookPackageData.cpp
+++ b/Engine/Source/Editor/UnrealEd/Private/Cooker/CookPackageData.cpp
@@ -1528,7 +1528,7 @@ void FPackageData::RemapTargetPlatforms(const TMap<ITargetPlatform*, ITargetPlat
 	// remaining unchanged. If we change that order due to the remap, we need to demote back to request.
 	if (IsInProgress() && GetState() != EPackageState::Request)
 	{
-		bool bDemote = true;
+		bool bDemote = false;
 		MapType::TConstIterator OldIter = PlatformDatas.CreateConstIterator();
 		MapType::TConstIterator NewIter = NewPlatformDatas.CreateConstIterator();
 		for (; OldIter; ++OldIter, ++NewIter)
