DiffPackageWriter Log format tweak, %lld for int64

diff --git a/Engine/Source/Editor/UnrealEd/Private/Cooker/DiffPackageWriter.cpp b/Engine/Source/Editor/UnrealEd/Private/Cooker/DiffPackageWriter.cpp
--- a/Engine/Source/Editor/UnrealEd/Private/Cooker/DiffPackageWriter.cpp
+++ b/Engine/Source/Editor/UnrealEd/Private/Cooker/DiffPackageWriter.cpp
@@ -691,7 +691,7 @@ void FDiffPackageWriter::LogAssetClassDiffSummary() const
 		if (Stats.TotalPackagesWithDifferences > 0)
 		{
 			const float AverageDifferenceCount = Stats.TotalDifferenceCount / static_cast<float>(Stats.TotalPackagesWithDifferences);
-			UE_LOG(LogDiff, Display, TEXT("\t%s | %d of %d | %.1f"),
+			UE_LOG(LogDiff, Display, TEXT("\t%s | %lld of %lld | %.1f"),
 				*AssetClassDifferences.Key.ToString(), Stats.TotalPackagesWithDifferences, Stats.TotalPackageCount, AverageDifferenceCount);
 		}
 	}
