Move core redirects dependencies to the global dependencies hash instead of adding it to each package dependencies.

Right now the core redirects dependencies are added to the stored key of each cooked package. There is no FCookDependency object for it so it doesn't show in the artifact. The result is I sometimes cook a package and the artifact doesn't change, except for the stored key, so I don't know what has changed and triggered the cook.

One downside is it will be added all the time in the global dependencies. Right now it is only added if there is at least one FCookDependency of type redirection. So it's possible we could run unnecessary recooks just because a core redirect we don't care about was changed. It can be fixed later if necessary, for example by adding certain global dependencies under certain condition.

Moving the core redirect dependencies to the global hash dependencies is still an improvement since it will make it visible (global hash dependencies are printed). Also it makes more sense to have it in the global hash dependencies since it's the same for the entire cook.

#rnx

diff --git a/Engine/Source/Editor/UnrealEd/Private/Cooker/CookGlobalDependencies.cpp b/Engine/Source/Editor/UnrealEd/Private/Cooker/CookGlobalDependencies.cpp
--- a/Engine/Source/Editor/UnrealEd/Private/Cooker/CookGlobalDependencies.cpp
+++ b/Engine/Source/Editor/UnrealEd/Private/Cooker/CookGlobalDependencies.cpp
@@ -10,6 +10,7 @@
 #include "Misc/ConfigCacheIni.h"
 #include "Misc/Guid.h"
 #include "Serialization/UnversionedPropertySerializationInternal.h"
+#include "UObject/CoreRedirects.h"
 
 namespace UE::Cook
 {
@@ -63,6 +64,17 @@ namespace UE::Cook
 			UE_LOG(LogCook, Display, TEXT("\tIsCookFlagSet(ECookInitializationFlags::Unversioned) = %s"), bIsUnversioned ? TEXT("True") : TEXT("False"));
 		}
 
+		{
+			FBlake3 GlobalRedirectHasher;
+			FCoreRedirects::AppendHashOfGlobalRedirects(GlobalRedirectHasher);
+
+			FBlake3Hash GlobalRedirectHash = GlobalRedirectHasher.Finalize();
+			Hasher.Update(&GlobalRedirectHash, sizeof(GlobalRedirectHash));
+
+			FString GlobalRedirectHashString = LexToString(GlobalRedirectHash);
+			UE_LOG(LogCook, Display, TEXT("\tFCoreRedirects::AppendHashOfGlobalRedirects = %s"), *GlobalRedirectHashString);
+		}
+
 		FBlake3Hash Hash = Hasher.Finalize();
 		UE_LOG(LogCook, Display, TEXT("\tGlobalDependenciesHash(%s) = %s"), *Platform->PlatformName(), *LexToString(Hash));
 		GlobalHash.Add(Platform, Hash);
diff --git a/Engine/Source/Editor/UnrealEd/Private/Cooker/CookPackageArtifacts.cpp b/Engine/Source/Editor/UnrealEd/Private/Cooker/CookPackageArtifacts.cpp
--- a/Engine/Source/Editor/UnrealEd/Private/Cooker/CookPackageArtifacts.cpp
+++ b/Engine/Source/Editor/UnrealEd/Private/Cooker/CookPackageArtifacts.cpp
@@ -217,7 +217,6 @@ FBuildDependencySet::ECurrentKeyResult FBuildDependencySet::TryCalculateCurrentK
 				KeyBuilder.Update(&Hash, sizeof(Hash));
 			}
 
-			FCoreRedirects::AppendHashOfGlobalRedirects(KeyBuilder);
 			break;
 		}
 
