Cooker: Softlock detection: Suppress the warning (or display) of ""Cooker has been blocked from saving the current packages..." when we detect that some UObjects are finishing up their IsCachedCookedPlatformData completed calls. The cooker should only give this display when it is making no progress at all. This change is important for large maps that may have thousands of expensive objects to cache.
#rnx

diff --git a/Engine/Source/Editor/UnrealEd/Classes/CookOnTheSide/CookOnTheFlyServer.h b/Engine/Source/Editor/UnrealEd/Classes/CookOnTheSide/CookOnTheFlyServer.h
--- a/Engine/Source/Editor/UnrealEd/Classes/CookOnTheSide/CookOnTheFlyServer.h
+++ b/Engine/Source/Editor/UnrealEd/Classes/CookOnTheSide/CookOnTheFlyServer.h
@@ -495,7 +495,7 @@ private:
 
 	/** Execute operations that need to be done after each Scheduler task, such as checking for new external requests. */
 	void TickCookStatus(UE::Cook::FTickStackData& StackData);
-	void SetSaveBusy(bool bInSaveBusy);
+	void SetSaveBusy(bool bInSaveBusy, int32 NumAsyncWorkRetired);
 	void SetLoadBusy(bool bInLoadBusy);
 	enum class EIdleStatus
 	{
@@ -558,9 +558,12 @@ private:
 	 * Try to save all packages in the SaveQueue until it's time to break.
 	 * Report the number of requests that were completed (either skipped or successfully saved or failed to save)
 	 */
-	void PumpSaves(UE::Cook::FTickStackData& StackData, uint32 DesiredQueueLength, int32& OutNumPushed, bool& bOutBusy);
-	void PumpRuntimeSaves(UE::Cook::FTickStackData& StackData, uint32 DesiredQueueLength, int32& OutNumPushed, bool& bOutBusy);
-	void PumpBuildDependencySaves(UE::Cook::FTickStackData& StackData, uint32 DesiredQueueLength, int32& OutNumPushed, bool& bOutBusy);
+	void PumpSaves(UE::Cook::FTickStackData& StackData, uint32 DesiredQueueLength, int32& OutNumPushed,
+		bool& bOutBusy, int32& OutNumAsyncWorkRetired);
+	void PumpRuntimeSaves(UE::Cook::FTickStackData& StackData, uint32 DesiredQueueLength, int32& OutNumPushed,
+		bool& bOutBusy, int32& OutNumAsyncWorkRetired);
+	void PumpBuildDependencySaves(UE::Cook::FTickStackData& StackData, uint32 DesiredQueueLength, int32& OutNumPushed,
+		bool& bOutBusy, int32& OutNumAsyncWorkRetired);
 
 	/** Find all packages that are build dependencies and request them for cook. */
 	void KickBuildDependencies(UE::Cook::FTickStackData& StackData);
@@ -1570,7 +1573,7 @@ private:
 	/** True when PumpSaves has detected it is blocked on async work and CookOnTheFlyServer should do work elsewhere. */
 	bool bSaveBusy = false;
 	/** We need to track whether the compiler has been inactive for a long time before issuing a warning about it. */
-	bool bShaderCompilerWasActiveeOnPreviousBusyReport = true;
+	bool bShaderCompilerWasActiveOnPreviousBusyReport = true;
 	/**
 	 * After cooking all requested packages in CookByTheBook we need to search for any build dependencies and request them
 	 * for commit as build dependencies. This bool records whether we have done that for the current session.
diff --git a/Engine/Source/Editor/UnrealEd/Private/CookOnTheFlyServer.cpp b/Engine/Source/Editor/UnrealEd/Private/CookOnTheFlyServer.cpp
--- a/Engine/Source/Editor/UnrealEd/Private/CookOnTheFlyServer.cpp
+++ b/Engine/Source/Editor/UnrealEd/Private/CookOnTheFlyServer.cpp
@@ -1573,30 +1573,29 @@ void UCookOnTheFlyServer::TickMainCookLoop(UE::Cook::FTickStackData& StackData)
 				SetLoadBusy(false);
 			}
 			break;
-		case ECookAction::Load:
-			PumpLoads(StackData, 0, NumPushed, bBusy);
-			SetLoadBusy(bBusy && NumPushed == 0); // Mark as busy if pump was blocked and we did not make any progress
-			if (NumPushed > 0)
-			{
-				SetSaveBusy(false);
-			}
-			break;
+		case ECookAction::Load: [[fallthrough]];
 		case ECookAction::LoadLimited:
-			PumpLoads(StackData, DesiredLoadQueueLength, NumPushed, bBusy);
+		{
+			uint32 CurrentLoadQueueLength = CookAction == ECookAction::Load ? 0 : DesiredLoadQueueLength;
+			PumpLoads(StackData, CurrentLoadQueueLength, NumPushed, bBusy);
 			SetLoadBusy(bBusy && NumPushed == 0); // Mark as busy if pump was blocked and we did not make any progress
 			if (NumPushed > 0)
 			{
-				SetSaveBusy(false);
+				SetSaveBusy(false, 0);
 			}
 			break;
-		case ECookAction::Save:
-			PumpSaves(StackData, 0, NumPushed, bBusy);
-			SetSaveBusy(bBusy && NumPushed == 0); // Mark as busy if pump was blocked and we did not make any progress
-			break;
+		}
+		case ECookAction::Save: [[fallthrough]];
 		case ECookAction::SaveLimited:
-			PumpSaves(StackData, DesiredSaveQueueLength, NumPushed, bBusy);
-			SetSaveBusy(bBusy && NumPushed == 0); // Mark as busy if pump was blocked and we did not make any progress
+		{
+			uint32 CurrentSaveQueueLength = CookAction == ECookAction::Save ? 0 : DesiredSaveQueueLength;
+			int32 NumAsyncWorkRetired;
+			PumpSaves(StackData, CurrentSaveQueueLength, NumPushed, bBusy, NumAsyncWorkRetired);
+			// Mark save as busy if pump was blocked and we have no more work we can do in pumpsaves.
+			// Also report NumAsyncWorkRetired for status that cares about whether we are making some progress.
+			SetSaveBusy(bBusy && NumPushed == 0, NumAsyncWorkRetired);
 			break;
+		}
 		case ECookAction::Poll:
 			PumpPollables(StackData, false /* bIsIdle */);
 			break;
@@ -1644,167 +1643,172 @@ void UCookOnTheFlyServer::TickCookStatus(UE::Cook::FTickStackData& StackData)
 	PumpExternalRequests(StackData.Timer);
 }
 
-void UCookOnTheFlyServer::SetSaveBusy(bool bInBusy)
+void UCookOnTheFlyServer::SetSaveBusy(bool bInBusy, int32 NumAsyncWorkRetired)
 {
 	using namespace UE::Cook;
 
-	if (bSaveBusy != bInBusy)
+	if (!bInBusy)
 	{
-		bSaveBusy = bInBusy;
 		if (bSaveBusy)
 		{
-			const double CurrentTime = FPlatformTime::Seconds();
-			SaveBusyStartTimeSeconds = CurrentTime;
-			SaveBusyRetryTimeSeconds = CurrentTime + CookProgressRetryBusyPeriodSeconds;
-			SaveBusyWarnTimeSeconds = CurrentTime + GCookProgressWarnBusyTime;
-		}
-		else
-		{
+			bSaveBusy = false;
 			SaveBusyStartTimeSeconds = MAX_flt;
 			SaveBusyRetryTimeSeconds = MAX_flt;
 			SaveBusyWarnTimeSeconds = MAX_flt;
 			// Whenever we set Save back to non-busy, reset the counter for how many busy reports with an
 			// idle shadercompiler we need before we issue a warning
-			bShaderCompilerWasActiveeOnPreviousBusyReport = true;
+			bShaderCompilerWasActiveOnPreviousBusyReport = true;
 		}
+		return;
 	}
-	else if (bSaveBusy)
+
+	const double CurrentTime = FPlatformTime::Seconds();
+	if (!bSaveBusy)
 	{
-		const double CurrentTime = FPlatformTime::Seconds();
-		SaveBusyRetryTimeSeconds = CurrentTime + CookProgressRetryBusyPeriodSeconds;
-		if (CurrentTime >= SaveBusyWarnTimeSeconds)
-		{
-			// Compiler users - classes using the shader compiler - can take multiple minutes to be compiled due to long
-			// queues and long compile times, so we do not issue a warning when they are the only objects holding us up,
-			// so long as the shadercompiler reports it is working on them.
-			bool bShaderCompilerIsActive = GShaderCompilingManager->IsCompiling();
-			bool bBusyCompilationUsersAreExpected = bShaderCompilerIsActive ||
-				// Even if the ShaderCompilerManager is not currently compiling, it might shortly begin or have recently finished.
-				// Issue a warning for blocked compiler users only if there are two reports in a row where the compiler is not active
-				bShaderCompilerWasActiveeOnPreviousBusyReport;
-			bShaderCompilerWasActiveeOnPreviousBusyReport = bShaderCompilerIsActive;
-
-			// Issue a status update. For each UObject we're still waiting on, check whether the long duration is expected using type-specific checks
-			// Make the status update a warning if the long duration is not reported as expected.
-			TArray<UObject*> NonExpectedObjects;
-			TSet<UPackage*> NonExpectedPackages;
-			TArray<UObject*> ExpectedObjects;
-			TSet<UPackage*> ExpectedPackages;
-			FPackageDataQueue& SaveQueue = PackageDatas->GetSaveQueue();
-			TArray<UClass*> CompilationUsers({ UMaterialInterface::StaticClass(), FindObject<UClass>(nullptr, TEXT("/Script/Niagara.NiagaraScript")) });
+		bSaveBusy = true;
+		SaveBusyStartTimeSeconds = CurrentTime;
+		SaveBusyWarnTimeSeconds = CurrentTime + GCookProgressWarnBusyTime;
+	}
+	SaveBusyRetryTimeSeconds = CurrentTime + CookProgressRetryBusyPeriodSeconds;
 
-			PackageDatas->ForEachPendingCookedPlatformData(
-			[&CompilationUsers, &ExpectedObjects, &ExpectedPackages, &NonExpectedObjects, &NonExpectedPackages, bBusyCompilationUsersAreExpected]
-			(const FPendingCookedPlatformData& Data)
+	if (NumAsyncWorkRetired)
+	{
+		// Some progress was made so reset the warning time, and the start time in case we later issue a warning.
+		SaveBusyStartTimeSeconds = CurrentTime;
+		SaveBusyWarnTimeSeconds = CurrentTime + GCookProgressWarnBusyTime;
+	}
+
+	if (CurrentTime >= SaveBusyWarnTimeSeconds)
+	{
+		// Compiler users - classes using the shader compiler - can take multiple minutes to be compiled due to long
+		// queues and long compile times, so we do not issue a warning when they are the only objects holding us up,
+		// so long as the shadercompiler reports it is working on them.
+		bool bShaderCompilerIsActive = GShaderCompilingManager->IsCompiling();
+		bool bBusyCompilationUsersAreExpected = bShaderCompilerIsActive ||
+			// Even if the ShaderCompilerManager is not currently compiling, it might shortly begin or have recently finished.
+			// Issue a warning for blocked compiler users only if there are two reports in a row where the compiler is not active
+			bShaderCompilerWasActiveOnPreviousBusyReport;
+		bShaderCompilerWasActiveOnPreviousBusyReport = bShaderCompilerIsActive;
+
+		// Issue a status update. For each UObject we're still waiting on, check whether the long duration is expected using type-specific checks
+		// Make the status update a warning if the long duration is not reported as expected.
+		TArray<UObject*> NonExpectedObjects;
+		TSet<UPackage*> NonExpectedPackages;
+		TArray<UObject*> ExpectedObjects;
+		TSet<UPackage*> ExpectedPackages;
+		FPackageDataQueue& SaveQueue = PackageDatas->GetSaveQueue();
+		TArray<UClass*> CompilationUsers({ UMaterialInterface::StaticClass(), FindObject<UClass>(nullptr, TEXT("/Script/Niagara.NiagaraScript")) });
+
+		PackageDatas->ForEachPendingCookedPlatformData(
+		[&CompilationUsers, &ExpectedObjects, &ExpectedPackages, &NonExpectedObjects, &NonExpectedPackages, bBusyCompilationUsersAreExpected]
+		(const FPendingCookedPlatformData& Data)
+		{
+			UObject* Object = Data.Object.Get();
+			if (!Object)
 			{
-				UObject* Object = Data.Object.Get();
-				if (!Object)
-				{
-					return;
-				}
-				bool bCompilationUser = false;
-				for (UClass* CompilationUserClass : CompilationUsers)
-				{
-					if (Object->IsA(CompilationUserClass))
-					{
-						bCompilationUser = true;
-						break;
-					}
-				}
-				if (bCompilationUser && bBusyCompilationUsersAreExpected)
-				{
-					ExpectedObjects.Add(Object);
-					ExpectedPackages.Add(Object->GetPackage());
-				}
-				else
-				{
-					NonExpectedObjects.Add(Object);
-					NonExpectedPackages.Add(Object->GetPackage());
-				}
-			});
-			TArray<UPackage*> RemovePackages;
-			for (UPackage* Package : ExpectedPackages)
+				return;
+			}
+			bool bCompilationUser = false;
+			for (UClass* CompilationUserClass : CompilationUsers)
 			{
-				if (NonExpectedPackages.Contains(Package))
+				if (Object->IsA(CompilationUserClass))
 				{
-					RemovePackages.Add(Package);
+					bCompilationUser = true;
+					break;
 				}
 			}
-			for (UPackage* Package : RemovePackages)
+			if (bCompilationUser && bBusyCompilationUsersAreExpected)
 			{
-				ExpectedPackages.Remove(Package);
+				ExpectedObjects.Add(Object);
+				ExpectedPackages.Add(Object->GetPackage());
 			}
-			
-			bool bExpectedDueToSlowBuildOperations = !ExpectedObjects.IsEmpty() && NonExpectedObjects.IsEmpty();
-			
-			FString Message = FString::Printf(TEXT("Cooker has been blocked from saving the current packages for %.0f seconds."),
-				(float)CurrentTime - SaveBusyStartTimeSeconds);
-			ELogVerbosity::Type MessageSeverity = ELogVerbosity::Display;
-			if (!bExpectedDueToSlowBuildOperations)
+			else
 			{
-				MessageSeverity = CookerIdleWarningSeverity;
+				NonExpectedObjects.Add(Object);
+				NonExpectedPackages.Add(Object->GetPackage());
 			}
+		});
+		TArray<UPackage*> RemovePackages;
+		for (UPackage* Package : ExpectedPackages)
+		{
+			if (NonExpectedPackages.Contains(Package))
+			{
+				RemovePackages.Add(Package);
+			}
+		}
+		for (UPackage* Package : RemovePackages)
+		{
+			ExpectedPackages.Remove(Package);
+		}
+			
+		bool bExpectedDueToSlowBuildOperations = !ExpectedObjects.IsEmpty() && NonExpectedObjects.IsEmpty();
+			
+		FString Message = FString::Printf(TEXT("Cooker has been blocked from saving the current packages for %.0f seconds."),
+			(float)CurrentTime - SaveBusyStartTimeSeconds);
+		ELogVerbosity::Type MessageSeverity = ELogVerbosity::Display;
+		if (!bExpectedDueToSlowBuildOperations)
+		{
+			MessageSeverity = CookerIdleWarningSeverity;
+		}
 #if !NO_LOGGING
-			FMsg::Logf(__FILE__, __LINE__, LogCook.GetCategoryName(), MessageSeverity, TEXT("%s"), *Message);
+		FMsg::Logf(__FILE__, __LINE__, LogCook.GetCategoryName(), MessageSeverity, TEXT("%s"), *Message);
 #endif
 
-			UE_LOG(LogCook, Display, TEXT("%d packages in the savequeue: "), SaveQueue.Num());
-			int DisplayCount = 0;
-			const int DisplayMax = 10;
-			for (TSet<UPackage*>* PackageSet : { &NonExpectedPackages, &ExpectedPackages })
+		UE_LOG(LogCook, Display, TEXT("%d packages in the savequeue: "), SaveQueue.Num());
+		int DisplayCount = 0;
+		const int DisplayMax = 10;
+		for (TSet<UPackage*>* PackageSet : { &NonExpectedPackages, &ExpectedPackages })
+		{
+			for (UPackage* Package : *PackageSet)
 			{
-				for (UPackage* Package : *PackageSet)
+				if (DisplayCount == DisplayMax)
 				{
-					if (DisplayCount == DisplayMax)
-					{
-						UE_LOG(LogCook, Display, TEXT("    ..."));
-						break;
-					}
-					UE_LOG(LogCook, Display, TEXT("    %s"), *Package->GetName());
-					++DisplayCount;
+					UE_LOG(LogCook, Display, TEXT("    ..."));
+					break;
 				}
+				UE_LOG(LogCook, Display, TEXT("    %s"), *Package->GetName());
+				++DisplayCount;
 			}
-			if (DisplayCount == 0)
-			{
-				UE_LOG(LogCook, Display, TEXT("    <None>"));
-			}
+		}
+		if (DisplayCount == 0)
+		{
+			UE_LOG(LogCook, Display, TEXT("    <None>"));
+		}
 
-			UE_LOG(LogCook, Display, TEXT("%d objects that have not yet returned true from IsCachedCookedPlatformDataLoaded:"),
-				PackageDatas->GetPendingCookedPlatformDataNum());
-			DisplayCount = 0;
-			for (TArray<UObject*>* ObjectArray : {  &NonExpectedObjects, &ExpectedObjects })
+		UE_LOG(LogCook, Display, TEXT("%d objects that have not yet returned true from IsCachedCookedPlatformDataLoaded:"),
+			PackageDatas->GetPendingCookedPlatformDataNum());
+		DisplayCount = 0;
+		for (TArray<UObject*>* ObjectArray : {  &NonExpectedObjects, &ExpectedObjects })
+		{
+			for (UObject* Object : *ObjectArray)
 			{
-				for (UObject* Object : *ObjectArray)
+				if (DisplayCount == DisplayMax)
 				{
-					if (DisplayCount == DisplayMax)
-					{
-						UE_LOG(LogCook, Display, TEXT("    ..."));
-						break;
-					}
-					UE_LOG(LogCook, Display, TEXT("    %s"), *Object->GetFullName());
+					UE_LOG(LogCook, Display, TEXT("    ..."));
+					break;
+				}
+				UE_LOG(LogCook, Display, TEXT("    %s"), *Object->GetFullName());
 
-					TStringBuilder<2048> AdditionalDebugInfo;
-					FDelegates::PackageBlocked.Broadcast(Object, AdditionalDebugInfo);
-					FStringView AdditionalDebugInfoStr = AdditionalDebugInfo.ToView();
-					UE::String::ParseTokens(AdditionalDebugInfo.ToView(), TEXT('\n'),
-						[](FStringView Line)
-						{
-							UE_LOG(LogCook, Display, TEXT("        %.*s"), Line.Len(), Line.GetData());
-						}, UE::String::EParseTokensOptions::SkipEmpty);
+				TStringBuilder<2048> AdditionalDebugInfo;
+				FDelegates::PackageBlocked.Broadcast(Object, AdditionalDebugInfo);
+				UE::String::ParseTokens(AdditionalDebugInfo.ToView(), TEXT('\n'),
+					[](FStringView Line)
+					{
+						UE_LOG(LogCook, Display, TEXT("        %.*s"), Line.Len(), Line.GetData());
+					}, UE::String::EParseTokensOptions::SkipEmpty);
 
-					++DisplayCount;
-				}
-			}
-			if (DisplayCount == 0)
-			{
-				UE_LOG(LogCook, Display, TEXT("    <None>"));
+				++DisplayCount;
 			}
+		}
+		if (DisplayCount == 0)
+		{
+			UE_LOG(LogCook, Display, TEXT("    <None>"));
+		}
 
-			SaveBusyWarnTimeSeconds = CurrentTime + GCookProgressWarnBusyTime;
+		SaveBusyWarnTimeSeconds = CurrentTime + GCookProgressWarnBusyTime;
 			
-			UE::Cook::FDelegates::CookSaveIdle.Broadcast(
-				*this, SaveQueue.Num(), PackageDatas->GetPendingCookedPlatformDataNum(), bExpectedDueToSlowBuildOperations);
-		}
+		UE::Cook::FDelegates::CookSaveIdle.Broadcast(
+			*this, SaveQueue.Num(), PackageDatas->GetPendingCookedPlatformDataNum(), bExpectedDueToSlowBuildOperations);
 	}
 }
 
@@ -4571,7 +4575,8 @@ void UCookOnTheFlyServer::ReleaseCookedPlatformData(UE::Cook::FPackageData& Pack
 
 void UCookOnTheFlyServer::TickCancels()
 {
-	PackageDatas->PollPendingCookedPlatformDatas(false, LastCookableObjectTickTime);
+	int32 UnusedNumRetired;
+	PackageDatas->PollPendingCookedPlatformDatas(false, LastCookableObjectTickTime, UnusedNumRetired);
 }
 
 bool UCookOnTheFlyServer::LoadPackageForCooking(UE::Cook::FPackageData& PackageData, UPackage*& OutPackage,
@@ -4881,23 +4886,27 @@ void UCookOnTheFlyServer::ProcessUnsolicitedPackages(TArray<FName>* OutDiscovere
 	}
 }
 
-void UCookOnTheFlyServer::PumpSaves(UE::Cook::FTickStackData& StackData, uint32 DesiredQueueLength, int32& OutNumPushed, bool& bOutBusy)
+void UCookOnTheFlyServer::PumpSaves(UE::Cook::FTickStackData& StackData, uint32 DesiredQueueLength,
+	int32& OutNumPushed, bool& bOutBusy, int32& OutNumAsyncWorkRetired)
 {
 	if (GetCookPhase() == UE::Cook::ECookPhase::Cook)
 	{
-		PumpRuntimeSaves(StackData, DesiredQueueLength, OutNumPushed, bOutBusy);
+		PumpRuntimeSaves(StackData, DesiredQueueLength, OutNumPushed, bOutBusy, OutNumAsyncWorkRetired);
 	}
 	else
 	{
-		// After we enter the build dependencies phase, we no longer save packages, we just commit them with dependencies but no cooked data.
-		PumpBuildDependencySaves(StackData, DesiredQueueLength, OutNumPushed, bOutBusy);
+		// After we enter the build dependencies phase, we no longer save packages,
+		// we just commit them with dependencies but no cooked data.
+		PumpBuildDependencySaves(StackData, DesiredQueueLength, OutNumPushed, bOutBusy, OutNumAsyncWorkRetired);
 	}
 }
 
-void UCookOnTheFlyServer::PumpRuntimeSaves(UE::Cook::FTickStackData& StackData, uint32 DesiredQueueLength, int32& OutNumPushed, bool& bOutBusy)
+void UCookOnTheFlyServer::PumpRuntimeSaves(UE::Cook::FTickStackData& StackData, uint32 DesiredQueueLength,
+	int32& OutNumPushed, bool& bOutBusy, int32& OutNumAsyncWorkRetired)
 {
 	using namespace UE::Cook;
 	OutNumPushed = 0;
+	OutNumAsyncWorkRetired = 0;
 	bOutBusy = false;
 
 	UE_SCOPED_HIERARCHICAL_COOKTIMER(SavingPackages);
@@ -4993,7 +5002,9 @@ void UCookOnTheFlyServer::PumpRuntimeSaves(UE::Cook::FTickStackData& StackData,
 		bool bForce = IsCookOnTheFlyMode() && !IsRealtimeMode();
 		{
 			UE_SCOPED_HIERARCHICAL_COOKTIMER(PollPendingCookedPlatformDatas);
-			PackageDatas->PollPendingCookedPlatformDatas(bForce, LastCookableObjectTickTime);
+			int32 NumAsyncWorkRetired;
+			PackageDatas->PollPendingCookedPlatformDatas(bForce, LastCookableObjectTickTime, NumAsyncWorkRetired);
+			OutNumAsyncWorkRetired += NumAsyncWorkRetired;
 		}
 
 		// If BeginCacheCookPlatformData is not ready then postpone the package, exit, or wait for it as appropriate
@@ -5054,7 +5065,9 @@ void UCookOnTheFlyServer::PumpRuntimeSaves(UE::Cook::FTickStackData& StackData,
 					// sleep for a bit
 					FPlatformProcess::Sleep(0.0f);
 					// Poll the results again and check whether we are now done
-					PackageDatas->PollPendingCookedPlatformDatas(true, LastCookableObjectTickTime);
+					int32 NumAsyncWorkRetired;
+					PackageDatas->PollPendingCookedPlatformDatas(true, LastCookableObjectTickTime, NumAsyncWorkRetired);
+					OutNumAsyncWorkRetired += NumAsyncWorkRetired;
 					PrepareSaveStatus = PrepareSave(PackageData, StackData.Timer, false /* bPrecaching */,
 						DemotionReason);
 				} while (!StackData.Timer.IsActionTimeUp() && PrepareSaveStatus == EPollStatus::Incomplete
@@ -5132,11 +5145,13 @@ void UCookOnTheFlyServer::PumpRuntimeSaves(UE::Cook::FTickStackData& StackData,
 	}
 }
 
-void UCookOnTheFlyServer::PumpBuildDependencySaves(UE::Cook::FTickStackData& StackData, uint32 DesiredQueueLength, int32& OutNumPushed, bool& bOutBusy)
+void UCookOnTheFlyServer::PumpBuildDependencySaves(UE::Cook::FTickStackData& StackData, uint32 DesiredQueueLength,
+	int32& OutNumPushed, bool& bOutBusy, int32& OutNumAsyncWorkRetired)
 {
 	using namespace UE::Cook;
 	OutNumPushed = 0;
 	bOutBusy = false;
+	OutNumAsyncWorkRetired = 0;
 
 	UE_SCOPED_HIERARCHICAL_COOKTIMER(SavingPackages);
 	check(IsInGameThread());
@@ -10479,7 +10494,7 @@ void UCookOnTheFlyServer::CancelAllQueues()
 	}
 
 	SetLoadBusy(false);
-	SetSaveBusy(false);
+	SetSaveBusy(false, 0);
 }
 
 
@@ -13108,7 +13123,13 @@ EPackageWriterResult UCookOnTheFlyServer::SavePackageBeginCacheForCookedPlatform
 		for (;;)
 		{
 			UE_SCOPED_HIERARCHICAL_COOKTIMER(PollPendingCookedPlatformDatas);
-			PackageDatas->PollPendingCookedPlatformDatas(true /* bForce */, LastCookableObjectTickTime);
+			// Note that we don't report the NumAsyncWorkRetired value here to the PumpSaves output variable.
+			// That reporting is currently only necessary for softlock detection, and since we're in the middle
+			// of a save we will already be reporting some progress and don't need to report the additional
+			// progress of NumAsyncWorkRetired occurrences.
+			int32 NumAsyncWorkRetired;
+			PackageDatas->PollPendingCookedPlatformDatas(true /* bForce */, LastCookableObjectTickTime,
+				NumAsyncWorkRetired);
 			if (PackageData->GetNumPendingCookedPlatformData() == 0)
 			{
 				break;
diff --git a/Engine/Source/Editor/UnrealEd/Private/Cooker/CookGarbageCollect.cpp b/Engine/Source/Editor/UnrealEd/Private/Cooker/CookGarbageCollect.cpp
--- a/Engine/Source/Editor/UnrealEd/Private/Cooker/CookGarbageCollect.cpp
+++ b/Engine/Source/Editor/UnrealEd/Private/Cooker/CookGarbageCollect.cpp
@@ -971,7 +971,7 @@ void UCookOnTheFlyServer::PostGarbageCollect()
 
 	// Whenever we collect garbage, reset the counter for how many busy reports with an
 	// idle shadercompiler we need before we issue a warning
-	bShaderCompilerWasActiveeOnPreviousBusyReport = true;
+	bShaderCompilerWasActiveOnPreviousBusyReport = true;
 }
 
 bool UCookOnTheFlyServer::NeedsDiagnosticSecondGC() const
diff --git a/Engine/Source/Editor/UnrealEd/Private/Cooker/CookPackageData.cpp b/Engine/Source/Editor/UnrealEd/Private/Cooker/CookPackageData.cpp
--- a/Engine/Source/Editor/UnrealEd/Private/Cooker/CookPackageData.cpp
+++ b/Engine/Source/Editor/UnrealEd/Private/Cooker/CookPackageData.cpp
@@ -2860,8 +2860,10 @@ void FPackageDatas::AddPendingCookedPlatformData(FPendingCookedPlatformData&& Da
 	++PendingCookedPlatformDataNum;
 }
 
-void FPackageDatas::PollPendingCookedPlatformDatas(bool bForce, double& LastCookableObjectTickTime)
+void FPackageDatas::PollPendingCookedPlatformDatas(bool bForce, double& LastCookableObjectTickTime,
+	int32& OutNumRetired)
 {
+	OutNumRetired = 0;
 	if (PendingCookedPlatformDataNum == 0)
 	{
 		return;
@@ -2942,6 +2944,7 @@ void FPackageDatas::PollPendingCookedPlatformDatas(bool bForce, double& LastCook
 				// We are destructing all elements of List after the for loop is done; we leave
 				// the completed Data on List to be destructed.
 				--PendingCookedPlatformDataNum;
+				++OutNumRetired;
 			}
 			else
 			{
@@ -2971,6 +2974,7 @@ void FPackageDatas::PollPendingCookedPlatformDatas(bool bForce, double& LastCook
 				{
 					ForceList.RemoveAtSwap(Index, EAllowShrinking::No);
 					--PendingCookedPlatformDataNum;
+					++OutNumRetired;
 				}
 				else
 				{
diff --git a/Engine/Source/Editor/UnrealEd/Private/Cooker/CookPackageData.h b/Engine/Source/Editor/UnrealEd/Private/Cooker/CookPackageData.h
--- a/Engine/Source/Editor/UnrealEd/Private/Cooker/CookPackageData.h
+++ b/Engine/Source/Editor/UnrealEd/Private/Cooker/CookPackageData.h
@@ -1437,7 +1437,7 @@ public:
 	 * Iterate over all elements in PendingCookedPlatformDatas and check whether they have completed,
 	 * releasing their resources and pending count if so.
 	 */
-	void PollPendingCookedPlatformDatas(bool bForce, double& LastCookableObjectTickTime);
+	void PollPendingCookedPlatformDatas(bool bForce, double& LastCookableObjectTickTime, int32& OutNumRetired);
 	void ClearCancelManager(FPackageData& PackageData);
 
 	/** Swap all ITargetPlatform* stored on this instance according to the mapping in @param Remap. */
