Core: Added core editor multicast delegates OnPreWorldGenerationProcess/OnPostWorldGenerationProcess that should be broadcasted by any type of data generation in the editor to reduce the amount of individual work doing on generated elements.
PCG: Graph executor now broadcast OnPreWorldGenerationProcess/OnPostWorldGenerationProcess delegates for each iteration in the editor.
Entity: Fixed UE::SceneGraphAPI::EndBatchOperation superfluous update calls when recreating entity actor proxies.

Baseline
- Default time slice: 30.6 sec
- Infinite time slice: 27.7 sec
PCG Broadcast Optim
- Default time slice: 20.3 sec
- Infinite time slice: 10.7 sec
EndBatchOperation Optim
- Default time slice: 14.9 sec
- Infinite time slice: 6.5 sec

diff --git a/Engine/Plugins/PCG/Source/PCG/Private/Graph/PCGGraphExecutor.cpp b/Engine/Plugins/PCG/Source/PCG/Private/Graph/PCGGraphExecutor.cpp
--- a/Engine/Plugins/PCG/Source/PCG/Private/Graph/PCGGraphExecutor.cpp
+++ b/Engine/Plugins/PCG/Source/PCG/Private/Graph/PCGGraphExecutor.cpp
@@ -1925,6 +1925,14 @@ void FPCGGraphExecutor::Execute(double& InOutEndTime)
 	bool bContinueExecute = GetNonScheduledRemainingTaskCount() > 0;
 	bool bFirstLoop = true;
 
+#if WITH_EDITOR
+	bool bOnWorldGenerationProcessBroadcasted = bContinueExecute && World;
+	if (bOnWorldGenerationProcessBroadcasted)
+	{
+		FEditorDelegates::OnPreWorldGenerationProcess.Broadcast(World);
+	}
+#endif // WITH_EDITOR
+
 	while ((bFirstLoop || (FPlatformTime::Seconds() < InOutEndTime)) && bContinueExecute)
 	{
 		TSharedPtr<FPCGGraphActiveTask> MainThreadTask;
@@ -1944,6 +1952,13 @@ void FPCGGraphExecutor::Execute(double& InOutEndTime)
 	}
 
 	ExecuteTasksEnded();
+
+#if WITH_EDITOR
+	if (bOnWorldGenerationProcessBroadcasted)
+	{
+		FEditorDelegates::OnPostWorldGenerationProcess.Broadcast(World);
+	}
+#endif // WITH_EDITOR
 }
 
 void FPCGGraphExecutor::ClearAllTasks()
@@ -3338,4 +3353,4 @@ FPCGElementPtr UPCGGridLinkageSettings::CreateElement() const
 	};
 
 	return MakeShared<PCGGraphExecutor::FPCGGridLinkageElement>(GridLinkageOperation, ContextAllocator, FromGrid, ToGrid, GenerationGrid, ResourceKey, ResolvedUpstreamPin);
-}
\ No newline at end of file
+}
diff --git a/Engine/Source/Editor/UnrealEd/Private/Editor.cpp b/Engine/Source/Editor/UnrealEd/Private/Editor.cpp
--- a/Engine/Source/Editor/UnrealEd/Private/Editor.cpp
+++ b/Engine/Source/Editor/UnrealEd/Private/Editor.cpp
@@ -184,6 +184,9 @@ FEditorDelegates::FOnEditorLoadDefaultStartupMap		FEditorDelegates::OnEditorLoad
 FEditorDelegates::FOnEditorInitialized					FEditorDelegates::OnEditorInitialized;
 FSimpleMulticastDelegate								FEditorDelegates::OnEditorPreExit;
 FEditorDelegates::FOnExternalContentResolved			FEditorDelegates::OnExternalContentResolved;
+FEditorDelegates::FOnWorldGenerationProcess				FEditorDelegates::OnPreWorldGenerationProcess;
+FEditorDelegates::FOnWorldGenerationProcess				FEditorDelegates::OnPostWorldGenerationProcess;
+
 
 PRAGMA_ENABLE_DEPRECATION_WARNINGS
 
diff --git a/Engine/Source/Editor/UnrealEd/Public/Editor.h b/Engine/Source/Editor/UnrealEd/Public/Editor.h
--- a/Engine/Source/Editor/UnrealEd/Public/Editor.h
+++ b/Engine/Source/Editor/UnrealEd/Public/Editor.h
@@ -186,6 +186,8 @@ struct FEditorDelegates
 	DECLARE_MULTICAST_DELEGATE_OneParam(FOnEditorInitialized, double Duration);
 	/** delegate when external content resolves and can replace placeholder data */
 	DECLARE_MULTICAST_DELEGATE_ThreeParams(FOnExternalContentResolved, const FGuid&, const FAssetData& /*PlaceholderAsset*/, const FAssetData& /*ResolvedAsset*/);
+	/** delegate for an editor process is generating a large set of actors */
+	DECLARE_MULTICAST_DELEGATE_OneParam(FOnWorldGenerationProcess, UWorld*);
 
 	/** Called when the CurrentLevel is switched to a new level.  Note that this event won't be fired for temporary
 		changes to the current level, such as when copying/pasting actors. */
@@ -388,6 +390,10 @@ struct FEditorDelegates
 	static UNREALED_API FSimpleMulticastDelegate OnEditorPreExit;
 	/** Called when external content gets resolved */
 	static UNREALED_API FOnExternalContentResolved OnExternalContentResolved;
+	/** Called when an editor process is about to generate a large set of actors */
+	static UNREALED_API FOnWorldGenerationProcess OnPreWorldGenerationProcess;
+	/** Called when an editor process is done generating a large set of actors */
+	static UNREALED_API FOnWorldGenerationProcess OnPostWorldGenerationProcess;
 };
 
 /**
