Cooker: Softlock detection: Add NiagaraSystem to list of compilation users. Compilation work has shifted over time from NiagaraScript to NiagaraSystem and eventually only NiagaraSystems will be cooked.

#rnx

diff --git a/Engine/Source/Editor/UnrealEd/Private/CookOnTheFlyServer.cpp b/Engine/Source/Editor/UnrealEd/Private/CookOnTheFlyServer.cpp
--- a/Engine/Source/Editor/UnrealEd/Private/CookOnTheFlyServer.cpp
+++ b/Engine/Source/Editor/UnrealEd/Private/CookOnTheFlyServer.cpp
@@ -1697,7 +1697,11 @@ void UCookOnTheFlyServer::SetSaveBusy(bool bInBusy, int32 NumAsyncWorkRetired)
 		TArray<UObject*> ExpectedObjects;
 		TSet<UPackage*> ExpectedPackages;
 		FPackageDataQueue& SaveQueue = PackageDatas->GetSaveQueue();
-		TArray<UClass*> CompilationUsers({ UMaterialInterface::StaticClass(), FindObject<UClass>(nullptr, TEXT("/Script/Niagara.NiagaraScript")) });
+		TArray<UClass*> CompilationUsers({
+			UMaterialInterface::StaticClass(),
+			FindObject<UClass>(nullptr, TEXT("/Script/Niagara.NiagaraScript")),
+			FindObject<UClass>(nullptr, TEXT("/Script/Niagara.NiagaraSystem"))
+		});
 
 		PackageDatas->ForEachPendingCookedPlatformData(
 		[&CompilationUsers, &ExpectedObjects, &ExpectedPackages, &NonExpectedObjects, &NonExpectedPackages, bBusyCompilationUsersAreExpected]
