IncrementalCook: Documentation: Change some names to reduce confusion. bLegacyIterative -> bLegacyBuildDependencies. In CookSettings.txt,  CookIncremental -> LegacyBuildDependencies.
#rnx
#incrementalcook

diff --git a/Engine/Source/Editor/UnrealEd/Classes/CookOnTheSide/CookOnTheFlyServer.h b/Engine/Source/Editor/UnrealEd/Classes/CookOnTheSide/CookOnTheFlyServer.h
--- a/Engine/Source/Editor/UnrealEd/Classes/CookOnTheSide/CookOnTheFlyServer.h
+++ b/Engine/Source/Editor/UnrealEd/Classes/CookOnTheSide/CookOnTheFlyServer.h
@@ -1585,10 +1585,20 @@ private:
 	 */
 	bool bPreloadingEnabled = false;
 	/**
-	 * If enabled, we load/save TargetDomainKey hashes and use those to test whether packages have already been
-	 * cooked, in every cook. This can be disabled with commandline -forcerecook or by bumping the CookIncrementalVersion.
+	 * By default (aka CookIncremental system), we load/save TargetDomainKey hashes and tweak behavior throughout the
+	 * cooker to support that. We decide whether to cook incrementally based on config/commandline options that specify
+	 * forcerecook; if cooking incrementally we decide whether packages have been cooked based on an unchanged
+	 * TargetDomainKey hash present in oplog.
+	 * 
+	 * This can be disabled and we use the old system instead - now called legacy iterative - which detects workspace
+	 * domain package changes by a change in their recorded PackageSaveHash, and propagates changed packages out to
+	 * AssetRegistry referencers. Legacy iterative is a simpler model that has more false positive and false negative
+	 * incremental skips. We do not collect TargetDomainKey hashes when running in legacy iterative.
+	 * 
+	 * This field specifies which dependency handling mode we are in, and is orthogonal to whether we are cooking
+	 * incrementally in the current mode; see PlatformManager.GetPlatformData().bFullCook for that property.
 	 */
-	bool bCookIncremental = true;
+	bool bLegacyBuildDependencies = false;
 	bool bCookIncrementalAllowAllClasses = false;
 	bool bHiddenDependenciesDebug = false;
 	bool bIncrementallyModifiedDiagnostics = false;
diff --git a/Engine/Source/Editor/UnrealEd/Private/CookOnTheFlyServer.cpp b/Engine/Source/Editor/UnrealEd/Private/CookOnTheFlyServer.cpp
--- a/Engine/Source/Editor/UnrealEd/Private/CookOnTheFlyServer.cpp
+++ b/Engine/Source/Editor/UnrealEd/Private/CookOnTheFlyServer.cpp
@@ -5274,7 +5274,7 @@ void UCookOnTheFlyServer::KickBuildDependencies(UE::Cook::FTickStackData& StackD
 	}
 	bKickedBuildDependencies = true;
 
-	if (!bCookIncremental)
+	if (bLegacyBuildDependencies)
 	{
 		// Uncooked BuildDependencies are not committed in legacy cooks that are not using IncrementalCook
 		return;
@@ -7499,9 +7499,9 @@ TMap<FName, FString> UCookOnTheFlyServer::CalculateCookSettingStrings(const FCoo
 	IPlatformFile& PlatformFile = FPlatformFileManager::Get().GetPlatformFile();
 
 	CookSettingStrings.Add(FName(TEXT("Version")), TEXT("21F52B9EDD4D456AB1AF381CA172BD28"));
-	CookSettingStrings.Add(FName(TEXT("CookIncremental")), *LexToString(bCookIncremental));
+	CookSettingStrings.Add(FName(TEXT("LegacyBuildDependencies")), *LexToString(bLegacyBuildDependencies));
 
-	if (bCookIncremental)
+	if (!bLegacyBuildDependencies)
 	{
 		// Store the CookIncrementalVersion in the global settings, so that it will cause deletion of all artifacts
 		// when it changes, in addition to invalidating all cooked packages.
@@ -7624,14 +7624,14 @@ bool UCookOnTheFlyServer::ArePreviousCookSettingsCompatible(const TMap<FName, FS
 		}
 	}
 
-	if (!bLegacyIterativeIgnoreIni && !bCookIncremental && IniSettingsOutOfDate(TargetPlatform))
+	if (!bLegacyIterativeIgnoreIni && bLegacyBuildDependencies && IniSettingsOutOfDate(TargetPlatform))
 	{
 		UE_LOG(LogCook, Display, TEXT("Cook invalidated for platform %s because ini settings have changed. Clearing previously cooked packages."),
 			*TargetPlatform->PlatformName());
 		return false;
 	}
 
-	if (!bLegacyIterativeIgnoreExe)
+	if (!bLegacyIterativeIgnoreExe && bLegacyBuildDependencies)
 	{
 		const FString* CurrentHash = CurrentCookSettings.Find(ExecutableHashName);
 		if (!CurrentHash)
@@ -7932,7 +7932,7 @@ void UCookOnTheFlyServer::PopulateCookedPackages(TArrayView<const ITargetPlatfor
 		}
 
 		TMap<FName, FAssetRegistryGenerator::FGeneratorPackageInfo> PreviousGeneratorPackages;
-		if (bCookIncremental)
+		if (!bLegacyBuildDependencies)
 		{
 			// Incremental cook does the equivalent operation of bRecurseModifications=true and bRecurseScriptModifications=true,
 			// but checks for out-of-datedness are done by FRequestCluster using the TargetDomainKey (which is built
@@ -10617,31 +10617,30 @@ namespace UE::Cook
 void FBeginCookConfigSettings::LoadLocal(FBeginCookContext& BeginContext)
 {
 	const TCHAR* CommandLine = FCommandLine::Get();
-	bool bLegacyIterative = false;
-	GConfig->GetBool(TEXT("CookSettings"), TEXT("LegacyIterative"), bLegacyIterative, GEditorIni);
-	bCookIncremental = !bLegacyIterative;
+	bLegacyBuildDependencies = false;
+	GConfig->GetBool(TEXT("CookSettings"), TEXT("LegacyIterative"), bLegacyBuildDependencies, GEditorIni);
 
-	if (bCookIncremental && FParse::Param(CommandLine, TEXT("ODSC")))
+	if (!bLegacyBuildDependencies && FParse::Param(CommandLine, TEXT("ODSC")))
 	{
-		// INCREMENTALCOOK_TODO: implement incremental cook for DLC
+		// INCREMENTALCOOK_TODO: implement incremental cook for ODSC
 		UE_LOG(LogCook, Display,
 			TEXT("INCREMENTAL COOK DEPENDENCIES: Disabled. The cooker is running as ODSC which forces use of legacyiterative."));
-		bCookIncremental = false;
+		bLegacyBuildDependencies = true;
 	}
-	else if (bCookIncremental && !BeginContext.COTFS.IsUsingZenStore())
+	else if (!bLegacyBuildDependencies && !BeginContext.COTFS.IsUsingZenStore())
 	{
 		// Incremental cook uses TargetDomain storage of dependencies which is only implemented in ZenStore
 		UE_LOG(LogCook, Display,
 			TEXT("INCREMENTAL COOK DEPENDENCIES: Disabled. Incremental Cook Dependencies are disabled because the cooker is not using ZenStore, ")
 			TEXT("and we do not yet support storage of Incremental Cook dependency data with loose cooked packages. ")
 			TEXT("Falling back to legacy iterative cook mode."));
-		bCookIncremental = false;
+		bLegacyBuildDependencies = true;
 	}
 	else
 	{
-		UE_CLOG(bCookIncremental, LogCook, Display,
+		UE_CLOG(!bLegacyBuildDependencies, LogCook, Display,
 			TEXT("INCREMENTAL COOK DEPENDENCIES: Enabled. Incremental Cook Dependencies are enabled in Editor.ini:[CookSettings]:LegacyIterative=false."));
-		UE_CLOG(!bCookIncremental, LogCook, Display,
+		UE_CLOG(bLegacyBuildDependencies, LogCook, Display,
 			TEXT("INCREMENTAL COOK DEPENDENCIES: Disabled. Incremental Cook Dependencies are disabled in Editor.ini:[CookSettings]:LegacyIterative=true."));
 	}
 		 
@@ -10668,7 +10667,7 @@ void FBeginCookConfigSettings::LoadLocal(FBeginCookContext& BeginContext)
 
 void UCookOnTheFlyServer::SetBeginCookConfigSettings(FBeginCookContext& BeginContext, UE::Cook::FBeginCookConfigSettings&& Settings)
 {
-	bCookIncremental = Settings.bCookIncremental;
+	bLegacyBuildDependencies = Settings.bLegacyBuildDependencies;
 	bCookIncrementalAllowAllClasses = Settings.bCookIncrementalAllowAllClasses;
 	PackageDatas->SetBeginCookConfigSettings(Settings.CookShowInstigator);
 	UE::Cook::FGenerationHelper::SetBeginCookConfigSettings();
@@ -10738,7 +10737,7 @@ void UCookOnTheFlyServer::LoadBeginCookIncrementalFlagsLocal(FBeginCookContext&
 	const bool bFullBuildAllowed = DiffModeHelper->IsFullBuildAllowed();
 
 	bool bDefaultIncremental = true;
-	if (bCookIncremental)
+	if (!bLegacyBuildDependencies)
 	{
 		GConfig->GetBool(TEXT("CookSettings"), TEXT("CookIncrementalDefaultIncremental"), bDefaultIncremental, GEditorIni);
 	}
@@ -10778,8 +10777,10 @@ void UCookOnTheFlyServer::LoadBeginCookIncrementalFlagsLocal(FBeginCookContext&
 		bForceRecookCommandLineArgPresent = true;
 	}
 
-	const bool bIncrementalOrLegacyIterative = !bForceRecook && (bCookIncremental || IsCookFlagSet(ECookInitializationFlags::LegacyIterative));
-	const bool bIsSharedLegacyIterativeCook = !bCookIncremental && IsCookFlagSet(ECookInitializationFlags::LegacyIterativeSharedBuild);
+	const bool bIncrementalOrLegacyIterative = !bForceRecook &&
+		(!bLegacyBuildDependencies || IsCookFlagSet(ECookInitializationFlags::LegacyIterative));
+	const bool bIsSharedLegacyIterativeCook = bLegacyBuildDependencies &&
+		IsCookFlagSet(ECookInitializationFlags::LegacyIterativeSharedBuild);
 
 	for (FBeginCookContextPlatform& PlatformContext : BeginContext.PlatformContexts)
 	{
@@ -10836,7 +10837,7 @@ void UCookOnTheFlyServer::LoadBeginCookIncrementalFlagsLocal(FBeginCookContext&
 			bool bIncrementalOrLegacyIterativeAllowed = true;
 			if (!bIncrementalOrLegacyIterative && !PlatformData->bIsSandboxInitialized)
 			{
-				if (bCookIncremental)
+				if (!bLegacyBuildDependencies)
 				{
 					if (!bDefaultIncremental && !bForceRecookCommandLineArgPresent)
 					{
@@ -10869,8 +10870,8 @@ void UCookOnTheFlyServer::LoadBeginCookIncrementalFlagsLocal(FBeginCookContext&
 			else if (!ArePreviousCookSettingsCompatible(PlatformContext.CurrentCookSettings, TargetPlatform))
 			{
 				UE_LOG(LogCook, Display,
-					TEXT("FULL COOK: %s was specified, but global settings have changed and all previously cook packages are invalidated. Deleting previously cooked packages for platform %s and recooking all packages discovered in the current cook."),
-					bCookIncremental ? TEXT("-cookincremental") : TEXT("-legacyiterative"),
+					TEXT("FULL COOK: %s was specified, but global settings have changed and all previously cooked packages are invalidated. Deleting previously cooked packages for platform %s and recooking all packages discovered in the current cook."),
+					bLegacyBuildDependencies ? TEXT("-legacyiterative") : TEXT("-cookincremental"),
 					*TargetPlatform->PlatformName());
 				if (bRunningAsShaderServer)
 				{
@@ -10885,7 +10886,7 @@ void UCookOnTheFlyServer::LoadBeginCookIncrementalFlagsLocal(FBeginCookContext&
 			{
 				UE_LOG(LogCook, Display,
 					TEXT("INCREMENTAL COOK: %s settings are still valid. Keeping previously cooked packages for platform %s and cooking only packages that have been modified."),
-					bCookIncremental ? TEXT("Global") : TEXT("-legacyiterative was specified and global"),
+					bLegacyBuildDependencies ? TEXT("-legacyiterative was specified and global") : TEXT("Global"),
 					*TargetPlatform->PlatformName());
 				PlatformContext.bFullBuild = false;
 				PlatformContext.bAllowIncrementalResults = true;
diff --git a/Engine/Source/Editor/UnrealEd/Private/Cooker/CookGenerationHelper.cpp b/Engine/Source/Editor/UnrealEd/Private/Cooker/CookGenerationHelper.cpp
--- a/Engine/Source/Editor/UnrealEd/Private/Cooker/CookGenerationHelper.cpp
+++ b/Engine/Source/Editor/UnrealEd/Private/Cooker/CookGenerationHelper.cpp
@@ -1322,7 +1322,7 @@ void FGenerationHelper::StartQueueGeneratedPackages(UCookOnTheFlyServer& COTFS)
 
 	NotifyStartQueueGeneratedPackages(COTFS, FWorkerId::Local());
 
-	bool bCookIncremental = COTFS.bCookIncremental;
+	bool bLegacyBuildDependencies = COTFS.bLegacyBuildDependencies;
 	bool bHasPreviousResultsForSomePlatform = false;
 	for (TPair<const ITargetPlatform*, FPlatformData>& PlatformPair : PlatformDatas)
 	{
@@ -1369,7 +1369,7 @@ void FGenerationHelper::StartQueueGeneratedPackages(UCookOnTheFlyServer& COTFS)
 				RemainingPreviousPackages.Remove(GeneratedPackageName);
 				if (PreviousAssetData)
 				{
-					if (!bCookIncremental)
+					if (bLegacyBuildDependencies)
 					{
 						bool bLegacyIterativeUnmodified;
 						GeneratedInfo.LegacyIterativeCookValidateOrClear(*this, TargetPlatform,
@@ -1393,7 +1393,7 @@ void FGenerationHelper::StartQueueGeneratedPackages(UCookOnTheFlyServer& COTFS)
 				}
 			}
 
-			if (!bCookIncremental)
+			if (bLegacyBuildDependencies)
 			{
 				UE_LOG(LogCook, Display,
 					TEXT("Found %d cooked package(s) in package store for generator package %s."),
diff --git a/Engine/Source/Editor/UnrealEd/Private/Cooker/CookRequestCluster.cpp b/Engine/Source/Editor/UnrealEd/Private/Cooker/CookRequestCluster.cpp
--- a/Engine/Source/Editor/UnrealEd/Private/Cooker/CookRequestCluster.cpp
+++ b/Engine/Source/Editor/UnrealEd/Private/Cooker/CookRequestCluster.cpp
@@ -2744,7 +2744,7 @@ TMap<FPackageData*, TArray<FPackageData*>>& FRequestCluster::FGraphSearch::GetGr
 
 bool FRequestCluster::IsIncrementalCook() const
 {
-	return bAllowIncrementalResults && COTFS.bCookIncremental;
+	return bAllowIncrementalResults && !COTFS.bLegacyBuildDependencies;
 }
 
 void FRequestCluster::IsRequestCookable(const ITargetPlatform* Platform, const FPackageData& PackageData,
diff --git a/Engine/Source/Editor/UnrealEd/Private/Cooker/CookSavePackage.cpp b/Engine/Source/Editor/UnrealEd/Private/Cooker/CookSavePackage.cpp
--- a/Engine/Source/Editor/UnrealEd/Private/Cooker/CookSavePackage.cpp
+++ b/Engine/Source/Editor/UnrealEd/Private/Cooker/CookSavePackage.cpp
@@ -807,109 +807,113 @@ private:
 void FSaveCookedPackageContext::CalculateCookDependencies(FGenerationHelper* GenerationHelper, bool bGenerated,
 	const TArray<FAssetDependency>* ExtraArDependencies, FBuildResultDependenciesMap* ExtraBuildResultDependencies)
 {
-	if (COTFS.bCookIncremental)
+	if (COTFS.bLegacyBuildDependencies)
 	{
-		UE_SCOPED_HIERARCHICAL_COOKTIMER(TargetDomainDependencies);
-		FBuildResultDependenciesMap BuildResultDependencies;
-		TArray<FName> PlatformRuntimeDependencies;
-		TConstArrayView<FName> UntrackedSoftPackageReferences;
-		TConstArrayView<UObject*> Imports;
-		TConstArrayView<UObject*> Exports;
-		TConstArrayView<UE::SavePackageUtilities::FPreloadDependency> PreloadDependencies;
+		return;
+	}
 
-		PlatformRuntimeDependencies = GetPlatformRuntimeDependencies();
-		if (ExtraArDependencies)
-		{
-			PlatformRuntimeDependencies.Reserve(PlatformRuntimeDependencies.Num() + ExtraArDependencies->Num());
-			for (const FAssetDependency& Dependency : *ExtraArDependencies)
-			{
-				PlatformRuntimeDependencies.Add(Dependency.AssetId.PackageName);
-			}
-		}
-		if (ExtraBuildResultDependencies)
-		{
-			BuildResultDependencies.Append(MoveTemp(*ExtraBuildResultDependencies));
-		}
-		bool bHasSavePackageResult = HasSavePackageResult();
-		FSavePackageResultStruct* LocalSavePackageResult = bHasSavePackageResult ? &SavePackageResult : nullptr;
+	UE_SCOPED_HIERARCHICAL_COOKTIMER(TargetDomainDependencies);
+	FBuildResultDependenciesMap BuildResultDependencies;
+	TArray<FName> PlatformRuntimeDependencies;
+	TConstArrayView<FName> UntrackedSoftPackageReferences;
+	TConstArrayView<UObject*> Imports;
+	TConstArrayView<UObject*> Exports;
+	TConstArrayView<UE::SavePackageUtilities::FPreloadDependency> PreloadDependencies;
 
-		if (!bHasSavePackageResult)
+	PlatformRuntimeDependencies = GetPlatformRuntimeDependencies();
+	if (ExtraArDependencies)
+	{
+		PlatformRuntimeDependencies.Reserve(PlatformRuntimeDependencies.Num() + ExtraArDependencies->Num());
+		for (const FAssetDependency& Dependency : *ExtraArDependencies)
 		{
-			// During cook saves of a Package, SavePackage is responsible for collecting BuildResultDependencies from
-			// objects, but if this SaveCookedPackage call is instead for the recording of a BuildDependency package,
-			// we need to collect them here.
-			// Call OnCookEvent(ECookEvent::LoadDependencies) on every UObject in the package, and to support UStructs,
-			// also serialize each object into an FArchive that identifies as a SavePackage archive and contains an
-			// FArchiveCookContext.
-			// Pass the reported BuildResult dependencies into FIncrementalCookAttachments::Collect.
-			FObjectSaveContextData ObjectSaveContextData;
-			ObjectSaveContextData.Set(Package, TargetPlatform, FPackagePath(), SaveFlags);
-			ObjectSaveContextData.CookType = COTFS.GetCookType();
-			ObjectSaveContextData.CookingDLC = COTFS.GetCookingDLC();
-			ObjectSaveContextData.CookInfo = &COTFS;
-			ObjectSaveContextData.ObjectSaveContextPhase = EObjectSaveContextPhase::CookDependencyHarvest;
+			PlatformRuntimeDependencies.Add(Dependency.AssetId.PackageName);
+		}
+	}
+	if (ExtraBuildResultDependencies)
+	{
+		BuildResultDependencies.Append(MoveTemp(*ExtraBuildResultDependencies));
+	}
+	bool bHasSavePackageResult = HasSavePackageResult();
+	FSavePackageResultStruct* LocalSavePackageResult = bHasSavePackageResult ? &SavePackageResult : nullptr;
 
-			UE::Cook::FCookEventContext CookEventContext(ObjectSaveContextData);
-			FBuildDependencyHarvestArchive Harvester(Package, *ArchiveCookContext, ObjectSaveContextData);
+	if (!bHasSavePackageResult)
+	{
+		// During cook saves of a Package, SavePackage is responsible for collecting BuildResultDependencies from
+		// objects, but if this SaveCookedPackage call is instead for the recording of a BuildDependency package,
+		// we need to collect them here.
+		// Call OnCookEvent(ECookEvent::LoadDependencies) on every UObject in the package, and to support UStructs,
+		// also serialize each object into an FArchive that identifies as a SavePackage archive and contains an
+		// FArchiveCookContext.
+		// Pass the reported BuildResult dependencies into FIncrementalCookAttachments::Collect.
+		FObjectSaveContextData ObjectSaveContextData;
+		ObjectSaveContextData.Set(Package, TargetPlatform, FPackagePath(), SaveFlags);
+		ObjectSaveContextData.CookType = COTFS.GetCookType();
+		ObjectSaveContextData.CookingDLC = COTFS.GetCookingDLC();
+		ObjectSaveContextData.CookInfo = &COTFS;
+		ObjectSaveContextData.ObjectSaveContextPhase = EObjectSaveContextPhase::CookDependencyHarvest;
 
-			for (FCachedObjectInOuter& CachedObjectInOuter : PackageData.GetCachedObjectsInOuter())
-			{
-				UObject* Object = CachedObjectInOuter.Object.Get();;
-				ObjectSaveContextData.Object = Object;
-				if (Object)
-				{
-					Object->OnCookEvent(ECookEvent::PlatformCookDependencies, CookEventContext);
-					Object->Serialize(Harvester);
-				}
-			}
+		UE::Cook::FCookEventContext CookEventContext(ObjectSaveContextData);
+		FBuildDependencyHarvestArchive Harvester(Package, *ArchiveCookContext, ObjectSaveContextData);
 
-			BuildResultDependencies = MoveTemp(ObjectSaveContextData.BuildResultDependencies);
-			for (FSoftObjectPath& RuntimeDependency : ObjectSaveContextData.CookRuntimeDependencies)
+		for (FCachedObjectInOuter& CachedObjectInOuter : PackageData.GetCachedObjectsInOuter())
+		{
+			UObject* Object = CachedObjectInOuter.Object.Get();;
+			ObjectSaveContextData.Object = Object;
+			if (Object)
 			{
-				FName PackageDependency = RuntimeDependency.GetLongPackageFName();
-				if (!PackageDependency.IsNone())
-				{
-					PlatformRuntimeDependencies.Add(PackageDependency);
-				}
+				Object->OnCookEvent(ECookEvent::PlatformCookDependencies, CookEventContext);
+				Object->Serialize(Harvester);
 			}
 		}
 
-		if (bHasSavePackageResult)
+		BuildResultDependencies = MoveTemp(ObjectSaveContextData.BuildResultDependencies);
+		for (FSoftObjectPath& RuntimeDependency : ObjectSaveContextData.CookRuntimeDependencies)
 		{
-			BuildResultDependencies.Append(MoveTemp(LocalSavePackageResult->BuildResultDependencies));
-			PlatformRuntimeDependencies.Append(MoveTemp(LocalSavePackageResult->SoftPackageReferences));
-			UntrackedSoftPackageReferences = LocalSavePackageResult->UntrackedSoftPackageReferences;
-			Imports = LocalSavePackageResult->Imports;
-			Exports = LocalSavePackageResult->Exports;
-			PreloadDependencies = LocalSavePackageResult->PreloadDependencies;
+			FName PackageDependency = RuntimeDependency.GetLongPackageFName();
+			if (!PackageDependency.IsNone())
+			{
+				PlatformRuntimeDependencies.Add(PackageDependency);
+			}
 		}
-		const FBuildResultDependenciesMap* LoadDependencies = PackageData.GetLoadDependencies();
-		checkf(LoadDependencies,
-			TEXT("LoadDependencies not found during save of package %s. LoadDependencies are supposted to be created by LoadPackageInQueue before entering the Save state."),
-			*PackageData.GetPackageName().ToString());
-		BuildResultDependencies.Append(*LoadDependencies);
+	}
 
-		PlatformCookAttachments = FIncrementalCookAttachments::Collect(Package, TargetPlatform,
-			MoveTemp(BuildResultDependencies), bHasSavePackageResult, UntrackedSoftPackageReferences, GenerationHelper,
-			bGenerated, MoveTemp(PlatformRuntimeDependencies), Imports, Exports, PreloadDependencies,
-			PackageData.GetLogMessages());
+	if (bHasSavePackageResult)
+	{
+		BuildResultDependencies.Append(MoveTemp(LocalSavePackageResult->BuildResultDependencies));
+		PlatformRuntimeDependencies.Append(MoveTemp(LocalSavePackageResult->SoftPackageReferences));
+		UntrackedSoftPackageReferences = LocalSavePackageResult->UntrackedSoftPackageReferences;
+		Imports = LocalSavePackageResult->Imports;
+		Exports = LocalSavePackageResult->Exports;
+		PreloadDependencies = LocalSavePackageResult->PreloadDependencies;
 	}
+	const FBuildResultDependenciesMap* LoadDependencies = PackageData.GetLoadDependencies();
+	checkf(LoadDependencies,
+		TEXT("LoadDependencies not found during save of package %s. LoadDependencies are supposted to be created by LoadPackageInQueue before entering the Save state."),
+		*PackageData.GetPackageName().ToString());
+	BuildResultDependencies.Append(*LoadDependencies);
+
+	PlatformCookAttachments = FIncrementalCookAttachments::Collect(Package, TargetPlatform,
+		MoveTemp(BuildResultDependencies), bHasSavePackageResult, UntrackedSoftPackageReferences, GenerationHelper,
+		bGenerated, MoveTemp(PlatformRuntimeDependencies), Imports, Exports, PreloadDependencies,
+		PackageData.GetLogMessages());
 }
 
 void FSaveCookedPackageContext::RecordPlatformBuildDependencies()
 {
-	if (COTFS.bCookIncremental)
+	if (COTFS.bLegacyBuildDependencies)
+	{
+		return;
+	}
+
+	TSet<FPackageData*>& CurrentDependencies = PlatformDependencies[PlatformIndex].BuildDependencies;
+	TArray<FName, TInlineAllocator<10>> TransitiveBuildDependencies;
+	PlatformCookAttachments.Artifacts.GetTransitiveBuildDependencies(TransitiveBuildDependencies);
+	for (FName TransitivePackageName : TransitiveBuildDependencies)
 	{
-		TSet<FPackageData*>& CurrentDependencies = PlatformDependencies[PlatformIndex].BuildDependencies;
-		TArray<FName, TInlineAllocator<10>> TransitiveBuildDependencies;
-		PlatformCookAttachments.Artifacts.GetTransitiveBuildDependencies(TransitiveBuildDependencies);
-		for (FName TransitivePackageName : TransitiveBuildDependencies)
+		FPackageData* BuildPackageData = COTFS.PackageDatas->TryAddPackageDataByPackageName(TransitivePackageName);
+		if (BuildPackageData)
 		{
-			FPackageData* BuildPackageData = COTFS.PackageDatas->TryAddPackageDataByPackageName(TransitivePackageName);
-			if (BuildPackageData)
-			{
-				CurrentDependencies.Add(BuildPackageData);
-			}
+			CurrentDependencies.Add(BuildPackageData);
 		}
 	}
 }
@@ -938,7 +942,7 @@ void FSaveCookedPackageContext::RecordCookImportsCheckerData()
 TArray<IPackageWriter::FCommitAttachmentInfo> FSaveCookedPackageContext::GetCommitAttachments()
 {
 	TArray<IPackageWriter::FCommitAttachmentInfo> Result;
-	if (COTFS.bCookIncremental)
+	if (!COTFS.bLegacyBuildDependencies)
 	{
 		PlatformCookAttachments.AppendCommitAttachments(Result);
 	}
diff --git a/Engine/Source/Editor/UnrealEd/Private/Cooker/CookTypes.cpp b/Engine/Source/Editor/UnrealEd/Private/Cooker/CookTypes.cpp
--- a/Engine/Source/Editor/UnrealEd/Private/Cooker/CookTypes.cpp
+++ b/Engine/Source/Editor/UnrealEd/Private/Cooker/CookTypes.cpp
@@ -851,7 +851,7 @@ void FInitializeConfigSettings::MoveToLocal(UCookOnTheFlyServer& COTFS)
 
 void FBeginCookConfigSettings::CopyFromLocal(const UCookOnTheFlyServer& COTFS)
 {
-	bCookIncremental = COTFS.bCookIncremental;
+	bLegacyBuildDependencies = COTFS.bLegacyBuildDependencies;
 	bCookIncrementalAllowAllClasses = COTFS.bCookIncrementalAllowAllClasses;
 	FParse::Value(FCommandLine::Get(), TEXT("-CookShowInstigator="), CookShowInstigator); // We don't store this on COTFS, so reparse it from commandLine
 	TSet<FName> COTFSNeverCookPackageList;
@@ -887,7 +887,7 @@ FString LexToString(FPlatformMemoryStats::EMemoryPressureStatus Value)
 FCbWriter& operator<<(FCbWriter& Writer, const UE::Cook::FBeginCookConfigSettings& Value)
 {
 	Writer.BeginObject();
-	Writer << "CookIncremental" << Value.bCookIncremental;
+	Writer << "LegacyBuildDependencies" << Value.bLegacyBuildDependencies;
 	Writer << "CookIncrementalAllowAllClasses" << Value.bCookIncrementalAllowAllClasses;
 	Writer << "CookShowInstigator" << Value.CookShowInstigator;
 	Writer << "NeverCookPackageList" << Value.NeverCookPackageList;
@@ -909,7 +909,7 @@ FCbWriter& operator<<(FCbWriter& Writer, const UE::Cook::FBeginCookConfigSetting
 bool LoadFromCompactBinary(FCbFieldView Field, UE::Cook::FBeginCookConfigSettings& OutValue)
 {
 	bool bOk = Field.IsObject();
-	bOk = LoadFromCompactBinary(Field["CookIncremental"], OutValue.bCookIncremental) & bOk;
+	bOk = LoadFromCompactBinary(Field["LegacyBuildDependencies"], OutValue.bLegacyBuildDependencies) & bOk;
 	bOk = LoadFromCompactBinary(Field["CookIncrementalAllowAllClasses"], OutValue.bCookIncrementalAllowAllClasses) & bOk;
 	bOk = LoadFromCompactBinary(Field["CookShowInstigator"], OutValue.CookShowInstigator) & bOk;
 	bOk = LoadFromCompactBinary(Field["NeverCookPackageList"], OutValue.NeverCookPackageList) & bOk;
diff --git a/Engine/Source/Editor/UnrealEd/Private/Cooker/CookTypes.h b/Engine/Source/Editor/UnrealEd/Private/Cooker/CookTypes.h
--- a/Engine/Source/Editor/UnrealEd/Private/Cooker/CookTypes.h
+++ b/Engine/Source/Editor/UnrealEd/Private/Cooker/CookTypes.h
@@ -496,7 +496,7 @@ namespace UE::Cook
 		void CopyFromLocal(const UCookOnTheFlyServer& COTFS);
 
 		FString CookShowInstigator;
-		bool bCookIncremental = true;
+		bool bLegacyBuildDependencies = false;
 		bool bCookIncrementalAllowAllClasses = false;
 		TArray<FName> NeverCookPackageList;
 		TMap<const ITargetPlatform*, TSet<FName>> PlatformSpecificNeverCookPackages;
