Convert depth bar far & near plane labels to custom widgets. Hide the depth bar when the viewport is not tall enough.

#rnx

diff --git a/Engine/Source/Editor/LevelEditor/Private/SLevelViewport.cpp b/Engine/Source/Editor/LevelEditor/Private/SLevelViewport.cpp
--- a/Engine/Source/Editor/LevelEditor/Private/SLevelViewport.cpp
+++ b/Engine/Source/Editor/LevelEditor/Private/SLevelViewport.cpp
@@ -3636,7 +3636,7 @@ void SLevelViewport::OnElementSelectionChanged(const UTypedElementSelectionSet*
 				}
 			}
 				
-			DepthBar->AddSlot()
+			DepthBar->AddDepthIndicatorSlot()
 			.Bounds(Bounds)
 			[
 				SNew(SImage)
diff --git a/Engine/Source/Editor/UnrealEd/Private/SEditorViewport.cpp b/Engine/Source/Editor/UnrealEd/Private/SEditorViewport.cpp
--- a/Engine/Source/Editor/UnrealEd/Private/SEditorViewport.cpp
+++ b/Engine/Source/Editor/UnrealEd/Private/SEditorViewport.cpp
@@ -43,6 +43,12 @@
 
 #define LOCTEXT_NAMESPACE "EditorViewport"
 
+static TAutoConsoleVariable<float> CVarEditorViewportDepthBarMinSpace(
+	TEXT("r.Editor.DepthBarMinSpace"),
+	600.0f,
+	TEXT("The minimum vertical space to show the depth bar. Defaults to 600."),
+	ECVF_Default);
+
 SEditorViewport::SEditorViewport()
 	: LastTickTime(0)
 	, bInvalidated(false)
@@ -375,6 +381,8 @@ void SEditorViewport::Tick( const FGeometry& AllottedGeometry, const double InCu
 	{
 		bShowDepthBar = true;
 	}
+	
+	bHasSpaceForDepthBar = AllottedGeometry.GetLocalSize().Y >= CVarEditorViewportDepthBarMinSpace.GetValueOnGameThread();
 }
 
 void SEditorViewport::BindCommands()
@@ -734,10 +742,11 @@ EVisibility SEditorViewport::OnGetViewportContentVisibility() const
 
 EVisibility SEditorViewport::OnGetDepthBarVisibility() const
 {
-	if (Client->SupportsManualOrthographicDepthControl() && Client->IsOrtho() && bShowDepthBar)
+	if (bHasSpaceForDepthBar && Client->SupportsManualOrthographicDepthControl() && Client->IsOrtho() && bShowDepthBar)
 	{
 		return EVisibility::Visible;
 	}
+	
 	return EVisibility::Collapsed;
 }
 
diff --git a/Engine/Source/Editor/UnrealEd/Private/SEditorViewportDepthBar.cpp b/Engine/Source/Editor/UnrealEd/Private/SEditorViewportDepthBar.cpp
--- a/Engine/Source/Editor/UnrealEd/Private/SEditorViewportDepthBar.cpp
+++ b/Engine/Source/Editor/UnrealEd/Private/SEditorViewportDepthBar.cpp
@@ -8,6 +8,8 @@
 #include "Elements/Interfaces/TypedElementWorldInterface.h"
 #include "Math/UnitConversion.h"
 #include "WorldPartition/WorldPartition.h"
+#include "Widgets/Layout/SBorder.h"
+#include "Widgets/Text/STextBlock.h"
 
 #define LOCTEXT_NAMESPACE "EditorViewportDepthBar"
 
@@ -23,7 +25,16 @@ namespace UE::Editor::ViewportDepthBar::Private
 			const TArray<EUnit>& AllDistanceUnits = FUnitConversion::Settings().GetDisplayUnits(EUnitType::Distance);
 			const EUnit BaseUnit = AllDistanceUnits.Num() ? AllDistanceUnits[0] : EUnit::Centimeters;
 			const FNumericUnit DisplayUnit = FUnitConversion::QuantizeUnitsToBestFit(Value, BaseUnit);
-			return FText::Format(LOCTEXT("DistanceFormat", "{0}{1}"), DisplayUnit.Value, FText::FromString(FUnitConversion::GetUnitDisplayString(DisplayUnit.Units)));
+			
+			const FNumberFormattingOptions Options = FNumberFormattingOptions()
+				.SetMaximumFractionalDigits(2)
+				.SetMinimumFractionalDigits(2);
+			
+			return FText::Format(
+				LOCTEXT("DistanceFormat", "{0}{1}"),
+				FText::AsNumber(DisplayUnit.Value, &Options),
+				FText::FromString(FUnitConversion::GetUnitDisplayString(DisplayUnit.Units))
+			);
 		}
 		return FText::Format(LOCTEXT("RawDistanceFormat", "{0}"), Value);
 	}
@@ -39,9 +50,35 @@ void SEditorViewportDepthBar::Construct(const FArguments& InArgs, const TSharedP
 	HoveredSliceTopBrush = FAppStyle::Get().GetBrush("EditorViewport.DepthBar.SliceTop");
 	HoveredSliceBottomBrush = FAppStyle::Get().GetBrush("EditorViewport.DepthBar.SliceBottom");
 	
-	Children.AddSlots(MoveTemp(const_cast<TArray<FSlot::FSlotArguments>&>(InArgs._Slots)));
+	DepthIndicatorChildren.AddSlots(MoveTemp(const_cast<TArray<FDepthIndicatorSlot::FSlotArguments>&>(InArgs._Slots)));
+	
+	AddDepthLabelSlot()
+	.VAlign(VAlign_Top)
+	.OnGetBarPosition_Lambda([](const FBarPositions& BarPositions) { return BarPositions.BarFarPlane; })
+	[
+		SNew(SBorder)
+		.BorderImage(BackgroundBrush)
+		.Padding(6.0f, 2.0f)
+		.Visibility(this, &SEditorViewportDepthBar::GetFarPlaneLabelVisibility)
+		[
+			SNew(STextBlock)
+			.Text(this, &SEditorViewportDepthBar::GetFarPlaneText)
+		]
+	];
 	
-	SetToolTipText(TAttribute<FText>::CreateSP(this, &SEditorViewportDepthBar::GetTooltipText));
+	AddDepthLabelSlot()
+	.VAlign(VAlign_Bottom)
+	.OnGetBarPosition_Lambda([](const FBarPositions& BarPositions) { return BarPositions.BarNearPlane; })
+	[
+		SNew(SBorder)
+		.BorderImage(BackgroundBrush)
+		.Padding(6.0f, 2.0f)
+		.Visibility(this, &SEditorViewportDepthBar::GetNearPlaneLabelVisibility)
+		[
+			SNew(STextBlock)
+			.Text(this, &SEditorViewportDepthBar::GetNearPlaneText)
+		]
+	];
 }
 
 void SEditorViewportDepthBar::Tick(const FGeometry& AllottedGeometry, const double InCurrentTime, const float InDeltaTime)
@@ -94,19 +131,19 @@ FCursorReply SEditorViewportDepthBar::OnCursorQuery(const FGeometry& MyGeometry,
 	}
 }
 
-SEditorViewportDepthBar::FSlot& SEditorViewportDepthBar::GetSlot(int32 Index)
+SEditorViewportDepthBar::FDepthIndicatorSlot& SEditorViewportDepthBar::GetDepthIndicatorSlot(int32 Index)
 {
-	return Children[Index];
+	return DepthIndicatorChildren[Index];
 }
 
-const SEditorViewportDepthBar::FSlot& SEditorViewportDepthBar::GetSlot(int32 Index) const
+const SEditorViewportDepthBar::FDepthIndicatorSlot& SEditorViewportDepthBar::GetDepthIndicatorSlot(int32 Index) const
 {
-	return Children[Index];
+	return DepthIndicatorChildren[Index];
 }
 
 void SEditorViewportDepthBar::ClearChildren()
 {
-	Children.Empty();
+	DepthIndicatorChildren.Empty();
 }
 
 FReply SEditorViewportDepthBar::OnMouseButtonDown(const FGeometry& MyGeometry, const FPointerEvent& MouseEvent)
@@ -243,9 +280,9 @@ void SEditorViewportDepthBar::OnArrangeChildren(const FGeometry& AllottedGeometr
 	{
 		FMargin Inset = Padding.Get();
 	
-		for (int32 ChildIndex = 0; ChildIndex < Children.Num(); ++ChildIndex)
+		for (int32 ChildIndex = 0; ChildIndex < DepthIndicatorChildren.Num(); ++ChildIndex)
 		{
-			const FSlot& CurrentChild = GetSlot(ChildIndex);
+			const FDepthIndicatorSlot& CurrentChild = GetDepthIndicatorSlot(ChildIndex);
 			
 			if (CurrentChild.GetBounds().IsSet())
 			{
@@ -277,12 +314,46 @@ void SEditorViewportDepthBar::OnArrangeChildren(const FGeometry& AllottedGeometr
 				ArrangedChildren.AddWidget(AllottedGeometry.MakeChild(CurrentChild.GetWidget(), BoundSize, FSlateLayoutTransform(Offset)));
 			}
 		}
+		
+		for (int32 ChildIndex = 0; ChildIndex < DepthLabelChildren.Num(); ++ChildIndex)
+		{
+			const FDepthLabelSlot& CurrentChild = DepthLabelChildren[ChildIndex];
+			
+			if (CurrentChild.GetWidget()->GetVisibility() != EVisibility::Visible)
+			{
+				continue;
+			}
+			
+			const float BarPosition = CurrentChild.GetBarPosition(*BarPositions);
+			
+			FVector2f ChildDesiredSize = CurrentChild.GetWidget()->GetDesiredSize();
+			FVector2f Offset(-ChildDesiredSize.X - 12.0f, BarPosition);
+			
+			switch (CurrentChild.GetVerticalAlignment())
+			{
+			case VAlign_Top:
+				Offset.Y -= ChildDesiredSize.Y;
+				break;
+			case VAlign_Center:
+			case VAlign_Fill:
+				Offset.Y -= ChildDesiredSize.Y * 0.5f;
+				break;
+			default:
+				break;
+			}
+			
+			ArrangedChildren.AddWidget(AllottedGeometry.MakeChild(
+				CurrentChild.GetWidget(),
+				Offset,
+				ChildDesiredSize
+			));
+		}
 	}	
 }
 
 FChildren* SEditorViewportDepthBar::GetChildren()
 {
-	return &Children;
+	return &AllChildren;
 }
 
 int32 SEditorViewportDepthBar::OnPaint(const FPaintArgs& Args, const FGeometry& AllottedGeometry, const FSlateRect& MyCullingRect, FSlateWindowElementList& OutDrawElements, int32 LayerId, const FWidgetStyle& InWidgetStyle, bool bParentEnabled) const
@@ -430,48 +501,44 @@ double SEditorViewportDepthBar::FBarPositions::BarToRealPosition(float BarPositi
 	return RelativePosition * (RealMax - RealMin) + RealMin;
 }
 
-FText SEditorViewportDepthBar::GetTooltipText() const
+EVisibility SEditorViewportDepthBar::GetFarPlaneLabelVisibility() const
 {
-	if (TSharedPtr<FEditorViewportClient> Client = WeakClient.Pin())
-	{
-		switch (HoverTarget)
-		{
-		case EDragTarget::Far:
-			return GetFarPlaneText(Client);
-		case EDragTarget::Near:
-			return GetNearPlaneText(Client);
-		case EDragTarget::Slice:
-			return FText::Format(LOCTEXT("CombinedPlanesFormat", "{0}\n\n{1}"), GetFarPlaneText(Client), GetNearPlaneText(Client));
-		default:
-			break;
-		}
-	}
-	
-	return FText::GetEmpty();
+	return HoverTarget == EDragTarget::Far || HoverTarget == EDragTarget::Slice ? EVisibility::Visible : EVisibility::Collapsed;
 }
 
-FText SEditorViewportDepthBar::GetFarPlaneText(const TSharedPtr<FEditorViewportClient>& Client)
+EVisibility SEditorViewportDepthBar::GetNearPlaneLabelVisibility() const
 {
-	if (Client->IsOrtho())
+	return HoverTarget == EDragTarget::Near || HoverTarget == EDragTarget::Slice ? EVisibility::Visible : EVisibility::Collapsed;
+}
+
+FText SEditorViewportDepthBar::GetFarPlaneText() const
+{
+	if (TSharedPtr<FEditorViewportClient> Client = WeakClient.Pin())
 	{
-		if (const double* FarPlane = Client->GetOrthographicFarPlaneOverride().GetPtrOrNull())
+		if (Client->IsOrtho())
 		{
-			return FText::Format(LOCTEXT("OrthoFarPlaneFormat", "Far Plane: {0}"), UE::Editor::ViewportDepthBar::Private::GetDistanceDisplayText(*FarPlane));
+			if (const double* FarPlane = Client->GetOrthographicFarPlaneOverride().GetPtrOrNull())
+			{
+				return FText::Format(LOCTEXT("OrthoFarPlaneFormat", "Far Plane: {0}"), UE::Editor::ViewportDepthBar::Private::GetDistanceDisplayText(*FarPlane));
+			}
+			return LOCTEXT("OrthoFarPlaneDefault", "Infinite Far Plane");
 		}
-		return LOCTEXT("OrthoFarPlaneDefault", "Infinite Far Plane");
 	}
 	return FText::GetEmpty();
 }
 
-FText SEditorViewportDepthBar::GetNearPlaneText(const TSharedPtr<FEditorViewportClient>& Client)
+FText SEditorViewportDepthBar::GetNearPlaneText() const
 {
-	if (Client->IsOrtho())
+	if (TSharedPtr<FEditorViewportClient> Client = WeakClient.Pin())
 	{
-		if (const double* NearPlane = Client->GetOrthographicNearPlaneOverride().GetPtrOrNull())
+		if (Client->IsOrtho())
 		{
-			return FText::Format(LOCTEXT("OrthoNearPlaneFormat", "Near Plane: {0}"), UE::Editor::ViewportDepthBar::Private::GetDistanceDisplayText(*NearPlane));
+			if (const double* NearPlane = Client->GetOrthographicNearPlaneOverride().GetPtrOrNull())
+			{
+				return FText::Format(LOCTEXT("OrthoNearPlaneFormat", "Near Plane: {0}"), UE::Editor::ViewportDepthBar::Private::GetDistanceDisplayText(*NearPlane));
+			}
+			return LOCTEXT("OrthoNearPlaneDefault", "Infinite Near Plane");
 		}
-		return LOCTEXT("OrthoNearPlaneDefault", "Infinite Near Plane");
 	}
 	return FText::GetEmpty();
 }
diff --git a/Engine/Source/Editor/UnrealEd/Public/SEditorViewport.h b/Engine/Source/Editor/UnrealEd/Public/SEditorViewport.h
--- a/Engine/Source/Editor/UnrealEd/Public/SEditorViewport.h
+++ b/Engine/Source/Editor/UnrealEd/Public/SEditorViewport.h
@@ -378,9 +378,12 @@ private:
 	/** Whether the viewport needs to update, even without input or realtime (e.g. inertial camera movement) */
 	bool bInvalidated;
 	
-	/** Whether the depth bar should be displayed */
+	/** Whether the depth bar should be displayed, if possible */
 	bool bShowDepthBar = false;
 	
+	/** Whether there is enough space to display the depth bar */
+	bool bHasSpaceForDepthBar = false;
+	
 	/** For automatic legacy toolbar upgrades. Registers the child widget to expect in order to allow an automatic upgrade. */
 	TWeakPtr<SWidget> AutoUpgradeWidgetChild = nullptr;
 	/** A flag indicating that the automatic toolbar upgrade has been approved */
diff --git a/Engine/Source/Editor/UnrealEd/Public/SEditorViewportDepthBar.h b/Engine/Source/Editor/UnrealEd/Public/SEditorViewportDepthBar.h
--- a/Engine/Source/Editor/UnrealEd/Public/SEditorViewportDepthBar.h
+++ b/Engine/Source/Editor/UnrealEd/Public/SEditorViewportDepthBar.h
@@ -15,11 +15,11 @@ class FEditorViewportClient;
 class SEditorViewportDepthBar : public SPanel
 {
 public:
-	class FSlot : public TSlotBase<FSlot>
+	class FDepthIndicatorSlot : public TSlotBase<FDepthIndicatorSlot>
 	{
 	public:
 	
-		SLATE_SLOT_BEGIN_ARGS(FSlot, TSlotBase<FSlot>)
+		SLATE_SLOT_BEGIN_ARGS(FDepthIndicatorSlot, TSlotBase<FDepthIndicatorSlot>)
 			/** Determines where in the bar the slot's contents appear. */
 			SLATE_ATTRIBUTE(FBoxSphereBounds, Bounds)
 		SLATE_SLOT_END_ARGS()
@@ -31,7 +31,7 @@ public:
 		
 		void Construct(const FChildren& SlotOwner, FSlotArguments&& InArgs)
 		{
-			TSlotBase<FSlot>::Construct(SlotOwner, MoveTemp(InArgs));
+			TSlotBase<FDepthIndicatorSlot>::Construct(SlotOwner, MoveTemp(InArgs));
 			Bounds = InArgs._Bounds;
 		}
 		
@@ -41,40 +41,45 @@ public:
 		TAttribute<FBoxSphereBounds> Bounds;
 	};
 
-	static FSlot::FSlotArguments Slot()
+	static FDepthIndicatorSlot::FSlotArguments Slot()
 	{
-		return FSlot::FSlotArguments(MakeUnique<FSlot>());
+		return FDepthIndicatorSlot::FSlotArguments(MakeUnique<FDepthIndicatorSlot>());
 	}
 
 	SLATE_BEGIN_ARGS(SEditorViewportDepthBar)
 	{}
 		SLATE_ATTRIBUTE(TAttribute<FMargin>, Padding)
 		
-		SLATE_SLOT_ARGUMENT(SEditorViewportDepthBar::FSlot, Slots)
+		SLATE_SLOT_ARGUMENT(SEditorViewportDepthBar::FDepthIndicatorSlot, Slots)
 	SLATE_END_ARGS()
 	
-	using FScopedWidgetSlotArguments = TPanelChildren<FSlot>::FScopedWidgetSlotArguments;
-	FScopedWidgetSlotArguments AddSlot()
+	using FScopedWidgetDepthIndicatorSlotArguments = TPanelChildren<FDepthIndicatorSlot>::FScopedWidgetSlotArguments;
+	FScopedWidgetDepthIndicatorSlotArguments AddDepthIndicatorSlot()
 	{
-		return InsertSlot(INDEX_NONE);
+		return InsertDepthIndicatorSlot(INDEX_NONE);
 	}
 
-	FScopedWidgetSlotArguments InsertSlot(int32 Index = INDEX_NONE)
+	FScopedWidgetDepthIndicatorSlotArguments InsertDepthIndicatorSlot(int32 Index = INDEX_NONE)
 	{
-		return FScopedWidgetSlotArguments(MakeUnique<FSlot>(), this->Children, Index);
+		return FScopedWidgetDepthIndicatorSlotArguments(MakeUnique<FDepthIndicatorSlot>(), this->DepthIndicatorChildren, Index);
 	}
 	
-	UNREALED_API FSlot& GetSlot(int32 Index);
-	UNREALED_API const FSlot& GetSlot(int32 Index) const;
+	UNREALED_API FDepthIndicatorSlot& GetDepthIndicatorSlot(int32 Index);
+	UNREALED_API const FDepthIndicatorSlot& GetDepthIndicatorSlot(int32 Index) const;
 	UNREALED_API void ClearChildren();
 	
 	SEditorViewportDepthBar()
-		: Children(this)
+		: DepthIndicatorChildren(this)
+		, DepthLabelChildren(this)
+		, AllChildren(this)
 		, Padding(FMargin(0.0f, 4.0f))
 		, SliderMargin(4.0f)
 		, SliderHeadMargin(2.0f)
 		, SliderHeadHeight(3.0f)
-	{}
+	{
+		AllChildren.AddChildren(DepthIndicatorChildren);
+		AllChildren.AddChildren(DepthLabelChildren);
+	}
 	
 	UNREALED_API void Construct(const FArguments& InArgs, const TSharedPtr<FEditorViewportClient>& Client);
 	
@@ -98,6 +103,7 @@ private:
 		FVector Origin;
 		FVector Forward;
 		
+		/** Get the point of the position on the line starting at the origin, facing along the forward. */
 		double GetAxisPosition(const FVector& Position) const;
 	};
 
@@ -120,16 +126,70 @@ private:
 		double BarToRealPosition(float BarPosition) const;
 	};
 
+	DECLARE_DELEGATE_RetVal_OneParam(float, FGetBarPosition, const FBarPositions& BarPositions)
+
+	// Depth label slots are displayed next to the bar
+	struct FDepthLabelSlot : public TSlotBase<FDepthLabelSlot>
+	{
+		SLATE_SLOT_BEGIN_ARGS(FDepthLabelSlot, TSlotBase<FDepthLabelSlot>)
+			/** Determines where in the bar the slot's contents appear. */
+			SLATE_EVENT(FGetBarPosition, OnGetBarPosition)
+			SLATE_ATTRIBUTE(EVerticalAlignment, VAlign);
+		SLATE_SLOT_END_ARGS()
+		
+		void Construct(const FChildren& SlotOwner, FSlotArguments&& InArgs)
+		{
+			TSlotBase<FDepthLabelSlot>::Construct(SlotOwner, MoveTemp(InArgs));
+			OnGetBarPosition = InArgs._OnGetBarPosition;
+			VAlign = InArgs._VAlign;
+		}
+		
+		float GetBarPosition(const FBarPositions& BarPositions) const
+		{
+			if (OnGetBarPosition.IsBound())
+			{
+				return OnGetBarPosition.Execute(BarPositions); 
+			}
+			return 0.0f;
+		}
+		
+		EVerticalAlignment GetVerticalAlignment() const
+		{
+			if (VAlign.IsSet())
+			{
+				return VAlign.Get();
+			}
+			return VAlign_Center;
+		}
+	protected:
+		FGetBarPosition OnGetBarPosition;
+		TAttribute<EVerticalAlignment> VAlign = VAlign_Center;
+	};
+	
+	using FScopedWidgetDepthLabelSlotArguments = TPanelChildren<FDepthLabelSlot>::FScopedWidgetSlotArguments;
+	FScopedWidgetDepthLabelSlotArguments AddDepthLabelSlot()
+	{
+		return InsertDepthLabelSlot(INDEX_NONE);
+	}
+
+	FScopedWidgetDepthLabelSlotArguments InsertDepthLabelSlot(int32 Index = INDEX_NONE)
+	{
+		return FScopedWidgetDepthLabelSlotArguments(MakeUnique<FDepthLabelSlot>(), this->DepthLabelChildren, Index);
+	}
+
 	TOptional<FBarPositions> GetBarPositions(const FGeometry& InGeometry) const;
 	
-	FText GetTooltipText() const;
-	static FText GetFarPlaneText(const TSharedPtr<FEditorViewportClient>& Client);
-	static FText GetNearPlaneText(const TSharedPtr<FEditorViewportClient>& Client);
+	EVisibility GetFarPlaneLabelVisibility() const;
+	EVisibility GetNearPlaneLabelVisibility() const;
+	FText GetFarPlaneText() const;
+	FText GetNearPlaneText() const;
 	
 	FBox WorldBounds;
 	TWeakPtr<FEditorViewportClient> WeakClient;
 	
-	TPanelChildren<FSlot> Children;
+	TPanelChildren<FDepthIndicatorSlot> DepthIndicatorChildren;
+	TPanelChildren<FDepthLabelSlot> DepthLabelChildren;
+	FCombinedChildren AllChildren;
 	
 	enum class EDragTarget
 	{
