Moved UFunction parameter tooltip parsing from K2Node_CallFunction to new ObjectTools::GetDefaultTooltipForFunctionParam and GetDefaultTooltipForFunctionReturn functions.

Also made inclusion of additional target info option for GetDefaultTooltipForFunction function to allow parsing pure function tooltips.

diff --git a/Engine/Source/Editor/BlueprintGraph/Private/K2Node_CallFunction.cpp b/Engine/Source/Editor/BlueprintGraph/Private/K2Node_CallFunction.cpp
--- a/Engine/Source/Editor/BlueprintGraph/Private/K2Node_CallFunction.cpp
+++ b/Engine/Source/Editor/BlueprintGraph/Private/K2Node_CallFunction.cpp
@@ -1685,123 +1685,9 @@ void UK2Node_CallFunction::GeneratePinTooltipFromFunction(UEdGraphPin& Pin, cons
 		return;
 	}
 
-	// figure what tag we should be parsing for (is this a return-val pin, or a parameter?)
-	FString ParamName;
-	FString TagStr = TEXT("@param");
 	const bool bReturnPin = Pin.PinName == UEdGraphSchema_K2::PN_ReturnValue;
-	if (bReturnPin)
-	{
-		TagStr = TEXT("@return");
-	}
-	else
-	{
-		ParamName = Pin.PinName.ToString();
-	}
-
-	// grab the the function's comment block for us to parse
-	FString FunctionToolTipText = Function->GetToolTipText().ToString();
-	
-	int32 CurStrPos = INDEX_NONE;
-	int32 FullToolTipLen = FunctionToolTipText.Len();
-	// parse the full function tooltip text, looking for tag lines
-	do 
-	{
-		CurStrPos = FunctionToolTipText.Find(TagStr, ESearchCase::IgnoreCase, ESearchDir::FromStart, CurStrPos);
-		if (CurStrPos == INDEX_NONE) // if the tag wasn't found
-		{
-			break;
-		}
-
-		// advance past the tag
-		CurStrPos += TagStr.Len();
-
-		// handle people having done @returns instead of @return
-		if (bReturnPin && CurStrPos < FullToolTipLen && FunctionToolTipText[CurStrPos] == TEXT('s'))
-		{
-			++CurStrPos;
-		}
-
-		// advance past whitespace
-		while(CurStrPos < FullToolTipLen && FChar::IsWhitespace(FunctionToolTipText[CurStrPos]))
-		{
-			++CurStrPos;
-		}
-
-		// if this is a parameter pin
-		if (!ParamName.IsEmpty())
-		{
-			FString TagParamName;
-
-			// copy the parameter name
-			while (CurStrPos < FullToolTipLen && !FChar::IsWhitespace(FunctionToolTipText[CurStrPos]))
-			{
-				TagParamName.AppendChar(FunctionToolTipText[CurStrPos++]);
-			}
-
-			// if this @param tag doesn't match the param we're looking for
-			if (TagParamName != ParamName)
-			{
-				continue;
-			}
-		}
-
-		// advance past whitespace (get to the meat of the comment)
-		// since many doxygen style @param use the format "@param <param name> - <comment>" we also strip - if it is before we get to any other non-whitespace
-		while(CurStrPos < FullToolTipLen && (FChar::IsWhitespace(FunctionToolTipText[CurStrPos]) || FunctionToolTipText[CurStrPos] == '-'))
-		{
-			++CurStrPos;
-		}
-
-
-		FString ParamDesc;
-		// collect the param/return-val description
-		while (CurStrPos < FullToolTipLen && FunctionToolTipText[CurStrPos] != TEXT('@'))
-		{
-			// advance past newline
-			while(CurStrPos < FullToolTipLen && FChar::IsLinebreak(FunctionToolTipText[CurStrPos]))
-			{
-				++CurStrPos;
-
-				// advance past whitespace at the start of a new line
-				while(CurStrPos < FullToolTipLen && FChar::IsWhitespace(FunctionToolTipText[CurStrPos]))
-				{
-					++CurStrPos;
-				}
-
-				// replace the newline with a single space
-				if(CurStrPos < FullToolTipLen && !FChar::IsLinebreak(FunctionToolTipText[CurStrPos]))
-				{
-					ParamDesc.AppendChar(TEXT(' '));
-				}
-			}
-
-			if (CurStrPos < FullToolTipLen && FunctionToolTipText[CurStrPos] != TEXT('@'))
-			{
-				ParamDesc.AppendChar(FunctionToolTipText[CurStrPos++]);
-			}
-		}
-
-		// trim any trailing whitespace from the descriptive text
-		ParamDesc.TrimEndInline();
-
-		// if we came up with a valid description for the param/return-val
-		if (!ParamDesc.IsEmpty())
-		{
-			Pin.PinToolTip += ParamDesc;
-			break; // we found a match, so there's no need to continue
-		}
-
-	} while (CurStrPos < FullToolTipLen);
-
-	// If we have no parameter or return value descriptions the full description will be relevant in describing the return value:
-	if( bReturnPin && 
-		Pin.PinToolTip.IsEmpty() && 
-		FunctionToolTipText.Find(TEXT("@param")) == INDEX_NONE && 
-		FunctionToolTipText.Find(TEXT("@return")) == INDEX_NONE)
-	{
-		// for the return pin, default to using the function description if no @return tag was provided:
-		Pin.PinToolTip = Function->GetToolTipText().ToString();
-	}
+	const FString ParamTooltip = bReturnPin ? ObjectTools::GetDefaultTooltipForFunctionReturn(Function) : ObjectTools::GetDefaultTooltipForFunctionParam(Function, Pin.PinName.ToString());
+	Pin.PinToolTip += ParamTooltip;
 
 	GetDefault<UEdGraphSchema_K2>()->ConstructBasicPinTooltip(Pin, FText::FromString(Pin.PinToolTip), Pin.PinToolTip);
 }
diff --git a/Engine/Source/Editor/UnrealEd/Private/ObjectTools.cpp b/Engine/Source/Editor/UnrealEd/Private/ObjectTools.cpp
--- a/Engine/Source/Editor/UnrealEd/Private/ObjectTools.cpp
+++ b/Engine/Source/Editor/UnrealEd/Private/ObjectTools.cpp
@@ -5325,10 +5325,14 @@ namespace ObjectTools
 		return FText::GetEmpty();
 	}
 
-	FString GetDefaultTooltipForFunction(const UFunction* Function)
+	FString GetDefaultTooltipForFunction(const UFunction* Function, bool bIncludeTarget)
 	{
-		FString Tooltip;
+		if (!ensure(Function))
+		{
+			return FString();
+		}
 
+		FString Tooltip;
 		if (Function != nullptr)
 		{
 			Tooltip = Function->GetToolTipText().ToString();
@@ -5351,25 +5355,147 @@ namespace ObjectTools
 			Tooltip.ReplaceInline(*DoxygenNote, *TooltipNote);
 
 			Tooltip.TrimStartAndEndInline();
-
-			UClass* CurrentSelfClass = (Function != nullptr) ? Function->GetOwnerClass() : nullptr;
-			UClass const* TrueSelfClass = CurrentSelfClass;
-			if (CurrentSelfClass && CurrentSelfClass->ClassGeneratedBy)
+			
+			if (bIncludeTarget)
 			{
-				TrueSelfClass = CurrentSelfClass->GetAuthoritativeClass();
+				UClass* CurrentSelfClass = (Function != nullptr) ? Function->GetOwnerClass() : nullptr;
+				UClass const* TrueSelfClass = CurrentSelfClass;
+				if (CurrentSelfClass && CurrentSelfClass->ClassGeneratedBy)
+				{
+					TrueSelfClass = CurrentSelfClass->GetAuthoritativeClass();
+				}
+	
+				FText TargetDisplayText = (TrueSelfClass != nullptr) ? TrueSelfClass->GetDisplayNameText() : LOCTEXT("None", "None");
+	
+				FFormatNamedArguments Args;
+				Args.Add(TEXT("TargetName"), TargetDisplayText);
+				Args.Add(TEXT("Tooltip"), FText::FromString(Tooltip));
+				Tooltip = FText::Format(LOCTEXT("CallFunction_Tooltip", "{Tooltip}\n\nTarget is {TargetName}"), Args).ToString();
 			}
-
-			FText TargetDisplayText = (TrueSelfClass != nullptr) ? TrueSelfClass->GetDisplayNameText() : LOCTEXT("None", "None");
-
-			FFormatNamedArguments Args;
-			Args.Add(TEXT("TargetName"), TargetDisplayText);
-			Args.Add(TEXT("Tooltip"), FText::FromString(Tooltip));
-			return FText::Format(LOCTEXT("CallFunction_Tooltip", "{Tooltip}\n\nTarget is {TargetName}"), Args).ToString();
 		}
-		else
+		else if (bIncludeTarget)
 		{
-			return GetUserFacingFunctionName(Function).ToString();
+			Tooltip = GetUserFacingFunctionName(Function).ToString();
 		}
+		
+		return Tooltip;
+	}
+
+	FString GetDefaultTooltipForFunctionParamOrReturn(const UFunction* Function, const FString& ParamName, const bool bReturnValue)
+	{
+		if (!ensure(Function))
+		{
+			return FString();
+		}
+		
+		const FString TagStr = bReturnValue ? TEXT("@return") : TEXT("@param");
+		
+		// grab the function's comment block for us to parse
+		FString FunctionToolTipText = Function->GetToolTipText().ToString();
+		
+		int32 CurStrPos = INDEX_NONE;
+		int32 FullToolTipLen = FunctionToolTipText.Len();
+		// parse the full function tooltip text, looking for tag lines
+		do 
+		{
+			CurStrPos = FunctionToolTipText.Find(TagStr, ESearchCase::IgnoreCase, ESearchDir::FromStart, CurStrPos);
+			if (CurStrPos == INDEX_NONE) // if the tag wasn't found
+			{
+				break;
+			}
+
+			// advance past the tag
+			CurStrPos += TagStr.Len();
+
+			// handle people having done @returns instead of @return
+			if (bReturnValue && CurStrPos < FullToolTipLen && FunctionToolTipText[CurStrPos] == TEXT('s'))
+			{
+				++CurStrPos;
+			}
+
+			// advance past whitespace
+			while(CurStrPos < FullToolTipLen && FChar::IsWhitespace(FunctionToolTipText[CurStrPos]))
+			{
+				++CurStrPos;
+			}
+
+			// if this is a parameter pin
+			if (!ParamName.IsEmpty())
+			{
+				FString TagParamName;
+
+				// copy the parameter name
+				while (CurStrPos < FullToolTipLen && !FChar::IsWhitespace(FunctionToolTipText[CurStrPos]))
+				{
+					TagParamName.AppendChar(FunctionToolTipText[CurStrPos++]);
+				}
+
+				// if this @param tag doesn't match the param we're looking for
+				if (TagParamName != ParamName)
+				{
+					continue;
+				}
+			}
+
+			// advance past whitespace (get to the meat of the comment)
+			// since many doxygen style @param use the format "@param <param name> - <comment>" we also strip - if it is before we get to any other non-whitespace
+			while(CurStrPos < FullToolTipLen && (FChar::IsWhitespace(FunctionToolTipText[CurStrPos]) || FunctionToolTipText[CurStrPos] == '-'))
+			{
+				++CurStrPos;
+			}
+
+
+			FString ParamDesc;
+			// collect the param/return-val description
+			while (CurStrPos < FullToolTipLen && FunctionToolTipText[CurStrPos] != TEXT('@'))
+			{
+				// advance past newline
+				while(CurStrPos < FullToolTipLen && FChar::IsLinebreak(FunctionToolTipText[CurStrPos]))
+				{
+					++CurStrPos;
+
+					// advance past whitespace at the start of a new line
+					while(CurStrPos < FullToolTipLen && FChar::IsWhitespace(FunctionToolTipText[CurStrPos]))
+					{
+						++CurStrPos;
+					}
+
+					// replace the newline with a single space
+					if(CurStrPos < FullToolTipLen && !FChar::IsLinebreak(FunctionToolTipText[CurStrPos]))
+					{
+						ParamDesc.AppendChar(TEXT(' '));
+					}
+				}
+
+				if (CurStrPos < FullToolTipLen && FunctionToolTipText[CurStrPos] != TEXT('@'))
+				{
+					ParamDesc.AppendChar(FunctionToolTipText[CurStrPos++]);
+				}
+			}
+
+			// trim any trailing whitespace from the descriptive text
+			ParamDesc.TrimEndInline();
+
+			// if we came up with a valid description for the param/return-val
+			if (!ParamDesc.IsEmpty())
+			{
+				// we found a match, so there's no need to continue
+				return ParamDesc;
+			}
+
+		} while (CurStrPos < FullToolTipLen);
+
+		return FString();
+	}
+
+	FString GetDefaultTooltipForFunctionParam(const UFunction* Function, const FString& ParamName)
+	{
+		return GetDefaultTooltipForFunctionParamOrReturn(Function, ParamName, /*bReturnValue*/false);
+	}
+	
+	FString GetDefaultTooltipForFunctionReturn(const UFunction* Function)
+	{
+		return GetDefaultTooltipForFunctionParamOrReturn(Function, FString(), /*bReturnValue*/true);
 	}
 }
 
diff --git a/Engine/Source/Editor/UnrealEd/Public/ObjectTools.h b/Engine/Source/Editor/UnrealEd/Public/ObjectTools.h
--- a/Engine/Source/Editor/UnrealEd/Public/ObjectTools.h
+++ b/Engine/Source/Editor/UnrealEd/Public/ObjectTools.h
@@ -719,7 +719,22 @@ namespace ObjectTools
 	 * @param Function the function to find the tooltip for
 	 * @return the tooltip as a string
 	 */
-	UNREALED_API FString GetDefaultTooltipForFunction(const UFunction* Function);
+	UNREALED_API FString GetDefaultTooltipForFunction(const UFunction* Function, bool bIncludeTarget = true);
+
+	/**
+	 * Gets the non-specific tooltip for the given UFunction parameter.
+	 * @param Function the function to parse the param tooltip from
+	 * @param ParamName the name of the parameter
+	 * @return the tooltip as a string
+	 */
+	UNREALED_API FString GetDefaultTooltipForFunctionParam(const UFunction* Function, const FString& ParamName);
+
+	/**
+	 * Gets the non-specific tooltip for the given UFunction's return value.
+	 * @param Function the function to parse the param tooltip from
+	 * @return the tooltip as a string
+	 */
+	UNREALED_API FString GetDefaultTooltipForFunctionReturn(const UFunction* Function);
 }
 
 
