Cooker: Add a commandline parameter CookingDLCBaseGamePlugin to distinguish between plugins that are cooking as DLCs but will be welded into the basegame by MultiProductCook versus plugins that are cooking as traditional DLCs that will be distributed separately and optionally from the basegame.
This is a prototype change that should be refactored before release.
#rnx

diff --git a/Engine/Source/Editor/UnrealEd/Classes/CookOnTheSide/CookOnTheFlyServer.h b/Engine/Source/Editor/UnrealEd/Classes/CookOnTheSide/CookOnTheFlyServer.h
--- a/Engine/Source/Editor/UnrealEd/Classes/CookOnTheSide/CookOnTheFlyServer.h
+++ b/Engine/Source/Editor/UnrealEd/Classes/CookOnTheSide/CookOnTheFlyServer.h
@@ -693,6 +693,7 @@ public:
 	UNREALED_API virtual TArray<UE::Cook::FInstigator> GetInstigatorChain(FName PackageName) override;
 	UNREALED_API virtual UE::Cook::ECookType GetCookType() override;
 	UNREALED_API virtual UE::Cook::ECookingDLC GetCookingDLC() override;
+	UNREALED_API virtual bool ExperimentalDoNotUse_IsCookingDLCBaseGamePlugin() override;
 	UNREALED_API virtual FString GetDLCName() override;
 	UNREALED_API virtual UE::Cook::EProcessType GetProcessType() override;
 	UNREALED_API virtual UE::Cook::ECookValidationOptions GetCookValidationOptions() override;
@@ -1649,6 +1650,12 @@ private:
 	bool bSkipSave = false;
 	/** Whether cooked packages should store extra data to debug indeterminism. */
 	bool bDeterminismDebug = false;
+	/**
+	 * True (on CookWorkers and CookDirector) if and only if commandline requested CookingDLCBaseGamePlugin. Only used during DLCCooks,
+	 * and exposed to system-specific code to e.g. do different validation based on whether the DLC will be welded into the base game or
+	 * released separately.
+	 */
+	bool bCookingDLCBaseGamePlugin = false;
 	/** Timers for tracking how long we have been busy, to manage retries and warnings of deadlock */
 	double SaveBusyStartTimeSeconds = MAX_flt;
 	double SaveBusyRetryTimeSeconds = MAX_flt;
diff --git a/Engine/Source/Editor/UnrealEd/Private/CookOnTheFlyServer.cpp b/Engine/Source/Editor/UnrealEd/Private/CookOnTheFlyServer.cpp
--- a/Engine/Source/Editor/UnrealEd/Private/CookOnTheFlyServer.cpp
+++ b/Engine/Source/Editor/UnrealEd/Private/CookOnTheFlyServer.cpp
@@ -6484,6 +6484,8 @@ void UCookOnTheFlyServer::SetInitializeConfigSettings(UE::Cook::FInitializeConfi
 	FParse::Value(CommandLine, TEXT("-DeviceProfile="), OverrideDeviceProfileName);
 
 	bIncrementallyModifiedDiagnostics = FParse::Param(CommandLine, TEXT("CookIncrementallyModifiedDiagnostics"));
+
+	bCookingDLCBaseGamePlugin = FParse::Param(CommandLine, TEXT("CookingDLCBaseGamePlugin"));
 }
 
 void UCookOnTheFlyServer::ParseCookFilters()
@@ -6770,6 +6772,11 @@ UE::Cook::ECookingDLC UCookOnTheFlyServer::GetCookingDLC()
 	return IsCookingDLC() ? ECookingDLC::Yes : ECookingDLC::No;
 }
 
+bool UCookOnTheFlyServer::ExperimentalDoNotUse_IsCookingDLCBaseGamePlugin()
+{
+	return IsCookingDLC() && bCookingDLCBaseGamePlugin;
+}
+
 FString UCookOnTheFlyServer::GetDLCName()
 {
 	return CookByTheBookOptions->DlcName;
diff --git a/Engine/Source/Runtime/CoreUObject/Public/UObject/ICookInfo.h b/Engine/Source/Runtime/CoreUObject/Public/UObject/ICookInfo.h
--- a/Engine/Source/Runtime/CoreUObject/Public/UObject/ICookInfo.h
+++ b/Engine/Source/Runtime/CoreUObject/Public/UObject/ICookInfo.h
@@ -213,6 +213,7 @@ public:
 	virtual UE::Cook::ECookType GetCookType() = 0;
 	/** Whether DLC is being cooked (e.g. via "-dlcname=<PluginName>"). This function will not return ECookingDLC::Unknown. */
 	virtual UE::Cook::ECookingDLC GetCookingDLC() = 0;
+	virtual bool ExperimentalDoNotUse_IsCookingDLCBaseGamePlugin() = 0;
 	/** When DLC is being cooked (@see GetCookingDLC), this will return the name of the DLC plugin. */
 	virtual FString GetDLCName() = 0;
 	/** The role the current process plays in its MPCook session, or EProcessType::SingleProcess if it is running standalone. */
