This fix ensures that the ?Hide Selected At Startup? action now correctly applies to all selected actors, regardless of their current visibility state. Previously, actors that were already hidden were skipped by this action. With this change, both visible and non-visible actors will be consistently marked to start hidden on load.

diff --git a/Engine/Source/Editor/UnrealEd/Private/EditorActor.cpp b/Engine/Source/Editor/UnrealEd/Private/EditorActor.cpp
--- a/Engine/Source/Editor/UnrealEd/Private/EditorActor.cpp
+++ b/Engine/Source/Editor/UnrealEd/Private/EditorActor.cpp
@@ -1724,7 +1724,7 @@ void UUnrealEdEngine::edactHideSelectedStartup( UWorld* InWorld )
 		checkSlow( Actor->IsA(AActor::StaticClass()) );
 
 		// Set the actor to hide at editor startup, if it's not already set that way
-		if ( !FActorEditorUtils::IsABuilderBrush(Actor) && !Actor->IsHiddenEd() && !Actor->IsHiddenEdAtStartup() )
+		if ( !FActorEditorUtils::IsABuilderBrush(Actor) && !Actor->IsHiddenEdAtStartup() )
 		{
 			Actor->Modify();
 			Actor->bHiddenEd = true;
