Fix undefined symbol after merge error.
#rnx

diff --git a/Engine/Source/Editor/UnrealEd/Private/CookOnTheFlyServer.cpp b/Engine/Source/Editor/UnrealEd/Private/CookOnTheFlyServer.cpp
--- a/Engine/Source/Editor/UnrealEd/Private/CookOnTheFlyServer.cpp
+++ b/Engine/Source/Editor/UnrealEd/Private/CookOnTheFlyServer.cpp
@@ -6485,7 +6485,7 @@ void UCookOnTheFlyServer::SetInitializeConfigSettings(UE::Cook::FInitializeConfi
 
 	bIncrementallyModifiedDiagnostics = FParse::Param(CommandLine, TEXT("CookIncrementallyModifiedDiagnostics"));
 
-	bCookingDLCBaseGamePlugin = FParse::Param(CommandLine, TEXT("CookingDLCBaseGamePlugin"));
+	bCookingDLCBaseGamePlugin = FParse::Param(FCommandLine::Get(), TEXT("CookingDLCBaseGamePlugin"));
 }
 
 void UCookOnTheFlyServer::ParseCookFilters()
