Removed a warning from GetAssetsInCollection when an asset could not be found.
Fixed an issue where CollectionManagerScripting's logs displayed as red despite actually being warnings.

[RN]

diff --git a/Engine/Source/Editor/UnrealEd/Private/Subsystems/CollectionManagerScriptingSubsystem.cpp b/Engine/Source/Editor/UnrealEd/Private/Subsystems/CollectionManagerScriptingSubsystem.cpp
--- a/Engine/Source/Editor/UnrealEd/Private/Subsystems/CollectionManagerScriptingSubsystem.cpp
+++ b/Engine/Source/Editor/UnrealEd/Private/Subsystems/CollectionManagerScriptingSubsystem.cpp
@@ -18,9 +18,9 @@ DEFINE_LOG_CATEGORY_STATIC(LogCollectionManagerScripting, Log, All);
 namespace CollectionManagerScriptingSubsystemUtil
 {
 
-void LogLastCollectionManagerError(const FText& Error)
+void LogLastCollectionManagerWarning(const FText& Warning)
 {
-	UE_LOG(LogCollectionManagerScripting, Warning, TEXT("Collection manager error: %s"), *Error.ToString());
+	UE_LOG(LogCollectionManagerScripting, Warning, TEXT("%s"), *Warning.ToString());
 }
 
 FString ECollectionScriptingShareType_To_String(const ECollectionScriptingShareType ShareType)
@@ -66,15 +66,13 @@ struct FConstants
 	const FText CollectionNotFound;
 	const FText CollectionExistsInContainer;
 	const FText ContainerMismatch;
-	const FText FailedToGetAsset;
 
 private:
 	FConstants() :
 		ContainerNotFound(LOCTEXT("ContainerNotFound", "No container with the name '{0}' could be found.")),
 		CollectionNotFound(LOCTEXT("CollectionNotFound", "No collection with the name '{0}' and share type '{1}' was found in container with name '{2}'")),
 		CollectionExistsInContainer(LOCTEXT("CollectionExistsInContainer", "Container '{0}' already has a collection named '{1}' with share type '{2}'.")),
-		ContainerMismatch(LOCTEXT("ContainerMismatch", "Collections must be in the same container to reparent. Found target collection in container '{0}' and new parent collection in container '{1}'.")),
-		FailedToGetAsset(LOCTEXT("FailedToGetAsset", "Failed to get Asset '{0}' from Collection '{1}'. Could not get the asset's Object Path."))
+		ContainerMismatch(LOCTEXT("ContainerMismatch", "Collections must be in the same container to reparent. Found target collection in container '{0}' and new parent collection in container '{1}'."))
 	{}
 };
 
@@ -104,27 +102,27 @@ bool UCollectionManagerScriptingSubsystem::CreateCollection(const FCollectionScr
 	const TSharedPtr<ICollectionContainer> CollectionContainer = FCollectionManagerModule::GetModule().Get().FindCollectionContainer(Container.Name);
 	if (!CollectionContainer.IsValid())
 	{
-		CollectionManagerScriptingSubsystemUtil::LogLastCollectionManagerError(FText::Format(CollectionManagerScriptingSubsystemUtil::FConstants::Get().ContainerNotFound, FText::FromName(Container.Name)));
+		CollectionManagerScriptingSubsystemUtil::LogLastCollectionManagerWarning(FText::Format(CollectionManagerScriptingSubsystemUtil::FConstants::Get().ContainerNotFound, FText::FromName(Container.Name)));
 		return false;
 	}
 
 	const ECollectionShareType::Type CastShareType = CollectionManagerScriptingSubsystemUtil::ECollectionScriptingShareType_To_ECollectionShareType(ShareType);
 	if (CollectionContainer->CollectionExists(Collection, CastShareType))
 	{
-		CollectionManagerScriptingSubsystemUtil::LogLastCollectionManagerError(FText::Format(CollectionManagerScriptingSubsystemUtil::FConstants::Get().CollectionExistsInContainer, { CollectionContainer->GetCollectionSource()->GetTitle(), FText::FromName(Collection), CastShareType }));
+		CollectionManagerScriptingSubsystemUtil::LogLastCollectionManagerWarning(FText::Format(CollectionManagerScriptingSubsystemUtil::FConstants::Get().CollectionExistsInContainer, { CollectionContainer->GetCollectionSource()->GetTitle(), FText::FromName(Collection), CastShareType }));
 		return false;
 	}
 
 	FText Error;
 	if (!CollectionContainer->IsValidCollectionName(Collection.ToString(), ECollectionShareType::CST_All, &Error))
 	{
-		CollectionManagerScriptingSubsystemUtil::LogLastCollectionManagerError(Error);
+		CollectionManagerScriptingSubsystemUtil::LogLastCollectionManagerWarning(Error);
 		return false;
 	}
 
 	if (!CollectionContainer->CreateCollection(Collection, CastShareType, ECollectionStorageMode::Static, &Error))
 	{
-		CollectionManagerScriptingSubsystemUtil::LogLastCollectionManagerError(Error);
+		CollectionManagerScriptingSubsystemUtil::LogLastCollectionManagerWarning(Error);
 		return false;
 	}
 
@@ -140,7 +138,7 @@ bool UCollectionManagerScriptingSubsystem::CreateOrEmptyCollection(const FCollec
 	const TSharedPtr<ICollectionContainer> CollectionContainer = FCollectionManagerModule::GetModule().Get().FindCollectionContainer(Container.Name);
 	if (!CollectionContainer.IsValid())
 	{
-		CollectionManagerScriptingSubsystemUtil::LogLastCollectionManagerError(FText::Format(CollectionManagerScriptingSubsystemUtil::FConstants::Get().ContainerNotFound, FText::FromName(Container.Name)));
+		CollectionManagerScriptingSubsystemUtil::LogLastCollectionManagerWarning(FText::Format(CollectionManagerScriptingSubsystemUtil::FConstants::Get().ContainerNotFound, FText::FromName(Container.Name)));
 		return false;
 	}
 
@@ -150,7 +148,7 @@ bool UCollectionManagerScriptingSubsystem::CreateOrEmptyCollection(const FCollec
 		FText Error;
 		if (!CollectionContainer->EmptyCollection(Collection, CastShareType, &Error))
 		{
-			CollectionManagerScriptingSubsystemUtil::LogLastCollectionManagerError(Error);
+			CollectionManagerScriptingSubsystemUtil::LogLastCollectionManagerWarning(Error);
 			return false;
 		}
 	}
@@ -159,13 +157,13 @@ bool UCollectionManagerScriptingSubsystem::CreateOrEmptyCollection(const FCollec
 		FText Error;
 		if (!CollectionContainer->IsValidCollectionName(Collection.ToString(), ECollectionShareType::CST_All, &Error))
 		{
-			CollectionManagerScriptingSubsystemUtil::LogLastCollectionManagerError(Error);
+			CollectionManagerScriptingSubsystemUtil::LogLastCollectionManagerWarning(Error);
 			return false;
 		}
 
 		if (!CollectionContainer->CreateCollection(Collection, CastShareType, ECollectionStorageMode::Static, &Error))
 		{
-			CollectionManagerScriptingSubsystemUtil::LogLastCollectionManagerError(Error);
+			CollectionManagerScriptingSubsystemUtil::LogLastCollectionManagerWarning(Error);
 			return false;
 		}
 	}
@@ -182,7 +180,7 @@ bool UCollectionManagerScriptingSubsystem::GetCollections(const FCollectionScrip
 	const TSharedPtr<ICollectionContainer> CollectionContainer = FCollectionManagerModule::GetModule().Get().FindCollectionContainer(Container.Name);
 	if (!CollectionContainer.IsValid())
 	{
-		CollectionManagerScriptingSubsystemUtil::LogLastCollectionManagerError(FText::Format(CollectionManagerScriptingSubsystemUtil::FConstants::Get().ContainerNotFound, { FText::FromName(Container.Name) }));
+		CollectionManagerScriptingSubsystemUtil::LogLastCollectionManagerWarning(FText::Format(CollectionManagerScriptingSubsystemUtil::FConstants::Get().ContainerNotFound, { FText::FromName(Container.Name) }));
 		return false;
 	}
 
@@ -207,14 +205,14 @@ bool UCollectionManagerScriptingSubsystem::DestroyCollection(const FCollectionSc
 	const TSharedPtr<ICollectionContainer> CollectionContainer = FCollectionManagerModule::GetModule().Get().FindCollectionContainer(Collection.Container);
 	if (!CollectionContainer.IsValid())
 	{
-		CollectionManagerScriptingSubsystemUtil::LogLastCollectionManagerError(FText::Format(CollectionManagerScriptingSubsystemUtil::FConstants::Get().ContainerNotFound, { FText::FromName(Collection.Container) }));
+		CollectionManagerScriptingSubsystemUtil::LogLastCollectionManagerWarning(FText::Format(CollectionManagerScriptingSubsystemUtil::FConstants::Get().ContainerNotFound, { FText::FromName(Collection.Container) }));
 		return false;
 	}
 
 	FText Error;
 	if (!CollectionContainer->DestroyCollection(Collection.Name, CollectionManagerScriptingSubsystemUtil::ECollectionScriptingShareType_To_ECollectionShareType(Collection.ShareType), &Error))
 	{
-		CollectionManagerScriptingSubsystemUtil::LogLastCollectionManagerError(Error);
+		CollectionManagerScriptingSubsystemUtil::LogLastCollectionManagerWarning(Error);
 		return false;
 	}
 
@@ -226,21 +224,21 @@ bool UCollectionManagerScriptingSubsystem::RenameCollection(const FCollectionScr
 	const TSharedPtr<ICollectionContainer> CollectionContainer = FCollectionManagerModule::GetModule().Get().FindCollectionContainer(Collection.Container);
 	if (!CollectionContainer.IsValid())
 	{
-		CollectionManagerScriptingSubsystemUtil::LogLastCollectionManagerError(FText::Format(CollectionManagerScriptingSubsystemUtil::FConstants::Get().ContainerNotFound, { FText::FromName(Collection.Container) }));
+		CollectionManagerScriptingSubsystemUtil::LogLastCollectionManagerWarning(FText::Format(CollectionManagerScriptingSubsystemUtil::FConstants::Get().ContainerNotFound, { FText::FromName(Collection.Container) }));
 		return false;
 	}
 
 	FText Error;
 	if (!CollectionContainer->IsValidCollectionName(NewName.ToString(), ECollectionShareType::CST_All, &Error))
 	{
-		CollectionManagerScriptingSubsystemUtil::LogLastCollectionManagerError(Error);
+		CollectionManagerScriptingSubsystemUtil::LogLastCollectionManagerWarning(Error);
 		return false;
 	}
 
 	if (!CollectionContainer->RenameCollection(Collection.Name, CollectionManagerScriptingSubsystemUtil::ECollectionScriptingShareType_To_ECollectionShareType(Collection.ShareType),
 		NewName, CollectionManagerScriptingSubsystemUtil::ECollectionScriptingShareType_To_ECollectionShareType(NewShareType), &Error))
 	{
-		CollectionManagerScriptingSubsystemUtil::LogLastCollectionManagerError(Error);
+		CollectionManagerScriptingSubsystemUtil::LogLastCollectionManagerWarning(Error);
 		return false;
 	}
 
@@ -252,14 +250,14 @@ bool UCollectionManagerScriptingSubsystem::ReparentCollection(const FCollectionS
 	const TSharedPtr<ICollectionContainer> CollectionContainer = FCollectionManagerModule::GetModule().Get().FindCollectionContainer(Collection.Container);
 	if (!CollectionContainer.IsValid())
 	{
-		CollectionManagerScriptingSubsystemUtil::LogLastCollectionManagerError(FText::Format(CollectionManagerScriptingSubsystemUtil::FConstants::Get().ContainerNotFound, { FText::FromName(Collection.Container) }));
+		CollectionManagerScriptingSubsystemUtil::LogLastCollectionManagerWarning(FText::Format(CollectionManagerScriptingSubsystemUtil::FConstants::Get().ContainerNotFound, { FText::FromName(Collection.Container) }));
 		return false;
 	}
 
 	// If the intention is to reparent to a different collection, make sure its in the same container.
 	if (!NewParentCollection.Name.IsNone() && Collection.Container != NewParentCollection.Container)
 	{
-		CollectionManagerScriptingSubsystemUtil::LogLastCollectionManagerError(FText::Format(CollectionManagerScriptingSubsystemUtil::FConstants::Get().ContainerMismatch, { FText::FromName(Collection.Container), FText::FromName(NewParentCollection.Container) }));
+		CollectionManagerScriptingSubsystemUtil::LogLastCollectionManagerWarning(FText::Format(CollectionManagerScriptingSubsystemUtil::FConstants::Get().ContainerMismatch, { FText::FromName(Collection.Container), FText::FromName(NewParentCollection.Container) }));
 		return false;
 	}
 
@@ -267,7 +265,7 @@ bool UCollectionManagerScriptingSubsystem::ReparentCollection(const FCollectionS
 	if (!CollectionContainer->ReparentCollection(Collection.Name, CollectionManagerScriptingSubsystemUtil::ECollectionScriptingShareType_To_ECollectionShareType(Collection.ShareType),
 		NewParentCollection.Name, CollectionManagerScriptingSubsystemUtil::ECollectionScriptingShareType_To_ECollectionShareType(NewParentCollection.ShareType), &Error))
 	{
-		CollectionManagerScriptingSubsystemUtil::LogLastCollectionManagerError(Error);
+		CollectionManagerScriptingSubsystemUtil::LogLastCollectionManagerWarning(Error);
 		return false;
 	}
 
@@ -279,14 +277,14 @@ bool UCollectionManagerScriptingSubsystem::EmptyCollection(const FCollectionScri
 	const TSharedPtr<ICollectionContainer> CollectionContainer = FCollectionManagerModule::GetModule().Get().FindCollectionContainer(Collection.Container);
 	if (!CollectionContainer.IsValid())
 	{
-		CollectionManagerScriptingSubsystemUtil::LogLastCollectionManagerError(FText::Format(CollectionManagerScriptingSubsystemUtil::FConstants::Get().ContainerNotFound, { FText::FromName(Collection.Container) }));
+		CollectionManagerScriptingSubsystemUtil::LogLastCollectionManagerWarning(FText::Format(CollectionManagerScriptingSubsystemUtil::FConstants::Get().ContainerNotFound, { FText::FromName(Collection.Container) }));
 		return false;
 	}
 
 	FText Error;
 	if (!CollectionContainer->EmptyCollection(Collection.Name, CollectionManagerScriptingSubsystemUtil::ECollectionScriptingShareType_To_ECollectionShareType(Collection.ShareType), &Error))
 	{
-		CollectionManagerScriptingSubsystemUtil::LogLastCollectionManagerError(Error);
+		CollectionManagerScriptingSubsystemUtil::LogLastCollectionManagerWarning(Error);
 		return false;
 	}
 
@@ -313,7 +311,7 @@ bool UCollectionManagerScriptingSubsystem::AddAssetsToCollection(const FCollecti
 	const TSharedPtr<ICollectionContainer> CollectionContainer = FCollectionManagerModule::GetModule().Get().FindCollectionContainer(Collection.Container);
 	if (!CollectionContainer.IsValid())
 	{
-		CollectionManagerScriptingSubsystemUtil::LogLastCollectionManagerError(FText::Format(CollectionManagerScriptingSubsystemUtil::FConstants::Get().ContainerNotFound, { FText::FromName(Collection.Container) }));
+		CollectionManagerScriptingSubsystemUtil::LogLastCollectionManagerWarning(FText::Format(CollectionManagerScriptingSubsystemUtil::FConstants::Get().ContainerNotFound, { FText::FromName(Collection.Container) }));
 		return false;
 	}
 
@@ -323,20 +321,20 @@ bool UCollectionManagerScriptingSubsystem::AddAssetsToCollection(const FCollecti
 	{
 		if (!CollectionContainer->IsValidCollectionName(Collection.Name.ToString(), ECollectionShareType::CST_All, &Error))
 		{
-			CollectionManagerScriptingSubsystemUtil::LogLastCollectionManagerError(Error);
+			CollectionManagerScriptingSubsystemUtil::LogLastCollectionManagerWarning(Error);
 			return false;
 		}
 
 		if (!CollectionContainer->CreateCollection(Collection.Name, ShareType, ECollectionStorageMode::Static, &Error))
 		{
-			CollectionManagerScriptingSubsystemUtil::LogLastCollectionManagerError(Error);
+			CollectionManagerScriptingSubsystemUtil::LogLastCollectionManagerWarning(Error);
 			return false;
 		}
 	}
 
 	if (!CollectionContainer->AddToCollection(Collection.Name, ShareType, AssetPaths, nullptr, &Error))
 	{
-		CollectionManagerScriptingSubsystemUtil::LogLastCollectionManagerError(Error);
+		CollectionManagerScriptingSubsystemUtil::LogLastCollectionManagerWarning(Error);
 		return false;
 	}
 
@@ -381,14 +379,14 @@ bool UCollectionManagerScriptingSubsystem::RemoveAssetsFromCollection(const FCol
 	const TSharedPtr<ICollectionContainer> CollectionContainer = FCollectionManagerModule::GetModule().Get().FindCollectionContainer(Collection.Container);
 	if (!CollectionContainer.IsValid())
 	{
-		CollectionManagerScriptingSubsystemUtil::LogLastCollectionManagerError(FText::Format(CollectionManagerScriptingSubsystemUtil::FConstants::Get().ContainerNotFound, { FText::FromName(Collection.Container) }));
+		CollectionManagerScriptingSubsystemUtil::LogLastCollectionManagerWarning(FText::Format(CollectionManagerScriptingSubsystemUtil::FConstants::Get().ContainerNotFound, { FText::FromName(Collection.Container) }));
 		return false;
 	}
 
 	FText Error;
 	if (!CollectionContainer->RemoveFromCollection(Collection.Name, CollectionManagerScriptingSubsystemUtil::ECollectionScriptingShareType_To_ECollectionShareType(Collection.ShareType), AssetPaths, nullptr, &Error))
 	{
-		CollectionManagerScriptingSubsystemUtil::LogLastCollectionManagerError(Error);
+		CollectionManagerScriptingSubsystemUtil::LogLastCollectionManagerWarning(Error);
 		return false;
 	}
 
@@ -418,7 +416,7 @@ bool UCollectionManagerScriptingSubsystem::CollectionExists(const FCollectionScr
 	const TSharedPtr<ICollectionContainer> CollectionContainer = FCollectionManagerModule::GetModule().Get().FindCollectionContainer(Container.Name);
 	if (!CollectionContainer.IsValid())
 	{
-		CollectionManagerScriptingSubsystemUtil::LogLastCollectionManagerError(FText::Format(CollectionManagerScriptingSubsystemUtil::FConstants::Get().ContainerNotFound, { FText::FromName(Container.Name) }));
+		CollectionManagerScriptingSubsystemUtil::LogLastCollectionManagerWarning(FText::Format(CollectionManagerScriptingSubsystemUtil::FConstants::Get().ContainerNotFound, { FText::FromName(Container.Name) }));
 		return false;
 	}
 
@@ -430,7 +428,7 @@ bool UCollectionManagerScriptingSubsystem::GetCollectionsByName(const FCollectio
 	const TSharedPtr<ICollectionContainer> CollectionContainer = FCollectionManagerModule::GetModule().Get().FindCollectionContainer(Container.Name);
 	if (!CollectionContainer.IsValid())
 	{
-		CollectionManagerScriptingSubsystemUtil::LogLastCollectionManagerError(FText::Format(CollectionManagerScriptingSubsystemUtil::FConstants::Get().ContainerNotFound, { FText::FromName(Container.Name) }));
+		CollectionManagerScriptingSubsystemUtil::LogLastCollectionManagerWarning(FText::Format(CollectionManagerScriptingSubsystemUtil::FConstants::Get().ContainerNotFound, { FText::FromName(Container.Name) }));
 		return false;
 	}
 
@@ -455,13 +453,13 @@ bool UCollectionManagerScriptingSubsystem::GetAssetsInCollection(const FCollecti
 	const TSharedPtr<ICollectionContainer> CollectionContainer = FCollectionManagerModule::GetModule().Get().FindCollectionContainer(Collection.Container);
 	if (!CollectionContainer.IsValid())
 	{
-		CollectionManagerScriptingSubsystemUtil::LogLastCollectionManagerError(FText::Format(CollectionManagerScriptingSubsystemUtil::FConstants::Get().ContainerNotFound, { FText::FromName(Collection.Container) }));
+		CollectionManagerScriptingSubsystemUtil::LogLastCollectionManagerWarning(FText::Format(CollectionManagerScriptingSubsystemUtil::FConstants::Get().ContainerNotFound, { FText::FromName(Collection.Container) }));
 		return false;
 	}
 
 	if (!CollectionContainer->CollectionExists(Collection.Name, CollectionManagerScriptingSubsystemUtil::ECollectionScriptingShareType_To_ECollectionShareType(Collection.ShareType)))
 	{
-		CollectionManagerScriptingSubsystemUtil::LogLastCollectionManagerError(FText::Format(CollectionManagerScriptingSubsystemUtil::FConstants::Get().CollectionNotFound,
+		CollectionManagerScriptingSubsystemUtil::LogLastCollectionManagerWarning(FText::Format(CollectionManagerScriptingSubsystemUtil::FConstants::Get().CollectionNotFound,
 			{ FText::FromName(Collection.Name), FText::FromString(CollectionManagerScriptingSubsystemUtil::ECollectionScriptingShareType_To_String(Collection.ShareType)), FText::FromName(Collection.Container) }));
 		return false;
 	}
@@ -478,10 +476,6 @@ bool UCollectionManagerScriptingSubsystem::GetAssetsInCollection(const FCollecti
 		{
 			OutAssets.Emplace(MoveTemp(AssetData));
 		}
-		else
-		{
-			CollectionManagerScriptingSubsystemUtil::LogLastCollectionManagerError(FText::Format(CollectionManagerScriptingSubsystemUtil::FConstants::Get().FailedToGetAsset, { FText::FromString(Asset.ToString()), FText::FromName(Collection.Name) }));
-		}
 	}
 
 	return true;
@@ -492,7 +486,7 @@ bool UCollectionManagerScriptingSubsystem::GetCollectionsContainingAsset(const F
 	const TSharedPtr<ICollectionContainer> CollectionContainer = FCollectionManagerModule::GetModule().Get().FindCollectionContainer(Container.Name);
 	if (!CollectionContainer.IsValid())
 	{
-		CollectionManagerScriptingSubsystemUtil::LogLastCollectionManagerError(FText::Format(CollectionManagerScriptingSubsystemUtil::FConstants::Get().ContainerNotFound, { FText::FromName(Container.Name) }));
+		CollectionManagerScriptingSubsystemUtil::LogLastCollectionManagerWarning(FText::Format(CollectionManagerScriptingSubsystemUtil::FConstants::Get().ContainerNotFound, { FText::FromName(Container.Name) }));
 		return false;
 	}
 
