[Shaders] Optimize pipeline to reduce single allocation size and improve performance by copying less data on resize

#tests Using the editor, running PIE and game client, making builds

diff --git a/Engine/Source/Editor/UnrealEd/Private/Commandlets/ShaderPipelineCacheToolsCommandlet.cpp b/Engine/Source/Editor/UnrealEd/Private/Commandlets/ShaderPipelineCacheToolsCommandlet.cpp
--- a/Engine/Source/Editor/UnrealEd/Private/Commandlets/ShaderPipelineCacheToolsCommandlet.cpp
+++ b/Engine/Source/Editor/UnrealEd/Private/Commandlets/ShaderPipelineCacheToolsCommandlet.cpp
@@ -145,7 +145,7 @@ static void LoadStableShaderKeys(TArray<FStableShaderKeyAndValue>& StableArray,
 	UE_LOG(LogShaderPipelineCacheTools, Display, TEXT("Loaded %d shader info lines from %.*s."), (StableArray.Num() - StableArrayOffset), FileName.Len(), FileName.GetData());
 }
 
-static void LoadStableShaderKeysMultiple(TMultiMap<FStableShaderKeyAndValue, FSHAHash>& StableMap, TArrayView<const FStringView> FileNames)
+static void LoadStableShaderKeysMultiple(FStableShaderSet& StableMap, TArrayView<const FStringView> FileNames)
 {
 	TArray<TArray<FStableShaderKeyAndValue>> StableArrays;
 	StableArrays.AddDefaulted(FileNames.Num());
@@ -153,25 +153,32 @@ static void LoadStableShaderKeysMultiple(TMultiMap<FStableShaderKeyAndValue, FSH
 
 	if (StableArrays.Num() > 0)
 	{
+		// Since stable keys are saved from a TSet, we assume that a single array does not have non-unique members, so add the largest one first in the most optimal manner
+		StableArrays.Sort([](const TArray<FStableShaderKeyAndValue>& A, const TArray<FStableShaderKeyAndValue>& B) { return (A.Num() > B.Num()); });
+
 		const int32 StableArrayCount = Algo::TransformAccumulate(StableArrays, &TArray<FStableShaderKeyAndValue>::Num, 0);
 		StableMap.Reserve(StableMap.Num() + StableArrayCount);
 
-		// Since stable keys are saved from a TSet, we assume that a single array does not have non-unique members, so add the largest one without using AddUnique
-		StableArrays.Sort([](const TArray<FStableShaderKeyAndValue>& A, const TArray<FStableShaderKeyAndValue>& B) { return (A.Num() > B.Num()); });
-		const TArray<FStableShaderKeyAndValue>& StableArrayLargest = StableArrays[0];
-		for (const FStableShaderKeyAndValue& Item : StableArrayLargest)
-		{
-			StableMap.Add(Item, Item.OutputHash);
-		}
-
-		if (StableArrays.Num() > 1)
+		bool bFirst = true;
+		for (const TArray<FStableShaderKeyAndValue>& StableArray : StableArrays)
 		{
-			for (int32 IdxStableArray = 1, StableArraysNum = StableArrays.Num(); IdxStableArray < StableArraysNum; ++IdxStableArray)
+			if (bFirst)
 			{
-				const TArray<FStableShaderKeyAndValue>& StableArray = StableArrays[IdxStableArray];
-				for (const FStableShaderKeyAndValue& Item : StableArray)
+				bFirst = false;
+				for (const FStableShaderKeyAndValue& Shader : StableArray)
 				{
-					StableMap.AddUnique(Item, Item.OutputHash);
+					StableMap.Add(MakeUnique<FStableShaderKeyAndValue>(Shader));
+				}
+			}
+			else
+			{
+				for (const FStableShaderKeyAndValue& Shader : StableArray)
+				{
+					// Double lookup to avoid MakeUnique for duplicates
+					if (!StableMap.ContainsByHash(GetTypeHash(Shader), Shader))
+					{
+						StableMap.Add(MakeUnique<FStableShaderKeyAndValue>(Shader));
+					}
 				}
 			}
 		}
@@ -435,13 +442,11 @@ int32 DumpPSOSC(FString& Token, const FString& StableKeyFileDir)
 			StableKeyFilesViews.Add(StableKeyFile);
 		}
 
-		TMultiMap<FStableShaderKeyAndValue, FSHAHash> StableMap;
+		FStableShaderSet StableMap;
 		LoadStableShaderKeysMultiple(StableMap, StableKeyFilesViews);
-		for (const auto& Pair : StableMap)
+		for (const TUniquePtr<FStableShaderKeyAndValue>& ShaderKeyPtr : StableMap)
 		{
-			FStableShaderKeyAndValue Temp = Pair.Key;
-			Temp.OutputHash = Pair.Value;
-			InverseMap.FindOrAdd(Pair.Value).Add(Temp.ToString());
+			InverseMap.FindOrAdd(ShaderKeyPtr->OutputHash).Add(ShaderKeyPtr->ToString());
 		}
 	}
 
@@ -551,21 +556,19 @@ static bool GetStableShaders(const TMap<FSHAHash, TArray<int32>>& InverseMap, TA
 	return true;
 }
 
-static void StableShadersSerializationSelfTest(const TMultiMap<FStableShaderKeyAndValue, FSHAHash>& StableMap)
+static void StableShadersSerializationSelfTest(const FStableShaderSet& StableMap)
 {
 	TAnsiStringBuilder<384> TestString;
-	for (const auto& Pair : StableMap)
+	for (const TUniquePtr<FStableShaderKeyAndValue>& ShaderKeyPtr : StableMap)
 	{
 		TestString.Reset();
-		FStableShaderKeyAndValue Item(Pair.Key);
-		Item.OutputHash = Pair.Value;
-		check(Pair.Value != FSHAHash());
-		Item.AppendString(TestString);
+		check(ShaderKeyPtr->OutputHash != FSHAHash());
+		ShaderKeyPtr->AppendString(TestString);
 		FStableShaderKeyAndValue TestItem;
 		TestItem.ParseFromString(UTF8_TO_TCHAR(TestString.ToString()));
-		check(Item == TestItem);
-		check(GetTypeHash(Item) == GetTypeHash(TestItem));
-		check(Item.OutputHash == TestItem.OutputHash);
+		check(*ShaderKeyPtr == TestItem);
+		check(GetTypeHash(*ShaderKeyPtr) == GetTypeHash(TestItem));
+		check(ShaderKeyPtr->OutputHash == TestItem.OutputHash);
 	}
 }
 
@@ -700,27 +703,26 @@ int32 DumpStableKeysFile(const FString& Token, const TMap<FString, FString>& Par
 	TArray<FString> OutputHashesToSearch = ParseParamStringList(ParamVals, TEXT("outputhash"));
 	TArray<FString> NamesToSearch = ParseParamStringList(ParamVals, TEXT("objname"));
 
-	TMultiMap<FStableShaderKeyAndValue, FSHAHash> StableMap;
+	FStableShaderSet StableMap;
 	LoadStableShaderKeysMultiple(StableMap, MakeArrayView(&File, 1));
 	UE_LOG(LogShaderPipelineCacheTools, Display, TEXT("    %s"), *FStableShaderKeyAndValue::HeaderLine());
-	for (const auto& Pair : StableMap)
+	for (const TUniquePtr<FStableShaderKeyAndValue>& ShaderKeyPtr : StableMap)
 	{
-		if (!OutputHashesToSearch.IsEmpty() && !OutputHashesToSearch.Contains(Pair.Key.OutputHash.ToString().ToLower()))
+		if (!OutputHashesToSearch.IsEmpty() && !OutputHashesToSearch.Contains(ShaderKeyPtr->OutputHash.ToString().ToLower()))
 		{
 			continue;
 		}
 
 		if (!NamesToSearch.IsEmpty())
 		{
-			const FString* Found = Algo::FindByPredicate(NamesToSearch, [ObjName = Pair.Key.ClassNameAndObjectPath.ToString().ToLower()](const FString& PartialName) { return ObjName.Contains(PartialName); });
+			const FString* Found = Algo::FindByPredicate(NamesToSearch, [ObjName = ShaderKeyPtr->ClassNameAndObjectPath.ToString().ToLower()](const FString& PartialName) { return ObjName.Contains(PartialName); });
 			if (!Found)
 			{
 				continue;
 			}
 		}
 
-		FStableShaderKeyAndValue Temp(Pair.Key);
-		UE_LOG(LogShaderPipelineCacheTools, Display, TEXT("    %s"), *Temp.ToString());
+		UE_LOG(LogShaderPipelineCacheTools, Display, TEXT("    %s"), *ShaderKeyPtr->ToString());
 	}
 	return 0;
 }
@@ -735,7 +737,7 @@ int32 CheckStableKeyAliasing(const TArray<FString>& Tokens)
 
 	UE_LOG(LogShaderPipelineCacheTools, Display, TEXT("Loading %d stable shader key files..."), Tokens.Num());
 
-	TArray<TMultiMap<FStableShaderKeyAndValue, FSHAHash>> StableMaps;
+	TArray<FStableShaderSet> StableMaps;
 	StableMaps.AddDefaulted(Tokens.Num());
 
 	for (int32 IdxFile = 0; IdxFile < Tokens.Num(); ++IdxFile)
@@ -744,31 +746,28 @@ int32 CheckStableKeyAliasing(const TArray<FString>& Tokens)
 		LoadStableShaderKeysMultiple(StableMaps[IdxFile], MakeArrayView(&File, 1));
 	}
 
-	TArray<TMultiMap<FSHAHash, FStableShaderKeyAndValue>> InverseStableMaps;
+	TArray<TMultiMap<FSHAHash, const FStableShaderKeyAndValue*>> InverseStableMaps;
 	InverseStableMaps.AddDefaulted(StableMaps.Num());
 
 	for (int32 IdxFile = 0; IdxFile < Tokens.Num(); ++IdxFile)
 	{
-		for (const auto& Pair : StableMaps[IdxFile])
+		for (const TUniquePtr<FStableShaderKeyAndValue>& ShaderKeyPtr : StableMaps[IdxFile])
 		{
-			FStableShaderKeyAndValue Temp(Pair.Key);
-			InverseStableMaps[IdxFile].Add(Pair.Value, Pair.Key);
+			InverseStableMaps[IdxFile].Add(ShaderKeyPtr->OutputHash, ShaderKeyPtr.Get());
 		}
-
-		StableMaps[IdxFile] = TMultiMap<FStableShaderKeyAndValue, FSHAHash>();
 	}
 
-	auto LogAllInFile = [](const FString& FileName, const TMultiMap<FSHAHash, FStableShaderKeyAndValue>& InverseMap, const FSHAHash& Hash)
+	auto LogAllInFile = [](const FString& FileName, const TMultiMap<FSHAHash, const FStableShaderKeyAndValue*>& InverseMap, const FSHAHash& Hash)
 	{
-		TArray<FStableShaderKeyAndValue> Values;
+		TArray<const FStableShaderKeyAndValue*> Values;
 		InverseMap.MultiFind(Hash, Values, false);
 		checkf(!Values.IsEmpty(), TEXT("This function should only be called with a hash known to exist in the inverse map"));
 
 		UE_LOG(LogShaderPipelineCacheTools, Display, TEXT("In file %s, it maps to %d shader(s):"), *FileName, Values.Num());
 		int32 ValuesToLog = 10;
-		for (const FStableShaderKeyAndValue& Value : Values)
+		for (const FStableShaderKeyAndValue* Value : Values)
 		{
-			UE_LOG(LogShaderPipelineCacheTools, Display, TEXT("    %s"), *Value.ToString());
+			UE_LOG(LogShaderPipelineCacheTools, Display, TEXT("    %s"), *Value->ToString());
 			if (ValuesToLog-- <= 0)
 			{
 				UE_LOG(LogShaderPipelineCacheTools, Display, TEXT("    etc (%d total)"), Values.Num());
@@ -1441,7 +1440,7 @@ static void ParseQuoteComma(const FStringView& InLine, TArray<FStringView, TInli
 	}
 }
 
-static TSet<FPipelineCacheFileFormatPSO> ParseStableCSV(const FString& FileName, const TArray<FString>& CSVLines, const TMultiMap<FStableShaderKeyAndValue, FSHAHash>& StableMap, FName& TargetPlatform, int32& PSOsRejected, int32& PSOsMerged)
+static TSet<FPipelineCacheFileFormatPSO> ParseStableCSV(const FString& FileName, const TArray<FString>& CSVLines, const FStableShaderSet& StableMap, FName& TargetPlatform, int32& PSOsRejected, int32& PSOsMerged)
 {
 	TSet<FPipelineCacheFileFormatPSO> PSOs;
 
@@ -1553,34 +1552,21 @@ static TSet<FPipelineCacheFileFormatPSO> ParseStableCSV(const FString& FileName,
 			}
 
 			FSHAHash Match;
-			int32 Count = 0;
-			for (auto Iter = StableMap.CreateConstKeyIterator(Shader); Iter; ++Iter)
+
+			if (const TUniquePtr<FStableShaderKeyAndValue>* ShaderKeyPtr = StableMap.Find(Shader))
 			{
-				check(Iter.Value() != FSHAHash());
-				Match = Iter.Value();
-				if (TargetPlatform == NAME_None)
-				{
-					TargetPlatform = Iter.Key().TargetPlatform;
-				}
-				else
-				{
-					check(TargetPlatform == Iter.Key().TargetPlatform);
-				}
-				++Count;
-			}
+				Match = (*ShaderKeyPtr)->OutputHash;
+				check(Match != FSHAHash());
 
-			if (!Count)
+				TargetPlatform = (*ShaderKeyPtr)->TargetPlatform;
+			}
+			else
 			{
 				UE_LOG(LogShaderPipelineCacheTools, Verbose, TEXT("Stable PSO not found, rejecting %s"), *Shader.ToString());
 				++PSOsRejected;
 				return;
 			}
 
-			if (Count > 1)
-			{
-				UE_LOG(LogShaderPipelineCacheTools, Error, TEXT("Stable PSO maps to multiple shaders. This is usually a bad thing and means you used %s files from multiple builds. Ignoring all but the last %s"), ShaderStableKeysFileExt, *Shader.ToString());
-			}
-
 			switch (AdjustedSlotIndex)
 			{
 			case SF_Vertex:
@@ -1734,7 +1720,7 @@ const TCHAR* VertexElementToString(EVertexElementType Type)
 }
 
 
-void FilterInvalidPSOs(TSet<FPipelineCacheFileFormatPSO>& InOutPSOs, const TMultiMap<FStableShaderKeyAndValue, FSHAHash>& StableMap)
+void FilterInvalidPSOs(TSet<FPipelineCacheFileFormatPSO>& InOutPSOs, const FStableShaderSet& StableMap)
 {
 	// list of Vertex Shaders known to be usable with empty vertex declaration without taking VF into consideration
 	const TCHAR* VShadersUsableWithEmptyVertexDecl_Table[] =
@@ -1924,11 +1910,9 @@ void FilterInvalidPSOs(TSet<FPipelineCacheFileFormatPSO>& InOutPSOs, const TMult
 
 	// inverse map is needed for VS checking
 	TMap<FSHAHash, TArray<FStableShaderKeyAndValue>> InverseMap;
-	for (const TTuple<FStableShaderKeyAndValue, FSHAHash>& Pair : StableMap)
+	for (const TUniquePtr<FStableShaderKeyAndValue>& ShaderKeyPtr : StableMap)
 	{
-		FStableShaderKeyAndValue Temp(Pair.Key);
-		Temp.OutputHash = Pair.Value;
-		InverseMap.FindOrAdd(Pair.Value).Add(Temp);
+		InverseMap.FindOrAdd(ShaderKeyPtr->OutputHash).Add(*ShaderKeyPtr);
 	}
 
 	// At this point we cannot really know what is the correct vertex format (input layout) for a given vertex shader. Instead, we're looking if we see the same VS used in multiple PSOs with incompatible vertex descriptors.
@@ -2165,7 +2149,7 @@ void FilterInvalidPSOs(TSet<FPipelineCacheFileFormatPSO>& InOutPSOs, const TMult
 }
 
 /** Adds compute PSOs directly from the stable shader map of this build */
-void AddComputePSOs(TSet<FPipelineCacheFileFormatPSO>& OutPSOs, const TMultiMap<FStableShaderKeyAndValue, FSHAHash>& StableShaderMap)
+void AddComputePSOs(TSet<FPipelineCacheFileFormatPSO>& OutPSOs, const FStableShaderSet& StableShaderMap)
 {
 	if (GShaderPipelineCacheTools_ComputePSOInclusionMode < 1)
 	{
@@ -2175,20 +2159,20 @@ void AddComputePSOs(TSet<FPipelineCacheFileFormatPSO>& OutPSOs, const TMultiMap<
 	static FName NAME_SF_Compute("SF_Compute");
 	static FName NAME_NiagaraShader("FNiagaraShader");
 
-	for (TMultiMap<FStableShaderKeyAndValue, FSHAHash>::TConstIterator Iter(StableShaderMap); Iter; ++Iter)
+	for (const TUniquePtr<FStableShaderKeyAndValue>& ShaderKeyPtr : StableShaderMap)
 	{
-		if (Iter.Key().TargetFrequency == NAME_SF_Compute)
+		if (ShaderKeyPtr->TargetFrequency == NAME_SF_Compute)
 		{
 			// add a new Compute PSO
 			// Check if we are only allowed to add Niagara PSOs
-			if (GShaderPipelineCacheTools_ComputePSOInclusionMode == 2 && Iter.Key().ShaderType != NAME_NiagaraShader)
+			if (GShaderPipelineCacheTools_ComputePSOInclusionMode == 2 &&ShaderKeyPtr->ShaderType != NAME_NiagaraShader)
 			{
 				continue;
 			}
 
 			FPipelineCacheFileFormatPSO NewPso;
 			NewPso.Type = FPipelineCacheFileFormatPSO::DescriptorType::Compute;
-			NewPso.ComputeDesc.ComputeShader = Iter.Value();
+			NewPso.ComputeDesc.ComputeShader = ShaderKeyPtr->OutputHash;
 			NewPso.UsageMask = uint64(-1);
 			NewPso.BindCount = 0;
 			OutPSOs.Add(NewPso);
@@ -2197,12 +2181,11 @@ void AddComputePSOs(TSet<FPipelineCacheFileFormatPSO>& OutPSOs, const TMultiMap<
 }
 
 /** Function that gets the target platform name from the stable map. Shouldn't exist, and we should be passing the target platform explicitly. */
-FName GetTargetPlatformFromStableShaderKeys(const TMultiMap<FStableShaderKeyAndValue, FSHAHash>& StableShaderMap)
+FName GetTargetPlatformFromStableShaderKeys(const FStableShaderSet& StableShaderMap)
 {
-	TMultiMap<FStableShaderKeyAndValue, FSHAHash>::TConstIterator Iter(StableShaderMap);
-	if (Iter)
+	if (const TUniquePtr<FStableShaderKeyAndValue>* ShaderKeyPtr = StableShaderMap.FindArbitraryElement())
 	{
-		return Iter.Key().TargetPlatform;
+		return (*ShaderKeyPtr)->TargetPlatform;
 	}
 
 	return NAME_None;
@@ -2213,7 +2196,7 @@ FName GetTargetPlatformFromStableShaderKeys(const TMultiMap<FStableShaderKeyAndV
  * 
  * @return commandlet return (0 is everything Ok, otherwise can return error codes 1-255, currently only 1 is used)
  */
-int32 SaveBinaryPipelineCacheFile(const FString& OutputFilename, const EShaderPlatform ShaderPlatform, const FString& ShaderFormat, const FString& ChunkInfoFilesPath, const FString& AssociatedShaderLibraryName, const FString& TargetPlatformName, const TSet<FPipelineCacheFileFormatPSO>& PSOs, const TMultiMap<FStableShaderKeyAndValue, FSHAHash>& StableMap)
+int32 SaveBinaryPipelineCacheFile(const FString& OutputFilename, const EShaderPlatform ShaderPlatform, const FString& ShaderFormat, const FString& ChunkInfoFilesPath, const FString& AssociatedShaderLibraryName, const FString& TargetPlatformName, const TSet<FPipelineCacheFileFormatPSO>& PSOs, const FStableShaderSet& StableMap)
 {
 	auto SaveSingleCacheFile = [](const FString& OutputFilename, const EShaderPlatform ShaderPlatform, const TSet<FPipelineCacheFileFormatPSO>& PSOs) -> int32
 	{
@@ -2264,10 +2247,10 @@ int32 SaveBinaryPipelineCacheFile(const FString& OutputFilename, const EShaderPl
 		TMultiMap<FName, FSHAHash> StableNameMap;
 		FGraphEventRef StableMapConvTask = FFunctionGraphTask::CreateAndDispatchWhenReady([&StableNameMap, &StableMap]
 			{
-				for (const TPair<FStableShaderKeyAndValue, FSHAHash>& Pair : StableMap)	// could be parallelized (skip first N*ThreadIdx iterations on each thread?)
+				for (const TUniquePtr<FStableShaderKeyAndValue>& ShaderKeyPtr : StableMap)	// could be parallelized (skip first N*ThreadIdx iterations on each thread?)
 				{
-					FName PackageName(*Pair.Key.ClassNameAndObjectPath.ToStringPathOnly());
-					StableNameMap.Add(PackageName, Pair.Value);
+					FName PackageName(*ShaderKeyPtr->ClassNameAndObjectPath.ToStringPathOnly());
+					StableNameMap.Add(PackageName, ShaderKeyPtr->OutputHash);
 				}
 			}, TStatId());
 
@@ -2454,18 +2437,16 @@ int32 BuildPSOSC(const TArray<FString>& Tokens, const TMap<FString, FString>& Pa
 	BuildDateSortedListOfFiles(Tokens, ExtensionFilterFn, StablePipelineCacheFiles);
 
 	// Start populating the files with stable keys in a task.
-	TMultiMap<FStableShaderKeyAndValue, FSHAHash> StableMap;
+	FStableShaderSet StableMap;
 	FGraphEventRef StableMapTask = FFunctionGraphTask::CreateAndDispatchWhenReady([&StableShaderFiles, &StableMap]
 	{
 		LoadStableShaderKeysMultiple(StableMap, StableShaderFiles);
 		if (UE_LOG_ACTIVE(LogShaderPipelineCacheTools, Verbose))
 		{
 			UE_LOG(LogShaderPipelineCacheTools, Verbose, TEXT("    %s"), *FStableShaderKeyAndValue::HeaderLine());
-			for (const auto& Pair : StableMap)
+			for (const TUniquePtr<FStableShaderKeyAndValue>& ShaderKeyPtr : StableMap)
 			{
-				FStableShaderKeyAndValue Temp(Pair.Key);
-				Temp.OutputHash = Pair.Value;
-				UE_LOG(LogShaderPipelineCacheTools, Verbose, TEXT("    %s"), *Temp.ToString());
+				UE_LOG(LogShaderPipelineCacheTools, Verbose, TEXT("    %s"), *ShaderKeyPtr->ToString());
 			}
 			StableShadersSerializationSelfTest(StableMap);
 		}
diff --git a/Engine/Source/Runtime/RenderCore/Private/PipelineCacheUtilities.cpp b/Engine/Source/Runtime/RenderCore/Private/PipelineCacheUtilities.cpp
--- a/Engine/Source/Runtime/RenderCore/Private/PipelineCacheUtilities.cpp
+++ b/Engine/Source/Runtime/RenderCore/Private/PipelineCacheUtilities.cpp
@@ -390,7 +390,7 @@ bool UE::PipelineCacheUtilities::LoadStableKeysFile(const FStringView& Filename,
 
 #if WITH_EDITOR
 
-bool UE::PipelineCacheUtilities::SaveStableKeysFile(const FStringView& Filename, const TSet<FStableShaderKeyAndValue>& Values)
+bool UE::PipelineCacheUtilities::SaveStableKeysFile(const FStringView& Filename, const FStableShaderSet& Values)
 {
 	TUniquePtr<FArchive> FileArchiveInner(IFileManager::Get().CreateFileWriter(*FString::ConstructFromPtrSize(Filename.GetData(), Filename.Len())));
 	if (!FileArchiveInner)
@@ -417,10 +417,10 @@ bool UE::PipelineCacheUtilities::SaveStableKeysFile(const FStringView& Filename,
 		}
 	};
 
-	for (const FStableShaderKeyAndValue& Item : Values)
+	for (const TUniquePtr<FStableShaderKeyAndValue>& Item : Values)
 	{
-		IndexHash(Item.PipelineHash);
-		IndexHash(Item.OutputHash);
+		IndexHash(Item->PipelineHash);
+		IndexHash(Item->OutputHash);
 	}
 
 	int32 NumHashes = Hashes.Num();
@@ -431,10 +431,10 @@ bool UE::PipelineCacheUtilities::SaveStableKeysFile(const FStringView& Filename,
 	}
 
 	// save the rest of the properties
-	for (const FStableShaderKeyAndValue& ConstItem : Values)
+	for (const TUniquePtr<FStableShaderKeyAndValue>& ConstItem : Values)
 	{
 		// serialization unfortunately needs non-const and this is easier than const-casting every field
-		FStableShaderKeyAndValue& Item = const_cast<FStableShaderKeyAndValue&>(ConstItem);
+		FStableShaderKeyAndValue& Item = const_cast<FStableShaderKeyAndValue&>(*ConstItem);
 
 		int8 CompactNamesNum = static_cast<int8>(Item.ClassNameAndObjectPath.ObjectClassAndPath.Num());
 		ensure(Item.ClassNameAndObjectPath.ObjectClassAndPath.Num() < 256);
@@ -586,7 +586,7 @@ bool UE::PipelineCacheUtilities::SaveStablePipelineCacheFile(const FString& Outp
 	return true;
 }
 
-bool UE::PipelineCacheUtilities::LoadStablePipelineCacheFile(const FString& Filename, const TMultiMap<FStableShaderKeyAndValue, FSHAHash>& StableMap, TSet<FPipelineCacheFileFormatPSO>& OutPSOs, FName& OutTargetPlatform, int32& OutPSOsRejected, int32& OutPSOsMerged)
+bool UE::PipelineCacheUtilities::LoadStablePipelineCacheFile(const FString& Filename, const FStableShaderSet& StableMap, TSet<FPipelineCacheFileFormatPSO>& OutPSOs, FName& OutTargetPlatform, int32& OutPSOsRejected, int32& OutPSOsMerged)
 {
 	TUniquePtr<FArchive> FileArchiveInner(IFileManager::Get().CreateFileReader(*Filename));
 	if (!FileArchiveInner)
@@ -716,11 +716,11 @@ bool UE::PipelineCacheUtilities::LoadStablePipelineCacheFile(const FString& File
 					SavedStableKeys[AbsKeyIdx].ComputeKeyHash();
 
 					FSHAHash Match = FSHAHash();
-					if (TMultiMap<FStableShaderKeyAndValue, FSHAHash>::TConstKeyIterator Iter = StableMap.CreateConstKeyIterator(SavedStableKeys[AbsKeyIdx]))
+					if (const TUniquePtr<FStableShaderKeyAndValue>* ShaderKeyPtr = StableMap.Find(SavedStableKeys[AbsKeyIdx]))
 					{
-						check(Iter.Value() != FSHAHash());
-						check(OutTargetPlatform == Iter.Key().TargetPlatform);
-						Match = Iter.Value();
+						check((*ShaderKeyPtr)->OutputHash != FSHAHash());
+						check(OutTargetPlatform == (*ShaderKeyPtr)->TargetPlatform);
+						Match = (*ShaderKeyPtr)->OutputHash;
 					}
 
 					HashesForStableKeys[AbsKeyIdx] = Match;
diff --git a/Engine/Source/Runtime/RenderCore/Private/ShaderCodeLibrary.cpp b/Engine/Source/Runtime/RenderCore/Private/ShaderCodeLibrary.cpp
--- a/Engine/Source/Runtime/RenderCore/Private/ShaderCodeLibrary.cpp
+++ b/Engine/Source/Runtime/RenderCore/Private/ShaderCodeLibrary.cpp
@@ -2656,7 +2656,7 @@ struct FEditorShaderStableInfo
 	};
 	void AddShader(FStableShaderKeyAndValue& StableKeyValue, EMergeRule MergeRule)
 	{
-		FStableShaderKeyAndValue* Existing = StableMap.Find(StableKeyValue);
+		const TUniquePtr<FStableShaderKeyAndValue>* Existing = StableMap.Find(StableKeyValue);
 		if (Existing)
 		{
 			switch (MergeRule)
@@ -2664,21 +2664,26 @@ struct FEditorShaderStableInfo
 			case EMergeRule::KeepExisting:
 				return;
 			case EMergeRule::OverwriteUnmodifiedWarnModified:
-				if (Existing->OutputHash != StableKeyValue.OutputHash)
+				if ((**Existing).OutputHash != StableKeyValue.OutputHash)
 				{
 					UE_LOG(LogShaderLibrary, Warning, TEXT("Duplicate key in stable shader library, but different output, skipping new item:"));
-					UE_LOG(LogShaderLibrary, Warning, TEXT("    Existing: %s"), *Existing->ToString());
+					UE_LOG(LogShaderLibrary, Warning, TEXT("    Existing: %s"), *(**Existing).ToString());
 					UE_LOG(LogShaderLibrary, Warning, TEXT("    New     : %s"), *StableKeyValue.ToString());
 					return;
 				}
-				break; // Otherwise fall through to overwrite
+				*Existing->Get() = StableKeyValue;
+				return; // Otherwise fall through to overwrite
 			default:
 				checkNoEntry();
 				break;
 			}
 		}
-		StableMap.Add(StableKeyValue);
-		MarkKeyValueDirty(StableKeyValue);
+
+		StableMap.Add(MakeUnique<FStableShaderKeyAndValue>(StableKeyValue));
+		if (bHasCopied)
+		{
+			StableMapsToCopy.Add(MakeUnique<FStableShaderKeyAndValue>(StableKeyValue));
+		}
 	}
 
 	void AddShaderCodeLibraryFromDirectory(const FString& BaseDir, EMergeRule MergeRule)
@@ -2802,14 +2807,6 @@ struct FEditorShaderStableInfo
 		return bSuccess;
 	}
 
-	void MarkKeyValueDirty(const FStableShaderKeyAndValue& StableKeyValue)
-	{
-		if (bHasCopied)
-		{
-			StableMapsToCopy.Add(StableKeyValue);
-		}
-	}
-
 	bool HasDataToCopy() const
 	{
 		// After the first copy dirty elements are added to StableMapsToCopy; during the first copy every element from StableMap is copied
@@ -2818,12 +2815,8 @@ struct FEditorShaderStableInfo
 
 	void CopyToCompactBinary(FCbWriter& Writer)
 	{
-		if (!bHasCopied)
-		{
-			// First Copy adds all StableMaps to StableMapsToCopy, because as an optimization we do not write to
-			// StableMapsToCopy until the first Copy call.
-			StableMapsToCopy = StableMap;
-		}
+		// First copy we use the whole stablemap, after that we'll only copy newly added elements
+		const FStableShaderSet& StableMapBeingCopied = bHasCopied ? StableMapsToCopy : StableMap;
 
 		// Deduplicate the hashes; there will be a significant number (up to 50%) of duplicate hashes.
 		TArray<FSHAHash> Hashes;
@@ -2836,19 +2829,19 @@ struct FEditorShaderStableInfo
 				HashToIndex.Add(Hash, Hashes.Num() - 1);
 			}
 		};
-		for (const FStableShaderKeyAndValue& StableInfo : StableMapsToCopy)
+		for (const TUniquePtr<FStableShaderKeyAndValue>& StableInfo : StableMapBeingCopied)
 		{
-			IndexHash(StableInfo.PipelineHash);
-			IndexHash(StableInfo.OutputHash);
+			IndexHash(StableInfo->PipelineHash);
+			IndexHash(StableInfo->OutputHash);
 		}
 
 		Writer.BeginObject();
 		{
 			Writer << "Hashes" << Hashes;
 			Writer.BeginArray("StableShaderKeyAndValues");
-			for (const FStableShaderKeyAndValue& StableInfo : StableMapsToCopy)
+			for (const TUniquePtr<FStableShaderKeyAndValue>& StableInfo : StableMapBeingCopied)
 			{
-				WriteToCompactBinary(Writer, StableInfo, HashToIndex);
+				WriteToCompactBinary(Writer, *StableInfo, HashToIndex);
 			}
 			Writer.EndArray();
 		}
@@ -2887,9 +2880,9 @@ struct FEditorShaderStableInfo
 private:
 	FName FormatName;
 	FString LibraryName;
-	TSet<FStableShaderKeyAndValue> StableMap;
+	FStableShaderSet StableMap;
 	/** Copies of elements in StableMap that have not yet been copied to the CookDirector. Used only by MultiprocessCookWorkers. */
-	TSet<FStableShaderKeyAndValue> StableMapsToCopy;
+	FStableShaderSet StableMapsToCopy;
 	/** True if CopyToCompactBinary has been called, otherwise false. If false we avoid doing some tracking since it might never be used. */
 	bool bHasCopied = false;
 };
diff --git a/Engine/Source/Runtime/RenderCore/Public/PipelineCacheUtilities.h b/Engine/Source/Runtime/RenderCore/Public/PipelineCacheUtilities.h
--- a/Engine/Source/Runtime/RenderCore/Public/PipelineCacheUtilities.h
+++ b/Engine/Source/Runtime/RenderCore/Public/PipelineCacheUtilities.h
@@ -77,7 +77,7 @@ namespace PipelineCacheUtilities
 	 * @param Values values to be saved
 	 * @return true if successful
 	 */
-	RENDERCORE_API bool SaveStableKeysFile(const FStringView& Filename, const TSet<FStableShaderKeyAndValue>& Values);
+	RENDERCORE_API bool SaveStableKeysFile(const FStringView& Filename, const FStableShaderSet& Values);
 
 	/**
 	 * Saves stable pipeline cache file.
@@ -100,7 +100,7 @@ namespace PipelineCacheUtilities
 	 * @param OutPSOsRejected number of PSOs that were rejected during loading (usually because the stable key it used is no longer present in StableMap)
 	 * @param OutPSOsMerged number of PSOs that mapped to the same shader code hashes despite using different build-agnostic ("stable") shader keys.
 	 */
-	RENDERCORE_API bool LoadStablePipelineCacheFile(const FString& Filename, const TMultiMap<FStableShaderKeyAndValue, FSHAHash>& StableMap, TSet<FPipelineCacheFileFormatPSO>& OutPSOs, FName& OutTargetPlatform, int32& OutPSOsRejected, int32& OutPSOsMerged);
+	RENDERCORE_API bool LoadStablePipelineCacheFile(const FString& Filename, const FStableShaderSet& StableMap, TSet<FPipelineCacheFileFormatPSO>& OutPSOs, FName& OutTargetPlatform, int32& OutPSOsRejected, int32& OutPSOsMerged);
 
 	/**
 	 * Saves description of assets in the chunk.
diff --git a/Engine/Source/Runtime/RenderCore/Public/ShaderCodeLibrary.h b/Engine/Source/Runtime/RenderCore/Public/ShaderCodeLibrary.h
--- a/Engine/Source/Runtime/RenderCore/Public/ShaderCodeLibrary.h
+++ b/Engine/Source/Runtime/RenderCore/Public/ShaderCodeLibrary.h
@@ -169,6 +169,28 @@ struct FStableShaderKeyAndValue
 		return Key.KeyHash;
 	}
 };
+
+struct FStableShaderKeyFuncs : BaseKeyFuncs<TUniquePtr<FStableShaderKeyAndValue>, FStableShaderKeyAndValue>
+{
+	[[nodiscard]] static FORCEINLINE const FStableShaderKeyAndValue& GetSetKey(const TUniquePtr<FStableShaderKeyAndValue>& Element)
+	{
+		return *Element;
+	}
+
+	[[nodiscard]] static FORCEINLINE bool Matches(const FStableShaderKeyAndValue& A, const FStableShaderKeyAndValue& B)
+	{
+		return A == B;
+	}
+
+	/** Calculates a hash index for a key. */
+	[[nodiscard]] static FORCEINLINE uint32 GetKeyHash(const FStableShaderKeyAndValue& Key)
+	{
+		return GetTypeHash(Key);
+	}
+};
+
+using FStableShaderSet = TSet<TUniquePtr<FStableShaderKeyAndValue>, FStableShaderKeyFuncs>;
+
 #if WITH_EDITOR
 RENDERCORE_API void WriteToCompactBinary(FCbWriter& Writer, const FStableShaderKeyAndValue& Key, 
 	const TMap<FSHAHash, int32>& HashToIndex);
