Write cook diagnostics to the logs :
- Refactorize FIncrementallyModifiedDiagnostics so it can write diagnostics to files and to the logs.
- New command line argument -WriteCookIncrementallyModifiedDiagnosticsToLogs to write the diagnostic to the logs. The default behaviour didn't change and still writes to files.

#rnx

diff --git a/Engine/Source/Editor/UnrealEd/Classes/CookOnTheSide/CookOnTheFlyServer.h b/Engine/Source/Editor/UnrealEd/Classes/CookOnTheSide/CookOnTheFlyServer.h
--- a/Engine/Source/Editor/UnrealEd/Classes/CookOnTheSide/CookOnTheFlyServer.h
+++ b/Engine/Source/Editor/UnrealEd/Classes/CookOnTheSide/CookOnTheFlyServer.h
@@ -1603,6 +1603,7 @@ private:
 	bool bCookIncrementalAllowAllClasses = false;
 	bool bHiddenDependenciesDebug = false;
 	bool bIncrementallyModifiedDiagnostics = false;
+	bool bWriteIncrementallyModifiedDiagnosticsToLogs = false;
 	bool bOnlyEditorOnlyDebug = false;
 	bool bHiddenDependenciesClassPathFilterListIsAllowList = true;
 	bool bFirstCookInThisProcessInitialized = false;
diff --git a/Engine/Source/Editor/UnrealEd/Private/CookOnTheFlyServer.cpp b/Engine/Source/Editor/UnrealEd/Private/CookOnTheFlyServer.cpp
--- a/Engine/Source/Editor/UnrealEd/Private/CookOnTheFlyServer.cpp
+++ b/Engine/Source/Editor/UnrealEd/Private/CookOnTheFlyServer.cpp
@@ -6484,6 +6484,7 @@ void UCookOnTheFlyServer::SetInitializeConfigSettings(UE::Cook::FInitializeConfi
 	FParse::Value(CommandLine, TEXT("-DeviceProfile="), OverrideDeviceProfileName);
 
 	bIncrementallyModifiedDiagnostics = FParse::Param(CommandLine, TEXT("CookIncrementallyModifiedDiagnostics"));
+	bWriteIncrementallyModifiedDiagnosticsToLogs = FParse::Param(CommandLine, TEXT("WriteCookIncrementallyModifiedDiagnosticsToLogs"));
 
 	bCookingDLCBaseGamePlugin = FParse::Param(FCommandLine::Get(), TEXT("CookingDLCBaseGamePlugin"));
 }
@@ -10113,7 +10114,14 @@ void UCookOnTheFlyServer::CookByTheBookFinishedInternal()
 
 		if (IncrementallyModifiedDiagnostics)
 		{
-			IncrementallyModifiedDiagnostics->WriteToFile(MetaDataDirectory);
+			if (bWriteIncrementallyModifiedDiagnosticsToLogs)
+			{
+				IncrementallyModifiedDiagnostics->WriteToLogs();
+			}
+			else
+			{
+				IncrementallyModifiedDiagnostics->WriteToFile(MetaDataDirectory);
+			}
 		}
 	}
 
diff --git a/Engine/Source/Editor/UnrealEd/Private/Cooker/CookDiagnostics.cpp b/Engine/Source/Editor/UnrealEd/Private/Cooker/CookDiagnostics.cpp
--- a/Engine/Source/Editor/UnrealEd/Private/Cooker/CookDiagnostics.cpp
+++ b/Engine/Source/Editor/UnrealEd/Private/Cooker/CookDiagnostics.cpp
@@ -501,107 +501,153 @@ void FIncrementallyModifiedDiagnostics::AddPackage(FIncrementallyModifiedContext
 	}
 }
 
-void FIncrementallyModifiedDiagnostics::WriteToFile(FStringView OutputDirectory)
+void FIncrementallyModifiedDiagnostics::WriteToFile(FStringView OutputDirectory) const
 {
 	TMap<FTopLevelAssetPath, TArray<const FPackageData*>> PackagesPerClass;
+	SortPackageData(PackagesPerClass);
+
+	{
+		TArray<FString> Lines;
+		CreateModifiedCookedPackageContent(PackagesPerClass, Lines);
+
+		FString FilePath = FPaths::Combine(OutputDirectory, TEXT("ModifiedCookedPackages.txt"));
+		FFileHelper::SaveStringArrayToFile(Lines, *FilePath);
+	}
+
+	{
+		TArray<FString> Lines;
+		CreateModifiedEditorPackageContent(PackagesPerClass, Lines);
+
+		FString FilePath = FPaths::Combine(OutputDirectory, TEXT("ModifiedEditorPackages.txt"));
+		FFileHelper::SaveStringArrayToFile(Lines, *FilePath);
+	}
+}
+
+void FIncrementallyModifiedDiagnostics::WriteToLogs() const
+{
+	TMap<FTopLevelAssetPath, TArray<const FPackageData*>> PackagesPerClass;
+	SortPackageData(PackagesPerClass);
+
+	{
+		TArray<FString> Lines;
+		CreateModifiedCookedPackageContent(PackagesPerClass, Lines);
+
+		UE_LOG(LogCook, Display, TEXT("Start ModifiedCookedPackages.txt"));
+		for (const FString& Line : Lines)
+		{
+			UE_LOG(LogCook, Display, TEXT("%s"), *Line);
+		}
+		UE_LOG(LogCook, Display, TEXT("End ModifiedCookedPackages.txt"));
+	}
+
+	{
+		TArray<FString> Lines;
+		CreateModifiedEditorPackageContent(PackagesPerClass, Lines);
+
+		UE_LOG(LogCook, Display, TEXT("Start ModifiedEditorPackages.txt"));
+		for (const FString& Line : Lines)
+		{
+			UE_LOG(LogCook, Display, TEXT("%s"), *Line);
+		}
+		UE_LOG(LogCook, Display, TEXT("End ModifiedEditorPackages.txt"));
+	}
+}
+
+void FIncrementallyModifiedDiagnostics::SortPackageData(TMap<FTopLevelAssetPath, TArray<const FPackageData*>>& OutPackagesPerClass) const
+{
 	for (const FPackageData& PackageData : PackageDatas)
 	{
-		PackagesPerClass.FindOrAdd(PackageData.AssetClass).Add(&PackageData);
+		OutPackagesPerClass.FindOrAdd(PackageData.AssetClass).Add(&PackageData);
 	}
 
-	PackagesPerClass.KeySort([](const FTopLevelAssetPath& A, const FTopLevelAssetPath& B)
+	OutPackagesPerClass.KeySort([](const FTopLevelAssetPath& A, const FTopLevelAssetPath& B)
 		{
 			return A.Compare(B) < 0;
 		});
-	for (TPair<FTopLevelAssetPath, TArray<const FPackageData*>>& Pair : PackagesPerClass)
+	for (TPair<FTopLevelAssetPath, TArray<const FPackageData*>>& Pair : OutPackagesPerClass)
 	{
 		Algo::Sort(Pair.Value, [](const FPackageData* A, const FPackageData* B)
 			{
 				return A->PackageName.LexicalLess(B->PackageName);
 			});
 	}
+}
 
-	{
-		FString FilePath = FPaths::Combine(OutputDirectory, TEXT("ModifiedCookedPackages.txt"));
-
-		TArray<FString> Lines;
-		Lines.Reserve(PackageDatas.Num() + PackagesPerClass.Num());
-		const TCHAR* Indent = TEXT("    ");
+void FIncrementallyModifiedDiagnostics::CreateModifiedCookedPackageContent(const TMap<FTopLevelAssetPath, TArray<const FPackageData*>>& PackagesPerClass, TArray<FString>& OutLines) const
+{
+	OutLines.Reserve(PackageDatas.Num() + PackagesPerClass.Num());
+	const TCHAR* Indent = TEXT("    ");
 
-		// Reserve leading line for the total
-		Lines.Emplace();
-		int32 NumTotalModified = 0;
+	// Reserve leading line for the total
+	OutLines.Emplace();
+	int32 NumTotalModified = 0;
 
-		for (TPair<FTopLevelAssetPath, TArray<const FPackageData*>>& Pair : PackagesPerClass)
+	for (const TPair<FTopLevelAssetPath, TArray<const FPackageData*>>& Pair : PackagesPerClass)
+	{
+		OutLines.Add(FString::Printf(TEXT("%s, %d"), *Pair.Key.ToString(), Pair.Value.Num()));
+		NumTotalModified += Pair.Value.Num();
+		for (const FPackageData* PackageData : Pair.Value)
 		{
-			Lines.Add(FString::Printf(TEXT("%s, %d"), *Pair.Key.ToString(), Pair.Value.Num()));
-			NumTotalModified += Pair.Value.Num();
-			for (const FPackageData* PackageData : Pair.Value)
+			if (PackageData->ReasonLines.Num() == 0)
 			{
-				if (PackageData->ReasonLines.Num() == 0)
-				{
-					Lines.Add(FString::Printf(TEXT("%s%s, %s"), Indent, *PackageData->PackageName.ToString(),
-						LexToString(PackageData->Reason)));
-				}
-				else if (PackageData->ReasonLines.Num() == 1)
-				{
-					Lines.Add(FString::Printf(TEXT("%s%s, %s: %s"), Indent, *PackageData->PackageName.ToString(),
-						LexToString(PackageData->Reason), *PackageData->ReasonLines[0]));
-				}
-				else
+				OutLines.Add(FString::Printf(TEXT("%s%s, %s"), Indent, *PackageData->PackageName.ToString(),
+					LexToString(PackageData->Reason)));
+			}
+			else if (PackageData->ReasonLines.Num() == 1)
+			{
+				OutLines.Add(FString::Printf(TEXT("%s%s, %s: %s"), Indent, *PackageData->PackageName.ToString(),
+					LexToString(PackageData->Reason), *PackageData->ReasonLines[0]));
+			}
+			else
+			{
+				OutLines.Add(FString::Printf(TEXT("%s%s, %s:"), Indent, *PackageData->PackageName.ToString(),
+					LexToString(PackageData->Reason)));
+				for (const FString& ReasonLine : PackageData->ReasonLines)
 				{
-					Lines.Add(FString::Printf(TEXT("%s%s, %s:"), Indent, *PackageData->PackageName.ToString(),
-						LexToString(PackageData->Reason)));
-					for (const FString& ReasonLine : PackageData->ReasonLines)
-					{
-						Lines.Add(FString::Printf(TEXT("%s%s%s"), Indent, Indent, *ReasonLine));
-					}
+					OutLines.Add(FString::Printf(TEXT("%s%s%s"), Indent, Indent, *ReasonLine));
 				}
 			}
 		}
-
-		Lines[0] = FString::Printf(TEXT("Total, %d"), NumTotalModified);
-		FFileHelper::SaveStringArrayToFile(Lines, *FilePath);
 	}
 
-	{
-		int32 NumTotalModified = 0;
-		FString FilePath = FPaths::Combine(OutputDirectory, TEXT("ModifiedEditorPackages.txt"));
+	OutLines[0] = FString::Printf(TEXT("Total, %d"), NumTotalModified);
+}
 
-		TArray<FString> Lines;
-		Lines.Reserve(PackageDatas.Num() + PackagesPerClass.Num() + 1);
-		const TCHAR* Indent = TEXT("    ");
+void FIncrementallyModifiedDiagnostics::CreateModifiedEditorPackageContent(const TMap<FTopLevelAssetPath, TArray<const FPackageData*>>& PackagesPerClass, TArray<FString>& OutLines) const
+{
+	int32 NumTotalModified = 0;
+
+	OutLines.Reserve(PackageDatas.Num() + PackagesPerClass.Num() + 1);
+	const TCHAR* Indent = TEXT("    ");
 
-		// Reserve leading line for the total
-		Lines.Emplace();
+	// Reserve leading line for the total
+	OutLines.Emplace();
 
-		for (TPair<FTopLevelAssetPath, TArray<const FPackageData*>>& Pair : PackagesPerClass)
+	for (const TPair<FTopLevelAssetPath, TArray<const FPackageData*>>& Pair : PackagesPerClass)
+	{
+		int NumClassModified = 0;
+		int32 ClassModifiedIndex = -1;
+		for (const FPackageData* PackageData : Pair.Value)
 		{
-			int NumClassModified = 0;
-			int32 ClassModifiedIndex = -1;
-			for (const FPackageData* PackageData : Pair.Value)
+			if (PackageData->bSelfModified)
 			{
-				if (PackageData->bSelfModified)
+				if (ClassModifiedIndex == -1)
 				{
-					if (ClassModifiedIndex == -1)
-					{
-						ClassModifiedIndex = Lines.Num();
-						Lines.Emplace();
-					}
-					Lines.Add(FString::Printf(TEXT("%s%s"), Indent, *PackageData->PackageName.ToString()));
-					++NumClassModified;
+					ClassModifiedIndex = OutLines.Num();
+					OutLines.Emplace();
 				}
-			}
-			if (ClassModifiedIndex != -1)
-			{
-				Lines[ClassModifiedIndex] = FString::Printf(TEXT("%s, %d"), *Pair.Key.ToString(), NumClassModified);
-				NumTotalModified += NumClassModified;
+				OutLines.Add(FString::Printf(TEXT("%s%s"), Indent, *PackageData->PackageName.ToString()));
+				++NumClassModified;
 			}
 		}
-
-		Lines[0] = FString::Printf(TEXT("Total, %d"), NumTotalModified);
-		FFileHelper::SaveStringArrayToFile(Lines, *FilePath);
+		if (ClassModifiedIndex != -1)
+		{
+			OutLines[ClassModifiedIndex] = FString::Printf(TEXT("%s, %d"), *Pair.Key.ToString(), NumClassModified);
+			NumTotalModified += NumClassModified;
+		}
 	}
+
+	OutLines[0] = FString::Printf(TEXT("Total, %d"), NumTotalModified);
 }
 
 } // namespace UE::Cook
\ No newline at end of file
diff --git a/Engine/Source/Editor/UnrealEd/Private/Cooker/CookDiagnostics.h b/Engine/Source/Editor/UnrealEd/Private/Cooker/CookDiagnostics.h
--- a/Engine/Source/Editor/UnrealEd/Private/Cooker/CookDiagnostics.h
+++ b/Engine/Source/Editor/UnrealEd/Private/Cooker/CookDiagnostics.h
@@ -68,7 +68,8 @@ class FIncrementallyModifiedDiagnostics
 {
 public:
 	void AddPackage(FIncrementallyModifiedContext& Context);
-	void WriteToFile(FStringView OutputDirectory);
+	void WriteToFile(FStringView OutputDirectory) const;
+	void WriteToLogs() const;
 
 private:
 	/** The saved diagnostic data about each reported modified Package. */
@@ -82,6 +83,10 @@ private:
 	};
 
 	TArray<FPackageData> PackageDatas;
+
+	void SortPackageData(TMap<FTopLevelAssetPath, TArray<const FPackageData*>>& OutPackagesPerClass) const;
+	void CreateModifiedCookedPackageContent(const TMap<FTopLevelAssetPath, TArray<const FPackageData*>>& PackagesPerClass, TArray<FString>& OutLines) const;
+	void CreateModifiedEditorPackageContent(const TMap<FTopLevelAssetPath, TArray<const FPackageData*>>& PackagesPerClass, TArray<FString>& OutLines) const;
 };
 
 } // namespace UE::Cook
\ No newline at end of file
