Adding a null check on GEditor and GEditor->Trans since it can be null in commandlets. Reported on UDN hapening in a commandlet that deletes assets.

diff --git a/Engine/Source/Editor/UnrealEd/Private/Kismet2/Kismet2.cpp b/Engine/Source/Editor/UnrealEd/Private/Kismet2/Kismet2.cpp
--- a/Engine/Source/Editor/UnrealEd/Private/Kismet2/Kismet2.cpp
+++ b/Engine/Source/Editor/UnrealEd/Private/Kismet2/Kismet2.cpp
@@ -759,9 +759,17 @@ bool FKismetEditorUtilities::IsReferencedByUndoBuffer(UBlueprint* Blueprint)
 
 	FReferencerInformationList ReferencesExcludingUndo;
 	// Determine the in-memory references, *excluding* the undo buffer
-	GEditor->Trans->DisableObjectSerialization();
+
+	UTransactor* Transactor = GEditor ? ToRawPtr(GEditor->Trans) : nullptr;
+	if (Transactor)
+	{
+		Transactor->DisableObjectSerialization();
+	}
 	IsReferenced(BlueprintObj, GARBAGE_COLLECTION_KEEPFLAGS, EInternalObjectFlags_GarbageCollectionKeepFlags, /*bCheckSubObjects =*/true, &ReferencesExcludingUndo);
-	GEditor->Trans->EnableObjectSerialization();
+	if (Transactor)
+	{
+		Transactor->EnableObjectSerialization();
+	}
 
 	// see if this object is the transaction buffer - set a flag so we know we need to clear the undo stack
 	const int32 TotalReferenceCount   = ReferencesIncludingUndo.ExternalReferences.Num() + ReferencesIncludingUndo.InternalReferences.Num();
diff --git a/Engine/Source/Editor/UnrealEd/Private/ObjectTools.cpp b/Engine/Source/Editor/UnrealEd/Private/ObjectTools.cpp
--- a/Engine/Source/Editor/UnrealEd/Private/ObjectTools.cpp
+++ b/Engine/Source/Editor/UnrealEd/Private/ObjectTools.cpp
@@ -320,7 +320,7 @@ namespace ObjectTools
 		
 		if (!CVarUseLegacyGetReferencersForDeletion.GetValueOnAnyThread())
 		{
-			const UTransactor* Transactor = GEditor ? ToRawPtr(GEditor->Trans) : nullptr;
+			UTransactor* Transactor = GEditor ? ToRawPtr(GEditor->Trans) : nullptr;
 			bool bIsGatheringPackageRef = InObject->IsA<UPackage>();
 
 			// Get the cluster of objects that are going to be deleted
@@ -439,14 +439,20 @@ namespace ObjectTools
 			{
 				// determine whether the transaction buffer is the only thing holding a reference to the object
 				// and if so, offer the user the option to reset the transaction buffer.
-				GEditor->Trans->DisableObjectSerialization();
+				if (Transactor)
+                {
+					Transactor->DisableObjectSerialization();
+				}
 				bOutIsReferenced = IsReferenced(InObject, GARBAGE_COLLECTION_KEEPFLAGS, EInternalObjectFlags_GarbageCollectionKeepFlags, true, OutMemoryReferences);
 				if (!bOutIsReferenced)
 				{
 					UE_LOG(LogObjectTools, Warning, TEXT("Detected inconsistencies between reference gathering algorithms. Switching 'Editor.UseLegacyGetReferencersForDeletion' on for the remainder of this editor session."));
 					CVarUseLegacyGetReferencersForDeletion->Set(1);
 				}
-				GEditor->Trans->EnableObjectSerialization();
+				if (Transactor)
+				{
+					Transactor->EnableObjectSerialization();
+				}
 			}
 		}
 		// This is the old/slower behavior that is kept for debug/comparison and is going to be removed in a future release
@@ -455,11 +461,18 @@ namespace ObjectTools
 			bOutIsReferenced = IsReferenced(InObject, GARBAGE_COLLECTION_KEEPFLAGS, EInternalObjectFlags_GarbageCollectionKeepFlags, true, OutMemoryReferences);
 			if (bOutIsReferenced)
 			{
+				UTransactor* Transactor = GEditor ? ToRawPtr(GEditor->Trans) : nullptr;
 				// determine whether the transaction buffer is the only thing holding a reference to the object
 				// and if so, offer the user the option to reset the transaction buffer.
-				GEditor->Trans->DisableObjectSerialization();
+				if (Transactor)
+				{
+					Transactor->DisableObjectSerialization();
+				}
 				bOutIsReferenced = IsReferenced(InObject, GARBAGE_COLLECTION_KEEPFLAGS, EInternalObjectFlags_GarbageCollectionKeepFlags, true, OutMemoryReferences);
-				GEditor->Trans->EnableObjectSerialization();
+				if (Transactor)
+				{
+					Transactor->EnableObjectSerialization();
+				}
 
 				// If object is referenced both in undo and non-undo, we can't determine which one it is but
 				// it doesn't matter since the undo stack is only cleared if objects are only referenced by it.
