Follow-up for UE-297732. Cook-on-the-fly now builds full inline shader maps for any cooked materials. Disabled ODSC forced recompilation when shader maps not owned by ODSC have been cooked already via cook-on-the-fly

#tests Tested Lyra, CitySample, AncientGame with + without COTF

diff --git a/Engine/Source/Editor/UnrealEd/Private/CookOnTheFlyServer.cpp b/Engine/Source/Editor/UnrealEd/Private/CookOnTheFlyServer.cpp
--- a/Engine/Source/Editor/UnrealEd/Private/CookOnTheFlyServer.cpp
+++ b/Engine/Source/Editor/UnrealEd/Private/CookOnTheFlyServer.cpp
@@ -246,7 +246,7 @@ static FAutoConsoleVariableRef CVarDisableShaderCompilationDuringCookOnTheFly(
 	TEXT("Controls whether shader compilation is disabled during cook on the fly\n"),
 	ECVF_Default);
 
-bool bAllowIncompleteShaderMapsDuringCookOnTheFly = true;
+bool bAllowIncompleteShaderMapsDuringCookOnTheFly = false;
 static FAutoConsoleVariableRef CVarAllowEmptyShaderMapsDuringCookOnTheFly(
 	TEXT("Cook.OnTheFly.AllowIncompleteShadermaps"),
 	bAllowIncompleteShaderMapsDuringCookOnTheFly,
diff --git a/Engine/Source/Runtime/Engine/Private/ODSC/ODSCManager.cpp b/Engine/Source/Runtime/Engine/Private/ODSC/ODSCManager.cpp
--- a/Engine/Source/Runtime/Engine/Private/ODSC/ODSCManager.cpp
+++ b/Engine/Source/Runtime/Engine/Private/ODSC/ODSCManager.cpp
@@ -384,6 +384,13 @@ bool FODSCManager::ShouldForceRecompileInternal(const FMaterialShaderMap* Materi
 		return ODSCMetaData == EODSCMetaDataType::IsDependentOnMaterialName;
 	}
 
+	// when running with cook on the fly there are full shader-maps built by the cooker already
+	// ODSC will only occur for this material if a shader is missing or changed
+	if (IsRunningCookOnTheFly())
+	{
+		return false;	
+	}
+
 	return true;
 #else
 	return false;
