Cooker: Undo merge from older stream.
#rnx

diff --git a/Engine/Source/Editor/UnrealEd/Private/CookOnTheFlyServer.cpp b/Engine/Source/Editor/UnrealEd/Private/CookOnTheFlyServer.cpp
--- a/Engine/Source/Editor/UnrealEd/Private/CookOnTheFlyServer.cpp
+++ b/Engine/Source/Editor/UnrealEd/Private/CookOnTheFlyServer.cpp
@@ -6486,7 +6486,7 @@ void UCookOnTheFlyServer::SetInitializeConfigSettings(UE::Cook::FInitializeConfi
 	bIncrementallyModifiedDiagnostics = FParse::Param(CommandLine, TEXT("CookIncrementallyModifiedDiagnostics"));
 	bWriteIncrementallyModifiedDiagnosticsToLogs = FParse::Param(CommandLine, TEXT("WriteCookIncrementallyModifiedDiagnosticsToLogs"));
 
-	bCookingDLCBaseGamePlugin = FParse::Param(FCommandLine::Get(), TEXT("CookingDLCBaseGamePlugin"));
+	bCookingDLCBaseGamePlugin = FParse::Param(CommandLine, TEXT("CookingDLCBaseGamePlugin"));
 }
 
 void UCookOnTheFlyServer::ParseCookFilters()
