Cooker: SkipOnlyEditorOnly: Fix incorrect reporting of hidden dependencies when an instanced package is loaded with imports of its own. We report all of those transitive dependencies as dependencies of the package loading the instanced package, but we were not correctly merging the flag that states whether the transitive import was an expected load with the flag that states whether the instanced package was an expected load.
#rnx

diff --git a/Engine/Source/Editor/UnrealEd/Private/Cooker/PackageTracker.cpp b/Engine/Source/Editor/UnrealEd/Private/Cooker/PackageTracker.cpp
--- a/Engine/Source/Editor/UnrealEd/Private/Cooker/PackageTracker.cpp
+++ b/Engine/Source/Editor/UnrealEd/Private/Cooker/PackageTracker.cpp
@@ -6,7 +6,9 @@
 #include "CookPackageData.h"
 #include "CookPlatformManager.h"
 #include "CookProfiling.h"
+#include "Engine/Level.h"
 #include "Misc/ScopeRWLock.h"
+#include "Misc/StringBuilder.h"
 #include "UObject/Package.h"
 #include "UObject/UObjectGlobals.h"
 #include "UObject/UObjectIterator.h"
@@ -325,11 +327,13 @@ void FPackageTracker::OnEndLoadPackage(const FEndLoadPackageContext& Context)
 		// during the top-most LoadPackage call, calculate the ancestor non-instanced package referencer for each of
 		// the instanced package loads.
 		{
+			TMap<FName, TOptional<TMap<FName, EDependencyProperty>>> DependenciesOfNonInstancedReferencers;
 			FReadScopeLock ActiveInstancesScopeLock(ActiveInstances->Lock);
 			for (TRefCountPtr<FPackageStreamInstancedPackage>& InstancedPackage : LoadedInstances)
 			{
 				TSet<FPackageStreamInstancedPackage*> Visited;
-				InstancedPackage->FlattenReferencer(ActiveInstancesScopeLock, Visited);
+				InstancedPackage->FlattenReferencer(ActiveInstancesScopeLock, Visited,
+					DependenciesOfNonInstancedReferencers, *this);
 			}
 		}
 
@@ -394,8 +398,12 @@ FPackageStreamInstancedPackage::~FPackageStreamInstancedPackage()
 }
 
 void FPackageStreamInstancedPackage::FlattenReferencer(FReadScopeLock& CalledInsideActiveInstancesLock,
-	TSet<FPackageStreamInstancedPackage*>& Visited)
+	TSet<FPackageStreamInstancedPackage*>& Visited,
+	TMap<FName, TOptional<TMap<FName, UE::AssetRegistry::EDependencyProperty>>>& DependenciesOfNonInstancedReferencers,
+	FPackageTracker& PackageTracker)
 {
+	using namespace UE::AssetRegistry;
+
 	bool bAlreadyVisited;
 	Visited.Add(this, &bAlreadyVisited);
 	if (bAlreadyVisited)
@@ -408,14 +416,78 @@ void FPackageStreamInstancedPackage::FlattenReferencer(FReadScopeLock& CalledIns
 	{
 		return;
 	}
+
+	// If the Instigator of this instanced package is itself an instanced package, recurse into it.
+	// And in either case, instanced package or not, get the instigator's AssetRegistry dependencies
+	// in case the Instigator.Category is EInstigator::Unsolicited.
 	FPackageStreamInstancedPackage** ReferencerPtr = Container->Map.Find(Instigator.Referencer);
+	FPackageStreamInstancedPackage* Referencer = ReferencerPtr ? *ReferencerPtr : nullptr;
+	TMap<FName, UE::AssetRegistry::EDependencyProperty>* ReferencerDependencies = nullptr;
+	if (ReferencerPtr)
+	{
+		// All Dependencies on FPackageStreamInstancedPackage are intialized when the
+		// FPackageStreamInstancedPackage is constructed, no need to initialize them here.
+		ReferencerDependencies = &Referencer->Dependencies;
+	}
+	else
+	{
+		TOptional<TMap<FName, UE::AssetRegistry::EDependencyProperty>>& ExistingDependencies
+			= DependenciesOfNonInstancedReferencers.FindOrAdd(Instigator.Referencer);
+		if (!ExistingDependencies)
+		{
+			TArray<FAssetDependency> PackageDependencies;
+			IAssetRegistry::GetChecked().GetDependencies(FAssetIdentifier(Instigator.Referencer), PackageDependencies,
+				EDependencyCategory::Package, EDependencyQuery::Hard);
+			ExistingDependencies.Emplace();
+			for (FAssetDependency& PackageDependency : PackageDependencies)
+			{
+				ExistingDependencies->Add(PackageDependency.AssetId.PackageName, PackageDependency.Properties);
+			}
+		}
+		ReferencerDependencies = &(*ExistingDependencies);
+	}
+
+	// If the Instigator.Category is Unsolicited, try to remap it based on AssetRegistry dependencies; all hard
+	// imports start off as being marked Unsolicited, and we need to convert them to HardDependency (or
+	// to EditorOnlyLoad if they are an EditorOnly HardDependency). It is also possible that this PackageLoad
+	// was not due to a hard import, and is indeed Unsolicited, in which case it is not present in the AR
+	// dependencies and we keep it marked as Unsolicited.
+	if (Instigator.Category != EInstigator::ForceExplorableSaveTimeSoftDependency
+		&& PackageTracker.NeverCookPackageList.Contains(LoadedName))
+	{
+		TStringBuilder<256> PackageNameStr(InPlace, LoadedName);
+		if (INDEX_NONE != UE::String::FindFirst(PackageNameStr,
+			ULevel::GetExternalActorsFolderName(), ESearchCase::IgnoreCase)
+			|| INDEX_NONE != UE::String::FindFirst(PackageNameStr,
+				FPackagePath::GetExternalObjectsFolderName(), ESearchCase::IgnoreCase))
+		{
+			// ONLYEDITORONLY_TODO: WorldPartition should mark these loads as ForceExplorableSaveTimeSoftDependency
+			// rather than needing to use a naming convention. We should also mark them once, during Save, rather than
+			// marking them and reexploring them every time they are loaded. 
+			Instigator.Category = EInstigator::ForceExplorableSaveTimeSoftDependency;
+		}
+	}
+
+	if (Instigator.Category == EInstigator::Unsolicited)
+	{
+		const EDependencyProperty* DependencyProperty = ReferencerDependencies->Find(LoadedName);
+		if (DependencyProperty)
+		{
+			Instigator.Category = EnumHasAnyFlags(*DependencyProperty, EDependencyProperty::Game)
+				? EInstigator::HardDependency : EInstigator::EditorOnlyLoad;
+		}
+	}
 	if (!ReferencerPtr)
 	{
+		// Now that we have remapped the category if necessary, if our parent is not an instanced package
+		// then we have no further work to do.
 		return;
 	}
 
-	FPackageStreamInstancedPackage* Referencer = *ReferencerPtr;
-	Referencer->FlattenReferencer(CalledInsideActiveInstancesLock, Visited);
+	// If our parent is an instanced package, flatten its instigator recursively and then merge
+	// its instigator with the data from its instigation of this instanced package.
+	Referencer->FlattenReferencer(CalledInsideActiveInstancesLock, Visited,
+		DependenciesOfNonInstancedReferencers, PackageTracker);
 
 	Instigator.Referencer = Referencer->Instigator.Referencer;
 	Instigator.Category = FPackageTracker::MergeReferenceCategories(
diff --git a/Engine/Source/Editor/UnrealEd/Private/Cooker/PackageTracker.h b/Engine/Source/Editor/UnrealEd/Private/Cooker/PackageTracker.h
--- a/Engine/Source/Editor/UnrealEd/Private/Cooker/PackageTracker.h
+++ b/Engine/Source/Editor/UnrealEd/Private/Cooker/PackageTracker.h
@@ -10,6 +10,8 @@
 #include "CookTypes.h"
 #include "HAL/CriticalSection.h"
 #include "HAL/Platform.h"
+#include "Misc/AssetRegistryInterface.h"
+#include "Misc/Optional.h"
 #include "Misc/ScopeLock.h"
 #include "Templates/UnrealTemplate.h"
 #include "UObject/NameTypes.h"
@@ -205,7 +207,9 @@ private:
 	 * on the parent referencers in between this and the ancestor.
 	 */
 	void FlattenReferencer(FReadScopeLock& CalledInsideActiveInstancesLock,
-		TSet<FPackageStreamInstancedPackage*>& Visited);
+		TSet<FPackageStreamInstancedPackage*>& Visited,
+		TMap<FName, TOptional<TMap<FName, UE::AssetRegistry::EDependencyProperty>>>&
+			DependenciesOfNonInstancedReferencers, FPackageTracker& PackageTracker);
 
 private:
 	TRefCountPtr<FPackageStreamInstancedPackageContainer> Container;
