Cloth - Deprecated ApexClothingUtils files which are no longer used anywhere.

#rnx

diff --git a/Engine/Source/Editor/UnrealEd/Private/ApexClothingUtils.cpp b/Engine/Source/Editor/UnrealEd/Private/ApexClothingUtils.cpp
--- a/Engine/Source/Editor/UnrealEd/Private/ApexClothingUtils.cpp
+++ b/Engine/Source/Editor/UnrealEd/Private/ApexClothingUtils.cpp
@@ -1,24 +1,15 @@
 // Copyright Epic Games, Inc. All Rights Reserved.
 
+// UE_DEPRECATED(5.7, "Apex clothing is no longer supported, this implementation will be removed")
+
+#if UE_ENABLE_INCLUDE_ORDER_DEPRECATED_IN_5_7
 #include "ApexClothingUtils.h"
+#endif // UE_ENABLE_INCLUDE_ORDER_DEPRECATED_IN_5_7
+
 #include "Components/SkeletalMeshComponent.h"
-#include "Misc/MessageDialog.h"
-#include "Misc/FileHelper.h"
-#include "Misc/Paths.h"
-#include "Modules/ModuleManager.h"
-#include "UObject/UObjectHash.h"
-#include "UObject/UObjectIterator.h"
-#include "EditorDirectories.h"
 #include "Engine/SkeletalMesh.h"
-#include "IDesktopPlatform.h"
-#include "DesktopPlatformModule.h"
-#include "Framework/Application/SlateApplication.h"
 #include "Rendering/SkeletalMeshModel.h"
-
-#include "ClothingAssetFactory.h"
-#include "PhysicsPublic.h"
-
-DEFINE_LOG_CATEGORY_STATIC(LogApexClothingUtils, Log, All);
+#include "ClothingAsset.h"
 
 #define LOCTEXT_NAMESPACE "ApexClothingUtils"
 
diff --git a/Engine/Source/Editor/UnrealEd/Private/Fbx/FbxReimportDialogs.cpp b/Engine/Source/Editor/UnrealEd/Private/Fbx/FbxReimportDialogs.cpp
--- a/Engine/Source/Editor/UnrealEd/Private/Fbx/FbxReimportDialogs.cpp
+++ b/Engine/Source/Editor/UnrealEd/Private/Fbx/FbxReimportDialogs.cpp
@@ -53,7 +53,6 @@
 #include "Components/SkinnedMeshComponent.h"
 #include "Components/SkeletalMeshComponent.h"
 #include "AnimEncoding.h"
-#include "ApexClothingUtils.h"
 #include "Engine/SkeletalMeshSocket.h"
 #include "ClothingAsset.h"
 #include "Factories/FbxSkeletalMeshImportData.h"
diff --git a/Engine/Source/Editor/UnrealEd/Private/Fbx/FbxSkeletalMeshImport.cpp b/Engine/Source/Editor/UnrealEd/Private/Fbx/FbxSkeletalMeshImport.cpp
--- a/Engine/Source/Editor/UnrealEd/Private/Fbx/FbxSkeletalMeshImport.cpp
+++ b/Engine/Source/Editor/UnrealEd/Private/Fbx/FbxSkeletalMeshImport.cpp
@@ -49,7 +49,6 @@
 
 #include "ObjectTools.h"
 
-#include "ApexClothingUtils.h"
 #include "MeshUtilities.h"
 #include "OverlappingCorners.h"
 
diff --git a/Engine/Source/Editor/UnrealEd/Public/ApexClothingUtils.h b/Engine/Source/Editor/UnrealEd/Public/ApexClothingUtils.h
--- a/Engine/Source/Editor/UnrealEd/Public/ApexClothingUtils.h
+++ b/Engine/Source/Editor/UnrealEd/Public/ApexClothingUtils.h
@@ -2,16 +2,18 @@
 
 #pragma once
 
+#include "Misc/CoreMiscDefines.h"
+UE_DEPRECATED_HEADER(5.7, "Apex clothing is no longer supported, remove this file from your header inclusions.")
+
+#if UE_ENABLE_INCLUDE_ORDER_DEPRECATED_IN_5_7
 #include "CoreMinimal.h"
 #include "EngineDefines.h"
+#endif // UE_ENABLE_INCLUDE_ORDER_DEPRECATED_IN_5_7
 
 class USkeletalMesh;
 
 namespace ApexClothingUtils
 {
-	// Functions below remain from previous clothing system and only remain to remove the bound
-	// data from a skeletal mesh. This is done when postloading USkeletalMesh when upgrading the assets
-
-	// Function to restore all clothing section to original mesh section related to specified asset index
+	UE_DEPRECATED(5.7, "Apex clothing is no longer supported, this implementation will be removed.")
 	UNREALED_API void RemoveAssetFromSkeletalMesh(USkeletalMesh* SkelMesh, uint32 AssetIndex, bool bReleaseAsset = true, bool bRecreateSkelMeshComponent = false);
 }
