CIS -  Fixed Issue 1012733: Compile errors in ApexClothingUtils.h

#rnx

diff --git a/Engine/Source/Editor/UnrealEd/Private/ApexClothingUtils.cpp b/Engine/Source/Editor/UnrealEd/Private/ApexClothingUtils.cpp
--- a/Engine/Source/Editor/UnrealEd/Private/ApexClothingUtils.cpp
+++ b/Engine/Source/Editor/UnrealEd/Private/ApexClothingUtils.cpp
@@ -2,9 +2,9 @@
 
 // UE_DEPRECATED(5.7, "Apex clothing is no longer supported, this implementation will be removed")
 
-#if UE_ENABLE_INCLUDE_ORDER_DEPRECATED_IN_5_7
+#if 0  // Workaround the first inclusion rule to help the deprecation
 #include "ApexClothingUtils.h"
-#endif // UE_ENABLE_INCLUDE_ORDER_DEPRECATED_IN_5_7
+#endif
 
 #include "Components/SkeletalMeshComponent.h"
 #include "Engine/SkeletalMesh.h"
