- Audit fixed delta time Cascade particle systems

#rnx

diff --git a/Engine/Source/Editor/UnrealEd/Classes/Commandlets/ParticleSystemAuditCommandlet.h b/Engine/Source/Editor/UnrealEd/Classes/Commandlets/ParticleSystemAuditCommandlet.h
--- a/Engine/Source/Editor/UnrealEd/Classes/Commandlets/ParticleSystemAuditCommandlet.h
+++ b/Engine/Source/Editor/UnrealEd/Classes/Commandlets/ParticleSystemAuditCommandlet.h
@@ -36,6 +36,8 @@ class UParticleSystemAuditCommandlet : public UCommandlet
 	TSet<FString> ParticleSystemsWithWarmupTime;
 	/** All particle systems with lights and the details */
 	TSet<FString> ParticleSystemsWithLights;
+	/** All particle systems with fixed delta time */
+	TSet<FString> ParticleSystemsWithFixedDeltaTime;
 
 	/** If a particle system has a spawn rate or burst count greater than this value, it will be reported */
 	UPROPERTY(config)
diff --git a/Engine/Source/Editor/UnrealEd/Private/Commandlets/ParticleSystemAuditCommandlet.cpp b/Engine/Source/Editor/UnrealEd/Private/Commandlets/ParticleSystemAuditCommandlet.cpp
--- a/Engine/Source/Editor/UnrealEd/Private/Commandlets/ParticleSystemAuditCommandlet.cpp
+++ b/Engine/Source/Editor/UnrealEd/Private/Commandlets/ParticleSystemAuditCommandlet.cpp
@@ -385,6 +385,11 @@ bool UParticleSystemAuditCommandlet::ProcessParticleSystems()
 				ParticleSystemsWithLights.Add(SystemString);
 			}
 
+			if (PSys->SystemUpdateMode == EParticleSystemUpdateMode::EPSUM_FixedTime)
+			{
+				ParticleSystemsWithFixedDeltaTime.Add(PSys->GetPathName());
+			}
+
 			if ((PSys->LODMethod == PARTICLESYSTEMLODMETHOD_Automatic) &&
 				(bInvalidLOD == false) && (bSingleLOD == false) &&
 				(PSys->LODDistanceCheckTime == 0.0f))
@@ -444,6 +449,7 @@ void UParticleSystemAuditCommandlet::DumpResults()
 	DumpSimplePSysSet(ParticleSystemsWithBoneLocationMismatches, TEXT("PSysBoneLocationLODMismatches"));
 	DumpSimplePSysSet(ParticleSystemsWithWarmupTime, TEXT("PSysWithWarmupTime"));
 	DumpSimplePSysSet(ParticleSystemsWithLights, TEXT("PSysWithLights"), TEXT("Name,HiQuality,ShadowCasting"));
+	DumpSimplePSysSet(ParticleSystemsWithFixedDeltaTime, TEXT("PSysWithFixedDeltaTime"), TEXT("Name"));
 }
 
 /**
