IsPlayInMobilePreview now correctly returns true when in preview mode PIE
Removed old VulkanPreview

diff --git a/Engine/Source/Editor/EditorStyle/Private/StarshipStyle.cpp b/Engine/Source/Editor/EditorStyle/Private/StarshipStyle.cpp
--- a/Engine/Source/Editor/EditorStyle/Private/StarshipStyle.cpp
+++ b/Engine/Source/Editor/EditorStyle/Private/StarshipStyle.cpp
@@ -6136,7 +6136,6 @@ void FStarshipEditorStyle::FStyle::SetupPersonaStyle()
 		Set( "PlayWorld.PlayInEditorFloating", new IMAGE_BRUSH_SVG("Starship/MainToolbar/PlayNewEditorWindow", Icon20x20));
 		Set( "PlayWorld.PlayInVR", new IMAGE_BRUSH_SVG("Starship/MainToolbar/PlayVRPreview", Icon20x20));
 		Set( "PlayWorld.PlayInMobilePreview", new IMAGE_BRUSH_SVG("Starship/MainToolbar/PlayMobilePreview", Icon20x20));
-		Set( "PlayWorld.PlayInVulkanPreview", new IMAGE_BRUSH_SVG("Starship/MainToolbar/PlayMobilePreview", Icon20x20));
 		Set( "PlayWorld.PlayInNewProcess", new IMAGE_BRUSH_SVG("Starship/MainToolbar/PlayStandaloneGame", Icon20x20));
 
 		Set( "PlayWorld.RepeatLastLaunch", new IMAGE_BRUSH_SVG( "Starship/MainToolbar/launch", Icon20x20 ) );
diff --git a/Engine/Source/Editor/UnrealEd/Classes/Settings/EditorExperimentalSettings.h b/Engine/Source/Editor/UnrealEd/Classes/Settings/EditorExperimentalSettings.h
--- a/Engine/Source/Editor/UnrealEd/Classes/Settings/EditorExperimentalSettings.h
+++ b/Engine/Source/Editor/UnrealEd/Classes/Settings/EditorExperimentalSettings.h
@@ -92,10 +92,6 @@ public:
 	UPROPERTY(EditAnywhere, config, Category = PIE, meta = (DisplayName = "Allow late joining"))
 	bool bAllowLateJoinInPIE;
 
-	/** Allow Vulkan Preview */
-	UPROPERTY(EditAnywhere, config, Category = PIE, meta = (DisplayName = "Allow Vulkan Mobile Preview"))
-	bool bAllowVulkanPreview;
-
 	/** Enable multithreaded lightmap encoding (decreases time taken to encode lightmaps) */
 	UPROPERTY(EditAnywhere, config, Category = LightingBuilds, meta = (DisplayName = "Enable Multithreaded lightmap encoding"))
 	bool bEnableMultithreadedLightmapEncoding;
diff --git a/Engine/Source/Editor/UnrealEd/Classes/Settings/LevelEditorPlaySettings.h b/Engine/Source/Editor/UnrealEd/Classes/Settings/LevelEditorPlaySettings.h
--- a/Engine/Source/Editor/UnrealEd/Classes/Settings/LevelEditorPlaySettings.h
+++ b/Engine/Source/Editor/UnrealEd/Classes/Settings/LevelEditorPlaySettings.h
@@ -70,7 +70,7 @@ enum EPlayModeType : int
 	PlayMode_InTargetedMobilePreview,
 
 	/** Runs a vulkan preview in a new process. */
-	PlayMode_InVulkanPreview,
+	PlayMode_InVulkanPreview UE_DEPRECATED(5.7, "PlayMode_InVulkanPreview is no longer supported, please use PlayMode_InMobilePreview instead"),
 
 	/** Runs in a new process. */
 	PlayMode_InNewProcess,
diff --git a/Engine/Source/Editor/UnrealEd/Private/Kismet2/DebuggerCommands.cpp b/Engine/Source/Editor/UnrealEd/Private/Kismet2/DebuggerCommands.cpp
--- a/Engine/Source/Editor/UnrealEd/Private/Kismet2/DebuggerCommands.cpp
+++ b/Engine/Source/Editor/UnrealEd/Private/Kismet2/DebuggerCommands.cpp
@@ -203,7 +203,6 @@ public:
 	static bool IsStoppedAtBreakpoint_InEngineMode();
 
 	static bool CanShowNonPlayWorldOnlyActions();
-	static bool CanShowVulkanNonPlayWorldOnlyActions();
 	static bool CanShowVROnlyActions();
 
 	static int32 GetNumberOfClients();
@@ -353,7 +352,6 @@ void FPlayWorldCommands::RegisterCommands()
 	UI_COMMAND(PlayInEditorFloating, "New Editor Window (PIE)", "Play this level in a new window", EUserInterfaceActionType::Check, FInputChord());
 	UI_COMMAND(PlayInVR, "VR Preview", "Play this level in VR", EUserInterfaceActionType::Check, FInputChord());
 	UI_COMMAND(PlayInMobilePreview, "Standalone Game Mobile", "Play this level as a mobile device preview in Mobile mode (runs in its own process)", EUserInterfaceActionType::Check, FInputChord());
-	UI_COMMAND(PlayInVulkanPreview, "Standalone Game Vulkan Mobile", "Play this level using mobile Vulkan rendering (runs in its own process)", EUserInterfaceActionType::Check, FInputChord());
 	UI_COMMAND(PlayInNewProcess, "Standalone Game", "Play this level in a new window that runs in its own process", EUserInterfaceActionType::Check, FInputChord());
 	UI_COMMAND(PlayInCameraLocation, "Current Camera Location", "Spawn the player at the current camera location", EUserInterfaceActionType::RadioButton, FInputChord());
 	UI_COMMAND(PlayInDefaultPlayerStart, "Default Player Start", "Spawn the player at the map's default player start", EUserInterfaceActionType::RadioButton, FInputChord());
@@ -478,13 +476,6 @@ void FPlayWorldCommands::BindGlobalPlayWorldCommands()
 		FIsActionButtonVisible::CreateStatic(&FInternalPlayWorldCommandCallbacks::CanShowNonPlayWorldOnlyActions)
 	);
 
-	ActionList.MapAction(Commands.PlayInVulkanPreview,
-		FExecuteAction::CreateStatic(&FInternalPlayWorldCommandCallbacks::PlayInNewProcess_Clicked, PlayMode_InVulkanPreview),
-		FCanExecuteAction::CreateStatic(&FInternalPlayWorldCommandCallbacks::PlayInNewProcess_CanExecute),
-		FIsActionChecked::CreateStatic(&FInternalPlayWorldCommandCallbacks::PlayInModeIsChecked, PlayMode_InVulkanPreview),
-		FIsActionButtonVisible::CreateStatic(&FInternalPlayWorldCommandCallbacks::CanShowVulkanNonPlayWorldOnlyActions)
-	);
-
 	ActionList.MapAction(Commands.PlayInNewProcess,
 		FExecuteAction::CreateStatic(&FInternalPlayWorldCommandCallbacks::PlayInNewProcess_Clicked, PlayMode_InNewProcess),
 		FCanExecuteAction::CreateStatic(&FInternalPlayWorldCommandCallbacks::PlayInNewProcess_CanExecute),
@@ -874,10 +865,6 @@ TSharedRef< SWidget > FPlayWorldCommands::GeneratePlayMenuContent(TSharedRef<FUI
 					PlayModeCommand = FPlayWorldCommands::Get().PlayInMobilePreview;
 					break;
 
-				case PlayMode_InVulkanPreview:
-					PlayModeCommand = FPlayWorldCommands::Get().PlayInVulkanPreview;
-					break;
-
 				case PlayMode_InNewProcess:
 					PlayModeCommand = FPlayWorldCommands::Get().PlayInNewProcess;
 					break;
@@ -917,7 +904,6 @@ TSharedRef< SWidget > FPlayWorldCommands::GeneratePlayMenuContent(TSharedRef<FUI
 			}
 
 			FLocal::AddPlayModeMenuEntry(Section, PlayMode_InMobilePreview);
-			FLocal::AddPlayModeMenuEntry(Section, PlayMode_InVulkanPreview);
 			FLocal::AddPlayModeMenuEntry(Section, PlayMode_InNewProcess);
 
 			Section.AddSeparator("SimulateSeparator");
@@ -1337,10 +1323,6 @@ const TSharedRef < FUICommandInfo > GetLastPlaySessionCommand()
 		break;
 	}
 
-	case PlayMode_InVulkanPreview:
-		Command = Commands.PlayInVulkanPreview.ToSharedRef();
-		break;
-
 	case PlayMode_InNewProcess:
 		Command = Commands.PlayInNewProcess.ToSharedRef();
 		break;
@@ -1400,10 +1382,6 @@ void RecordLastExecutedPlayMode()
 			PlayModeString = TEXT("InTargetedMobilePreview");
 			break;
 
-		case PlayMode_InVulkanPreview:
-			PlayModeString = TEXT("InVulkanPreview");
-			break;
-
 		case PlayMode_InNewProcess:
 			PlayModeString = TEXT("InNewProcess");
 			break;
@@ -1710,8 +1688,9 @@ void FInternalPlayWorldCommandCallbacks::PlayInNewProcessPreviewDevice_Clicked(F
 
 void FInternalPlayWorldCommandCallbacks::PlayInNewProcess_Clicked(EPlayModeType PlayModeType)
 {
-	check(PlayModeType == PlayMode_InNewProcess || PlayModeType == PlayMode_InMobilePreview
-		|| PlayModeType == PlayMode_InTargetedMobilePreview || PlayModeType == PlayMode_InVulkanPreview);
+	check(PlayModeType == PlayMode_InNewProcess 
+		|| PlayModeType == PlayMode_InMobilePreview
+		|| PlayModeType == PlayMode_InTargetedMobilePreview);
 
 	if (IsStoppedAtBreakpoint_InEngineMode() || DebuggerCommands::IsVREditorActive())
 	{
@@ -1747,10 +1726,6 @@ void FInternalPlayWorldCommandCallbacks::PlayInNewProcess_Clicked(EPlayModeType
 
 			SessionParams.SessionPreviewTypeOverride = EPlaySessionPreviewType::MobilePreview;
 		}
-		else if (PlayModeType == PlayMode_InVulkanPreview)
-		{
-			SessionParams.SessionPreviewTypeOverride = EPlaySessionPreviewType::VulkanPreview;
-		}
 
 		SessionParams.SessionDestination = EPlaySessionDestinationType::NewProcess;
 
@@ -1981,11 +1956,6 @@ bool FInternalPlayWorldCommandCallbacks::CanShowNonPlayWorldOnlyActions()
 	return !HasPlayWorld() && DebuggerCommands::AllowPlayWorldFeature();
 }
 
-bool FInternalPlayWorldCommandCallbacks::CanShowVulkanNonPlayWorldOnlyActions()
-{
-	return !HasPlayWorld() && GetDefault<UEditorExperimentalSettings>()->bAllowVulkanPreview && FModuleManager::Get().ModuleExists(TEXT("VulkanRHI")) && DebuggerCommands::AllowPlayWorldFeature();
-}
-
 bool FInternalPlayWorldCommandCallbacks::CanShowVROnlyActions()
 {
 	return !HasPlayWorld() && DebuggerCommands::AllowPlayWorldFeature();
diff --git a/Engine/Source/Editor/UnrealEd/Private/PlayLevelNewProcess.cpp b/Engine/Source/Editor/UnrealEd/Private/PlayLevelNewProcess.cpp
--- a/Engine/Source/Editor/UnrealEd/Private/PlayLevelNewProcess.cpp
+++ b/Engine/Source/Editor/UnrealEd/Private/PlayLevelNewProcess.cpp
@@ -189,12 +189,6 @@ void UEditorEngine::LaunchNewProcess(const FRequestPlaySessionParams& InParams,
 		);
 	}
 
-	if (InParams.SessionPreviewTypeOverride.Get(EPlaySessionPreviewType::NoPreview) == EPlaySessionPreviewType::VulkanPreview)
-	{
-		// Vulkan only supports a sub-set
-		CommandLine += TEXT(" -vulkan -faketouches -featureleveles31");
-	}
-
 	// VRPreview handling
 	if (InParams.SessionPreviewTypeOverride.Get(EPlaySessionPreviewType::NoPreview) == EPlaySessionPreviewType::VRPreview)
 	{
diff --git a/Engine/Source/Editor/UnrealEd/Public/Kismet2/DebuggerCommands.h b/Engine/Source/Editor/UnrealEd/Public/Kismet2/DebuggerCommands.h
--- a/Engine/Source/Editor/UnrealEd/Public/Kismet2/DebuggerCommands.h
+++ b/Engine/Source/Editor/UnrealEd/Public/Kismet2/DebuggerCommands.h
@@ -89,7 +89,6 @@ public:
 	TSharedPtr<FUICommandInfo> PlayInEditorFloating;
 	TSharedPtr<FUICommandInfo> PlayInVR;
 	TSharedPtr<FUICommandInfo> PlayInMobilePreview;
-	TSharedPtr<FUICommandInfo> PlayInVulkanPreview;
 	TSharedPtr<FUICommandInfo> PlayInNewProcess;
 	TSharedPtr<FUICommandInfo> PlayInCameraLocation;
 	TSharedPtr<FUICommandInfo> PlayInDefaultPlayerStart;
diff --git a/Engine/Source/Editor/UnrealEd/Public/PlayInEditorDataTypes.h b/Engine/Source/Editor/UnrealEd/Public/PlayInEditorDataTypes.h
--- a/Engine/Source/Editor/UnrealEd/Public/PlayInEditorDataTypes.h
+++ b/Engine/Source/Editor/UnrealEd/Public/PlayInEditorDataTypes.h
@@ -118,7 +118,7 @@ enum class EPlaySessionPreviewType : uint8
 	/* Emulate Mobile ES3 Rendering. */
 	MobilePreview,
 	/* Emulate Vulkan Rendering. */
-	VulkanPreview,
+	VulkanPreview UE_DEPRECATED(5.7, "VulkanPreview is no longer supported, please use MobilePreview instead"),
 	/* Attempt to launch in a HMD (Not used for VREditor) */
 	VRPreview,
 };
diff --git a/Engine/Source/Runtime/Engine/Classes/Engine/World.h b/Engine/Source/Runtime/Engine/Classes/Engine/World.h
--- a/Engine/Source/Runtime/Engine/Classes/Engine/World.h
+++ b/Engine/Source/Runtime/Engine/Classes/Engine/World.h
@@ -4040,7 +4040,8 @@ public:
 	UE_API bool IsPlayInMobilePreview() const;
 
 	/** Returns whether script is executing within a vulkan preview window */
-	UE_API bool IsPlayInVulkanPreview() const;
+	UE_DEPRECATED(5.7, "PlayInVulkanPreview is no longer supported, please use PlayInMobilePreview instead")
+	UE_API bool IsPlayInVulkanPreview() const { return false; };
 
 	/** Returns true if this world is any kind of game world (including PIE worlds) */
 	UE_API bool IsGameWorld() const;
diff --git a/Engine/Source/Runtime/Engine/Private/World.cpp b/Engine/Source/Runtime/Engine/Private/World.cpp
--- a/Engine/Source/Runtime/Engine/Private/World.cpp
+++ b/Engine/Source/Runtime/Engine/Private/World.cpp
@@ -9143,17 +9143,13 @@ bool UWorld::IsPlayInMobilePreview() const
 {
 #if WITH_EDITOR
 	if (FPIEPreviewDeviceModule::IsRequestingPreviewDevice()
-		|| FParse::Param(FCommandLine::Get(), TEXT("featureleveles31")))
+		|| FParse::Param(FCommandLine::Get(), TEXT("featureleveles31"))
+		|| FDataDrivenShaderPlatformInfo::GetIsPreviewPlatform(Scene ? Scene->GetShaderPlatform(): GMaxRHIShaderPlatform))
 	{
 		return true;
 	}
 #endif // WITH_EDITOR
-	return FParse::Param(FCommandLine::Get(), TEXT("simmobile")) && !IsPlayInVulkanPreview();
-}
-
-bool UWorld::IsPlayInVulkanPreview() const
-{
-	return FParse::Param(FCommandLine::Get(), TEXT("vulkan"));
+	return FParse::Param(FCommandLine::Get(), TEXT("simmobile"));
 }
 
 bool UWorld::IsGameWorld() const
