World Partition - HLOD: Make sure -ReuseParentBranchHLODs works for non-world partition worlds

diff --git a/Engine/Source/Editor/UnrealEd/Private/WorldPartition/WorldPartitionHLODsBuilder.cpp b/Engine/Source/Editor/UnrealEd/Private/WorldPartition/WorldPartitionHLODsBuilder.cpp
--- a/Engine/Source/Editor/UnrealEd/Private/WorldPartition/WorldPartitionHLODsBuilder.cpp
+++ b/Engine/Source/Editor/UnrealEd/Private/WorldPartition/WorldPartitionHLODsBuilder.cpp
@@ -670,14 +670,7 @@ bool UWorldPartitionHLODsBuilder::BuildHLODActors()
 
 			if (!ParentBranchHLODFileToCopy.IsEmpty())
 			{
-				ResetLoaders(HLODActorPackage);
-
-				SourceControlHelper->Copy(ParentBranchHLODFileToCopy, USourceControlHelpers::PackageFilename(HLODActorPackage));
-
-				// Clean up temp file
-				IFileManager::Get().DeleteDirectory(*FPaths::GetPath(ParentBranchHLODFileToCopy), /*RequireExists=*/false, /*Tree=*/true);
-
-				ParentBranchHLODFileToCopy.Empty();
+				CopyParentBranchHLODFile(HLODActorPackage);
 			}
 
 			TRACE_BOOKMARK(TEXT("BuildHLOD End - %d"), CurrentActor);
@@ -694,15 +687,24 @@ bool UWorldPartitionHLODsBuilder::BuildHLODActors()
 	{
 		IWorldPartitionHLODProvider::FBuildHLODActorParams BuildHLODActorParams;
 		BuildHLODActorParams.bForceRebuild = bForceBuild;
-		BuildHLODActorParams.OnPackageModified.BindLambda([this](UPackage* ModifiedPackage)
+		BuildHLODActorParams.OnPackageProcessed.BindLambda([this](UPackage* ProcessedPackage)
 		{
-			bool bSaved = SourceControlHelper->Save(ModifiedPackage);
-			if (!bSaved)
+			bool bSuccess = true;
+
+			if (!ParentBranchHLODFileToCopy.IsEmpty())
+			{
+				CopyParentBranchHLODFile(ProcessedPackage);
+			}
+			else if (ProcessedPackage->IsDirty())
 			{
-				UE_LOG(LogWorldPartitionHLODsBuilder, Error, TEXT("Failed to save %s, exiting..."), *USourceControlHelpers::PackageFilename(ModifiedPackage));
+				bSuccess = SourceControlHelper->Save(ProcessedPackage);
+				if (!bSuccess)
+				{
+					UE_LOG(LogWorldPartitionHLODsBuilder, Error, TEXT("Failed to save %s, exiting..."), *USourceControlHelpers::PackageFilename(ProcessedPackage));
+				}
 			}
 
-			return bSaved;
+			return bSuccess;
 		});
 
 		TArray<IWorldPartitionHLODProvider*> HLODProviders;
@@ -1533,12 +1535,6 @@ bool UWorldPartitionHLODsBuilder::EvaluateHLODBuildConditions(AWorldPartitionHLO
 	bool bShouldPerformBuild = OldHash != NewHash;
 	EvaluationStats.Built += bShouldPerformBuild ? 1 : 0;
 
-	// Parent branch hash comparison can only work on OFPA HLODs
-	if (HLODActor->GetExternalPackage() == nullptr)
-	{
-		return bShouldPerformBuild;
-	}
-
 	// Parent branch hash comparison can't be performed without SCC
 	if (!ISourceControlModule::Get().IsEnabled())
 	{
@@ -1703,4 +1699,16 @@ bool UWorldPartitionHLODsBuilder::EvaluateHLODBuildConditions(AWorldPartitionHLO
 	return !bAbortBuild;
 }
 
+void UWorldPartitionHLODsBuilder::CopyParentBranchHLODFile(UPackage* HLODActorPackage)
+{
+	ResetLoaders(HLODActorPackage);
+
+	SourceControlHelper->Copy(ParentBranchHLODFileToCopy, USourceControlHelpers::PackageFilename(HLODActorPackage));
+
+	// Clean up temp file
+	IFileManager::Get().DeleteDirectory(*FPaths::GetPath(ParentBranchHLODFileToCopy), /*RequireExists=*/false, /*Tree=*/true);
+
+	ParentBranchHLODFileToCopy.Empty();
+}
+
 #undef LOCTEXT_NAMESPACE
diff --git a/Engine/Source/Editor/UnrealEd/Public/WorldPartition/WorldPartitionHLODsBuilder.h b/Engine/Source/Editor/UnrealEd/Public/WorldPartition/WorldPartitionHLODsBuilder.h
--- a/Engine/Source/Editor/UnrealEd/Public/WorldPartition/WorldPartitionHLODsBuilder.h
+++ b/Engine/Source/Editor/UnrealEd/Public/WorldPartition/WorldPartitionHLODsBuilder.h
@@ -68,6 +68,7 @@ protected:
 	void AllowExternalDataLayerInjection(const UWorld* InWorld, const UExternalDataLayerAsset* InExternalDataLayerAsset, bool& bOutAllowInjection);
 
 	bool EvaluateHLODBuildConditions(class AWorldPartitionHLOD* HLODActor, uint32 OldHash, uint32 NewHash);
+	void CopyParentBranchHLODFile(UPackage* HLODActorPackage);
 
 private:
 	UWorld* World;
diff --git a/Engine/Source/Runtime/Engine/Private/WorldPartition/HLOD/HLODActor.cpp b/Engine/Source/Runtime/Engine/Private/WorldPartition/HLOD/HLODActor.cpp
--- a/Engine/Source/Runtime/Engine/Private/WorldPartition/HLOD/HLODActor.cpp
+++ b/Engine/Source/Runtime/Engine/Private/WorldPartition/HLOD/HLODActor.cpp
@@ -364,7 +364,7 @@ void AWorldPartitionHLOD::PreRegisterAllComponents()
 	// In PIE, collisions should be disabled
 	if (GetWorld() && !IsRunningCommandlet() && !FApp::IsUnattended())
 	{
-		bool bShouldEnableCollision = !GetWorld()->IsGameWorld();
+		bool bShouldEnableCollision = false;/*!GetWorld()->IsGameWorld();*/
 		if (GetActorEnableCollision() != bShouldEnableCollision)
 		{
 			SetActorEnableCollision(bShouldEnableCollision);
@@ -583,11 +583,7 @@ void AWorldPartitionHLOD::BuildHLOD(bool bForceBuild)
 void AWorldPartitionHLOD::GetAssetRegistryTags(FAssetRegistryTagsContext Context) const
 {
 	Super::GetAssetRegistryTags(Context);
-
-	if (IsPackageExternal())
-	{
-		Context.AddTag(UObject::FAssetRegistryTag(NAME_HLODHash_AssetTag, LexToString(HLODHash), UObject::FAssetRegistryTag::TT_Hidden));
-	}
+	Context.AddTag(UObject::FAssetRegistryTag(NAME_HLODHash_AssetTag, LexToString(HLODHash), UObject::FAssetRegistryTag::TT_Hidden));
 }
 
 #endif // WITH_EDITOR
\ No newline at end of file
diff --git a/Engine/Source/Runtime/Engine/Public/WorldPartition/HLOD/HLODProviderInterface.h b/Engine/Source/Runtime/Engine/Public/WorldPartition/HLOD/HLODProviderInterface.h
--- a/Engine/Source/Runtime/Engine/Public/WorldPartition/HLOD/HLODProviderInterface.h
+++ b/Engine/Source/Runtime/Engine/Public/WorldPartition/HLOD/HLODProviderInterface.h
@@ -22,12 +22,21 @@ class IWorldPartitionHLODProvider
 	GENERATED_IINTERFACE_BODY()
 
 public:
+	DECLARE_DELEGATE_RetVal_OneParam(bool, FPackageProcessedDelegate, UPackage*)
+
+	UE_DEPRECATED(5.7, "Use FPackageProcessedDelegate along with FBuildHLODActorParams::OnPackageProcessed")
 	DECLARE_DELEGATE_RetVal_OneParam(bool, FPackageModifiedDelegate, UPackage*)
 
 	struct FBuildHLODActorParams
 	{
 		bool bForceRebuild;
+
+		FPackageProcessedDelegate OnPackageProcessed;
+
+		PRAGMA_DISABLE_DEPRECATION_WARNINGS
+		UE_DEPRECATED(5.7, "Use OnPackageProcessed. Test if packaged is modified with Package->IsDirty()")
 		FPackageModifiedDelegate OnPackageModified;
+		PRAGMA_ENABLE_DEPRECATION_WARNINGS
 	};
 
 	virtual bool BuildHLODActor(const FBuildHLODActorParams& BuildHLODActorParams) = 0;
