Fixes the issue where unloading assets within a plugin before the plugin itself is unloaded is preventing the GC from sweeping those assets, because the compiled package registry in Solaris is still holding strong references to the `UVerseClass` objects.

#rnx

diff --git a/Engine/Source/Editor/UnrealEd/Internal/PackageTools_Private.h b/Engine/Source/Editor/UnrealEd/Internal/PackageTools_Private.h
new file mode 100644
--- /dev/null
+++ b/Engine/Source/Editor/UnrealEd/Internal/PackageTools_Private.h
@@ -0,0 +1,17 @@
+// Copyright Epic Games, Inc. All Rights Reserved.
+
+#pragma once
+
+#include "Delegates/Delegate.h"
+
+namespace UE::UnrealEd
+{
+namespace PackageTools_Private
+{
+DECLARE_MULTICAST_DELEGATE_OneParam(FOnUnloadPackagesDelegate, const TSet<UPackage*> /*PackagesToUnload*/);
+
+/// Broadcast before a given set of packages are unloaded.
+UNREALED_API extern FOnUnloadPackagesDelegate OnUnloadPackagesDelegate;
+
+}	 // namespace PackageTools
+}	 // namespace UE::UnrealEd
diff --git a/Engine/Source/Editor/UnrealEd/Private/PackageTools.cpp b/Engine/Source/Editor/UnrealEd/Private/PackageTools.cpp
--- a/Engine/Source/Editor/UnrealEd/Private/PackageTools.cpp
+++ b/Engine/Source/Editor/UnrealEd/Private/PackageTools.cpp
@@ -2,6 +2,7 @@
 
 
 #include "PackageTools.h"
+#include "PackageTools_Private.h"
 #include "Algo/Transform.h"
 #include "BlueprintCompilationManager.h"
 #include "UObject/PackageReload.h"
@@ -74,11 +75,16 @@
 
 DEFINE_LOG_CATEGORY_STATIC(LogPackageTools, Log, All);
 
+namespace UE::UnrealEd
+{
 namespace PackageTools_Private
 {
-	static bool bUnloadPackagesUnloadsPrimaryAssets = true;
-	FAutoConsoleVariableRef CVarUnloadPackagesUnloadsPrimaryAssets(TEXT("PackageTools.UnloadPackagesUnloadsPrimaryAssets"), bUnloadPackagesUnloadsPrimaryAssets, TEXT("During unload packages, also unload primary assets"), ECVF_Default);
-}
+static bool bUnloadPackagesUnloadsPrimaryAssets = true;
+FAutoConsoleVariableRef CVarUnloadPackagesUnloadsPrimaryAssets(TEXT("PackageTools.UnloadPackagesUnloadsPrimaryAssets"), bUnloadPackagesUnloadsPrimaryAssets, TEXT("During unload packages, also unload primary assets"), ECVF_Default);
+
+FOnUnloadPackagesDelegate OnUnloadPackagesDelegate;
+}	 // namespace PackageTools_Private
+}	 // namespace UE::UnrealEd
 
 /** State passed to RestoreStandaloneOnReachableObjects. */
 TSet<UPackage*>* UPackageTools::PackagesBeingUnloaded = nullptr;
@@ -465,6 +471,10 @@ UPackageTools::UPackageTools(const FObjectInitializer& ObjectInitializer)
 			// so that it is safe to call ::ResetLoaders
 			FlushAsyncCompilation(PackagesToUnload.Array());
 
+			// Before we clean up the assets in the packages that are about to be unloaded, we give a chance for any registered observers
+			// to clean up any state that they need to for the assets in those packages.
+			UE::UnrealEd::PackageTools_Private::OnUnloadPackagesDelegate.Broadcast(PackagesToUnload);
+
 			// Now try to clean up assets in all packages to unload.
 			bool bScriptPackageWasUnloaded = false;
 			int32 PackageIndex = 0;
@@ -501,7 +511,7 @@ UPackageTools::UPackageTools(const FObjectInitializer& ObjectInitializer)
 				for (UObject* Obj : ObjectsInPackage)
 				{
 					// Asset manager can hold hard references to this object and prevent GC
-					if (PackageTools_Private::bUnloadPackagesUnloadsPrimaryAssets)
+					if (UE::UnrealEd::PackageTools_Private::bUnloadPackagesUnloadsPrimaryAssets)
 					{
 						const FPrimaryAssetId PrimaryAssetId = UAssetManager::Get().GetPrimaryAssetIdForObject(Obj);
 						if (PrimaryAssetId.IsValid())
