Anim In Engine: Make AnimSequencerExport Blueprint types
#trivial

diff --git a/Engine/Source/Editor/UnrealEd/Classes/Exporters/AnimSeqExportOption.h b/Engine/Source/Editor/UnrealEd/Classes/Exporters/AnimSeqExportOption.h
--- a/Engine/Source/Editor/UnrealEd/Classes/Exporters/AnimSeqExportOption.h
+++ b/Engine/Source/Editor/UnrealEd/Classes/Exporters/AnimSeqExportOption.h
@@ -12,7 +12,7 @@
 #include "AnimSeqExportOption.generated.h"
 
 
-UCLASS(MinimalAPI)
+UCLASS(MinimalAPI, BlueprintType)
 class UAnimSeqExportOption : public UObject
 {
 	GENERATED_UCLASS_BODY()
