[Reimport] Asset file type extension fallback mechanism introduced for use cases where AssetImportData is missing.

#rnx

diff --git a/Engine/Source/Editor/UnrealEd/Private/Editor.cpp b/Engine/Source/Editor/UnrealEd/Private/Editor.cpp
--- a/Engine/Source/Editor/UnrealEd/Private/Editor.cpp
+++ b/Engine/Source/Editor/UnrealEd/Private/Editor.cpp
@@ -282,6 +282,20 @@ void FReimportManager::UpdateReimportPath(UObject* Obj, const FString& Filename,
 			}
 			Obj->MarkPackageDirty();
 		}
+		else
+		{
+			//In case no Handler was found with the SourcefileName acquisition,
+			// we then try to set the reimportpath 1 by 1 and break out as soon as the the SetReimportPaths then provided a CanReimport state.
+			TArray<FString> FileNames;
+			for (int32 HandlerIndex = 0; HandlerIndex < Handlers.Num(); ++HandlerIndex)
+			{
+				Handlers[HandlerIndex]->SetReimportPaths(Obj, Filename, SourceFileIndex);
+				if (Handlers[HandlerIndex]->CanReimport(Obj, FileNames))
+				{
+					break;
+				}
+			}
+		}
 	}
 }
 
@@ -325,21 +339,73 @@ UE::Interchange::FAssetImportResultRef FReimportManager::ReimportAsync(UObject*
 		{
 			CanReimportHandler->SetPreferredReimportPath(PreferredReimportFile);
 		}
+		auto TryToSetCanReimportHandler = [&]()
+			{
+				for (int32 HandlerIndex = 0; HandlerIndex < Handlers.Num(); ++HandlerIndex)
+				{
+					SourceFilenames.Empty();
+					Handlers[HandlerIndex]->SetPreferredReimportPath(PreferredReimportFile);
+					if (Handlers[HandlerIndex]->CanReimport(Obj, SourceFilenames))
+					{
+						CanReimportHandler = Handlers[HandlerIndex];
+						break;
+					}
+				}
+			};
+
 		if (CanReimportHandler == nullptr || !CanReimportHandler->CanReimport(Obj, SourceFilenames))
 		{
-			for (int32 HandlerIndex = 0; HandlerIndex < Handlers.Num(); ++HandlerIndex)
+			TryToSetCanReimportHandler();
+		}
+
+		bool bHasSourceFile = false;
+		if (CanReimportHandler == nullptr)
+		{
+			if (const UAssetDefinition* AssetDefinition = UAssetDefinitionRegistry::Get()->GetAssetDefinitionForAsset(Obj))
 			{
-				SourceFilenames.Empty();
-				Handlers[HandlerIndex]->SetPreferredReimportPath(PreferredReimportFile);
-				if (Handlers[HandlerIndex]->CanReimport(Obj, SourceFilenames))
+				FAssetData AssetData(Obj);
+				FAssetSourceFilesArgs GetSourceFilesArgs;
+				GetSourceFilesArgs.Assets = TConstArrayView<FAssetData>(&AssetData, 1);
+
+				// Doesn't need to resolve the paths so it is a bit quicker
+				GetSourceFilesArgs.FilePathFormat = EPathUse::Display;
+
+				AssetDefinition->GetSourceFiles(GetSourceFilesArgs, [&bHasSourceFile](const FAssetSourceFilesResult& AssetImportInfo)
+					{
+						bHasSourceFile = true;
+
+						return false;
+					});
+			}
+
+			if (!bHasSourceFile)
+			{
+				if (SourceFilenames.IsValidIndex(RealSourceFileIndex))
 				{
-					CanReimportHandler = Handlers[HandlerIndex];
-					break;
+					SourceFilenames[RealSourceFileIndex].Empty();
+				}
+				else
+				{
+					//Add the missing entries
+					SourceFilenames.AddDefaulted(RealSourceFileIndex - (SourceFilenames.Num() - 1));
 				}
+
+				//If Source file was not found, thenw e force a new file selection here so we can re-check the CanReimportHandler
+				for (int32 FileIndex = 0; FileIndex < SourceFilenames.Num(); ++FileIndex)
+				{
+					GetNewReimportPath(Obj, SourceFilenames, FileIndex);
+				}
+
+				UpdateReimportPath(Obj, SourceFilenames[RealSourceFileIndex], RealSourceFileIndex);
+
+				//Also setting the bForceNewFile to false as we already did force it.
+				bForceNewFile = false;
+
+				TryToSetCanReimportHandler();
 			}
 		}
 
-		if(CanReimportHandler != nullptr)
+		if (CanReimportHandler != nullptr)
 		{
 			TArray<int32> MissingFileIndex;
 
@@ -641,6 +707,36 @@ void FReimportManager::ValidateAllSourceFileAndReimport(TArray<UObject*> &ToImpo
 				}
 			}
 		}
+		else
+		{
+			if (const UAssetDefinition* AssetDefinition = UAssetDefinitionRegistry::Get()->GetAssetDefinitionForAsset(Asset))
+			{
+				FAssetData AssetData(Asset);
+				FAssetSourceFilesArgs GetSourceFilesArgs;
+				GetSourceFilesArgs.Assets = TConstArrayView<FAssetData>(&AssetData, 1);
+
+				// Doesn't need to resolve the paths so it is a bit quicker
+				GetSourceFilesArgs.FilePathFormat = EPathUse::Display;
+
+				bool bHasSourceFile = false;
+				AssetDefinition->GetSourceFiles(GetSourceFilesArgs, [&bHasSourceFile](const FAssetSourceFilesResult& AssetImportInfo)
+					{
+						bHasSourceFile = true;
+
+						return false;
+					});
+
+				if (!bHasSourceFile)
+				{
+					TArray<int32>& SourceIndexArray = MissingFileSelectedAssets.FindOrAdd(Asset);
+					if (SourceIndexArray.Num() == 0)
+					{
+						// Insert an invalid index to indicate no file
+						SourceIndexArray.Add(INDEX_NONE);
+					}
+				}
+			}
+		}
 	}
 
 	if (MissingFileSelectedAssets.Num() > 0)
@@ -867,6 +963,34 @@ void FReimportManager::GetNewReimportPath(UObject* Obj, TArray<FString>& InOutFi
 	{
 		//Get the extension interchange can translate for this object
 		TArray<FString> TranslatorFormats = UInterchangeManager::GetInterchangeManager().GetSupportedFormatsForObject(Obj, RealSourceFileIndex);
+
+		if (TranslatorFormats.Num() == 0)
+		{
+			//In case the asset was not imported or loose the sources file name, we want to be able to retrieve the possible formats.
+			if (const UClass* FactoryClass = UInterchangeManager::GetInterchangeManager().GetRegisteredFactoryClass(Obj->GetClass()))
+			{
+				UInterchangeFactoryBase* FactoryBase = FactoryClass->GetDefaultObject<UInterchangeFactoryBase>();
+				EInterchangeTranslatorAssetType TranslatorAssetType = EInterchangeTranslatorAssetType::None;
+				switch (FactoryBase->GetFactoryAssetType())
+				{
+					case EInterchangeFactoryAssetType::Animations:
+						TranslatorAssetType = EInterchangeTranslatorAssetType::Animations;
+						break;
+					case EInterchangeFactoryAssetType::Materials:
+						TranslatorAssetType = EInterchangeTranslatorAssetType::Materials;
+						break;
+					case EInterchangeFactoryAssetType::Meshes:
+					case EInterchangeFactoryAssetType::Physics:
+						TranslatorAssetType = EInterchangeTranslatorAssetType::Meshes;
+						break;
+					case EInterchangeFactoryAssetType::Textures:
+						TranslatorAssetType = EInterchangeTranslatorAssetType::Textures;
+						break;
+				};
+				TranslatorFormats = UInterchangeManager::GetInterchangeManager().GetSupportedAssetTypeFormats(TranslatorAssetType, EInterchangeTranslatorType::Assets, false /*bStrictMatchTranslatorType*/);
+			}
+		}
+
 		ObjectTools::AppendFormatsFileExtensions(TranslatorFormats, FileTypes, AllExtensions, DummyFilterIndexToFactory);
 	}
 
