[Reimport] Fixing UpdateReimportPath function's fallback, adding package dirty and handling SourceFileIndex==INDEX_NONE scenarios.

#rnx

diff --git a/Engine/Source/Editor/UnrealEd/Private/Editor.cpp b/Engine/Source/Editor/UnrealEd/Private/Editor.cpp
--- a/Engine/Source/Editor/UnrealEd/Private/Editor.cpp
+++ b/Engine/Source/Editor/UnrealEd/Private/Editor.cpp
@@ -286,15 +286,45 @@ void FReimportManager::UpdateReimportPath(UObject* Obj, const FString& Filename,
 		{
 			//In case no Handler was found with the SourcefileName acquisition,
 			// we then try to set the reimportpath 1 by 1 and break out as soon as the the SetReimportPaths then provided a CanReimport state.
-			TArray<FString> FileNames;
-			for (int32 HandlerIndex = 0; HandlerIndex < Handlers.Num(); ++HandlerIndex)
+
+			bool bMarkPackageDirty = false;
+			if (SourceFileIndex == INDEX_NONE)
 			{
-				Handlers[HandlerIndex]->SetReimportPaths(Obj, Filename, SourceFileIndex);
-				if (Handlers[HandlerIndex]->CanReimport(Obj, FileNames))
+				TArray<FString> FileNames;
+				FileNames.Add(Filename);
+
+				TArray<FString> OutFileNames;
+
+				for (int32 HandlerIndex = 0; HandlerIndex < Handlers.Num(); ++HandlerIndex)
 				{
-					break;
+					OutFileNames.Empty();
+					Handlers[HandlerIndex]->SetReimportPaths(Obj, FileNames);
+					if (Handlers[HandlerIndex]->CanReimport(Obj, OutFileNames))
+					{
+						bMarkPackageDirty = true;
+						break;
+					}
 				}
 			}
+			else
+			{
+				TArray<FString> FileNames;
+				for (int32 HandlerIndex = 0; HandlerIndex < Handlers.Num(); ++HandlerIndex)
+				{
+					FileNames.Empty();
+					Handlers[HandlerIndex]->SetReimportPaths(Obj, Filename, SourceFileIndex);
+					if (Handlers[HandlerIndex]->CanReimport(Obj, FileNames))
+					{
+						bMarkPackageDirty = true;
+						break;
+					}
+				}
+			}
+
+			if (bMarkPackageDirty)
+			{
+				Obj->MarkPackageDirty();
+			}
 		}
 	}
 }
