[Reimport] Setting package dirty regardless of CanReimport' outcome.

#rnx

diff --git a/Engine/Source/Editor/UnrealEd/Private/Editor.cpp b/Engine/Source/Editor/UnrealEd/Private/Editor.cpp
--- a/Engine/Source/Editor/UnrealEd/Private/Editor.cpp
+++ b/Engine/Source/Editor/UnrealEd/Private/Editor.cpp
@@ -287,7 +287,6 @@ void FReimportManager::UpdateReimportPath(UObject* Obj, const FString& Filename,
 			//In case no Handler was found with the SourcefileName acquisition,
 			// we then try to set the reimportpath 1 by 1 and break out as soon as the the SetReimportPaths then provided a CanReimport state.
 
-			bool bMarkPackageDirty = false;
 			if (SourceFileIndex == INDEX_NONE)
 			{
 				TArray<FString> FileNames;
@@ -301,7 +300,6 @@ void FReimportManager::UpdateReimportPath(UObject* Obj, const FString& Filename,
 					Handlers[HandlerIndex]->SetReimportPaths(Obj, FileNames);
 					if (Handlers[HandlerIndex]->CanReimport(Obj, OutFileNames))
 					{
-						bMarkPackageDirty = true;
 						break;
 					}
 				}
@@ -315,16 +313,12 @@ void FReimportManager::UpdateReimportPath(UObject* Obj, const FString& Filename,
 					Handlers[HandlerIndex]->SetReimportPaths(Obj, Filename, SourceFileIndex);
 					if (Handlers[HandlerIndex]->CanReimport(Obj, FileNames))
 					{
-						bMarkPackageDirty = true;
 						break;
 					}
 				}
 			}
 
-			if (bMarkPackageDirty)
-			{
-				Obj->MarkPackageDirty();
-			}
+			Obj->MarkPackageDirty();
 		}
 	}
 }
