Cache world bounds on the Editor Viewport Client. Expose orthographic clipping planes in world units in hover decorators and view menus.

#rnx

diff --git a/Engine/Source/Editor/LevelEditor/Private/SLevelViewport.cpp b/Engine/Source/Editor/LevelEditor/Private/SLevelViewport.cpp
--- a/Engine/Source/Editor/LevelEditor/Private/SLevelViewport.cpp
+++ b/Engine/Source/Editor/LevelEditor/Private/SLevelViewport.cpp
@@ -3604,7 +3604,7 @@ void SLevelViewport::OnElementSelectionChanged(const UTypedElementSelectionSet*
 		return;
 	}
 	
-	DepthBar->ClearChildren();
+	DepthBar->ClearDepthIndicators();
 	FTypedElementListConstRef ElementList = SelectionSet->GetElementList();
 	for (int32 Index = 0; Index < ElementList->Num(); ++Index)
 	{
diff --git a/Engine/Source/Editor/UnrealEd/Private/EditorViewportClient.cpp b/Engine/Source/Editor/UnrealEd/Private/EditorViewportClient.cpp
--- a/Engine/Source/Editor/UnrealEd/Private/EditorViewportClient.cpp
+++ b/Engine/Source/Editor/UnrealEd/Private/EditorViewportClient.cpp
@@ -85,8 +85,10 @@
 #include "Misc/AxisDisplayInfo.h"
 #include "Tools/EdModeInteractiveToolsContext.h"
 #include "ViewportClientNavigationHelper.h"
+#include "Engine/LevelBounds.h"
 #include "ViewportInteractions/ViewportInteractionsBehaviorSource.h"
 #include "ViewportInteractions/ViewportInteraction.h"
+#include "WorldPartition/WorldPartition.h"
 
 #define LOCTEXT_NAMESPACE "EditorViewportClient"
 
@@ -533,6 +535,7 @@ FEditorViewportClient::FEditorViewportClient(FEditorModeTools* InModeTools, FPre
 	, ViewModeParam(-1)
 	, NearPlane(-1.0f)
 	, FarPlane(0.0f)
+	, CachedWorldBounds(ForceInit)
 	, bInGameViewMode(false)
 	, bInVREditViewMode(false)
 	, bShouldInvalidateViewportWidget(false)
@@ -604,6 +607,9 @@ FEditorViewportClient::FEditorViewportClient(FEditorModeTools* InModeTools, FPre
 	}
 	ViewTransformPerspective.SetMaxLocation(MaxViewLocation);
 	ViewTransformOrthographic.SetMaxLocation(MaxViewLocation);
+	
+	FCoreUObjectDelegates::OnObjectTransacted.AddRaw(this, &FEditorViewportClient::OnObjectTransacted);
+	FEditorDelegates::MapChange.AddRaw(this, &FEditorViewportClient::OnPostLoadMap);
 }
 
 FEditorViewportClient::~FEditorViewportClient()
@@ -643,6 +649,9 @@ FEditorViewportClient::~FEditorViewportClient()
 	}
 
 	ModeTools.Reset();
+	
+	FCoreUObjectDelegates::OnObjectTransacted.RemoveAll(this);
+	FEditorDelegates::MapChange.RemoveAll(this);
 }
 
 FLevelEditorViewportClient::FRealtimeOverride::FRealtimeOverride(bool bInIsRealtime, FText InSystemDisplayName)
@@ -1603,6 +1612,36 @@ void FEditorViewportClient::Tick(float DeltaTime)
 {
 	SCOPED_NAMED_EVENT(FEditorViewportClient_Tick, FColor::Red);
 	
+	if (!CachedWorldBounds.IsValid)
+	{
+		if (UWorld* World = GetWorld())
+		{
+			if (UWorldPartition* WorldPartition = World->GetWorldPartition())
+			{
+				CachedWorldBounds = WorldPartition->GetEditorWorldBounds();
+			}
+			else
+			{
+				CachedWorldBounds = FBox(ForceInitToZero);
+				
+				for (int32 LevelIndex = 0; LevelIndex < World->GetNumLevels(); ++LevelIndex)
+				{
+					if (ULevel* Level = World->GetLevel(LevelIndex))
+					{
+						if (ALevelBounds* LevelBounds = Level->LevelBoundsActor.Get())
+						{
+							CachedWorldBounds += LevelBounds->GetComponentsBoundingBox();
+						}
+						else
+						{
+							CachedWorldBounds = ALevelBounds::CalculateLevelBounds(Level);
+						}
+					}
+				}
+			}
+		}
+	}
+	
 	ConditionalCheckHoveredHitProxy();
 
 	FViewportCameraTransform& ViewTransform = GetViewTransform();
@@ -5691,6 +5730,11 @@ void FEditorViewportClient::SetOrthographicFarPlaneOverride(const TOptional<doub
 	OrthographicFarPlane = InValue;
 }
 
+FBox FEditorViewportClient::GetWorldBounds() const
+{
+	return CachedWorldBounds;
+}
+
 float FEditorViewportClient::GetSceneDepthAtLocation(int32 X, int32 Y)
 {
 	// #todo: in the future we will just sample the depth buffer
@@ -6917,6 +6961,29 @@ float FEditorViewportClient::GetMinimumOrthoZoom() const
 	return FMath::Max(GetDefault<ULevelEditorViewportSettings>()->MinimumOrthographicZoom, MIN_ORTHOZOOM);
 }
 
+void FEditorViewportClient::OnObjectTransacted(UObject* Object, const FTransactionObjectEvent& Event)
+{
+	if (Event.GetEventType() == ETransactionObjectEventType::Snapshot)
+	{
+		return;
+	}
+
+	if (Object->GetWorld() == GetWorld())
+	{
+		CachedWorldBounds.Init();
+	}
+}
+
+void FEditorViewportClient::OnPostLoadMap(uint32 MapChangeFlags)
+{
+	if (GWorld == GetWorld())
+	{
+		CachedWorldBounds.Init();
+		OrthographicFarPlane.Reset();
+		OrthographicNearPlane.Reset();
+	}
+}
+
 ////////////////
 
 bool FEditorViewportStats::bInitialized(false);
diff --git a/Engine/Source/Editor/UnrealEd/Private/EditorViewportCommands.cpp b/Engine/Source/Editor/UnrealEd/Private/EditorViewportCommands.cpp
--- a/Engine/Source/Editor/UnrealEd/Private/EditorViewportCommands.cpp
+++ b/Engine/Source/Editor/UnrealEd/Private/EditorViewportCommands.cpp
@@ -196,6 +196,8 @@ void FEditorViewportCommands::RegisterCommands()
 	UI_COMMAND( FocusOutlinerToSelection, "Frame Selected in Outliner", "Focuses the selected actors in all open outliners", EUserInterfaceActionType::Button, FInputChord() );
 	UI_COMMAND( FocusOutlinerToContextFolder, "Frame Current Folder in Outliner", "Focuses the current folder (if it's set) in the outliner", EUserInterfaceActionType::Button, FInputChord() );
 
+	UI_COMMAND( UseOrthographicClippingPlanes, "Ortho Clipping Planes", "Whether or not the depth of the orthographic view should be clipped.", EUserInterfaceActionType::Check, FInputChord(EKeys::O, EModifierKey::Shift) );
+
 	UI_COMMAND( LocationGridSnap, "Grid Snap", "Toggle whether objects snap to a grid when you move them.\nUse the submenu to choose the grid size.", EUserInterfaceActionType::ToggleButton, FInputChord() );
 	UI_COMMAND( RotationGridSnap, "Rotation Snap", "Toggle whether objects snap to preset angles when you rotate them.\nUse the submenu to choose the angle increment.", EUserInterfaceActionType::ToggleButton, FInputChord() );
 	UI_COMMAND( Layer2DSnap, "Layer2D Snap", "Enables or disables snapping objects to a 2D layer", EUserInterfaceActionType::ToggleButton, FInputChord());
diff --git a/Engine/Source/Editor/UnrealEd/Private/SEditorViewport.cpp b/Engine/Source/Editor/UnrealEd/Private/SEditorViewport.cpp
--- a/Engine/Source/Editor/UnrealEd/Private/SEditorViewport.cpp
+++ b/Engine/Source/Editor/UnrealEd/Private/SEditorViewport.cpp
@@ -588,13 +588,19 @@ void SEditorViewport::BindCommands()
 		FExecuteAction::CreateSP( this, &SEditorViewport::OnCycleCoordinateSystem )
 		);
 
-							
 	CommandListRef.MapAction( 
 		Commands.FocusViewportToSelection, 
 		FExecuteAction::CreateSP( this, &SEditorViewport::OnFocusViewportToSelection )
 		//FExecuteAction::CreateStatic( &FLevelEditorActionCallbacks::ExecuteExecCommand, FString( TEXT("CAMERA ALIGN ACTIVEVIEWPORTONLY") ) )
 		);
-
+		
+	CommandListRef.MapAction(
+		Commands.UseOrthographicClippingPlanes,
+		FExecuteAction::CreateSP( this, &SEditorViewport::OnToggleUseDepthBar ),
+		FCanExecuteAction::CreateSP( this, &SEditorViewport::CanUseDepthBar ),
+		FIsActionChecked::CreateSP( this, &SEditorViewport::IsDepthBarActive )
+		);
+	
 	CommandListRef.MapAction(
 		Commands.SurfaceSnapping,
 		FExecuteAction::CreateStatic( &SEditorViewport::OnToggleSurfaceSnap ),
@@ -1104,6 +1110,31 @@ void SEditorViewport::OnCycleCoordinateSystem()
 	Client->SetWidgetCoordSystemSpace( (ECoordSystem)CoordSystemAsInt );
 }
 
+void SEditorViewport::OnToggleUseDepthBar()
+{
+	if (Client->IsOrtho())
+	{
+		if (bShowDepthBar)
+		{
+			DismissDepthBar();
+		}
+		else
+		{
+			bShowDepthBar = true;
+		}
+	}
+}
+
+bool SEditorViewport::CanUseDepthBar() const
+{
+	return Client->IsOrtho();	
+}
+
+bool SEditorViewport::IsDepthBarActive() const
+{
+	return bShowDepthBar;
+}
+
 UWorld* SEditorViewport::GetWorld() const
 {
 	return Client->GetWorld();
diff --git a/Engine/Source/Editor/UnrealEd/Private/SEditorViewportDepthBar.cpp b/Engine/Source/Editor/UnrealEd/Private/SEditorViewportDepthBar.cpp
--- a/Engine/Source/Editor/UnrealEd/Private/SEditorViewportDepthBar.cpp
+++ b/Engine/Source/Editor/UnrealEd/Private/SEditorViewportDepthBar.cpp
@@ -81,31 +81,6 @@ void SEditorViewportDepthBar::Construct(const FArguments& InArgs, const TSharedP
 	];
 }
 
-void SEditorViewportDepthBar::Tick(const FGeometry& AllottedGeometry, const double InCurrentTime, const float InDeltaTime)
-{
-	if (TSharedPtr<FEditorViewportClient> Client = WeakClient.Pin())
-	{
-		UWorld* World = Client->GetWorld();
-		if (UWorldPartition* WorldPartition = World->GetWorldPartition())
-		{
-			WorldBounds = WorldPartition->GetEditorWorldBounds();
-		}
-		else
-		{
-			// TODO: Not this every frame, good lord
-			WorldBounds = FBox(EForceInit::ForceInitToZero);
-			for (TActorIterator<AActor> It(Client->GetWorld()); It; ++It)
-			{
-				AActor* Actor = *It;
-				if (IsValid(Actor))
-				{
-					WorldBounds += Actor->GetComponentsBoundingBox(true);
-				}
-			}	
-		}
-	}
-}
-
 FCursorReply SEditorViewportDepthBar::OnCursorQuery(const FGeometry& MyGeometry, const FPointerEvent& CursorEvent) const
 {
 	using namespace UE::Editor::ViewportDepthBar::Private;
@@ -141,7 +116,7 @@ const SEditorViewportDepthBar::FDepthIndicatorSlot& SEditorViewportDepthBar::Get
 	return DepthIndicatorChildren[Index];
 }
 
-void SEditorViewportDepthBar::ClearChildren()
+void SEditorViewportDepthBar::ClearDepthIndicators()
 {
 	DepthIndicatorChildren.Empty();
 }
@@ -192,7 +167,7 @@ FReply SEditorViewportDepthBar::OnMouseMove(const FGeometry& MyGeometry, const F
 					}
 					else
 					{
-						Client->SetOrthographicFarPlaneOverride(Positions->BarToRealPosition(ClampedBarPosition));
+						Client->SetOrthographicFarPlaneOverride(Positions->BarToAlignedPosition(ClampedBarPosition));
 					}
 					
 					Client->Invalidate();
@@ -207,7 +182,7 @@ FReply SEditorViewportDepthBar::OnMouseMove(const FGeometry& MyGeometry, const F
 					}
 					else
 					{
-						Client->SetOrthographicNearPlaneOverride(Positions->BarToRealPosition(ClampedBarPosition));
+						Client->SetOrthographicNearPlaneOverride(Positions->BarToAlignedPosition(ClampedBarPosition));
 					}
 					
 					Client->Invalidate();
@@ -218,8 +193,8 @@ FReply SEditorViewportDepthBar::OnMouseMove(const FGeometry& MyGeometry, const F
 					const float RawPosition = MousePosition.Y - DragOffset;
 					const float ClampedBarPosition = FMath::Clamp(RawPosition, Positions->BarMin, Positions->BarMax - (Positions->BarNearPlane - Positions->BarFarPlane));
 					
-					Client->SetOrthographicFarPlaneOverride(Positions->BarToRealPosition(ClampedBarPosition));
-					Client->SetOrthographicNearPlaneOverride(Positions->BarToRealPosition(ClampedBarPosition + (Positions->BarNearPlane - Positions->BarFarPlane)));
+					Client->SetOrthographicFarPlaneOverride(Positions->BarToAlignedPosition(ClampedBarPosition));
+					Client->SetOrthographicNearPlaneOverride(Positions->BarToAlignedPosition(ClampedBarPosition + (Positions->BarNearPlane - Positions->BarFarPlane)));
 					
 					Client->Invalidate();
 					return FReply::Handled();
@@ -289,8 +264,8 @@ void SEditorViewportDepthBar::OnArrangeChildren(const FGeometry& AllottedGeometr
 				FBoxSphereBounds Bounds = CurrentChild.GetBounds().Get();
 				const FSphere Sphere = Bounds.GetSphere();
 				
-				const FVector Max = Sphere.Center + BarPositions->DepthSpace.Forward * Sphere.W;
-				const FVector Min = Sphere.Center - BarPositions->DepthSpace.Forward * Sphere.W;
+				const FVector Max = Sphere.Center + BarPositions->DepthSpace.GetForward() * Sphere.W;
+				const FVector Min = Sphere.Center - BarPositions->DepthSpace.GetForward() * Sphere.W;
 					
 				const float BarMax = BarPositions->WorldToBarPosition(Max);
 				const float BarMin = BarPositions->WorldToBarPosition(Min);
@@ -434,33 +409,27 @@ TOptional<SEditorViewportDepthBar::FBarPositions> SEditorViewportDepthBar::GetBa
 		return TOptional<FBarPositions>();
 	}
 	
-	const FDepthSpace DepthSpace
+	const FBox WorldBounds = Client->GetWorldBounds();
+	if (!WorldBounds.IsValid)
 	{
-		.Origin = Client->GetViewTransform().GetLocation(),
-		.Forward = Client->GetForwardVector()
-	};
-	
-	const double WorldMaxAxis = DepthSpace.GetAxisPosition(WorldBounds.Max);
-	const double WorldMinAxis = DepthSpace.GetAxisPosition(WorldBounds.Min);
+		return TOptional<FBarPositions>();
+	}
 	
-	const double RealMax = WorldMaxAxis > 0.0 ? WorldMaxAxis : WorldMinAxis;
-	const double RealMin = WorldMinAxis < 0.0 ? WorldMinAxis : WorldMaxAxis;
+	const FDepthSpace DepthSpace(Client->GetViewTransform().GetLocation(), Client->GetForwardVector(), WorldBounds);
 	
-	FMargin Inset = Padding.Get();
+	const FMargin Inset = Padding.Get();
 	
 	if (Client->IsOrtho())
 	{
 		FBarPositions Positions
 		{
 			.DepthSpace = DepthSpace,
-			.RealMax = RealMax,
-			.RealMin = RealMin,
 			.BarMin = Inset.Top,
 			.BarMax = InGeometry.GetLocalSize().Y - Inset.Bottom
 		};
 		
-		Positions.BarFarPlane = Positions.RelativeToBarPosition(Client->GetOrthographicFarPlaneOverride().IsSet() ? Positions.AxisToRelativePosition(Client->GetOrthographicFarPlaneOverride().GetValue()) : 1.0);
-		Positions.BarNearPlane = Positions.RelativeToBarPosition(Client->GetOrthographicNearPlaneOverride().IsSet() ? Positions.AxisToRelativePosition(Client->GetOrthographicNearPlaneOverride().GetValue()) : 0.0f);
+		Positions.BarFarPlane = Positions.RelativeToBarPosition(Client->GetOrthographicFarPlaneOverride().IsSet() ? DepthSpace.AlignedToRelativePosition(Client->GetOrthographicFarPlaneOverride().GetValue()) : 1.0);
+		Positions.BarNearPlane = Positions.RelativeToBarPosition(Client->GetOrthographicNearPlaneOverride().IsSet() ? DepthSpace.AlignedToRelativePosition(Client->GetOrthographicNearPlaneOverride().GetValue()) : 0.0f);
 		
 		return Positions;
 	}
@@ -468,25 +437,72 @@ TOptional<SEditorViewportDepthBar::FBarPositions> SEditorViewportDepthBar::GetBa
 	return TOptional<FBarPositions>();
 }
 
-double SEditorViewportDepthBar::FDepthSpace::GetAxisPosition(const FVector& Position) const
+SEditorViewportDepthBar::FDepthSpace::FDepthSpace(const FVector& InOrigin, const FVector& InForward, const FBox& Bounds)
+{
+	Origin = InOrigin;
+	Forward = InForward;
+	
+	const float BoundsMax = GetAlignedPosition(Bounds.Max);
+	const float BoundsMin = GetAlignedPosition(Bounds.Min);
+	
+	if (BoundsMax > BoundsMin)
+	{
+		bFlipped = false;
+		AlignedMax = BoundsMax;
+		AlignedMin = BoundsMin;
+	}
+	else
+	{
+		bFlipped = true;
+		AlignedMax = BoundsMin;
+		AlignedMin = BoundsMax;
+	}
+}
+
+double SEditorViewportDepthBar::FDepthSpace::GetAlignedPosition(const FVector& Position) const
 {
 	const FVector OriginToPosition = Position - Origin;
 	const FVector ProjectedPosition = OriginToPosition.ProjectOnToNormal(Forward);
-	const double Dot = ProjectedPosition.Dot(Forward);
-	return ProjectedPosition.Length() * (Dot > 0.0 ? 1.0 : -1.0); 
+	return ProjectedPosition.Dot(Forward); 
 }
 
-float SEditorViewportDepthBar::FBarPositions::WorldToBarPosition(const FVector& WorldPosition) const
+double SEditorViewportDepthBar::FDepthSpace::AlignedToRelativePosition(double AlignedPosition) const
 {
-	const double AxisPosition = DepthSpace.GetAxisPosition(WorldPosition);
-	const double RelativePosition = AxisToRelativePosition(AxisPosition);
-	return RelativeToBarPosition(RelativePosition);
+	const double Size = AlignedMax - AlignedMin;
+	return Size != 0.0 ? ((AlignedPosition - AlignedMin) / Size) : 0.0;
+}
+
+double SEditorViewportDepthBar::FDepthSpace::RelativeToAlignedPosition(double RelativePosition) const
+{
+	return FMath::Lerp(AlignedMin, AlignedMax, RelativePosition);
+}
+
+double SEditorViewportDepthBar::FDepthSpace::RelativeToBoundsPosition(double RelativePosition) const
+{
+	const double Adjust = Origin.ProjectOnToNormal(Forward).Dot(Forward);
+	if (bFlipped)
+	{
+		return FMath::Lerp(AlignedMax, AlignedMin, RelativePosition) + Adjust;
+	}
+	return FMath::Lerp(AlignedMin, AlignedMax, RelativePosition) + Adjust;
+}
+
+double SEditorViewportDepthBar::FDepthSpace::BoundsToRelativePosition(double BoundsPosition) const
+{
+	const double Adjusted = BoundsPosition - Origin.ProjectOnToNormal(Forward).Dot(Forward);
+	if (bFlipped)
+	{
+		const double Size = AlignedMin - AlignedMax;
+		return Size != 0.0 ? ((Adjusted - AlignedMax) / Size) : 0.0;
+	}
+	return AlignedToRelativePosition(Adjusted);
 }
 
-double SEditorViewportDepthBar::FBarPositions::AxisToRelativePosition(double AxisPosition) const
+float SEditorViewportDepthBar::FBarPositions::WorldToBarPosition(const FVector& WorldPosition) const
 {
-	const double Size = RealMax - RealMin;
-	return Size != 0.0 ? ((AxisPosition - RealMin) / Size) : 0.0;
+	const double AxisPosition = DepthSpace.GetAlignedPosition(WorldPosition);
+	const double RelativePosition = DepthSpace.AlignedToRelativePosition(AxisPosition);
+	return RelativeToBarPosition(RelativePosition);
 }
 
 float SEditorViewportDepthBar::FBarPositions::RelativeToBarPosition(double RelativePosition) const
@@ -494,11 +510,11 @@ float SEditorViewportDepthBar::FBarPositions::RelativeToBarPosition(double Relat
 	return FMath::Clamp(1.0f - RelativePosition, 0.0f, 1.0f) * (BarMax - BarMin) + BarMin;
 }
 
-double SEditorViewportDepthBar::FBarPositions::BarToRealPosition(float BarPosition) const
+double SEditorViewportDepthBar::FBarPositions::BarToAlignedPosition(float BarPosition) const
 {
 	const float BarSize = BarMax - BarMin;
 	const float RelativePosition = BarSize > 0.0f ? (1.0f - (BarPosition - BarMin) / BarSize) : 0.0f;
-	return RelativePosition * (RealMax - RealMin) + RealMin;
+	return DepthSpace.RelativeToAlignedPosition(RelativePosition);
 }
 
 EVisibility SEditorViewportDepthBar::GetFarPlaneLabelVisibility() const
@@ -519,7 +535,13 @@ FText SEditorViewportDepthBar::GetFarPlaneText() const
 		{
 			if (const double* FarPlane = Client->GetOrthographicFarPlaneOverride().GetPtrOrNull())
 			{
-				return FText::Format(LOCTEXT("OrthoFarPlaneFormat", "Far Plane: {0}"), UE::Editor::ViewportDepthBar::Private::GetDistanceDisplayText(*FarPlane));
+				FBox Bounds = Client->GetWorldBounds();
+				FDepthSpace Space(Client->GetViewLocation(), Client->GetForwardVector(), Bounds);
+
+				const float RelativePosition = Space.AlignedToRelativePosition(*FarPlane);
+				const float WorldPosition = Space.RelativeToBoundsPosition(RelativePosition);	
+			
+				return FText::Format(LOCTEXT("OrthoFarPlaneFormat", "Far Plane: {0}"), UE::Editor::ViewportDepthBar::Private::GetDistanceDisplayText(WorldPosition));
 			}
 			return LOCTEXT("OrthoFarPlaneDefault", "Infinite Far Plane");
 		}
@@ -535,7 +557,13 @@ FText SEditorViewportDepthBar::GetNearPlaneText() const
 		{
 			if (const double* NearPlane = Client->GetOrthographicNearPlaneOverride().GetPtrOrNull())
 			{
-				return FText::Format(LOCTEXT("OrthoNearPlaneFormat", "Near Plane: {0}"), UE::Editor::ViewportDepthBar::Private::GetDistanceDisplayText(*NearPlane));
+				FBox Bounds = Client->GetWorldBounds();
+				FDepthSpace Space(Client->GetViewLocation(), Client->GetForwardVector(), Bounds);
+
+				const float RelativePosition = Space.AlignedToRelativePosition(*NearPlane);
+				const float WorldPosition = Space.RelativeToBoundsPosition(RelativePosition);			
+
+				return FText::Format(LOCTEXT("OrthoNearPlaneFormat", "Near Plane: {0}"), UE::Editor::ViewportDepthBar::Private::GetDistanceDisplayText(WorldPosition));
 			}
 			return LOCTEXT("OrthoNearPlaneDefault", "Infinite Near Plane");
 		}
diff --git a/Engine/Source/Editor/UnrealEd/Private/ViewportToolbar/UnrealEdViewportToolbar.cpp b/Engine/Source/Editor/UnrealEd/Private/ViewportToolbar/UnrealEdViewportToolbar.cpp
--- a/Engine/Source/Editor/UnrealEd/Private/ViewportToolbar/UnrealEdViewportToolbar.cpp
+++ b/Engine/Source/Editor/UnrealEd/Private/ViewportToolbar/UnrealEdViewportToolbar.cpp
@@ -20,6 +20,7 @@
 #include "LevelEditorActions.h"
 #include "RayTracingDebugVisualizationMenuCommands.h"
 #include "SEditorViewport.h"
+#include "SEditorViewportDepthBar.h"
 #include "SScalabilitySettings.h"
 #include "ShowFlagMenuCommands.h"
 #include "Settings/EditorProjectSettings.h"
@@ -772,6 +773,175 @@ void SetCameraSpeed(const TSharedRef<SEditorViewport>& InEditorViewport, float N
 	}
 }
 
+class SOrthographicViewPlaneSlider : public SCompoundWidget
+{
+public:
+	enum class EMode
+	{
+		Near,
+		Far
+	};
+
+	SLATE_BEGIN_ARGS(SOrthographicViewPlaneSlider) {}
+	SLATE_END_ARGS()
+	
+	void Construct(const FArguments& InArgs, TWeakPtr<FEditorViewportClient> InClient, EMode InMode)
+	{
+		WeakClient = InClient;
+		Mode = InMode;
+		
+		ChildSlot
+		[
+			SNew(SBox)
+			.Padding(FMargin(4.0f, 0.0f, 0.0f, 0.0f))
+			.WidthOverride(100.0f)
+			[
+				SNew(SBorder)
+				.BorderImage(FAppStyle::Get().GetBrush("Menu.WidgetBorder"))
+				.Padding(FMargin(1.0f))
+				[
+					SNew(SSpinBox<float>)
+					.Style(&FAppStyle::Get(), "Menu.SpinBox")
+					.Font(FAppStyle::GetFontStyle(TEXT("MenuItem.Font")))
+					.MinValue(this, &SOrthographicViewPlaneSlider::GetMin)
+					.MaxValue(this, &SOrthographicViewPlaneSlider::GetMax)
+					.OnGetDisplayValue(this, &SOrthographicViewPlaneSlider::GetDisplayValue)
+					.Value(this, &SOrthographicViewPlaneSlider::GetValue)
+					.OnValueChanged(this, &SOrthographicViewPlaneSlider::OnValueChanged)
+					.MaxFractionalDigits(2)
+				]
+			]
+		];
+	}
+	
+	virtual void Tick(const FGeometry& AllottedGeometry, const double InCurrentTime, const float InDeltaTime) override
+	{
+		if (TSharedPtr<FEditorViewportClient> Client = WeakClient.Pin())
+		{
+			const FBox Bounds = Client->GetWorldBounds();
+			if (Bounds.IsValid)
+			{
+				const SEditorViewportDepthBar::FDepthSpace DepthSpace = GetDepthSpace();
+				const double BoundsMin = DepthSpace.RelativeToBoundsPosition(0.0f);
+				const double BoundsMax = DepthSpace.RelativeToBoundsPosition(1.0f);
+				
+				Min = BoundsMin;
+				Max = BoundsMax;
+			}
+		}
+	
+		SCompoundWidget::Tick(AllottedGeometry, InCurrentTime, InDeltaTime);
+	}
+	
+private:
+	TOptional<float> GetMin() const
+	{
+		return Min < Max ? Min : Max;
+	}
+	
+	TOptional<float> GetMax() const
+	{
+		return Min < Max ? Max : Min;
+	}
+	
+	TOptional<double> GetRawValue() const
+	{
+		if (TSharedPtr<FEditorViewportClient> Client = WeakClient.Pin())
+		{
+			if (Mode == EMode::Far)
+			{
+				return Client->GetOrthographicFarPlaneOverride();
+			}
+			return Client->GetOrthographicNearPlaneOverride();
+		}
+	
+		return TOptional<double>();
+	}
+	
+	float GetValue() const
+	{
+		TOptional<double> ValueD = GetRawValue();
+		if (const double* Value = ValueD.GetPtrOrNull())
+		{
+			// The overrides are "aligned" offsets relative to the camera position
+			// we want to present them as world coordinates.
+			const SEditorViewportDepthBar::FDepthSpace DepthSpace = GetDepthSpace();
+			const double RelativePosition = DepthSpace.AlignedToRelativePosition(*Value);
+			return static_cast<float>(DepthSpace.RelativeToBoundsPosition(RelativePosition));
+		}
+		return Mode == EMode::Far ? Max : Min;
+	}
+	
+	void OnValueChanged(float NewValue)
+	{
+		if (TSharedPtr<FEditorViewportClient> Client = WeakClient.Pin())
+		{
+			const SEditorViewportDepthBar::FDepthSpace DepthSpace = GetDepthSpace();
+			const double RelativePosition = DepthSpace.BoundsToRelativePosition(NewValue);
+			
+			if (Mode == EMode::Near && FMath::IsNearlyEqual(RelativePosition, 0.0))
+			{
+				Client->SetOrthographicNearPlaneOverride(TOptional<double>());
+			}
+			else if (Mode == EMode::Far && FMath::IsNearlyEqual(RelativePosition, 1.0f))
+			{
+				Client->SetOrthographicFarPlaneOverride(TOptional<double>());
+			}
+			else
+			{
+				const double AlignedPosition = DepthSpace.RelativeToAlignedPosition(RelativePosition);
+			
+				if (Mode == EMode::Far)
+				{
+					Client->SetOrthographicFarPlaneOverride(AlignedPosition);
+					if (const double* NearPlane = Client->GetOrthographicNearPlaneOverride().GetPtrOrNull())
+					{
+						if (*NearPlane >= AlignedPosition)
+						{
+							Client->SetOrthographicNearPlaneOverride(AlignedPosition - 1.0f);
+						}
+					}
+				}
+				else
+				{
+					Client->SetOrthographicNearPlaneOverride(AlignedPosition);
+					if (const double* FarPlane = Client->GetOrthographicFarPlaneOverride().GetPtrOrNull())
+					{
+						if (*FarPlane <= AlignedPosition)
+						{
+							Client->SetOrthographicFarPlaneOverride(AlignedPosition + 1.0f);
+						}
+					}
+				}
+			}
+			
+			Client->Invalidate();
+		}
+	}
+	
+	SEditorViewportDepthBar::FDepthSpace GetDepthSpace() const
+	{
+		TSharedPtr<FEditorViewportClient> Client = WeakClient.Pin();
+		
+		return SEditorViewportDepthBar::FDepthSpace(Client->GetViewLocation(), Client->GetForwardVector(), Client->GetWorldBounds());
+	}
+	
+	TOptional<FText> GetDisplayValue(float Value) const
+	{
+		if (!GetRawValue().IsSet())
+		{
+			return TOptional<FText>(LOCTEXT("Infinity", "Infinity"));
+		}
+		return TOptional<FText>();
+	}
+
+	TWeakPtr<FEditorViewportClient> WeakClient;
+	EMode Mode = EMode::Near;
+	
+	float Min = -10000.0f;
+	float Max = 10000.0f;
+};
+
 } // namespace UE::UnrealEd::Private
 
 // The accessor functions are deprecated, but the functionality remains so that licensee implementers can swap back and forth easily while transitioning.
@@ -3505,6 +3675,7 @@ void PopulateCameraMenu(UToolMenu* InMenu, const FViewportCameraMenuOptions& InO
 	// View
 	{
 		FToolMenuSection& ViewSection = InMenu->FindOrAddSection("View", LOCTEXT("CameraViewSectionLabel", "View"));
+		const TAttribute<bool> IsPerspectiveAttribute = UE::UnrealEd::GetIsPerspectiveAttribute(EditorViewport->GetViewportClient());  
 
 		if (InOptions.bShowFieldOfView)
 		{
@@ -3515,11 +3686,10 @@ void PopulateCameraMenu(UToolMenu* InMenu, const FViewportCameraMenuOptions& InO
 			);
 			CameraFOV.Icon = FSlateIcon(FAppStyle::GetAppStyleSetName(), "ViewportToolbar.FieldOfView");
 			CameraFOV.ToolTip = LOCTEXT("CameraMovementTooltip", "Sets the field of view of the viewport's camera.");
+			CameraFOV.Visibility = IsPerspectiveAttribute;
 			ViewSection.AddEntry(CameraFOV);
 		}
 
-		ViewSection.Visibility = UE::UnrealEd::GetIsPerspectiveAttribute(EditorViewport->GetViewportClient());
-
 		if (InOptions.bShowNearAndFarPlanes)
 		{
 			FToolMenuEntry CameraNearViewPlane = FToolMenuEntry::InitWidget(
@@ -3528,6 +3698,7 @@ void PopulateCameraMenu(UToolMenu* InMenu, const FViewportCameraMenuOptions& InO
 				LOCTEXT("CameraSubmenu_NearViewPlaneLabel", "Near View Plane")
 			);
 			CameraNearViewPlane.Icon = FSlateIcon(FAppStyle::GetAppStyleSetName(), "ViewportToolbar.NearViewPlane");
+			CameraNearViewPlane.Visibility = IsPerspectiveAttribute;
 			ViewSection.AddEntry(CameraNearViewPlane);
 
 			FToolMenuEntry CameraFarViewPlane = FToolMenuEntry::InitWidget(
@@ -3536,10 +3707,17 @@ void PopulateCameraMenu(UToolMenu* InMenu, const FViewportCameraMenuOptions& InO
 				LOCTEXT("CameraSubmenu_FarViewPlaneLabel", "Far View Plane")
 			);
 			CameraFarViewPlane.Icon = FSlateIcon(FAppStyle::GetAppStyleSetName(), "ViewportToolbar.FarViewPlane");
+			CameraFarViewPlane.Visibility = IsPerspectiveAttribute;
 			ViewSection.AddEntry(CameraFarViewPlane);
+			
+			FToolMenuEntry& UseOrthoClippingPlanes = ViewSection.AddMenuEntry(FEditorViewportCommands::Get().UseOrthographicClippingPlanes);
+			UseOrthoClippingPlanes.Visibility = UE::UnrealEd::GetIsOrthographicAttribute(EditorViewport->GetViewportClient());
+					
+			FToolMenuEntry& OrthoSubmenu = ViewSection.AddEntry(CreateOrthographicClippingPlanesSubmenu());
+			OrthoSubmenu.Visibility = UE::UnrealEd::GetIsOrthographicAttribute(EditorViewport->GetViewportClient());
 		}
 	}
-
+	
 	// Auto Exposure
 	{
 		if (UE::UnrealEd::Private::ViewModesSubmenu::IsMenuSectionAvailable(
@@ -4116,6 +4294,40 @@ FText GetCameraSpeedTooltip()
 	return LOCTEXT("CameraSpeedSubMenu_ToolTip", "Set the camera speed.\nShortcut: Hold either mouse button and use the scroll wheel.");
 }
 
+FToolMenuEntry CreateOrthographicClippingPlanesSubmenu()
+{
+	return FToolMenuEntry::InitSubMenu(
+		"OrthographicClippingPlanes",
+		LOCTEXT("OrthographicClippingPlanesSubMenuLabel", "Ortho Clipping Planes"),
+		LOCTEXT("OrthographicClippingPlanesSubMenuTooltip", "Control the Orthographic Clipping Planes"),
+		FNewToolMenuDelegate::CreateLambda([](UToolMenu* Menu)
+		{
+			TSharedPtr<SEditorViewport> Viewport = UUnrealEdViewportToolbarContext::GetEditorViewport(Menu->Context);
+			if (!Viewport || !Viewport->GetViewportClient())
+			{
+				return;
+			}
+			
+			FToolMenuSection& ClippingSection = Menu->AddSection("ClippingPlanes", LOCTEXT("OrthographicClippingPlanesSectionLabel", "Ortho Clipping Planes"));
+			
+			TSharedRef<UE::UnrealEd::Private::SOrthographicViewPlaneSlider> FarSlider = SNew(
+				UE::UnrealEd::Private::SOrthographicViewPlaneSlider,
+				Viewport->GetViewportClient(),
+				UE::UnrealEd::Private::SOrthographicViewPlaneSlider::EMode::Far
+			);
+			
+			TSharedRef<UE::UnrealEd::Private::SOrthographicViewPlaneSlider> NearSlider = SNew(
+				UE::UnrealEd::Private::SOrthographicViewPlaneSlider,
+				Viewport->GetViewportClient(),
+				UE::UnrealEd::Private::SOrthographicViewPlaneSlider::EMode::Near
+			);
+			
+			ClippingSection.AddEntry(FToolMenuEntry::InitWidget("FarPlane", FarSlider, LOCTEXT("OrthoFarPlaneLabel", "Far Plane")));
+			ClippingSection.AddEntry(FToolMenuEntry::InitWidget("NearPlane", NearSlider, LOCTEXT("OrthoNearPlaneLabel", "Near Plane")));
+		})
+	);
+}
+
 bool ShouldShowViewportRealtimeWarning(const FEditorViewportClient& ViewportClient)
 {
 	// Almost all usages of viewport realtime overrides are for overriding realtime to be true.
@@ -4554,6 +4766,25 @@ TAttribute<bool> GetIsPerspectiveAttribute(const TSharedPtr<FEditorViewportClien
 	);
 }
 
+TAttribute<bool> GetIsOrthographicAttribute(const TSharedPtr<FEditorViewportClient>& InViewportClient)
+{
+	if (!InViewportClient.IsValid())
+	{
+		return false;
+	}
+
+	return TAttribute<bool>::CreateLambda(
+		[WeakViewport = InViewportClient.ToWeakPtr()]
+		{
+			if (const TSharedPtr<FEditorViewportClient>& ViewportClient = WeakViewport.Pin())
+			{
+				return ViewportClient->IsOrtho();
+			}
+			return false;
+		}
+	);
+}
+
 } // namespace UE::UnrealEd
 
 #undef LOCTEXT_NAMESPACE
diff --git a/Engine/Source/Editor/UnrealEd/Public/EditorViewportClient.h b/Engine/Source/Editor/UnrealEd/Public/EditorViewportClient.h
--- a/Engine/Source/Editor/UnrealEd/Public/EditorViewportClient.h
+++ b/Engine/Source/Editor/UnrealEd/Public/EditorViewportClient.h
@@ -19,6 +19,7 @@
 #include "Framework/Commands/Commands.h"
 #include "Editor.h"
 #include "ViewportClient.h"
+#include "Misc/TransactionObjectEvent.h"
 #include "Settings/EditorViewportSettings.h"
 
 struct FAssetData;
@@ -1340,6 +1341,9 @@ public:
 	
 	/** Set the orthographic far plane override for this viewport. */
 	UNREALED_API void SetOrthographicFarPlaneOverride(const TOptional<double>& InValue);
+	
+	/** Gets the bounds of the world being represented. */
+	UNREALED_API FBox GetWorldBounds() const;
 
 	/** When collision draw mode changes, this function allows hidden objects to be drawn, so hidden colliding objects can be seen */
 	UNREALED_API void UpdateHiddenCollisionDrawing();
@@ -2161,6 +2165,9 @@ public:
 
 private:
 
+	void OnObjectTransacted(UObject* Object, const FTransactionObjectEvent& Event);
+	void OnPostLoadMap(uint32 MapChangeFlags);
+
 	TSharedPtr<ICustomTemporalUpscalerData> CustomTemporalUpscalerData;
 
 private:
@@ -2206,7 +2213,10 @@ private:
 	
 	/** The location of the far clipping plane in orthographic view. */
 	TOptional<double> OrthographicFarPlane;
-
+	
+	/** The cached bounds of the world. */
+	FBox CachedWorldBounds;
+	
 	/** If true, we are in Game View mode*/
 	bool bInGameViewMode;
 
diff --git a/Engine/Source/Editor/UnrealEd/Public/EditorViewportCommands.h b/Engine/Source/Editor/UnrealEd/Public/EditorViewportCommands.h
--- a/Engine/Source/Editor/UnrealEd/Public/EditorViewportCommands.h
+++ b/Engine/Source/Editor/UnrealEd/Public/EditorViewportCommands.h
@@ -228,6 +228,8 @@ public:
 	TSharedPtr< FUICommandInfo > FocusViewportToSelection;
 	TSharedPtr< FUICommandInfo > FocusOutlinerToSelection;
 	TSharedPtr< FUICommandInfo > FocusOutlinerToContextFolder;
+	
+	TSharedPtr<FUICommandInfo> UseOrthographicClippingPlanes;
 
 	/** Toggle automatic exposure */
 	TSharedPtr< FUICommandInfo > ToggleAutoExposure;
diff --git a/Engine/Source/Editor/UnrealEd/Public/SEditorViewport.h b/Engine/Source/Editor/UnrealEd/Public/SEditorViewport.h
--- a/Engine/Source/Editor/UnrealEd/Public/SEditorViewport.h
+++ b/Engine/Source/Editor/UnrealEd/Public/SEditorViewport.h
@@ -270,6 +270,21 @@ protected:
 	 * Called when the user wants to focus the viewport to the current selection
 	 */
 	virtual void OnFocusViewportToSelection(){}
+	
+	/**
+	 * Called when enabling or disabling the depth bar.
+	 */
+	void OnToggleUseDepthBar();
+
+	/**
+	 * @return True if the depth bar can be enabled. 
+	 */
+	bool CanUseDepthBar() const;
+
+	/**
+	 * @return True if the depth bar is active 
+	 */
+	bool IsDepthBarActive() const;
 
 	/** Gets the world this viewport is for */
 	UNREALED_API virtual UWorld* GetWorld() const;
diff --git a/Engine/Source/Editor/UnrealEd/Public/SEditorViewportDepthBar.h b/Engine/Source/Editor/UnrealEd/Public/SEditorViewportDepthBar.h
--- a/Engine/Source/Editor/UnrealEd/Public/SEditorViewportDepthBar.h
+++ b/Engine/Source/Editor/UnrealEd/Public/SEditorViewportDepthBar.h
@@ -66,7 +66,7 @@ public:
 	
 	UNREALED_API FDepthIndicatorSlot& GetDepthIndicatorSlot(int32 Index);
 	UNREALED_API const FDepthIndicatorSlot& GetDepthIndicatorSlot(int32 Index) const;
-	UNREALED_API void ClearChildren();
+	UNREALED_API void ClearDepthIndicators();
 	
 	SEditorViewportDepthBar()
 		: DepthIndicatorChildren(this)
@@ -83,7 +83,6 @@ public:
 	
 	UNREALED_API void Construct(const FArguments& InArgs, const TSharedPtr<FEditorViewportClient>& Client);
 	
-	UNREALED_API virtual void Tick(const FGeometry& AllottedGeometry, const double InCurrentTime, const float InDeltaTime) override;
 	UNREALED_API virtual FCursorReply OnCursorQuery(const FGeometry& MyGeometry, const FPointerEvent& CursorEvent) const override;
 	UNREALED_API virtual void OnArrangeChildren(const FGeometry& AllottedGeometry, FArrangedChildren& ArrangedChildren) const override;
 	
@@ -92,28 +91,54 @@ public:
 	UNREALED_API virtual FReply OnMouseButtonUp(const FGeometry& MyGeometry, const FPointerEvent& MouseEvent) override;
 	UNREALED_API virtual void OnMouseLeave(const FPointerEvent& MouseEvent) override;
 	
+	struct FDepthSpace
+	{
+		FDepthSpace(const FVector& InOrigin, const FVector& InForward, const FBox& Bounds);
+		
+		const FVector& GetPosition() const { return Origin; }
+		const FVector& GetForward() const { return Forward; } 
+		
+		/** Gets the minimum value relative to originating direction */
+		double GetAlignedMin() const { return AlignedMin; }
+		/** Gets the maximum value relative to originating direction */
+		double GetAlignedMax() const { return AlignedMin; }
+		
+		/** Get the point of the position on the line starting at the origin, facing along the forward. */
+		double GetAlignedPosition(const FVector& Position) const;
+		
+		/** Get the relative value of the given axis-aligned position */
+		double AlignedToRelativePosition(double AlignedPosition) const;
+		
+		/** Get the axis-aligned value of the given relative position */
+		double RelativeToAlignedPosition(double RelativePosition) const;
+		
+		/** Get the bounds position of the given relative position */
+		double RelativeToBoundsPosition(double RelativePosition) const;
+		
+		/** Get the relative value of the given bounds position */
+		double BoundsToRelativePosition(double BoundsPosition) const;
+		
+	private:
+		FVector Origin;
+		FVector Forward;	
+		
+		double AlignedMin;
+		double AlignedMax;
+		
+		bool bFlipped;
+	};
+	
 protected:
 	UNREALED_API virtual FChildren* GetChildren() override;
 	UNREALED_API virtual int32 OnPaint(const FPaintArgs& Args, const FGeometry& AllottedGeometry, const FSlateRect& MyCullingRect, FSlateWindowElementList& OutDrawElements, int32 LayerId, const FWidgetStyle& InWidgetStyle, bool bParentEnabled) const override;
 	UNREALED_API virtual FVector2D ComputeDesiredSize(float LayoutScaleMultiplier) const override;
 	
 private:
-	struct FDepthSpace
-	{
-		FVector Origin;
-		FVector Forward;
-		
-		/** Get the point of the position on the line starting at the origin, facing along the forward. */
-		double GetAxisPosition(const FVector& Position) const;
-	};
-
+	
 	struct FBarPositions
 	{
 		FDepthSpace DepthSpace;
 	
-		double RealMax;
-		double RealMin;
-		
 		float BarMin;
 		float BarMax;
 		
@@ -121,9 +146,8 @@ private:
 		float BarNearPlane;
 		
 		float WorldToBarPosition(const FVector& WorldPosition) const;
-		double AxisToRelativePosition(double AxisPosition) const;
 		float RelativeToBarPosition(double RelativePosition) const;
-		double BarToRealPosition(float BarPosition) const;
+		double BarToAlignedPosition(float BarPosition) const;
 	};
 
 	DECLARE_DELEGATE_RetVal_OneParam(float, FGetBarPosition, const FBarPositions& BarPositions)
@@ -184,7 +208,6 @@ private:
 	FText GetFarPlaneText() const;
 	FText GetNearPlaneText() const;
 	
-	FBox WorldBounds;
 	TWeakPtr<FEditorViewportClient> WeakClient;
 	
 	TPanelChildren<FDepthIndicatorSlot> DepthIndicatorChildren;
diff --git a/Engine/Source/Editor/UnrealEd/Public/ViewportToolbar/UnrealEdViewportToolbar.h b/Engine/Source/Editor/UnrealEd/Public/ViewportToolbar/UnrealEdViewportToolbar.h
--- a/Engine/Source/Editor/UnrealEd/Public/ViewportToolbar/UnrealEdViewportToolbar.h
+++ b/Engine/Source/Editor/UnrealEd/Public/ViewportToolbar/UnrealEdViewportToolbar.h
@@ -284,6 +284,8 @@ UNREALED_API FToolMenuEntry CreateResetScalabilitySubmenu();
 
 UNREALED_API FText GetCameraSpeedTooltip();
 
+UNREALED_API FToolMenuEntry CreateOrthographicClippingPlanesSubmenu();
+
 /**
  * Returns Visible if the provided Viewport Client is in perspective view, Collapsed if not
  */
@@ -292,6 +294,10 @@ UNREALED_API TAttribute<EVisibility> GetPerspectiveOnlyVisibility(const TSharedP
  * Creates an attribute that provides whether the viewport client is in perspective view.
  */
 UNREALED_API TAttribute<bool> GetIsPerspectiveAttribute(const TSharedPtr<FEditorViewportClient>& InViewportClient);
+/**
+ * Creates an attribute that provides whether the viewport client is in an orthographic view.
+ */
+UNREALED_API TAttribute<bool> GetIsOrthographicAttribute(const TSharedPtr<FEditorViewportClient>& InViewportClient);
 
 // Camera Menu Widgets
 TSharedRef<SWidget> CreateCameraMenuWidget(const TSharedRef<SEditorViewport>& InViewport, bool bInShowExposureSettings = false);
