Fix issues found by BugHawk & Tests:
- Missing += for bounds accumulation.
- Returning Min value from Max() function.
- Remove default command colission.

#rnx

diff --git a/Engine/Source/Editor/UnrealEd/Private/EditorViewportClient.cpp b/Engine/Source/Editor/UnrealEd/Private/EditorViewportClient.cpp
--- a/Engine/Source/Editor/UnrealEd/Private/EditorViewportClient.cpp
+++ b/Engine/Source/Editor/UnrealEd/Private/EditorViewportClient.cpp
@@ -1634,7 +1634,7 @@ void FEditorViewportClient::Tick(float DeltaTime)
 						}
 						else
 						{
-							CachedWorldBounds = ALevelBounds::CalculateLevelBounds(Level);
+							CachedWorldBounds += ALevelBounds::CalculateLevelBounds(Level);
 						}
 					}
 				}
diff --git a/Engine/Source/Editor/UnrealEd/Private/EditorViewportCommands.cpp b/Engine/Source/Editor/UnrealEd/Private/EditorViewportCommands.cpp
--- a/Engine/Source/Editor/UnrealEd/Private/EditorViewportCommands.cpp
+++ b/Engine/Source/Editor/UnrealEd/Private/EditorViewportCommands.cpp
@@ -196,7 +196,7 @@ void FEditorViewportCommands::RegisterCommands()
 	UI_COMMAND( FocusOutlinerToSelection, "Frame Selected in Outliner", "Focuses the selected actors in all open outliners", EUserInterfaceActionType::Button, FInputChord() );
 	UI_COMMAND( FocusOutlinerToContextFolder, "Frame Current Folder in Outliner", "Focuses the current folder (if it's set) in the outliner", EUserInterfaceActionType::Button, FInputChord() );
 
-	UI_COMMAND( UseOrthographicClippingPlanes, "Ortho Clipping Planes", "Whether or not the depth of the orthographic view should be clipped.", EUserInterfaceActionType::Check, FInputChord(EKeys::O, EModifierKey::Shift) );
+	UI_COMMAND( UseOrthographicClippingPlanes, "Ortho Clipping Planes", "Whether or not the depth of the orthographic view should be clipped.", EUserInterfaceActionType::Check, FInputChord() );
 
 	UI_COMMAND( LocationGridSnap, "Grid Snap", "Toggle whether objects snap to a grid when you move them.\nUse the submenu to choose the grid size.", EUserInterfaceActionType::ToggleButton, FInputChord() );
 	UI_COMMAND( RotationGridSnap, "Rotation Snap", "Toggle whether objects snap to preset angles when you rotate them.\nUse the submenu to choose the angle increment.", EUserInterfaceActionType::ToggleButton, FInputChord() );
diff --git a/Engine/Source/Editor/UnrealEd/Public/SEditorViewportDepthBar.h b/Engine/Source/Editor/UnrealEd/Public/SEditorViewportDepthBar.h
--- a/Engine/Source/Editor/UnrealEd/Public/SEditorViewportDepthBar.h
+++ b/Engine/Source/Editor/UnrealEd/Public/SEditorViewportDepthBar.h
@@ -101,7 +101,7 @@ public:
 		/** Gets the minimum value relative to originating direction */
 		double GetAlignedMin() const { return AlignedMin; }
 		/** Gets the maximum value relative to originating direction */
-		double GetAlignedMax() const { return AlignedMin; }
+		double GetAlignedMax() const { return AlignedMax; }
 		
 		/** Get the point of the position on the line starting at the origin, facing along the forward. */
 		double GetAlignedPosition(const FVector& Position) const;
