Added a namespace to the transform provider class asset filter help, to prevent symbol clashes

diff --git a/Engine/Source/Editor/UnrealEd/Private/Factories/TransformProviderFactory.cpp b/Engine/Source/Editor/UnrealEd/Private/Factories/TransformProviderFactory.cpp
--- a/Engine/Source/Editor/UnrealEd/Private/Factories/TransformProviderFactory.cpp
+++ b/Engine/Source/Editor/UnrealEd/Private/Factories/TransformProviderFactory.cpp
@@ -32,6 +32,9 @@ UTransformProviderDataFactory::UTransformProviderDataFactory(const FObjectInitia
 	ProviderDataClass = nullptr;
 }
 
+namespace TransformProvider
+{
+
 class FAssetClassParentFilter : public IClassViewerFilter
 {
 public:
@@ -77,6 +80,8 @@ public:
 	}
 };
 
+} // namespace
+
 bool UTransformProviderDataFactory::ConfigureProperties()
 {
 	// Null the ProviderDataClass so we can get a clean class
@@ -89,7 +94,7 @@ bool UTransformProviderDataFactory::ConfigureProperties()
 	FClassViewerInitializationOptions Options;
 	Options.Mode = EClassViewerMode::ClassPicker;
 
-	TSharedPtr<FAssetClassParentFilter> Filter = MakeShareable(new FAssetClassParentFilter);
+	TSharedPtr<TransformProvider::FAssetClassParentFilter> Filter = MakeShareable(new TransformProvider::FAssetClassParentFilter);
 	Options.ClassFilters.Add(Filter.ToSharedRef());
 
 	Filter->DisallowedClassFlags = CLASS_Abstract | CLASS_Deprecated | CLASS_NewerVersionExists;
