MPCook: Fix assertion failure in MPCook: CachedObjectsInOuterInfo.IsEmpty() in FCookGenerationInfo::Uninitialize.
FCookWorkerServer::RecordResults needs to correctly handle cases of the CookDirector starting on the save of a generated package, but then retracting it from itself and sending it to a CookWorker. That scenario puts the packaget into the SaveStalledAssignedToWorker with the incomplete results from the save still loaded. RecordResults needs to make sure FCookGenerationHelper::ResetSaveState is called for all generated packagedatas to handle that case.
#rnx

diff --git a/Engine/Source/Editor/UnrealEd/Private/Cooker/CookWorkerServer.cpp b/Engine/Source/Editor/UnrealEd/Private/Cooker/CookWorkerServer.cpp
--- a/Engine/Source/Editor/UnrealEd/Private/Cooker/CookWorkerServer.cpp
+++ b/Engine/Source/Editor/UnrealEd/Private/Cooker/CookWorkerServer.cpp
@@ -1186,6 +1186,14 @@ void FCookWorkerServer::RecordResults(FPackageResultsMessage& Message)
 						ParentGenerationHelper->MarkPackageSavedRemotely(COTFS, *PackageData, TargetPlatform, GetWorkerId());
 					}
 				}
+				// ReleaseCookedPlatformData will called when we leave the SaveState during PromoteToSaveComplete or
+				// SendToState below, but we need to call it now before we SetParentGenerationHelper(nullptr), so that
+				// it can call ParentGenerationHelper->ResetSaveState. It is safe to call twice, so call it now and it
+				// will be a noop during SaveState. Call with EStateChangeReason::DoneForNow and
+				// EPackageState::AssignedToWorker so that it does the minimum shutdown for leaving the SaveState
+				// and lets us handle the work to return to idle ourselves in our required order.
+				COTFS.ReleaseCookedPlatformData(*PackageData, EStateChangeReason::DoneForNow,
+					EPackageState::AssignedToWorker);
 				PackageData->SetParentGenerationHelper(nullptr, StateChangeReason);
 			}
 		}
