[UnrealEd] Added optional ClampToBounds toggle to FViewportCursorLocation, defaulting to true (current behavior). Allows calculations based on cursor location beyond the viewport bounds

diff --git a/Engine/Source/Editor/UnrealEd/Private/EditorViewportClient.cpp b/Engine/Source/Editor/UnrealEd/Private/EditorViewportClient.cpp
--- a/Engine/Source/Editor/UnrealEd/Private/EditorViewportClient.cpp
+++ b/Engine/Source/Editor/UnrealEd/Private/EditorViewportClient.cpp
@@ -237,11 +237,13 @@ public:
 //	coordinate system for the viewport.
 //
 ///////////////////////////////////////////////////////////////////////////////////////////////////
-FViewportCursorLocation::FViewportCursorLocation(const FSceneView* View, FEditorViewportClient* InViewportClient, int32 X, int32 Y)
+FViewportCursorLocation::FViewportCursorLocation(const FSceneView* View, FEditorViewportClient* InViewportClient, int32 X, int32 Y, const bool bInClampToBounds)
 	: Origin(ForceInit), Direction(ForceInit), CursorPos(X, Y)
 {
-
-	FVector4 ScreenPos = View->CursorToScreen(static_cast<float>(X), static_cast<float>(Y), 0);
+	const FVector4 ScreenPos =
+		bInClampToBounds
+		? View->CursorToScreen(static_cast<float>(X), static_cast<float>(Y), 0)
+		: View->PixelToScreen(static_cast<float>(X), static_cast<float>(Y), 0);
 
 	const FMatrix InvViewMatrix = View->ViewMatrices.GetInvViewMatrix();
 	const FMatrix InvProjMatrix = View->ViewMatrices.GetInvProjectionMatrix();
diff --git a/Engine/Source/Editor/UnrealEd/Private/Tools/EdModeInteractiveToolsContext.cpp b/Engine/Source/Editor/UnrealEd/Private/Tools/EdModeInteractiveToolsContext.cpp
--- a/Engine/Source/Editor/UnrealEd/Private/Tools/EdModeInteractiveToolsContext.cpp
+++ b/Engine/Source/Editor/UnrealEd/Private/Tools/EdModeInteractiveToolsContext.cpp
@@ -51,6 +51,17 @@ DEFINE_LOG_CATEGORY_STATIC(LogInteractiveToolsContext, Log, All);
 
 #define LOCTEXT_NAMESPACE "EdModeInteractiveToolsContext"
 
+namespace UE::Editor::Private
+{
+	static TAutoConsoleVariable<int32> CVarUnboundedCursorDrag(
+		TEXT("Editor.UnboundedCursorDrag"),
+		0,
+		TEXT("Whether to allow (ITF) mouse drag operations to continue beyond the viewport bounds.\n")
+		TEXT("0: Do not allow mouse drag operations to continue beyond the viewport bounds (default)\n")
+		TEXT("1: Allow mouse drag operations to continue beyond the viewport bounds"),
+		ECVF_Default);
+}
+
 class FEdModeToolsContextQueriesImpl : public IToolsContextQueriesAPI
 {
 public:
@@ -696,6 +707,12 @@ bool UEditorInteractiveToolsContext::ProcessEditDelete()
 
 
 FRay UEditorInteractiveToolsContext::GetRayFromMousePos(FEditorViewportClient* ViewportClient, FViewport* Viewport, int MouseX, int MouseY)
+{
+	constexpr bool bClampMouseToBounds = true;
+	return GetRayFromMousePos(ViewportClient, Viewport, MouseX, MouseY, bClampMouseToBounds);
+}
+
+FRay UEditorInteractiveToolsContext::GetRayFromMousePos(FEditorViewportClient* ViewportClient, FViewport* Viewport, int MouseX, int MouseY, const bool bInClampMouseToBounds)
 {
 	if (!ViewportClient)
 	{
@@ -708,7 +725,7 @@ FRay UEditorInteractiveToolsContext::GetRayFromMousePos(FEditorViewportClient* V
 		ViewportClient->EngineShowFlags).SetRealtimeUpdate(ViewportClient->IsRealtime()));		// why SetRealtimeUpdate here??
 	// this View is deleted by the FSceneViewFamilyContext destructor
 	FSceneView* View = ViewportClient->CalcSceneView(&ViewFamily);
-	FViewportCursorLocation MouseViewportRay(View, ViewportClient, MouseX, MouseY);
+	FViewportCursorLocation MouseViewportRay(View, ViewportClient, MouseX, MouseY, bInClampMouseToBounds);
 
 	FVector RayOrigin = MouseViewportRay.GetOrigin();
 	FVector RayDirection = MouseViewportRay.GetDirection();
@@ -1155,11 +1172,14 @@ bool UModeManagerInteractiveToolsContext::StartTracking(FEditorViewportClient* I
 
 bool UModeManagerInteractiveToolsContext::CapturedMouseMove(FEditorViewportClient* InViewportClient, FViewport* InViewport, int32 InMouseX, int32 InMouseY)
 {
+	const bool bClampMouseToBoundsWhenCaptured = UE::Editor::Private::CVarUnboundedCursorDrag.GetValueOnGameThread() == 0;
+	const bool bHasActiveMouseCapture = InputRouter->HasActiveMouseCapture();
+
 	FVector2D OldPosition = CurrentMouseState.Mouse.Position2D;
 	CurrentMouseState.Mouse.Position2D = FVector2D(InMouseX, InMouseY);
-	CurrentMouseState.Mouse.WorldRay = GetRayFromMousePos(InViewportClient, InViewport, InMouseX, InMouseY);
+	CurrentMouseState.Mouse.WorldRay = GetRayFromMousePos(InViewportClient, InViewport, InMouseX, InMouseY, bHasActiveMouseCapture && bClampMouseToBoundsWhenCaptured);
 
-	if (InputRouter->HasActiveMouseCapture())
+	if (bHasActiveMouseCapture)
 	{
 #ifdef ENABLE_DEBUG_PRINTING
 		UE_LOG(LogInteractiveToolsContext, Warning, TEXT("CAPTURED MOUSE MOVE"));
diff --git a/Engine/Source/Editor/UnrealEd/Public/EditorViewportClient.h b/Engine/Source/Editor/UnrealEd/Public/EditorViewportClient.h
--- a/Engine/Source/Editor/UnrealEd/Public/EditorViewportClient.h
+++ b/Engine/Source/Editor/UnrealEd/Public/EditorViewportClient.h
@@ -170,7 +170,7 @@ private:
 struct FViewportCursorLocation
 {
 public:
-	UNREALED_API FViewportCursorLocation( const FSceneView* View, FEditorViewportClient* InViewportClient, int32 X, int32 Y );
+	UNREALED_API FViewportCursorLocation( const FSceneView* View, FEditorViewportClient* InViewportClient, int32 X, int32 Y, const bool bInClampToBounds = true );
 	UNREALED_API virtual ~FViewportCursorLocation();
 
 	const FVector&		GetOrigin()			const	{ return Origin; }
diff --git a/Engine/Source/Editor/UnrealEd/Public/Tools/EdModeInteractiveToolsContext.h b/Engine/Source/Editor/UnrealEd/Public/Tools/EdModeInteractiveToolsContext.h
--- a/Engine/Source/Editor/UnrealEd/Public/Tools/EdModeInteractiveToolsContext.h
+++ b/Engine/Source/Editor/UnrealEd/Public/Tools/EdModeInteractiveToolsContext.h
@@ -235,6 +235,10 @@ protected:
 	// Copy-pasted from other Editor code, seems kind of expensive?
 	static UE_API FRay GetRayFromMousePos(FEditorViewportClient* ViewportClient, FViewport* Viewport, int MouseX, int MouseY);
 
+	// Utility function to convert viewport x/y from mouse events (and others?) into scene ray.
+	// Copy-pasted from other Editor code, seems kind of expensive?
+	static UE_API FRay GetRayFromMousePos(FEditorViewportClient* ViewportClient, FViewport* Viewport, int MouseX, int MouseY, const bool bInClampMouseToBounds);
+
 	// editor UI state that we set before starting tool and when exiting tool
 	// Currently disabling anti-aliasing during active Tools because it causes PDI flickering
 	UE_API void SetEditorStateForTool();
