[Level Editor] Ensure Focused Viewport Client is the one currently being interacted with

FLevelEditorViewportClient::InputKey updates the currently focused viewport, and calls LostFocus on previously focused viewport.

FLevelEditorViewportClient::LostFocus makes sure viewport loses focus only if it's actually focused.

This prevents issues with ITF captures being terminated right after they began, and helps preventing crashes when an un-focused, uninitialized viewport is used instead of the viewport the user is currently interacting with

diff --git a/Engine/Source/Editor/UnrealEd/Private/LevelEditorViewport.cpp b/Engine/Source/Editor/UnrealEd/Private/LevelEditorViewport.cpp
--- a/Engine/Source/Editor/UnrealEd/Private/LevelEditorViewport.cpp
+++ b/Engine/Source/Editor/UnrealEd/Private/LevelEditorViewport.cpp
@@ -2779,7 +2779,12 @@ void FLevelEditorViewportClient::ReceivedFocus(FViewport* InViewport)
 
 void FLevelEditorViewportClient::LostFocus(FViewport* InViewport)
 {
-	FEditorViewportClient::LostFocus(InViewport);
+	// Make sure we are calling LostFocus on currently focused viewport
+	// We don't need to call LostFocus multiple times on the same viewport if it's already un-focused.
+	if (InViewport == ModeTools->GetFocusedViewportClient()->Viewport)
+	{
+		FEditorViewportClient::LostFocus(InViewport);
+	}
 
 	GEditor->SetPreviewMeshMode(false);
 }
@@ -3492,6 +3497,18 @@ bool FLevelEditorViewportClient::InputKey(const FInputKeyEventArgs& InEventArgs)
 	// Compute the click location.
 	if ( InputState.IsMouseButtonEvent() && InputState.IsAnyMouseButtonDown() )
 	{
+		// This Viewport Client received a click press, let's retrieve the currently focused viewport client.
+		// If that's not This viewport client, let's make sure This is focused, instead of the currently focused one.
+		FEditorViewportClient* CurrentlyFocusedViewport = ModeTools->GetFocusedViewportClient();
+		if (CurrentlyFocusedViewport != this)
+		{
+			// Call LostFocus on the Viewport Client losing focus
+			LostFocus(CurrentlyFocusedViewport->Viewport);
+
+			// Focusing on This viewport
+			ReceivedFocus(Viewport);
+		}
+
 		const FViewportCursorLocation Cursor(View, this, HitX, HitY);
 		const FActorPositionTraceResult TraceResult = FActorPositioning::TraceWorldForPositionWithDefault(Cursor, *View);
 		GEditor->UnsnappedClickLocation = TraceResult.Location;
