Added a wrap text property to a text block used by the Asset Drag And Drop Operation in the editor - this fixes a case where the widget would not be visible to the user when the message was too long

diff --git a/Engine/Source/Editor/UnrealEd/Private/DragAndDrop/AssetDragDropOp.cpp b/Engine/Source/Editor/UnrealEd/Private/DragAndDrop/AssetDragDropOp.cpp
--- a/Engine/Source/Editor/UnrealEd/Private/DragAndDrop/AssetDragDropOp.cpp
+++ b/Engine/Source/Editor/UnrealEd/Private/DragAndDrop/AssetDragDropOp.cpp
@@ -205,7 +205,8 @@ TSharedPtr<SWidget> FAssetDragDropOp::GetDefaultDecorator() const
 							.Padding(0,0,8,0)
 							.VAlign(VAlign_Center)
 							[
-								SNew(STextBlock) 
+								SNew(STextBlock)
+								.WrapTextAt(800.f)
 								.Text(this, &FAssetDragDropOp::GetDecoratorText)
 							]
 						]
@@ -289,7 +290,8 @@ TSharedPtr<SWidget> FAssetDragDropOp::GetDefaultDecorator() const
 					.Padding(0,0,3,0)
 					.VAlign(VAlign_Center)
 					[
-						SNew(STextBlock) 
+						SNew(STextBlock)
+						.WrapTextAt(800.f)
 						.Text(this, &FAssetDragDropOp::GetDecoratorText)
 					]
 				]
