[Outliner] Add visibility function versions that work on the selected actor's hierarchy and set H to default to toggle selected hierarchy

diff --git a/Engine/Source/Editor/LevelEditor/Private/LevelEditor.cpp b/Engine/Source/Editor/LevelEditor/Private/LevelEditor.cpp
--- a/Engine/Source/Editor/LevelEditor/Private/LevelEditor.cpp
+++ b/Engine/Source/Editor/LevelEditor/Private/LevelEditor.cpp
@@ -1118,6 +1118,24 @@ void FLevelEditorModule::BindGlobalLevelEditorCommands()
 		FCanExecuteAction::CreateStatic( &FLevelEditorActionCallbacks::CanToggleSelectedActorsVisibility )
 		);
 
+	ActionList.MapAction(
+		Commands.ShowSelectedHierarchy,
+		FExecuteAction::CreateStatic( &FLevelEditorActionCallbacks::ExecuteExecCommand, FString( TEXT("ACTOR UNHIDE SELECTED HIERARCHY") ) ),
+		FCanExecuteAction::CreateStatic( &FLevelEditorActionCallbacks::CanShowSelectedActors )
+		);
+
+	ActionList.MapAction(
+		Commands.HideSelectedHierarchy,
+		FExecuteAction::CreateStatic( &FLevelEditorActionCallbacks::ExecuteExecCommand, FString( TEXT("ACTOR HIDE SELECTED HIERARCHY") ) ),
+		FCanExecuteAction::CreateStatic( &FLevelEditorActionCallbacks::CanHideSelectedActors )
+		);
+
+	ActionList.MapAction(
+		Commands.ToggleSelectedHierarchyVisibility,
+		FExecuteAction::CreateStatic( &FLevelEditorActionCallbacks::ExecuteExecCommand, FString( TEXT("ACTOR TOGGLE SELECTED VISIBILITY HIERARCHY") ) ),
+		FCanExecuteAction::CreateStatic( &FLevelEditorActionCallbacks::CanToggleSelectedActorsVisibility )
+		);
+
 	ActionList.MapAction(
 		Commands.ShowAllStartup,
 		FExecuteAction::CreateStatic( &FLevelEditorActionCallbacks::ExecuteExecCommand, FString( TEXT("ACTOR UNHIDE ALL STARTUP") ) )
@@ -1133,6 +1151,16 @@ void FLevelEditorModule::BindGlobalLevelEditorCommands()
 		FExecuteAction::CreateStatic( &FLevelEditorActionCallbacks::ExecuteExecCommand, FString( TEXT("ACTOR HIDE SELECTED STARTUP") ) )
 		);
 
+	ActionList.MapAction(
+		Commands.ShowSelectedHierarchyStartup,
+		FExecuteAction::CreateStatic( &FLevelEditorActionCallbacks::ExecuteExecCommand, FString( TEXT("ACTOR UNHIDE SELECTED STARTUP HIERARCHY") ) )
+		);
+		
+	ActionList.MapAction(
+        Commands.HideSelectedHierarchyStartup,
+        FExecuteAction::CreateStatic( &FLevelEditorActionCallbacks::ExecuteExecCommand, FString( TEXT("ACTOR HIDE SELECTED STARTUP HIERARCHY") ) )
+        );
+
 	ActionList.MapAction(
 		Commands.CycleNavigationDataDrawn,
 		FExecuteAction::CreateStatic( &FLevelEditorActionCallbacks::ExecuteExecCommand, FString( TEXT("CYCLENAVDRAWN") ) )
diff --git a/Engine/Source/Editor/LevelEditor/Private/LevelEditorActions.cpp b/Engine/Source/Editor/LevelEditor/Private/LevelEditorActions.cpp
--- a/Engine/Source/Editor/LevelEditor/Private/LevelEditorActions.cpp
+++ b/Engine/Source/Editor/LevelEditor/Private/LevelEditorActions.cpp
@@ -4308,12 +4308,17 @@ void FLevelEditorCommands::RegisterCommands()
 	UI_COMMAND( ShowAll, "Show All Actors", "Shows all actors", EUserInterfaceActionType::Button, FInputChord( EModifierKey::Control, EKeys::H ) );
 #endif
 	UI_COMMAND( ShowSelectedOnly, "Show Only Selected", "Shows only the selected actors", EUserInterfaceActionType::Button, FInputChord() );
-	UI_COMMAND( ShowSelected, "Show Selected", "Shows the selected actors", EUserInterfaceActionType::Button, FInputChord( EModifierKey::Shift, EKeys::H ) );
+	UI_COMMAND( ShowSelected, "Show Selected", "Shows the selected actors", EUserInterfaceActionType::Button, FInputChord() );
 	UI_COMMAND( HideSelected, "Hide Selected", "Hides the selected actors", EUserInterfaceActionType::Button, FInputChord() );
-	UI_COMMAND( ToggleSelectedVisibility, "Toggle Selected Visibility", "Hides or shows the selected actors", EUserInterfaceActionType::Button, FInputChord( EKeys::H ) );
+	UI_COMMAND( ToggleSelectedVisibility, "Toggle Selected Visibility", "Hides or shows the selected actors", EUserInterfaceActionType::Button, FInputChord() );
+	UI_COMMAND( ShowSelectedHierarchy, "Show Selected Hierarchy", "Shows the selected actors and their children", EUserInterfaceActionType::Button, FInputChord( EModifierKey::Shift, EKeys::H ) );
+	UI_COMMAND( HideSelectedHierarchy, "Hide Selected Hierarchy", "Hides the selected actors and their children", EUserInterfaceActionType::Button, FInputChord() );
+	UI_COMMAND( ToggleSelectedHierarchyVisibility, "Toggle Selected Hierarchy Visibility", "Hides or shows the selected actors and their children", EUserInterfaceActionType::Button, FInputChord( EKeys::H ) );
 	UI_COMMAND( ShowAllStartup, "Show All At Startup", "Shows all actors at startup", EUserInterfaceActionType::Button, FInputChord() );
 	UI_COMMAND( ShowSelectedStartup, "Show Selected At Startup", "Shows selected actors at startup", EUserInterfaceActionType::Button, FInputChord() );
 	UI_COMMAND( HideSelectedStartup, "Hide Selected At Startup", "Hide selected actors at startup", EUserInterfaceActionType::Button, FInputChord() );
+	UI_COMMAND( ShowSelectedHierarchyStartup, "Show Selected Hierarchy At Startup", "Shows selected actors and their current children at startup", EUserInterfaceActionType::Button, FInputChord() );
+	UI_COMMAND( HideSelectedHierarchyStartup, "Hide Selected Hierarchy At Startup", "Hide selected actors and their current children at startup", EUserInterfaceActionType::Button, FInputChord() );
 	UI_COMMAND( CycleNavigationDataDrawn, "Cycle Navigation Data Drawn", "Cycles through navigation data (navmeshes for example) to draw one at a time", EUserInterfaceActionType::Button, FInputChord( EModifierKey::Alt, EKeys::N ) );
 
 	UI_COMMAND( SelectNone, "Unselect All", "Unselects all actors", EUserInterfaceActionType::Button, FInputChord( EKeys::Escape ) ) ;
diff --git a/Engine/Source/Editor/LevelEditor/Private/LevelEditorContextMenu.cpp b/Engine/Source/Editor/LevelEditor/Private/LevelEditorContextMenu.cpp
--- a/Engine/Source/Editor/LevelEditor/Private/LevelEditorContextMenu.cpp
+++ b/Engine/Source/Editor/LevelEditor/Private/LevelEditorContextMenu.cpp
@@ -1074,6 +1074,9 @@ void FLevelEditorContextMenuImpl::FillActorVisibilityMenu(UToolMenu* Menu)
 		Section.AddMenuEntry(FLevelEditorCommands::Get().ToggleSelectedVisibility);
 		Section.AddMenuEntry(FLevelEditorCommands::Get().HideSelected);
 		Section.AddMenuEntry(FLevelEditorCommands::Get().ShowSelected);
+		Section.AddMenuEntry(FLevelEditorCommands::Get().ToggleSelectedHierarchyVisibility);
+		Section.AddMenuEntry(FLevelEditorCommands::Get().HideSelectedHierarchy);
+		Section.AddMenuEntry(FLevelEditorCommands::Get().ShowSelectedHierarchy);
 	}
 
 	{
@@ -1087,6 +1090,8 @@ void FLevelEditorContextMenuImpl::FillActorVisibilityMenu(UToolMenu* Menu)
 		Section.AddMenuEntry(FLevelEditorCommands::Get().ShowAllStartup);
 		Section.AddMenuEntry(FLevelEditorCommands::Get().ShowSelectedStartup);
 		Section.AddMenuEntry(FLevelEditorCommands::Get().HideSelectedStartup);
+		Section.AddMenuEntry(FLevelEditorCommands::Get().ShowSelectedHierarchyStartup);
+		Section.AddMenuEntry(FLevelEditorCommands::Get().HideSelectedHierarchyStartup);
 	}
 }
 
diff --git a/Engine/Source/Editor/LevelEditor/Public/LevelEditorActions.h b/Engine/Source/Editor/LevelEditor/Public/LevelEditorActions.h
--- a/Engine/Source/Editor/LevelEditor/Public/LevelEditorActions.h
+++ b/Engine/Source/Editor/LevelEditor/Public/LevelEditorActions.h
@@ -363,6 +363,15 @@ public:
 	/** Hides or shows selected actors */
 	TSharedPtr< FUICommandInfo > ToggleSelectedVisibility;
 
+	/** Unhides selected actors and their children */
+	TSharedPtr< FUICommandInfo > ShowSelectedHierarchy;
+	
+	/** Hides selected actors and their children */
+	TSharedPtr< FUICommandInfo > HideSelectedHierarchy;
+
+	/** Hides or shows selected actors and their children */
+	TSharedPtr< FUICommandInfo > ToggleSelectedHierarchyVisibility;
+
 	/** Shows all actors at startup */
 	TSharedPtr< FUICommandInfo > ShowAllStartup;
 
@@ -372,6 +381,12 @@ public:
 	/** Hides selected actors at startup */
 	TSharedPtr< FUICommandInfo > HideSelectedStartup;
 
+	/** Shows selected actors and their current children at startup */
+	TSharedPtr< FUICommandInfo > ShowSelectedHierarchyStartup;
+
+	/** Hides selected actors and their current children at startup */
+	TSharedPtr< FUICommandInfo > HideSelectedHierarchyStartup;
+
 	/** Cycles through all navigation data to show one at a time */
 	TSharedPtr< FUICommandInfo > CycleNavigationDataDrawn;
 
diff --git a/Engine/Source/Editor/UnrealEd/Classes/Editor/UnrealEdEngine.h b/Engine/Source/Editor/UnrealEd/Classes/Editor/UnrealEdEngine.h
--- a/Engine/Source/Editor/UnrealEd/Classes/Editor/UnrealEdEngine.h
+++ b/Engine/Source/Editor/UnrealEd/Classes/Editor/UnrealEdEngine.h
@@ -597,14 +597,26 @@ public:
 	 * modify/dirty actors/BSP.
 	 */
 	UNREALED_API virtual void edactHideUnselected( UWorld* InWorld );
-	
+
 	/**
-	 * Toggles visibility of selected actors. If any selected actors are visible, they will be made invisible.
-	 * If all selected actors are invisible, the selection will be made visible.
+	 * Hide selected actors, their child hierarchy, and BSP models by marking their bHiddenEdTemporary flags true.
 	 * Will not modify/dirty actors/BSP.
 	 */
+	UNREALED_API virtual void edactHideSelectedHierarchy( UWorld* InWorld );
+
+	/**
+	 * Toggles visibility of selected actors. If any selected actors are visible, they will be made invisible.
+	 * If all selected actors are invisible, the selection will be made visible. Will not modify/dirty actors/BSP.
+	 */
 	UNREALED_API virtual void edactToggleSelectedVisibility( UWorld* InWorld );
 
+	/**
+	 * Toggles visibility of selected actors and their child hierarchy. If any selected actors are visible, they will be 
+	 * made invisible. If all selected actors are invisible, the selection will be made visible. Will not modify/dirty
+	 * actors/BSP.
+	 */
+	UNREALED_API virtual void edactToggleSelectedHierarchyVisibility( UWorld* InWorld );
+
 	/**
 	 * Attempt to unhide all actors and BSP models by setting their bHiddenEdTemporary flags to false if they
 	 * are true. Note: Will not unhide actors/BSP hidden by higher priority visibility settings, such as bHiddenEdGroup,
@@ -618,6 +630,12 @@ public:
 	 */
 	UNREALED_API virtual void edactHideSelectedStartup( UWorld* InWorld );
 
+	/**
+	 * Mark all selected actors, their child hierarchy, and BSP models to be hidden upon editor startup, by setting
+	 * their bHiddenEd flag to true, if it is not already. This directly modifies/dirties the relevant actors/BSP.
+	 */
+	UNREALED_API virtual void edactHideSelectedHierarchyStartup( UWorld* InWorld );
+
 	/**
 	 * Mark all actors and BSP models to be shown upon editor startup, by setting their bHiddenEd flag to false, if it is
 	 * not already. This directly modifies/dirties the relevant actors/BSP.
@@ -626,15 +644,26 @@ public:
 
 	/**
 	 * Mark all selected actors and BSP models to be shown upon editor startup, by setting their bHiddenEd flag to false, if it
-	 * not already. This directly modifies/dirties the relevant actors/BSP.
+	 * is not already. This directly modifies/dirties the relevant actors/BSP.
 	 */
 	UNREALED_API virtual void edactUnHideSelectedStartup( UWorld* InWorld );
 
 	/**
-	 * Show selected actors and BSP models by marking their bHiddenEdTemporary flags false. Will not
-	 * modify/dirty actors/BSP.
+	 * Mark all selected actors, their child hierarchy, and BSP models to be shown upon editor startup, by setting their bHiddenEd
+	 * flag to false, if it is not already. This directly modifies/dirties the relevant actors/BSP.
+	 */
+	UNREALED_API virtual void edactUnHideSelectedHierarchyStartup( UWorld* InWorld );
+
+	/**
+	 * Show selected actors and BSP models by marking their bHiddenEdTemporary flags false. Will not modify/dirty actors/BSP.
 	 */
 	UNREALED_API virtual void edactUnhideSelected( UWorld* InWorld );
+	
+	/**
+	 * Show selected actors, their child hierarchy, and BSP models by marking their bHiddenEdTemporary flags false. Will
+	 * not modify/dirty actors/BSP.
+	 */
+	UNREALED_API virtual void edactUnhideSelectedHierarchy( UWorld* InWorld );
 
 	/** Will create a map of currently visible BSP surfaces. */
 	UNREALED_API virtual void CreateBSPVisibilityMap( UWorld* InWorld, TMap<AActor*, TArray<int32>>& OutBSPMap, bool& bOutAllVisible );
@@ -920,6 +949,43 @@ protected:
 	TWeakPtr<SNotificationItem> CheckOutNotificationWeakPtr;
 
 private:
+	/** 
+	 * Helper function to recursively add the children of a given actor to the OutActorArray. This is used for Hide/Show Selected functions to
+	 * hide/show their child hierarchy. Set bExcludeSelectedChildren to true when calling this function in a loop of all selected actors, we
+	 * want the helper function to not travel down pre-selected children so there won't be duplicate entries and unnecessary complexity.
+	 */
+	static void GetChildHierarchyOfSelectedActor(TArray<AActor*>& OutActorArray, AActor* ParentActor, const bool bExcludeSelectedChildren = false, const TOptional<bool>& bIsFindingHidden = TOptional<bool>());
+	
+	/**
+	 * Internal function to hide selected Actors, allows bApplyToChildHierarchy to define whether we also want to hide the
+	 * child hierarchies of the selected actors.
+	 */
+	void Internal_HideSelected(UWorld* InWorld, bool bApplyToChildHierarchy);
+	
+	/**
+	 * Internal function to unhide selected Actors, allows bApplyToChildHierarchy to define whether we also want to unhide
+	 * the child hierarchies of the selected actors.
+	 */
+	void Internal_UnhideSelected(UWorld* InWorld, bool bApplyToChildHierarchy);
+	
+	/**
+	 * Internal function to toggle the visibility of selected Actors, allows bApplyToChildHierarchy to define whether we also
+	 * want to toggle the child hierarchies of the selected actors.
+	 */
+	void Internal_ToggleSelectedVisibility(UWorld* InWorld, bool bApplyToChildHierarchy);
+	
+	/**
+	 * Internal function to hide selected Actors on startup, allows bApplyToChildHierarchy to define whether we also want to
+	 * hide the current child hierarchies of the selected actors.
+	 */
+	void Internal_HideSelectedStartup(UWorld* InWorld, bool bApplyToChildHierarchy);
+	
+	/**
+	 * Internal function to unhide selected Actors on startup, allows bApplyToChildHierarchy to define whether we also want
+	 * to unhide the current child hierarchies of the selected actors.
+	 */
+	void Internal_UnHideSelectedStartup(UWorld* InWorld, bool bApplyToChildHierarchy);
+	
 	/** Verify if we have write permission under the specified mount point. */
 	UNREALED_API bool VerifyMountPointWritePermission(FName MountPoint);
 
diff --git a/Engine/Source/Editor/UnrealEd/Private/EditorActor.cpp b/Engine/Source/Editor/UnrealEd/Private/EditorActor.cpp
--- a/Engine/Source/Editor/UnrealEd/Private/EditorActor.cpp
+++ b/Engine/Source/Editor/UnrealEd/Private/EditorActor.cpp
@@ -1517,7 +1517,31 @@ static void SetActorVisibility(AActor* Actor, bool bVisible)
 	Actor->SetIsTemporarilyHiddenInEditor(!bVisible);
 }
 
-void UUnrealEdEngine::edactHideSelected( UWorld* InWorld )
+void UUnrealEdEngine::GetChildHierarchyOfSelectedActor(TArray<AActor*>& OutActorArray, AActor* ParentActor, bool bExcludeSelectedChildren, const TOptional<bool>& bIsFindingHidden)
+{
+	// If the ParentActor is a builder brush or doesn't match the find flag (if it is set), then don't
+	// add it to the array or continue searching the hierarchy
+	if ( !FActorEditorUtils::IsABuilderBrush(ParentActor) )
+	{
+		// If the actor doesn't match the bIsFindingHidden, we still want to search its children in case they match
+		if( !bIsFindingHidden.IsSet() || ParentActor->IsHiddenEd() == bIsFindingHidden.GetValue() )
+		{
+			OutActorArray.Add( ParentActor );
+		}
+
+		TArray<AActor*> ActorChildren;
+		ParentActor->GetAttachedActors(ActorChildren, false, false);
+		for (AActor* Child : ActorChildren)
+		{
+			if (!bExcludeSelectedChildren || !GEditor->GetSelectedActors()->IsSelected(Child))
+			{
+				GetChildHierarchyOfSelectedActor(OutActorArray, Child, bExcludeSelectedChildren, bIsFindingHidden);
+			}
+		}
+	}
+}
+
+void UUnrealEdEngine::Internal_HideSelected( UWorld* InWorld, bool bApplyToChildHierarchy )
 {
 	// Assemble a list of actors to hide.
 	TArray<AActor*> ActorsToHide;
@@ -1525,9 +1549,12 @@ void UUnrealEdEngine::edactHideSelected( UWorld* InWorld )
 	{
 		AActor* Actor = static_cast<AActor*>( *It );
 		checkSlow( Actor->IsA(AActor::StaticClass()) );
-		
-		// Don't consider already hidden actors or the builder brush
-		if ( !FActorEditorUtils::IsABuilderBrush(Actor) && !Actor->IsHiddenEd() )
+
+		if ( bApplyToChildHierarchy )
+		{
+			GetChildHierarchyOfSelectedActor(ActorsToHide, Actor, true, false);
+		}
+		else if ( !FActorEditorUtils::IsABuilderBrush(Actor) && !Actor->IsHiddenEd() )
 		{
 			ActorsToHide.Add( Actor );
 		}
@@ -1573,6 +1600,15 @@ void UUnrealEdEngine::edactHideSelected( UWorld* InWorld )
 	RedrawLevelEditingViewports();
 }
 
+void UUnrealEdEngine::edactHideSelected( UWorld* InWorld )
+{
+	Internal_HideSelected(InWorld, false);
+}
+
+void UUnrealEdEngine::edactHideSelectedHierarchy( UWorld* InWorld )
+{
+	Internal_HideSelected(InWorld, true);
+}
 
 void UUnrealEdEngine::edactHideUnselected( UWorld* InWorld )
 {
@@ -1608,22 +1644,27 @@ void UUnrealEdEngine::edactHideUnselected( UWorld* InWorld )
 	RedrawLevelEditingViewports();
 }
 
-
-void UUnrealEdEngine::edactToggleSelectedVisibility(UWorld* InWorld)
+void UUnrealEdEngine::Internal_ToggleSelectedVisibility(UWorld* InWorld, bool bApplyToChildHierarchy)
 {
 	// Assemble a list of actors to operate on.
 	TArray<AActor*> ActorsToConsider;
 	bool bTargetVisibility = true;
-	
+
 	for ( FSelectionIterator It( GetSelectedActorIterator() ) ; It ; ++It )
 	{
 		AActor* Actor = static_cast<AActor*>( *It );
 		checkSlow( Actor->IsA(AActor::StaticClass()) );
-		
-		// Don't consider already hidden actors or the builder brush
-		if ( !FActorEditorUtils::IsABuilderBrush(Actor) )
+
+		if (!FActorEditorUtils::IsABuilderBrush(Actor))
 		{
-			ActorsToConsider.Add( Actor );
+			if ( bApplyToChildHierarchy )
+			{
+				GetChildHierarchyOfSelectedActor(ActorsToConsider, Actor, true);
+			}
+			else
+			{
+				ActorsToConsider.Add( Actor );
+			}
 			bTargetVisibility = bTargetVisibility && Actor->IsHiddenEd();
 		}
 	}
@@ -1678,6 +1719,15 @@ void UUnrealEdEngine::edactToggleSelectedVisibility(UWorld* InWorld)
 	RedrawLevelEditingViewports();
 }
 
+void UUnrealEdEngine::edactToggleSelectedVisibility(UWorld* InWorld)
+{
+	Internal_ToggleSelectedVisibility(InWorld, false);
+}
+
+void UUnrealEdEngine::edactToggleSelectedHierarchyVisibility(UWorld* InWorld)
+{
+	Internal_ToggleSelectedVisibility(InWorld, true);
+}
 
 void UUnrealEdEngine::edactUnHideAll( UWorld* InWorld )
 {
@@ -1711,20 +1761,31 @@ void UUnrealEdEngine::edactUnHideAll( UWorld* InWorld )
 	RedrawLevelEditingViewports();
 }
 
-
-void UUnrealEdEngine::edactHideSelectedStartup( UWorld* InWorld )
+void UUnrealEdEngine::Internal_HideSelectedStartup( UWorld* InWorld, bool bApplyToChildHierarchy )
 {
 	// Fires ULevel::LevelDirtiedEvent when falling out of scope.
 	FScopedLevelDirtied LevelDirtyCallback;
 
-	// Iterate through all of the selected actors
+	TArray<AActor*> ActorsToHideAtStartup;
 	for ( FSelectionIterator It( GetSelectedActorIterator() ) ; It ; ++It )
 	{
 		AActor* Actor = static_cast<AActor*>( *It );
 		checkSlow( Actor->IsA(AActor::StaticClass()) );
 
+		if ( bApplyToChildHierarchy )
+		{
+			GetChildHierarchyOfSelectedActor(ActorsToHideAtStartup, Actor, true);
+		}
+		else if ( !FActorEditorUtils::IsABuilderBrush(Actor) )
+		{
+			ActorsToHideAtStartup.Add( Actor );
+		}
+	}
+
+	for (AActor* Actor : ActorsToHideAtStartup)
+	{
 		// Set the actor to hide at editor startup, if it's not already set that way
-		if ( !FActorEditorUtils::IsABuilderBrush(Actor) && !Actor->IsHiddenEdAtStartup() )
+		if ( !Actor->IsHiddenEdAtStartup() )
 		{
 			Actor->Modify();
 			Actor->bHiddenEd = true;
@@ -1744,7 +1805,7 @@ void UUnrealEdEngine::edactHideSelectedStartup( UWorld* InWorld )
 				
 				// Set the BSP surface to hide at editor startup, if it's not already set that way
 				const bool bSelected = CurSurface.Actor->IsActorOrSelectionParentSelected() || (CurSurface.PolyFlags & PF_Selected);
-				if (bSelected && !CurSurface.IsHiddenEdAtStartup() && !CurSurface.IsHiddenEd())
+				if (bSelected && !CurSurface.IsHiddenEdAtStartup())
 				{
 					CurLevelModel.Modify(false);
 					CurLevelModel.ModifySurf( SurfaceIterator.GetIndex(), false );
@@ -1757,6 +1818,15 @@ void UUnrealEdEngine::edactHideSelectedStartup( UWorld* InWorld )
 	RedrawLevelEditingViewports();
 }
 
+void UUnrealEdEngine::edactHideSelectedStartup( UWorld* InWorld )
+{
+	Internal_HideSelectedStartup(InWorld, false);
+}
+
+void UUnrealEdEngine::edactHideSelectedHierarchyStartup( UWorld* InWorld )
+{
+	Internal_HideSelectedStartup(InWorld, true);
+}
 
 void UUnrealEdEngine::edactUnHideAllStartup( UWorld* InWorld )
 {
@@ -1802,20 +1872,31 @@ void UUnrealEdEngine::edactUnHideAllStartup( UWorld* InWorld )
 	RedrawLevelEditingViewports();
 }
 
-
-void UUnrealEdEngine::edactUnHideSelectedStartup( UWorld* InWorld )
+void UUnrealEdEngine::Internal_UnHideSelectedStartup( UWorld* InWorld, bool bApplyToChildHierarchy )
 {
 	// Fires ULevel::LevelDirtiedEvent when falling out of scope.
 	FScopedLevelDirtied LevelDirtyCallback;
 
-	// Iterate over all selected actors
+	TArray<AActor*> ActorsToShowAtStartup;
 	for ( FSelectionIterator It( GetSelectedActorIterator() ) ; It ; ++It )
 	{
 		AActor* Actor = static_cast<AActor*>( *It );
 		checkSlow( Actor->IsA(AActor::StaticClass()) );
 
-		// Mark the selected actor as showing at editor startup if it was currently set to be hidden
-		if ( !FActorEditorUtils::IsABuilderBrush(Actor) && Actor->IsHiddenEdAtStartup() )
+		if( bApplyToChildHierarchy )
+		{
+			GetChildHierarchyOfSelectedActor(ActorsToShowAtStartup, Actor, true);
+		}
+		else if ( !FActorEditorUtils::IsABuilderBrush(Actor) )
+		{
+			ActorsToShowAtStartup.Add( Actor );
+		}
+	}
+
+	for (AActor* Actor : ActorsToShowAtStartup)
+	{
+		// Set the actor to show at editor startup, if it's not already set that way
+		if ( Actor->IsHiddenEdAtStartup() )
 		{
 			Actor->Modify();
 			Actor->bHiddenEd = false;
@@ -1848,18 +1929,30 @@ void UUnrealEdEngine::edactUnHideSelectedStartup( UWorld* InWorld )
 	RedrawLevelEditingViewports();
 }
 
+void UUnrealEdEngine::edactUnHideSelectedStartup( UWorld* InWorld )
+{
+	Internal_UnHideSelectedStartup(InWorld, false);
+}
 
-void UUnrealEdEngine::edactUnhideSelected( UWorld* InWorld )
+void UUnrealEdEngine::edactUnHideSelectedHierarchyStartup( UWorld* InWorld )
 {
-	// Assemble a list of actors to hide.
+	Internal_UnHideSelectedStartup(InWorld, true);
+}
+
+void UUnrealEdEngine::Internal_UnhideSelected(UWorld* InWorld, bool bApplyToChildHierarchy)
+{
+	// Assemble a list of actors to show.
 	TArray<AActor*> ActorsToShow;
 	for ( FSelectionIterator It( GetSelectedActorIterator() ) ; It ; ++It )
 	{
 		AActor* Actor = static_cast<AActor*>( *It );
 		checkSlow( Actor->IsA(AActor::StaticClass()) );
-		
-		// Don't consider already visible actors or the builder brush
-		if ( !FActorEditorUtils::IsABuilderBrush(Actor) && Actor->IsActorOrSelectionParentSelected() && Actor->IsHiddenEd() )
+
+		if( bApplyToChildHierarchy )
+		{
+			GetChildHierarchyOfSelectedActor(ActorsToShow, Actor, true, true);
+		}
+		else if ( !FActorEditorUtils::IsABuilderBrush(Actor) && Actor->IsHiddenEd() )
 		{
 			ActorsToShow.Add( Actor );
 		}
@@ -1899,6 +1992,16 @@ void UUnrealEdEngine::edactUnhideSelected( UWorld* InWorld )
 	RedrawLevelEditingViewports();
 }
 
+void UUnrealEdEngine::edactUnhideSelected( UWorld* InWorld )
+{
+	Internal_UnhideSelected(InWorld, false);
+}
+
+void UUnrealEdEngine::edactUnhideSelectedHierarchy( UWorld* InWorld )
+{
+	Internal_UnhideSelected(InWorld, true);
+}
+
 void UUnrealEdEngine::CreateBSPVisibilityMap(UWorld* InWorld, TMap<AActor*, TArray<int32>>& OutBSPMap, bool& bOutAllVisible  )
 {
 	// Start out true, we do not know otherwise.
diff --git a/Engine/Source/Editor/UnrealEd/Private/UnrealEdSrv.cpp b/Engine/Source/Editor/UnrealEd/Private/UnrealEdSrv.cpp
--- a/Engine/Source/Editor/UnrealEd/Private/UnrealEdSrv.cpp
+++ b/Engine/Source/Editor/UnrealEd/Private/UnrealEdSrv.cpp
@@ -2319,8 +2319,24 @@ bool UUnrealEdEngine::Exec_Actor( UWorld* InWorld, const TCHAR* Str, FOutputDevi
 		{
 			if ( FParse::Command(&Str,TEXT("STARTUP")) ) // ACTOR HIDE SELECTED STARTUP
 			{
-				const FScopedTransaction Transaction( NSLOCTEXT("UnrealEd", "HideSelectedAtStartup", "Hide Selected at Editor Startup") );
-				edactHideSelectedStartup( InWorld );
+				if (FParse::Command(&Str, TEXT("HIERARCHY")))
+				{
+					const FScopedTransaction Transaction( NSLOCTEXT("UnrealEd", "HideSelectedHierarchyAtStartup", "Hide Selected Hierarchy at Editor Startup") );
+					edactHideSelectedHierarchyStartup( InWorld );
+					return true;
+				}
+				else
+				{
+					const FScopedTransaction Transaction( NSLOCTEXT("UnrealEd", "HideSelectedAtStartup", "Hide Selected at Editor Startup") );
+					edactHideSelectedStartup( InWorld );
+					return true;
+				}
+			}
+			else if (FParse::Command(&Str, TEXT("HIERARCHY")))
+			{
+				const FScopedTransaction Transaction( NSLOCTEXT("UnrealEd", "HideSelectedHierarchy", "Hide Selected Hierarchy") );
+				edactHideSelectedHierarchy( InWorld );
+				SelectNone( true, true );
 				return true;
 			}
 			else
@@ -2360,8 +2376,23 @@ bool UUnrealEdEngine::Exec_Actor( UWorld* InWorld, const TCHAR* Str, FOutputDevi
 		{
 			if ( FParse::Command(&Str,TEXT("STARTUP")) ) // ACTOR UNHIDE SELECTED STARTUP
 			{
-				const FScopedTransaction Transaction( NSLOCTEXT("UnrealEd", "ShowSelectedAtStartup", "Show Selected at Editor Startup") );
-				edactUnHideSelectedStartup( InWorld );
+				if (FParse::Command(&Str, TEXT("HIERARCHY")))
+				{
+					const FScopedTransaction Transaction( NSLOCTEXT("UnrealEd", "ShowSelectedHierarchyAtStartup", "Show Selected Hierarchy at Editor Startup") );
+					edactUnHideSelectedHierarchyStartup( InWorld );
+					return true;
+				}
+				else
+				{
+					const FScopedTransaction Transaction( NSLOCTEXT("UnrealEd", "ShowSelectedAtStartup", "Show Selected at Editor Startup") );
+					edactUnHideSelectedStartup( InWorld );
+					return true;
+				}
+			}
+			else if (FParse::Command(&Str, TEXT("HIERARCHY")))
+			{
+				const FScopedTransaction Transaction( NSLOCTEXT("UnrealEd", "UnhideHierarchySelected", "Unhide Selected Hierarchy") );
+				edactUnhideSelectedHierarchy( InWorld );
 				return true;
 			}
 			else
@@ -2374,9 +2405,18 @@ bool UUnrealEdEngine::Exec_Actor( UWorld* InWorld, const TCHAR* Str, FOutputDevi
 	}
 	else if ( FParse::Command(&Str, TEXT("TOGGLE SELECTED VISIBILITY")) )
 	{
-		const FScopedTransaction Transaction( NSLOCTEXT("UnrealEd", "ToggleSelectedVisibility", "Toggle Selected Visibility") );
-		edactToggleSelectedVisibility( InWorld );
-		return true;
+		if (FParse::Command(&Str, TEXT("HIERARCHY")))
+		{
+			const FScopedTransaction Transaction( NSLOCTEXT("UnrealEd", "ToggleSelectedHierarchyVisibility", "Toggle Selected Hierarchy Visibility") );
+			edactToggleSelectedHierarchyVisibility( InWorld );
+			return true;
+		}
+		else
+		{
+			const FScopedTransaction Transaction( NSLOCTEXT("UnrealEd", "ToggleSelectedVisibility", "Toggle Selected Visibility") );
+			edactToggleSelectedVisibility( InWorld );
+			return true;
+		}
 	}
 	else if( FParse::Command(&Str, TEXT("APPLYTRANSFORM")) )
 	{
