[Level Editor] Add missing nullptr checks to FLevelEditorViewportClient::LostFocus and ::InputKey

diff --git a/Engine/Source/Editor/UnrealEd/Private/LevelEditorViewport.cpp b/Engine/Source/Editor/UnrealEd/Private/LevelEditorViewport.cpp
--- a/Engine/Source/Editor/UnrealEd/Private/LevelEditorViewport.cpp
+++ b/Engine/Source/Editor/UnrealEd/Private/LevelEditorViewport.cpp
@@ -2779,11 +2779,17 @@ void FLevelEditorViewportClient::ReceivedFocus(FViewport* InViewport)
 
 void FLevelEditorViewportClient::LostFocus(FViewport* InViewport)
 {
-	// Make sure we are calling LostFocus on currently focused viewport
-	// We don't need to call LostFocus multiple times on the same viewport if it's already un-focused.
-	if (InViewport == ModeTools->GetFocusedViewportClient()->Viewport)
+	if (ModeTools)
 	{
-		FEditorViewportClient::LostFocus(InViewport);
+		if (const FEditorViewportClient* FocusedViewportClient = ModeTools->GetFocusedViewportClient())
+		{
+			// Make sure we are calling LostFocus on currently focused viewport
+			// We don't need to call LostFocus multiple times on the same viewport if it's already un-focused.
+			if (InViewport == FocusedViewportClient->Viewport)
+            {
+            	FEditorViewportClient::LostFocus(InViewport);
+            }
+		}
 	}
 
 	GEditor->SetPreviewMeshMode(false);
@@ -3502,8 +3508,11 @@ bool FLevelEditorViewportClient::InputKey(const FInputKeyEventArgs& InEventArgs)
 		FEditorViewportClient* CurrentlyFocusedViewport = ModeTools->GetFocusedViewportClient();
 		if (CurrentlyFocusedViewport != this)
 		{
-			// Call LostFocus on the Viewport Client losing focus
-			LostFocus(CurrentlyFocusedViewport->Viewport);
+			if (CurrentlyFocusedViewport)
+			{
+				// Call LostFocus on the Viewport Client losing focus
+				LostFocus(CurrentlyFocusedViewport->Viewport);
+			}
 
 			// Focusing on This viewport
 			ReceivedFocus(Viewport);
