Modify FCookDependencyContext so I can easily retrieve the reason a package was invalidated:
- ReportInvalidated and ReportError are now deprecated. We should use instead LogInvalidated and LogError. It forces the user to pass a log message.
- Remove the callback OnInvalidated. Only the callback OnLog is used for every case.
- OnLog : add parameter bInvalidated to know if the message is from invalidation or not.

FBuildDependencySet::TryGetModifiedDependencies : log the invalidation message to the array of modified dependencies so we can see for what reason the package was invalidated in CookModifiedPackages.txt.

#rnx

diff --git a/Engine/Source/Editor/UnrealEd/Private/Cooker/CookPackageArtifacts.cpp b/Engine/Source/Editor/UnrealEd/Private/Cooker/CookPackageArtifacts.cpp
--- a/Engine/Source/Editor/UnrealEd/Private/Cooker/CookPackageArtifacts.cpp
+++ b/Engine/Source/Editor/UnrealEd/Private/Cooker/CookPackageArtifacts.cpp
@@ -96,18 +96,19 @@ FBuildDependencySet::ECurrentKeyResult FBuildDependencySet::TryCalculateCurrentK
 
 	ECurrentKeyResult Result = ECurrentKeyResult::Success;
 	FCookDependencyContext Context(&KeyBuilder,
-		[&Result, OutMessages](ELogVerbosity::Type Verbosity, FString&& Message)
+		[&Result, OutMessages](ELogVerbosity::Type Verbosity, FString&& Message, bool bInvalidated)
 		{
 			if (OutMessages)
 			{
 				OutMessages->Emplace(Verbosity, MoveTemp(Message));
 			}
-		},
-		[&Result](ELogVerbosity::Type Verbosity)
-		{
-			ECurrentKeyResult NewResult = Verbosity <= ELogVerbosity::Error
-				? ECurrentKeyResult::Error : ECurrentKeyResult::Invalidated;
-			Result = FMath::Max(Result, NewResult);
+
+			if (bInvalidated)
+			{
+				ECurrentKeyResult NewResult = Verbosity <= ELogVerbosity::Error
+					? ECurrentKeyResult::Error : ECurrentKeyResult::Invalidated;
+				Result = FMath::Max(Result, NewResult);
+			}
 		},
 		PackageName);
 
@@ -263,18 +264,6 @@ bool FBuildDependencySet::TryGetModifiedDependencies(FName PackageName, const IT
 		return false;
 	}
 
-	bool bInvalidated = false;
-	FBlake3 KeyBuilder;
-	FCookDependencyContext Context(&KeyBuilder,
-		[](ELogVerbosity::Type Verbosity, FString&& Message)
-		{
-		},
-		[&bInvalidated](ELogVerbosity::Type Verbosity)
-		{
-			bInvalidated = true;
-		},
-		PackageName);
-
 	for (const FCookDependency& Dependency : Dependencies)
 	{
 		// Some CookDependency types can not handle Update being called, because their Update relies on functions
@@ -367,7 +356,18 @@ bool FBuildDependencySet::TryGetModifiedDependencies(FName PackageName, const IT
 
 		default:
 		{
-			bInvalidated = false;
+			FBlake3 KeyBuilder;
+			FCookDependencyContext Context(&KeyBuilder,
+				[&OutModifiedDependencies, &Dependency](ELogVerbosity::Type Verbosity, FString&& Message, bool bInvalidated)
+				{
+					if (bInvalidated)
+					{
+						OutModifiedDependencies.Add(FString::Printf(TEXT("%s: %s: ReportInvalidated was called : %s"),
+							LexToString(Dependency.GetType()), *Dependency.GetDebugIdentifier(), *Message));
+					}
+				},
+				PackageName);
+	
 			FBlake3Hash CurrentValue;
 			Dependency.UpdateHash(Context, &CurrentValue);
 			if (CurrentValue != Dependency.GetHashValue())
@@ -376,11 +376,6 @@ bool FBuildDependencySet::TryGetModifiedDependencies(FName PackageName, const IT
 					LexToString(Dependency.GetType()), *Dependency.GetDebugIdentifier(),
 					*LexToString(Dependency.GetHashValue()), *LexToString(CurrentValue)));
 			}
-			else if (bInvalidated)
-			{
-				OutModifiedDependencies.Add(FString::Printf(TEXT("%s: %s: ReportInvalidated was called"),
-					LexToString(Dependency.GetType()), *Dependency.GetDebugIdentifier()));
-			}
 			break;
 		}
 		}
diff --git a/Engine/Source/Runtime/CoreUObject/Private/Cooker/CookDependencyContext.cpp b/Engine/Source/Runtime/CoreUObject/Private/Cooker/CookDependencyContext.cpp
--- a/Engine/Source/Runtime/CoreUObject/Private/Cooker/CookDependencyContext.cpp
+++ b/Engine/Source/Runtime/CoreUObject/Private/Cooker/CookDependencyContext.cpp
@@ -12,22 +12,35 @@ FCookDependencyContext::FCookDependencyContext(void* InHasher,
 	TUniqueFunction<void(ELogVerbosity::Type, FString&&)>&& InOnLog,
 	TUniqueFunction<void(ELogVerbosity::Type)>&& InOnInvalidated,
 	FName InPackageName)
+	: PackageName(InPackageName)
+	, Hasher(InHasher)
+{
+	OnLog = ([ParamOnLog = MoveTemp(InOnLog), ParamOnInvalidated = MoveTemp(InOnInvalidated)](ELogVerbosity::Type Type, FString&& Message, bool bInvalidated)
+		{ 
+			ParamOnLog(Type, MoveTemp(Message));
+			if (bInvalidated)
+			{
+				ParamOnInvalidated(Type);
+			}
+		});
+}
+
+FCookDependencyContext::FCookDependencyContext(void* InHasher,
+	TUniqueFunction<void(ELogVerbosity::Type, FString&&, bool)>&& InOnLog,
+	FName InPackageName)
 	: OnLog(MoveTemp(InOnLog))
-	, OnInvalidated(MoveTemp(InOnInvalidated))
 	, PackageName(InPackageName)
 	, Hasher(InHasher)
 { }
 
 void FCookDependencyContext::LogError(FString Message)
 {
-	Log(ELogVerbosity::Error, MoveTemp(Message));
-	ReportError();
+	Log(ELogVerbosity::Error, MoveTemp(Message), /*bInvalidated*/ true);
 }
 
 void FCookDependencyContext::LogInvalidated(FString Message)
 {
-	Log(ELogVerbosity::Display, MoveTemp(Message));
-	ReportInvalidated();
+	Log(ELogVerbosity::Display, MoveTemp(Message), /*bInvalidated*/ true);
 }
 
 void FCookDependencyContext::Update(const void* Data, uint64 Size)
@@ -38,21 +51,21 @@ void FCookDependencyContext::Update(const void* Data, uint64 Size)
 
 void FCookDependencyContext::ReportInvalidated()
 {
-	OnInvalidated(ELogVerbosity::Display);
+	Log(ELogVerbosity::Display, "", /*bInvalidated*/ true);
 }
 
 void FCookDependencyContext::ReportError()
 {
-	OnInvalidated(ELogVerbosity::Error);
+	Log(ELogVerbosity::Error, "", /*bInvalidated*/ true);
 }
 
-void FCookDependencyContext::Log(ELogVerbosity::Type Verbosity, FString Message)
+void FCookDependencyContext::Log(ELogVerbosity::Type Verbosity, FString Message, bool bInvalidated)
 {
 	for (TUniqueFunction<FString(FString&&)>& Handler : ErrorHandlers)
 	{
 		Message = Handler(MoveTemp(Message));
 	}
-	OnLog(Verbosity, MoveTemp(Message));
+	OnLog(Verbosity, MoveTemp(Message), bInvalidated);
 }
 
 void* FCookDependencyContext::SetHasher(void* NewHasher)
diff --git a/Engine/Source/Runtime/CoreUObject/Public/Cooker/CookDependencyContext.h b/Engine/Source/Runtime/CoreUObject/Public/Cooker/CookDependencyContext.h
--- a/Engine/Source/Runtime/CoreUObject/Public/Cooker/CookDependencyContext.h
+++ b/Engine/Source/Runtime/CoreUObject/Public/Cooker/CookDependencyContext.h
@@ -16,11 +16,16 @@ struct FCookDependencyContext
 {
 public:
 	/** InHasher is void* to mask the implementation details of the hashbuilder. See Update function. */
+	UE_DEPRECATED(5.7, "OnInvalidated doesn't exist anymore. Use the new constructor instead.")
 	COREUOBJECT_API explicit FCookDependencyContext(void* InHasher,
 		TUniqueFunction<void(ELogVerbosity::Type, FString&&)>&& InOnLog,
 		TUniqueFunction<void(ELogVerbosity::Type)>&& InOnInvalidated,
 		FName InPackageName);
 
+	COREUOBJECT_API explicit FCookDependencyContext(void* InHasher,
+		TUniqueFunction<void(ELogVerbosity::Type, FString&&, bool)>&& InOnLog,
+		FName InPackageName);
+
 	/**
 	 * Update the hashbuilder for the key being constructed (e.g. TargetDomainKey for cooked packages)
 	 * with the given Data of Size bytes.
@@ -32,6 +37,7 @@ public:
 	 * cannot be reported as data passed into Update. When called while calculating the initial hash this call is
 	 * ignored. When called while testing incrementally skippable the package is marked modified and recooked.
 	 */
+	UE_DEPRECATED(5.7, "Use function LogInvalidated instead.")
 	COREUOBJECT_API void ReportInvalidated();
 
 	/**
@@ -39,13 +45,14 @@ public:
 	 * When called while calculating the initial hash the storage of the key fails and the package will be recooked on
 	 * the next cook. When called while testing incrementally skippable the package is marked modified and recooked.
 	 */
+	UE_DEPRECATED(5.7, "Use function LogError instead.")
 	COREUOBJECT_API void ReportError();
 
 	/**
 	 * Send a message to the cook dependency context with the given severity. The message may be suppressed or
 	 * reduced in verbosity based on the calling context.
 	 */
-	COREUOBJECT_API void Log(ELogVerbosity::Type, FString Message);
+	COREUOBJECT_API void Log(ELogVerbosity::Type, FString Message, bool bInvalidated = false);
 
 	/** Calls Log(Error, Message) and ReportError. */
 	COREUOBJECT_API void LogError(FString Message);
@@ -86,8 +93,7 @@ public:
 	void* SetHasher(void* NewHasher);
 
 private:
-	TUniqueFunction<void(ELogVerbosity::Type, FString&&)> OnLog;
-	TUniqueFunction<void(ELogVerbosity::Type)> OnInvalidated;
+	TUniqueFunction<void(ELogVerbosity::Type, FString&&, bool)> OnLog;
 	TArray<TUniqueFunction<FString(FString&&)>, TInlineAllocator<1>> ErrorHandlers;
 	FName PackageName;
 	void* Hasher; // Type is void* to mask the implementation detail
