Remove all code related to generating and show material layer parameters in the "Parameters" panel in the Material Editor.  Instead of the low-level SWidget deciding whether or not to display layer parameters we simple prevent the parameters from making it into the View-Model data.

- Currently these aren't exposed in this panel due to the check in (FMaterialPropertyHelpers::ShouldCreatePropertyRowForParameter) where we don't show layer parameters.
- Remove redundant `if (ParameterGroup.GroupName == FMaterialPropertyHelpers::LayerParamName)`.
- Remove all code that was populating the layer parameters in the first place.  If we don't generate them they in the data model, the view (STree) widget won't display them in the first place.

#tests I removed the validation code in ShouldCreatePropertyRowForParameter to allow for Layer Parameters to show up in the UI. Then I removed the code that added the layer parameters in the first place and observed that they no longer show up in the UI.
#rnx

diff --git a/Engine/Source/Editor/MaterialEditor/Private/SMaterialParametersOverviewWidget.cpp b/Engine/Source/Editor/MaterialEditor/Private/SMaterialParametersOverviewWidget.cpp
--- a/Engine/Source/Editor/MaterialEditor/Private/SMaterialParametersOverviewWidget.cpp
+++ b/Engine/Source/Editor/MaterialEditor/Private/SMaterialParametersOverviewWidget.cpp
@@ -312,12 +312,6 @@ void SMaterialParametersOverviewTree::CreateGroupsWidget()
 			}
 
 			const FEditorParameterGroup& ParameterGroup = *ParameterGroupPtr;
-			if (ParameterGroup.GroupName == FMaterialPropertyHelpers::LayerParamName)
-			{
-				// Don't create or show the material layer parameter info in this UI
-				continue;
-			}
-
 			for (int32 ParamIdx = 0; ParamIdx < ParameterGroup.Parameters.Num(); ParamIdx++)
 			{
 				UDEditorParameterValue* Parameter = ParameterGroup.Parameters[ParamIdx];
diff --git a/Engine/Source/Editor/UnrealEd/Private/PreviewMaterial.cpp b/Engine/Source/Editor/UnrealEd/Private/PreviewMaterial.cpp
--- a/Engine/Source/Editor/UnrealEd/Private/PreviewMaterial.cpp
+++ b/Engine/Source/Editor/UnrealEd/Private/PreviewMaterial.cpp
@@ -315,28 +315,18 @@ void UMaterialEditorPreviewParameters::AssignParameterToGroup(UDEditorParameterV
 	{
 		ParameterGroupName = TEXT("None");
 	}
-	IMaterialEditorModule* MaterialEditorModule = &FModuleManager::LoadModuleChecked<IMaterialEditorModule>("MaterialEditor");
 	
-	// Material layers
-	UDEditorMaterialLayersParameterValue* MaterialLayerParam = Cast<UDEditorMaterialLayersParameterValue>(ParameterValue);
 	if (ParameterValue->ParameterInfo.Association == EMaterialParameterAssociation::GlobalParameter)
 	{
-		if (MaterialLayerParam)
+		FString AppendedGroupName = GlobalGroupPrefix.ToString();
+		if (ParameterGroupName != TEXT("None"))
 		{
-			ParameterGroupName = FMaterialPropertyHelpers::LayerParamName;
+			ParameterGroupName.AppendString(AppendedGroupName);
+			ParameterGroupName = FName(*AppendedGroupName);
 		}
 		else
 		{
-			FString AppendedGroupName = GlobalGroupPrefix.ToString();
-			if (ParameterGroupName != TEXT("None"))
-			{
-				ParameterGroupName.AppendString(AppendedGroupName);
-				ParameterGroupName = FName(*AppendedGroupName);
-			}
-			else
-			{
-				ParameterGroupName = TEXT("Global");
-			}
+			ParameterGroupName = TEXT("Global");
 		}
 	}
 
@@ -371,18 +361,6 @@ void UMaterialEditorPreviewParameters::RegenerateArrays()
 				AssignParameterToGroup(Parameter, It.Value.Group);
 			}
 		}
-
-		// Static Material Layers
-		{
-			FMaterialLayersFunctions MaterialLayers;
-			if (ParentMaterial->GetMaterialLayers(MaterialLayers))
-			{
-				UDEditorMaterialLayersParameterValue* ParameterValue = NewObject<UDEditorMaterialLayersParameterValue>(this);
-				ParameterValue->bOverride = true;
-				ParameterValue->ParameterValue = MoveTemp(MaterialLayers);
-				AssignParameterToGroup(ParameterValue, FName());
-			}
-		}
 	}
 	// sort contents of groups
 	for (int32 ParameterIdx = 0; ParameterIdx < ParameterGroups.Num(); ParameterIdx++)
