MPCook: IncrementalCook: Fix bug in PumpPhaseTransitionFence: if we received a message that caused new messages to be sent, we would not recreate the fence, and therefore would skip the necessary wait at end of phase after the work done for the new messages completed.
This also adds an optimization to skip waiting for the CookWorkers to startup, in the case of an empty incremental cook.
#rnx

diff --git a/Engine/Source/Editor/UnrealEd/Classes/CookOnTheSide/CookOnTheFlyServer.h b/Engine/Source/Editor/UnrealEd/Classes/CookOnTheSide/CookOnTheFlyServer.h
--- a/Engine/Source/Editor/UnrealEd/Classes/CookOnTheSide/CookOnTheFlyServer.h
+++ b/Engine/Source/Editor/UnrealEd/Classes/CookOnTheSide/CookOnTheFlyServer.h
@@ -569,9 +569,14 @@ private:
 	void KickBuildDependencies(UE::Cook::FTickStackData& StackData);
 
 	/**
-	 * Send a fence to CookWorkers if not already sent since the last phase transition or non-idle action.
-	 * Return bOutCompleted=true only when the fence has been passed.
-	 * This is a Noop in MPCook; no fence is used and the function sets bOutComplete=true whenever called.
+	 * Send a fence to CookWorkers, and record its number for use by PumpPhaseTransitionFence. Noop if not
+	 * in MPCook or not on Director.
+	 */
+	void InsertPhaseTransitionFence();
+	/**
+	 * Check whether the fence inserted by the last InsertPhaseTransitionFence has been passed if it exists.
+	 * Return bOutCompleted=true only when the fence does not exist or has been passed.
+	 * Noop if not MPCook or not on director; no fence is used and the function sets bOutComplete=true whenever called.
 	 */
 	void PumpPhaseTransitionFence(bool& bOutComplete);
 
@@ -1540,6 +1545,8 @@ private:
 
 	/** Returns the current Phase of the cook; this is separate axis of behavior from CookMode. @see ECookPhase. */
 	UE::Cook::ECookPhase GetCookPhase() const;
+	/** Sets the current phase of the cook to the TargetPhase, and updates per-state data. */
+	void SetCookPhase(UE::Cook::ECookPhase TargetPhase);
 
 	static UCookOnTheFlyServer* ActiveCOTFS;
 	uint32		StatLoadedPackageCount = 0;
diff --git a/Engine/Source/Editor/UnrealEd/Private/CookOnTheFlyServer.cpp b/Engine/Source/Editor/UnrealEd/Private/CookOnTheFlyServer.cpp
--- a/Engine/Source/Editor/UnrealEd/Private/CookOnTheFlyServer.cpp
+++ b/Engine/Source/Editor/UnrealEd/Private/CookOnTheFlyServer.cpp
@@ -1276,6 +1276,24 @@ UE::Cook::ECookPhase UCookOnTheFlyServer::GetCookPhase() const
 	return !bKickedBuildDependencies ? ECookPhase::Cook : ECookPhase::BuildDependencies;
 }
 
+void UCookOnTheFlyServer::SetCookPhase(UE::Cook::ECookPhase TargetPhase)
+{
+	using namespace UE::Cook;
+
+	switch (TargetPhase)
+	{
+	case ECookPhase::Cook:
+		bKickedBuildDependencies = false;
+		break;
+	case ECookPhase::BuildDependencies:
+		bKickedBuildDependencies = true;
+		break;
+	default:
+		checkNoEntry();
+		break;
+	}
+}
+
 bool UCookOnTheFlyServer::IsUsingLegacyCookOnTheFlyScheduling() const
 {
 	return CookOnTheFlyRequestManager && CookOnTheFlyRequestManager->ShouldUseLegacyScheduling();
@@ -2555,6 +2573,18 @@ UCookOnTheFlyServer::ECookAction UCookOnTheFlyServer::DecideNextCookAction(UE::C
 
 	if (IsCookByTheBookMode() && GetCookPhase() != ECookPhase::BuildDependencies)
 	{
+		// In MPCook, we need to wait for any late-arrival discovery messages before we can know that we have finished
+		// all work for the Cook phase. Send a fence to the CookWorkers and wait for their acknowledgement to be sure
+		// that we received all messages they had in progress. We only need to insert a new fence if there have been
+		// messages sent to CookWorkers since the last fence we created. (Note this also provides an optimization in
+		// empty cooks: no messages have been sent so we do not need to wait, which lets us skip having to wait for
+		// the CookWorkers to start up, which might take minutes.)
+		if (CookDirector && CookDirector->HasSentActions())
+		{
+			InsertPhaseTransitionFence();
+			CookDirector->ClearHasSentActions();
+		}
+
 		bool bCompleted;
 		PumpPhaseTransitionFence(bCompleted);
 		if (!bCompleted)
@@ -5219,19 +5249,24 @@ void UCookOnTheFlyServer::PumpBuildDependencySaves(UE::Cook::FTickStackData& Sta
 	}
 }
 
-void UCookOnTheFlyServer::PumpPhaseTransitionFence(bool& bOutComplete)
+void UCookOnTheFlyServer::InsertPhaseTransitionFence()
 {
-	using namespace UE::Cook;
-
 	if (!CookDirector)
 	{
-		bOutComplete = true;
 		return;
 	}
 
-	if (PhaseTransitionFence == -1)
+	PhaseTransitionFence = CookDirector->InsertBroadcastFence();
+}
+
+void UCookOnTheFlyServer::PumpPhaseTransitionFence(bool& bOutComplete)
+{
+	using namespace UE::Cook;
+
+	if (!CookDirector || PhaseTransitionFence == -1)
 	{
-		PhaseTransitionFence = CookDirector->InsertBroadcastFence();
+		bOutComplete = true;
+		return;
 	}
 
 	TArray<FWorkerId> PendingWorkers;
@@ -5267,12 +5302,12 @@ void UCookOnTheFlyServer::KickBuildDependencies(UE::Cook::FTickStackData& StackD
 {
 	using namespace UE::Cook;
 
-	check(!IsCookWorkerMode()); // KickBuildDependencies should only be called on the Director.
+	check(!IsCookWorkerMode()); // KickBuildDependencies should only be called on the Director or in SPCook.
 	if (CookDirector)
 	{
 		CookDirector->BroadcastMessage(FDirectorEventMessage(EDirectorEvent::KickBuildDependencies));
 	}
-	bKickedBuildDependencies = true;
+	SetCookPhase(ECookPhase::BuildDependencies);
 
 	if (bLegacyBuildDependencies)
 	{
@@ -6601,7 +6636,7 @@ void UCookOnTheFlyServer::InitializeSession()
 
 	NumObjectsHistory.Initialize(GUObjectArray.GetObjectArrayNumMinusAvailable());
 	VirtualMemoryHistory.Initialize(FPlatformMemory::GetStats().UsedVirtual);
-	bKickedBuildDependencies = false;
+	SetCookPhase(UE::Cook::ECookPhase::Cook);
 	InitialRequestCount = 0;
 }
 
diff --git a/Engine/Source/Editor/UnrealEd/Private/Cooker/CookDirector.cpp b/Engine/Source/Editor/UnrealEd/Private/Cooker/CookDirector.cpp
--- a/Engine/Source/Editor/UnrealEd/Private/Cooker/CookDirector.cpp
+++ b/Engine/Source/Editor/UnrealEd/Private/Cooker/CookDirector.cpp
@@ -489,6 +489,7 @@ void FCookDirector::AssignRequests(TArray<FWorkerId>&& InWorkers,
 			}
 			(*RemoteWorker)->AppendAssignments(RequestsToSend, GetAssignPackageExtraDatas(RequestsToSend),
 				GetInfoPackagesForRequests(RequestsToSend), ECookDirectorThread::SchedulerThread);
+			bHasSentActions = true;
 		}
 		return;
 	}
@@ -611,6 +612,7 @@ void FCookDirector::AssignRequests(TArray<FWorkerId>&& InWorkers,
 			TArray<FPackageData*>& RemoteBatch = RemoteBatches[WorkerId.GetRemoteIndex()];
 			(*RemoteWorker)->AppendAssignments(RemoteBatch,
 				GetAssignPackageExtraDatas(RemoteBatch), GetInfoPackagesForRequests(RemoteBatch), ECookDirectorThread::SchedulerThread);
+			bHasSentActions = true;
 		}
 	}
 
@@ -696,6 +698,7 @@ void FCookDirector::BroadcastMessage(UE::CompactBinaryTCP::FMarshalledMessage&&
 			checkNoEntry();
 			break;
 		}
+		bHasSentActions = true;
 	}
 }
 
@@ -1242,6 +1245,16 @@ void FCookDirector::MoveFileSynchronized(const TCHAR* TempFileName, const TCHAR*
 	}
 }
 
+bool FCookDirector::HasSentActions() const
+{
+	return bHasSentActions;
+}
+
+void FCookDirector::ClearHasSentActions()
+{
+	bHasSentActions = false;
+}
+
 FCookDirector::FPendingConnection::FPendingConnection(FPendingConnection&& Other)
 {
 	Swap(Socket, Other.Socket);
@@ -2429,6 +2442,7 @@ FCookDirector::FRetractionHandler::ReassignPackages(const FWorkerId& FromWorker,
 				PackageData->SendToState(EPackageState::AssignedToWorker, ESendFlags::QueueAddAndRemove, EStateChangeReason::Retraction);
 			}
 			PackageData->SetWorkerAssignment(Assignment);
+			Director.bHasSentActions = true;
 		}
 	}
 	Director.DisplayRemainingPackages();
diff --git a/Engine/Source/Editor/UnrealEd/Private/Cooker/CookDirector.h b/Engine/Source/Editor/UnrealEd/Private/Cooker/CookDirector.h
--- a/Engine/Source/Editor/UnrealEd/Private/Cooker/CookDirector.h
+++ b/Engine/Source/Editor/UnrealEd/Private/Cooker/CookDirector.h
@@ -158,6 +158,15 @@ public:
 	 */
 	void ReportFileTransferFromLocalWorker(FStringView TempFileName, FStringView TargetFileName);
 
+	/**
+	 * Report whether functions have been called on the director to send actions to a remote worker,
+	 * e.g. AssignPackages. This flag is not true for messages about communication (e.g. InitialConfig), it is only
+	 * true for messages that request or unblock work on the CookWorker and might cause a response from the worker.
+	 * Used to decide whether we need to WaitOnFence during cook phase transitions.
+	 */
+	bool HasSentActions() const;
+	void ClearHasSentActions();
+
 private:
 	enum class ELoadBalanceAlgorithm
 	{
@@ -313,6 +322,7 @@ private:
 	bool bReceivingMessages = false;
 	bool bForceNextHeartbeat = false;
 	bool bCookProcessCountSetByCommandLine = false;
+	bool bHasSentActions = false;
 
 	// Data that is read-only while the CommunicationThread is active and is readable from any thread
 	FBeginCookContextForWorker BeginCookContext;
