Improving a localization log to make the fix clearer and structured log to support auto-assign feature on Horde

diff --git a/Engine/Source/Editor/UnrealEd/Private/Commandlets/GatherTextFromAssetsCommandlet.cpp b/Engine/Source/Editor/UnrealEd/Private/Commandlets/GatherTextFromAssetsCommandlet.cpp
--- a/Engine/Source/Editor/UnrealEd/Private/Commandlets/GatherTextFromAssetsCommandlet.cpp
+++ b/Engine/Source/Editor/UnrealEd/Private/Commandlets/GatherTextFromAssetsCommandlet.cpp
@@ -1223,7 +1223,11 @@ void UGatherTextFromAssetsCommandlet::ProcessAndRemoveCachedPackages(TMap<FName,
 				// if this package's localization ID is a duplicate
 				if (const FName* ExistingLongPackageName = AssignedPackageLocalizationIds.Find(PackagePendingGather.PackageLocalizationId))
 				{
-					UE_LOGFMT(LogGatherTextFromAssetsCommandlet, Warning, "Package '{file}' and '{conflictFile}' have the same localization ID ({locKey}). Please reset one of these (Asset Localization -> Reset Localization ID) to avoid conflicts.", ("file", *PackageNameStr), ("conflictFile", *ExistingLongPackageName->ToString()), ("locKey", *PackagePendingGather.PackageLocalizationId), ("id", GatherTextFromAssetsCommandlet::LocalizationLogIdentifier));
+					UE_LOGFMT(LogGatherTextFromAssetsCommandlet, Warning, "Package '{file}' and '{conflictFile}' have the same localization ID ({locKey}). Please reset both of these to avoid conflicts (one will keep the ID and the other one should actually reset). In the Content Browser, Right-Click the assets -> Asset Localization -> Reset Localization ID.", 
+						("file", *FPackageName::LongPackageNameToFilename(*PackageNameStr, TEXT(".uasset"))),
+						("conflictFile", *FPackageName::LongPackageNameToFilename(ExistingLongPackageName->ToString(), TEXT(".uasset"))),
+						("locKey", *PackagePendingGather.PackageLocalizationId), 
+						("id", GatherTextFromAssetsCommandlet::LocalizationLogIdentifier));
 					if (bFixPackageLocalizationIdConflict) 
 					{
 						bThisPackageHasLocIdConflictToFix = true;
