[UE] Fix editor crash on copy/paste of a level actor with an instance-added subobject containing one or more default-instanced subobject properties.

diff --git a/Engine/Source/Editor/UnrealEd/Private/EditorObject.cpp b/Engine/Source/Editor/UnrealEd/Private/EditorObject.cpp
--- a/Engine/Source/Editor/UnrealEd/Private/EditorObject.cpp
+++ b/Engine/Source/Editor/UnrealEd/Private/EditorObject.cpp
@@ -755,13 +755,23 @@ static const TCHAR* ImportProperties(
 					ComponentTemplate->SetFlags(NewFlags);
 				}
 
+				// Make sure we have an up-to-date mapping in the object graph. This may not always be the case after defining an object.
+				//
+				// For example, if it was exported as an instance-added subobject with no non-default archetype, UE will bypass the subobject
+				// instancing pass at construction time, and so the object graph will not include a mapping for it. This can cause the explicit
+				// instancing pass that gets run after importing properties to assert, for example, if this object has its own default subobjects.
+				//
+				// Notes:
+				//	- It is ok if the archetype is NULL. Including it avoids a redundant GetArchetype() call when it has already been set.
+				//	- It is safe to assume that we're always inside a root object scope here, because we're importing a Begin Object block.
+				InstanceGraph.AddNewInstance(ComponentTemplate, Archetype);
+
 				// replace all properties in this subobject outer' class that point to the original subobject with the new subobject
 				TMap<UObject*, UObject*> ReplacementMap;
 				if (Archetype)
 				{
 					checkSlow(ComponentTemplate->GetArchetype() == Archetype);
 					ReplacementMap.Add(Archetype, ComponentTemplate);
-					InstanceGraph.AddNewInstance(ComponentTemplate, Archetype);
 				}
 				if (OldComponent)
 				{
@@ -1384,13 +1394,23 @@ static const TCHAR* ImportCreateSubObjectsStep(
 					ComponentTemplate->SetFlags(NewFlags);
 				}
 
+				// Make sure we have an up-to-date mapping in the object graph. This may not always be the case after defining an object.
+				//
+				// For example, if it was exported as an instance-added subobject with no non-default archetype, UE will bypass the subobject
+				// instancing pass at construction time, and so the object graph will not include a mapping for it. This can cause the explicit
+				// instancing pass that gets run after importing properties to assert, for example, if this object has its own default subobjects.
+				//
+				// Notes:
+				//	- It is ok if the archetype is NULL. Including it avoids a redundant GetArchetype() call when it has already been set.
+				//	- It is safe to assume that we're always inside a root object scope here, because we're importing a Begin Object block.
+				InstanceGraph.AddNewInstance(ComponentTemplate, Archetype);
+
 				// replace all properties in this subobject outer' class that point to the original subobject with the new subobject
 				TMap<UObject*, UObject*> ReplacementMap;
 				if (Archetype)
 				{
 					checkSlow(ComponentTemplate->GetArchetype() == Archetype);
 					ReplacementMap.Add(Archetype, ComponentTemplate);
-					InstanceGraph.AddNewInstance(ComponentTemplate, Archetype);
 				}
 				if (OldComponent)
 				{
