Add ability for commandlets to act like the cooker commandlet. This can let you know all the startup packages that might be used/tracked by the cooker.
-RunAsCookCommandlet
Will have your commandlet return true for the function IsRunningCookCommandlet()

diff --git a/Engine/Source/Editor/UnrealEd/Private/Cooker/PackageTracker.cpp b/Engine/Source/Editor/UnrealEd/Private/Cooker/PackageTracker.cpp
--- a/Engine/Source/Editor/UnrealEd/Private/Cooker/PackageTracker.cpp
+++ b/Engine/Source/Editor/UnrealEd/Private/Cooker/PackageTracker.cpp
@@ -81,24 +81,9 @@ void FPackageTracker::InitializeTracking(TSet<FName>& OutStartupPackages)
 	LLM_SCOPE_BYTAG(Cooker);
 
 	FWriteScopeLock ScopeLock(Lock);
-	OutStartupPackages.Empty();
-	for (TObjectIterator<UPackage> It; It; ++It)
-	{
-		UPackage* Package = *It;
-
-		if (Package->GetOuter() == nullptr && Package != GetTransientPackage())
-		{
-			LoadedPackages.Add(Package);
-			OutStartupPackages.Add(Package->GetFName());
-		}
-	}
-
-	TMap<FName, FInstigator> MapOfNewPackages;
-	MapOfNewPackages.Reserve(LoadedPackages.Num());
-	for (UPackage* Package : LoadedPackages)
-	{
-		MapOfNewPackages.Add(Package->GetFName(), FInstigator(EInstigator::StartupPackage));
-	}
+	
+	TMap<FName, UE::Cook::FInstigator> MapOfNewPackages;
+	FCookLoadScope::GetCookerStartupPackages(MapOfNewPackages);
 
 	GUObjectArray.AddUObjectDeleteListener(this);
 	GUObjectArray.AddUObjectCreateListener(this);
@@ -106,41 +91,22 @@ void FPackageTracker::InitializeTracking(TSet<FName>& OutStartupPackages)
 
 	bSubscribed = true;
 
-	TArray<TPair<FName, ECookLoadType>> StartupPackageLoadTypes;
-	FCookLoadScope::SetCookerStartupComplete(StartupPackageLoadTypes);
+	bTrackingInitialized = true;
 
-	for (TPair<FName, ECookLoadType>& Pair : StartupPackageLoadTypes)
+	PackageStream.Reserve(MapOfNewPackages.Num());
+	OutStartupPackages.Empty(MapOfNewPackages.Num());
+	for (TPair<FName, UE::Cook::FInstigator>& Pair : MapOfNewPackages)
 	{
-		switch (Pair.Value)
-		{
-		case ECookLoadType::EditorOnly:
+		switch (Pair.Value.Category)
 		{
-			OutStartupPackages.Remove(Pair.Key);
-			FInstigator* Instigator = MapOfNewPackages.Find(Pair.Key);
-			if (Instigator)
-			{
-				*Instigator = FInstigator(EInstigator::EditorOnlyLoad);
-			}
-			break;
-		}
-		case ECookLoadType::UsedInGame:
-		{
-			FInstigator* Instigator = MapOfNewPackages.Find(Pair.Key);
-			if (Instigator)
-			{
-				*Instigator = FInstigator(EInstigator::StartupPackageCookLoadScope);
-			}
+		case UE::Cook::EInstigator::StartupPackage:
+		case UE::Cook::EInstigator::StartupPackageCookLoadScope:
+			OutStartupPackages.Add(Pair.Key);
 			break;
-		}
+		case UE::Cook::EInstigator::EditorOnlyLoad:
 		default:
 			break;
 		}
-	}
-	bTrackingInitialized = true;
-
-	PackageStream.Reserve(MapOfNewPackages.Num());
-	for (TPair<FName, FInstigator>& Pair : MapOfNewPackages)
-	{
 		PackageStream.Add({ Pair.Key, MoveTemp(Pair.Value), EPackageStreamEvent::PackageLoad });
 	}
 }
diff --git a/Engine/Source/Runtime/CoreUObject/Private/UObject/ICookInfo.cpp b/Engine/Source/Runtime/CoreUObject/Private/UObject/ICookInfo.cpp
--- a/Engine/Source/Runtime/CoreUObject/Private/UObject/ICookInfo.cpp
+++ b/Engine/Source/Runtime/CoreUObject/Private/UObject/ICookInfo.cpp
@@ -7,6 +7,7 @@
 #include "Misc/ScopeLock.h"
 #include "Templates/UniquePtr.h"
 #include "UObject/UObjectArray.h"
+#include "UObject/UObjectIterator.h"
 #include "UObject/Package.h"
 
 #include <atomic>
@@ -204,4 +205,47 @@ void FCookLoadScope::SetCookerStartupComplete(TArray<TPair<FName, ECookLoadType>
 	delete Listener;
 }
 
+void FCookLoadScope::GetCookerStartupPackages(TMap<FName, UE::Cook::FInstigator>& OutStartupPackages)
+{
+	OutStartupPackages.Empty();
+	for (TObjectIterator<UPackage> It; It; ++It)
+	{
+		UPackage* Package = *It;
+		if (Package->GetOuter() == nullptr && Package != GetTransientPackage())
+		{
+			OutStartupPackages.Add(Package->GetFName(), UE::Cook::FInstigator(UE::Cook::EInstigator::StartupPackage));
+		}
+	}
+
+	TArray<TPair<FName, ECookLoadType>> StartupPackageLoadTypes;
+	FCookLoadScope::SetCookerStartupComplete(StartupPackageLoadTypes);
+
+	for (TPair<FName, ECookLoadType>& Pair : StartupPackageLoadTypes)
+	{
+		switch (Pair.Value)
+		{
+		case ECookLoadType::EditorOnly:
+		{
+			UE::Cook::FInstigator* Instigator = OutStartupPackages.Find(Pair.Key);
+			if (Instigator)
+			{
+				*Instigator = UE::Cook::FInstigator(UE::Cook::EInstigator::EditorOnlyLoad);
+			}
+			break;
+		}
+		case ECookLoadType::UsedInGame:
+		{
+			UE::Cook::FInstigator* Instigator = OutStartupPackages.Find(Pair.Key);
+			if (Instigator)
+			{
+				*Instigator = UE::Cook::FInstigator(UE::Cook::EInstigator::StartupPackageCookLoadScope);
+			}
+			break;
+		}
+		default:
+			break;
+		}
+	}
+}
+
 #endif // WITH_EDITOR
diff --git a/Engine/Source/Runtime/CoreUObject/Public/UObject/ICookInfo.h b/Engine/Source/Runtime/CoreUObject/Public/UObject/ICookInfo.h
--- a/Engine/Source/Runtime/CoreUObject/Public/UObject/ICookInfo.h
+++ b/Engine/Source/Runtime/CoreUObject/Public/UObject/ICookInfo.h
@@ -351,9 +351,10 @@ struct FCookLoadScope
 	COREUOBJECT_API ~FCookLoadScope();
 
 	COREUOBJECT_API static ECookLoadType GetCurrentValue();
-	COREUOBJECT_API static void SetCookerStartupComplete(TArray<TPair<FName, ECookLoadType>>& OutStartupPackageLoadTypes);
+	COREUOBJECT_API static void GetCookerStartupPackages(TMap<FName, UE::Cook::FInstigator>& OutStartupPackages);
 
 private:
+	static void SetCookerStartupComplete(TArray<TPair<FName, ECookLoadType>>& OutStartupPackageLoadTypes);
 	ECookLoadType PreviousScope;
 };
 
diff --git a/Engine/Source/Runtime/Launch/Private/LaunchEngineLoop.cpp b/Engine/Source/Runtime/Launch/Private/LaunchEngineLoop.cpp
--- a/Engine/Source/Runtime/Launch/Private/LaunchEngineLoop.cpp
+++ b/Engine/Source/Runtime/Launch/Private/LaunchEngineLoop.cpp
@@ -2140,7 +2140,7 @@ int32 FEngineLoop::PreInitPreStartupScreen(const TCHAR* CmdLine)
 
 		CommandletCommandLine = FString::Join(QuotedTokenArray, TEXT(" "));
 #if WITH_EDITOR
-		if (CommandletName == TEXTVIEW("cookcommandlet"))
+		if (CommandletName == TEXTVIEW("cookcommandlet") || FParse::Param(FCommandLine::Get(), TEXT("RunAsCookCommandlet")))
 		{
 			PRIVATE_GIsRunningCookCommandlet = true;
 			UE::Cook::InitializeCookGlobals();
