Fix issue where character validation would always fail on null terminator of a char array.

diff --git a/Engine/Source/Editor/UnrealEd/Private/Commandlets/AssetRegistryGenerator.cpp b/Engine/Source/Editor/UnrealEd/Private/Commandlets/AssetRegistryGenerator.cpp
--- a/Engine/Source/Editor/UnrealEd/Private/Commandlets/AssetRegistryGenerator.cpp
+++ b/Engine/Source/Editor/UnrealEd/Private/Commandlets/AssetRegistryGenerator.cpp
@@ -1192,13 +1192,13 @@ TMap<int32, FString> FAssetRegistryGenerator::GetPakChunkIdToStringOverrideMap()
 	for (const TPair<int32, FString>& PakChunkToName : ChunkIdStringOverride)
 	{
 		const FString& PakChunkName = PakChunkToName.Value;
-
-		// Iterate over the real characters. FString::Len() excludes the null terminator.
-		for (int32 CharIndex = 0; CharIndex < PakChunkName.Len(); ++CharIndex)
+		for (TCHAR Char : PakChunkName.GetCharArray())
 		{
-			const TCHAR& Char = PakChunkName[CharIndex];
-
-			if (!ValidCharacters.Contains(Char))
+			if (Char == '\0')
+			{
+				break;
+			}
+			else if (!ValidCharacters.Contains(Char))
 			{
 				UE_LOG(LogCook, Error, TEXT("GetPakChunkIdToStringMapping contains invalid string for chunk mapping. Character '%c' within %s. Only the following character are allowed:\n%s"), Char, *PakChunkName, VALID_SAVEDDIRSUFFIX_CHARACTERS);
 				InvalidOverrides.Add(PakChunkToName.Key);
