Clear external references to objects that are invalid at reinstancing time, this prevents reachable garbage errors after reinstancing operations
[RN] Fixed one cause of 'Reachable Garbage' when recompiling blueprints

diff --git a/Engine/Source/Editor/UnrealEd/Private/Kismet2/KismetReinstanceUtilities.cpp b/Engine/Source/Editor/UnrealEd/Private/Kismet2/KismetReinstanceUtilities.cpp
--- a/Engine/Source/Editor/UnrealEd/Private/Kismet2/KismetReinstanceUtilities.cpp
+++ b/Engine/Source/Editor/UnrealEd/Private/Kismet2/KismetReinstanceUtilities.cpp
@@ -2768,13 +2768,22 @@ void FBlueprintCompileReinstancer::ReplaceInstancesOfClass_Inner(const TMap<UCla
 
 					// Skip archetype instances, EXCEPT for component templates and child actor templates
 					const bool bIsChildActorTemplate = OldActor && OldActor->GetOuter()->IsA<UChildActorComponent>();
-					if (!bIsValid || 
-						  (!bIsComponent && !bIsChildActorTemplate &&
+					if (!bIsValid)
+					{
+						// the object is invalid, we cannot replace it but we can ensure any external references
+						// to it are nulled out, this is necessary because we've given it a new (invalid) REINST
+						// class.
+						if(!OldToNewInstanceMap.Contains(OldObject))
+						{
+							OldToNewInstanceMap.Add(OldObject, nullptr);
+						}
+						continue;
+					}
+					else if (!bIsComponent && !bIsChildActorTemplate &&
 						    ( (bArchetypesAreUpToDate && OldObject->IsTemplate()) ||
 						      (!bArchetypesAreUpToDate && OldObject->HasAnyFlags(RF_ClassDefaultObject))
 						    )
 						  )
-						)
 					{
 						continue;
 					}
