Disable the script maximum loop iteration count during editor map loads.
Temporary solution to avoid triggering the BPVM runaway detection when loading large maps with Verse objects.

diff --git a/Engine/Source/Editor/UnrealEd/Private/EditorServer.cpp b/Engine/Source/Editor/UnrealEd/Private/EditorServer.cpp
--- a/Engine/Source/Editor/UnrealEd/Private/EditorServer.cpp
+++ b/Engine/Source/Editor/UnrealEd/Private/EditorServer.cpp
@@ -2317,6 +2317,10 @@ bool UEditorEngine::PackageIsAMapFile( const TCHAR* PackageFilename, FText& OutN
 
 bool UEditorEngine::Map_Load(const TCHAR* Str, FOutputDevice& Ar)
 {
+	// Temporary solution to avoid triggering the BPVM runaway detection when loading large maps with Verse objects.
+	FDelegateHandle FlushUpdateHandle = FCoreDelegates::OnAsyncLoadingFlushUpdate.AddLambda([]() { FBlueprintCoreDelegates::SetScriptMaximumLoopIterations(MAX_int32); });
+	ON_SCOPE_EXIT { FBlueprintCoreDelegates::SetScriptMaximumLoopIterations(GEngine->MaximumLoopIterationCount); };
+
 	TRACE_CPUPROFILER_EVENT_SCOPE(UEditorEngine::Map_Load);
 	TRACE_LOADTIME_REQUEST_GROUP_SCOPE(TEXT("LoadMap - %s"), Str);
 
