UE: Bump CookIncrementalVersion due to shader change

diff --git a/Engine/Source/Editor/UnrealEd/Private/Cooker/CookIncrementalVersion.cpp b/Engine/Source/Editor/UnrealEd/Private/Cooker/CookIncrementalVersion.cpp
--- a/Engine/Source/Editor/UnrealEd/Private/Cooker/CookIncrementalVersion.cpp
+++ b/Engine/Source/Editor/UnrealEd/Private/Cooker/CookIncrementalVersion.cpp
@@ -8,6 +8,6 @@ namespace UE::Cook
 {
 
 // Change CookIncrementalVersion to a new guid when all packages in an incremental cook need to be invalidated
-FGuid CookIncrementalVersion(0x7f1d6986, 0xfa474247, 0xa96dd4db, 0x1b106a2f);
+FGuid CookIncrementalVersion(0x8491f969, 0x2a6c4f42, 0xa31ad43b, 0x193a3f24);
 
 }
