[Misc] Fix inverted logic from previous submit

diff --git a/Engine/Source/Editor/UnrealEd/Private/Commandlets/AssetRegUtilCommandlet.cpp b/Engine/Source/Editor/UnrealEd/Private/Commandlets/AssetRegUtilCommandlet.cpp
--- a/Engine/Source/Editor/UnrealEd/Private/Commandlets/AssetRegUtilCommandlet.cpp
+++ b/Engine/Source/Editor/UnrealEd/Private/Commandlets/AssetRegUtilCommandlet.cpp
@@ -531,16 +531,24 @@ bool UAssetRegUtilCommandlet::LoadOrderFiles(const FString& OrderFilePath, TSet<
 
 bool UAssetRegUtilCommandlet::GeneratePartiallyUpdatedOrderFile(const FString& OldOrderFilePath, const FString& NewOrderFilePath, const FString& OutOrderFilePath, const float PatchSizePerfBalanceFactor)
 {
-	TSet<FName> OldOrderFileLineSet;
+	TSet<FName> OldOrderFileLineFullSet;
 	TSet<FName> NewOrderFileLineSet;
-	if (!LoadOrderFiles(OldOrderFilePath, OldOrderFileLineSet) ||
+	if (!LoadOrderFiles(OldOrderFilePath, OldOrderFileLineFullSet) ||
 		!LoadOrderFiles(NewOrderFilePath, NewOrderFileLineSet))
 	{
 		return false;
 	}
 
-	// Remove deleted files from old order file lines
-	OldOrderFileLineSet = OldOrderFileLineSet.Difference(NewOrderFileLineSet);
+	TSet<FName> OldOrderFileLineSet;
+
+	// Remove deleted files from old order file lines, don't use .Intersect so we keep the original order
+	for (const FName& OldOrderFileIter : OldOrderFileLineFullSet)
+	{
+		if (NewOrderFileLineSet.Contains(OldOrderFileIter))
+		{
+			OldOrderFileLineSet.Add(OldOrderFileIter);
+		}
+	}
 
 	// Add new files to old order file lines
 	int LastFoundIndexInOldFileLines = -1;
