MPCook: Fix warning "Package %s was retracted from %s, but already has no platforms needing commit....".
When retracting packages from the local worker in state WantToRetract, we need to execute WaitingForResponse in the same frame. If we wait until the next frame, the packages we queued up for retraction remain active on the local worker and might be save-completed, and then when we returned to state WaitingForResponse we would trigger the warning.
#rnx

diff --git a/Engine/Source/Editor/UnrealEd/Private/Cooker/CookDirector.cpp b/Engine/Source/Editor/UnrealEd/Private/Cooker/CookDirector.cpp
--- a/Engine/Source/Editor/UnrealEd/Private/Cooker/CookDirector.cpp
+++ b/Engine/Source/Editor/UnrealEd/Private/Cooker/CookDirector.cpp
@@ -2047,11 +2047,17 @@ void FCookDirector::FRetractionHandler::TickFromSchedulerThread(bool bAllWorkers
 	int32 BusiestNumAssignments)
 {
 	bool bHadStateChange;
-	int32 NumTransitions = 0;
-	constexpr int32 MaxNumTransitions = static_cast<int32>(ERetractionState::Count);
-	do
+
+	// Enforce a constraint that the loop terminates. Some states require followup states to be executed in
+	// the same frame, so we can only enforce this constraint during certain transitions. We enforce it based
+	// on returning to an earlier state; we currently know that those are never required to execute in the same
+	// frame.
+	int32 NumReturnsToEarlierState = 0;
+	constexpr int32 MaxNumReturnsToEarlierState = 2;
+	for (;;)
 	{
 		bHadStateChange = false;
+		ERetractionState OldState = RetractionState;
 		switch (RetractionState)
 		{
 		case ERetractionState::Idle:
@@ -2084,9 +2090,20 @@ void FCookDirector::FRetractionHandler::TickFromSchedulerThread(bool bAllWorkers
 			checkNoEntry();
 			break;
 		}
-	} while (bHadStateChange
-		&& RetractionState != ERetractionState::Idle
-		&& ++NumTransitions <= MaxNumTransitions);
+
+		if (!bHadStateChange || RetractionState == ERetractionState::Idle)
+		{
+			break;
+		}
+		if (RetractionState < OldState)
+		{
+			++NumReturnsToEarlierState;
+			if (NumReturnsToEarlierState > MaxNumReturnsToEarlierState)
+			{
+				break;
+			}
+		}
+	}
 }
 
 void FCookDirector::FRetractionHandler::SetRetractionState(ERetractionState NewState, bool &bOutHadStateChange)
