Add parallel enumeration of source texts to FLocTextHelper

#rnx

diff --git a/Engine/Source/Developer/Localization/Private/LocTextHelper.cpp b/Engine/Source/Developer/Localization/Private/LocTextHelper.cpp
--- a/Engine/Source/Developer/Localization/Private/LocTextHelper.cpp
+++ b/Engine/Source/Developer/Localization/Private/LocTextHelper.cpp
@@ -15,6 +15,7 @@
 #include "Serialization/JsonInternationalizationArchiveSerializer.h"
 #include "Serialization/JsonInternationalizationManifestSerializer.h"
 #include "Serialization/JsonInternationalizationMetadataSerializer.h"
+#include "Async/ParallelFor.h"
 
 #include UE_INLINE_GENERATED_CPP_BY_NAME(LocTextHelper)
 
@@ -1062,6 +1063,42 @@ void FLocTextHelper::EnumerateSourceTexts(const FEnumerateSourceTextsFuncPtr& In
 			break;
 		}
 	}
+	
+}
+
+void FLocTextHelper::EnumerateSourceTextsParallel(const FEnumerateSourceTextsFuncPtr& InCallback, const bool InCheckDependencies) const
+{
+	std::atomic<bool> bBreak = false;
+
+	TArray<TPair<FString, TSharedRef<FManifestEntry>>> ManifestEntries = Manifest->GetEntriesBySourceTextAsArray();
+	ParallelFor(ManifestEntries.Num(), [this, &ManifestEntries, &bBreak, &InCallback, &InCheckDependencies](int32 Index)
+		{
+			if (bBreak.load(std::memory_order_acquire))
+			{
+				return;
+			}
+
+			const TSharedRef<FManifestEntry> ManifestEntry = ManifestEntries[Index].Value;
+
+			bool bShouldEnumerate = true;
+			if (InCheckDependencies)
+			{
+				for (const TSharedPtr<FInternationalizationManifest>& DepManifest : Dependencies)
+				{
+					const TSharedPtr<FManifestEntry> DepEntry = DepManifest->FindEntryBySource(ManifestEntry->Namespace, ManifestEntry->Source);
+					if (DepEntry.IsValid())
+					{
+						bShouldEnumerate = false;
+						break;
+					}
+				}
+			}
+
+			if (bShouldEnumerate && !InCallback(ManifestEntry))
+			{
+				bBreak.store(true, std::memory_order_release);
+			}
+		});
 }
 
 bool FLocTextHelper::AddTranslation(const FString& InCulture, const FLocKey& InNamespace, const FLocKey& InKey, const TSharedPtr<FLocMetadataObject>& InKeyMetadataObj, const FLocItem& InSource, const FLocItem& InTranslation, const bool InOptional)
diff --git a/Engine/Source/Developer/Localization/Public/LocTextHelper.h b/Engine/Source/Developer/Localization/Public/LocTextHelper.h
--- a/Engine/Source/Developer/Localization/Public/LocTextHelper.h
+++ b/Engine/Source/Developer/Localization/Public/LocTextHelper.h
@@ -693,6 +693,16 @@ public:
 	typedef TFunctionRef<bool(TSharedRef<FManifestEntry>)> FEnumerateSourceTextsFuncPtr;
 	UE_API void EnumerateSourceTexts(const FEnumerateSourceTextsFuncPtr& InCallback, const bool InCheckDependencies) const;
 
+	/**
+	 * Enumerate in parallel all the source texts in the manifest, optionally skipping those entries from a dependent manifest.
+	 * InCallback shall be called from multiple threads so it is the calllers responsibility to guarantee thread safety inside
+	 * the callback's scope.
+	 *
+	 * @param InCallback			Function to call for each source text. Returns true to continue enumeration.
+	 * @param InCheckDependencies	True to skip entries that can be found in a dependent manifest, false to include them.
+	 */
+	UE_API void EnumerateSourceTextsParallel(const FEnumerateSourceTextsFuncPtr& InCallback, const bool InCheckDependencies) const;
+
 	/**
 	 * Add a new translation to the given archive.
 	 *
diff --git a/Engine/Source/Editor/UnrealEd/Private/LocalizationChunkDataGenerator.cpp b/Engine/Source/Editor/UnrealEd/Private/LocalizationChunkDataGenerator.cpp
--- a/Engine/Source/Editor/UnrealEd/Private/LocalizationChunkDataGenerator.cpp
+++ b/Engine/Source/Editor/UnrealEd/Private/LocalizationChunkDataGenerator.cpp
@@ -19,6 +19,7 @@
 #include "TextLocalizationResourceGenerator.h"
 #include "Trace/Detail/Channel.h"
 #include "UObject/NameTypes.h"
+#include "Async/ParallelFor.h"
 
 class FArchiveEntry;
 
@@ -65,6 +66,8 @@ FLocalizationChunkDataGenerator::FLocalizationChunkDataGenerator(const int32 InC
 
 void FLocalizationChunkDataGenerator::GenerateChunkDataFiles(const int32 InChunkId, const TSet<FName>& InPackagesInChunk, const ITargetPlatform* TargetPlatform, FSandboxPlatformFile* InSandboxFile, TArray<FString>& OutChunkFilenames)
 {
+	TRACE_CPUPROFILER_EVENT_SCOPE(FLocalizationChunkDataGenerator::GenerateChunkDataFiles);
+
 	// The primary chunk doesn't gain a suffix to make it unique, as it is replacing the offline localization data that is usually staged verbatim
 	const bool bIsPrimaryChunk = InChunkId == 0;
 	
@@ -107,7 +110,8 @@ void FLocalizationChunkDataGenerator::GenerateChunkDataFiles(const int32 InChunk
 		bool bChunkHasText = false;
 		FLocTextHelper ChunkLocTextHelper(ChunkTargetMetadataRoot, FString::Printf(TEXT("%s.manifest"), *ChunkTargetName), FString::Printf(TEXT("%s.archive"), *ChunkTargetName), SourceLocTextHelper->GetNativeCulture(), SourceLocTextHelper->GetForeignCultures(), nullptr);
 		ChunkLocTextHelper.LoadAll(ELocTextHelperLoadFlags::Create); // Create the in-memory manifest and archives
-		SourceLocTextHelper->EnumerateSourceTexts([this, bIsCatchAllChunk, SourceLocTextHelper, &bChunkHasText, &ChunkLocTextHelper, &InPackagesInChunk, &AvailableCulturesToCook](TSharedRef<FManifestEntry> InManifestEntry)
+		FCriticalSection EnumerateGuard;
+		SourceLocTextHelper->EnumerateSourceTextsParallel([this, bIsCatchAllChunk, SourceLocTextHelper, &bChunkHasText, &ChunkLocTextHelper, &InPackagesInChunk, &AvailableCulturesToCook, &EnumerateGuard](TSharedRef<FManifestEntry> InManifestEntry)
 		{
 			for (const FManifestContext& ManifestContext : InManifestEntry->Contexts)
 			{
@@ -139,6 +143,8 @@ void FLocalizationChunkDataGenerator::GenerateChunkDataFiles(const int32 InChunk
 
 				if (bIncludeInChunk)
 				{
+					FScopeLock _(&EnumerateGuard);
+
 					bChunkHasText = true;
 					ChunkLocTextHelper.AddSourceText(InManifestEntry->Namespace, InManifestEntry->Source, ManifestContext);
 					for (const FString& CultureToCook : AvailableCulturesToCook)
diff --git a/Engine/Source/Runtime/Core/Private/Internationalization/TextLocalizationResource.cpp b/Engine/Source/Runtime/Core/Private/Internationalization/TextLocalizationResource.cpp
--- a/Engine/Source/Runtime/Core/Private/Internationalization/TextLocalizationResource.cpp
+++ b/Engine/Source/Runtime/Core/Private/Internationalization/TextLocalizationResource.cpp
@@ -480,6 +480,22 @@ bool FTextLocalizationResource::SaveToArchive(FArchive& Archive, const FTextKey&
 		KeysTable.Emplace(EntryPair.Key.GetKey(), &EntryPair.Value);
 	}
 
+	// Sort the namespace keys
+	auto SortTextKeysPredicate = [](const FTextKey& A, const FTextKey& B) -> bool
+		{
+			TStringBuilder<256> AString, BString;
+			A.ToString(AString);
+			B.ToString(BString);
+			
+			return AString.ToView() < BString.ToView();
+		};
+
+	for (auto& NamespacePair : Namespaces)
+	{
+		NamespacePair.Value.KeySort(SortTextKeysPredicate);
+	}
+	Namespaces.KeySort(SortTextKeysPredicate);
+
 	// Write entries count
 	uint32 EntriesCount = Entries.Num();
 	Archive << EntriesCount;
diff --git a/Engine/Source/Runtime/Core/Public/Internationalization/InternationalizationManifest.h b/Engine/Source/Runtime/Core/Public/Internationalization/InternationalizationManifest.h
--- a/Engine/Source/Runtime/Core/Public/Internationalization/InternationalizationManifest.h
+++ b/Engine/Source/Runtime/Core/Public/Internationalization/InternationalizationManifest.h
@@ -163,6 +163,11 @@ public:
 		return EntriesBySourceText.CreateConstIterator();
 	}
 
+	TArray<TPair<FString, TSharedRef<FManifestEntry>>> GetEntriesBySourceTextAsArray() const
+	{
+		return EntriesBySourceText.Array();
+	}
+
 	int32 GetNumEntriesBySourceText() const
 	{
 		return EntriesBySourceText.Num();
