Fix static validation warning in UnrealEd module for Clang 20

Remove potential to accidentally dereference dangling pointer to previous stack usage.

diff --git a/Engine/Source/Editor/UnrealEd/Private/EditorObject.cpp b/Engine/Source/Editor/UnrealEd/Private/EditorObject.cpp
--- a/Engine/Source/Editor/UnrealEd/Private/EditorObject.cpp
+++ b/Engine/Source/Editor/UnrealEd/Private/EditorObject.cpp
@@ -923,6 +923,7 @@ const TCHAR* ImportObjectProperties( FImportObjectParams& InParams )
 	FObjectInstancingGraph TempGraph;
 	FObjectInstancingGraph& InstanceGraph = InParams.InInstanceGraph && !bShouldUseDedicatedInstancingGraph ? *InParams.InInstanceGraph : TempGraph;
 
+	TMap<FSoftObjectPath, UObject*>* ExportedPathToObjectCreatedPtr = nullptr;
 	if (bShouldUseDedicatedInstancingGraph)
 	{
 		InstanceGraph.SetDestinationRoot(InParams.SubobjectOuter);
@@ -931,7 +932,8 @@ const TCHAR* ImportObjectProperties( FImportObjectParams& InParams )
 		TMap<FSoftObjectPath, UObject*> ExportedPathToObjectCreated;
 		if (!InParams.ObjectRemapper)
 		{
-			InParams.ObjectRemapper = &ExportedPathToObjectCreated;
+			ExportedPathToObjectCreatedPtr = &ExportedPathToObjectCreated;
+			InParams.ObjectRemapper = ExportedPathToObjectCreatedPtr;
 		}
 	}
 	else if ( InParams.SubobjectRoot && InParams.SubobjectRoot != UObject::StaticClass()->GetDefaultObject() )
@@ -986,6 +988,12 @@ const TCHAR* ImportObjectProperties( FImportObjectParams& InParams )
 		}
 	}
 
+	if (InParams.ObjectRemapper == ExportedPathToObjectCreatedPtr)
+	{
+		// Remove dangling reference to ExportedPathToObjectCreated on the stack
+		InParams.ObjectRemapper = nullptr;
+	}
+
 	return NewSourceText;
 }
 
