Ensure TMap is alive during ImportProperties call in ImportObjectProperties

diff --git a/Engine/Source/Editor/UnrealEd/Private/EditorObject.cpp b/Engine/Source/Editor/UnrealEd/Private/EditorObject.cpp
--- a/Engine/Source/Editor/UnrealEd/Private/EditorObject.cpp
+++ b/Engine/Source/Editor/UnrealEd/Private/EditorObject.cpp
@@ -923,17 +923,16 @@ const TCHAR* ImportObjectProperties( FImportObjectParams& InParams )
 	FObjectInstancingGraph TempGraph;
 	FObjectInstancingGraph& InstanceGraph = InParams.InInstanceGraph && !bShouldUseDedicatedInstancingGraph ? *InParams.InInstanceGraph : TempGraph;
 
-	TMap<FSoftObjectPath, UObject*>* ExportedPathToObjectCreatedPtr = nullptr;
+	TOptional<TMap<FSoftObjectPath, UObject*>> ExportedPathToObjectCreated;
 	if (bShouldUseDedicatedInstancingGraph)
 	{
 		InstanceGraph.SetDestinationRoot(InParams.SubobjectOuter);
 
 		// When using override T3D, always map exported subobject paths so they can be resolved at redefinition time
-		TMap<FSoftObjectPath, UObject*> ExportedPathToObjectCreated;
 		if (!InParams.ObjectRemapper)
 		{
-			ExportedPathToObjectCreatedPtr = &ExportedPathToObjectCreated;
-			InParams.ObjectRemapper = ExportedPathToObjectCreatedPtr;
+			ExportedPathToObjectCreated.Emplace();
+			InParams.ObjectRemapper = ExportedPathToObjectCreated.GetPtrOrNull();
 		}
 	}
 	else if ( InParams.SubobjectRoot && InParams.SubobjectRoot != UObject::StaticClass()->GetDefaultObject() )
@@ -988,7 +987,7 @@ const TCHAR* ImportObjectProperties( FImportObjectParams& InParams )
 		}
 	}
 
-	if (InParams.ObjectRemapper == ExportedPathToObjectCreatedPtr)
+	if (ExportedPathToObjectCreated.IsSet())
 	{
 		// Remove dangling reference to ExportedPathToObjectCreated on the stack
 		InParams.ObjectRemapper = nullptr;
