Fixed comment inconsistencies regarding GetSubsystemArrayCopy latest implementation

diff --git a/Engine/Source/Editor/UnrealEd/Classes/Editor/EditorEngine.h b/Engine/Source/Editor/UnrealEd/Classes/Editor/EditorEngine.h
--- a/Engine/Source/Editor/UnrealEd/Classes/Editor/EditorEngine.h
+++ b/Engine/Source/Editor/UnrealEd/Classes/Editor/EditorEngine.h
@@ -3298,9 +3298,7 @@ public:
 	}
 
 	/**
-	 * Get all Subsystem of specified type, this is only necessary for interfaces that can have multiple implementations instanced at a time.
-	 *
-	 * Do not hold onto this Array reference unless you are sure the lifetime is less than that of UGameInstance
+	 * Get all Subsystems of specified type, this is only necessary for interfaces that can have multiple implementations instanced at a time.
 	 */
 	template <typename TSubsystemClass>
 	TArray<TSubsystemClass*> GetEditorSubsystemArrayCopy() const
@@ -3309,7 +3307,7 @@ public:
 	}
 
 	/**
-	 * Performs an operation on all all Subsystem of specified type, this is only necessary for interfaces that can have multiple implementations instanced at a time.
+	 * Performs an operation on all Subsystems of specified type, this is only necessary for interfaces that can have multiple implementations instanced at a time.
 	 */
 	template <typename TSubsystemClass>
 	void ForEachEditorSubsystem(TFunctionRef<void(TSubsystemClass*)> Operation) const
diff --git a/Engine/Source/Runtime/Engine/Classes/Engine/Engine.h b/Engine/Source/Runtime/Engine/Classes/Engine/Engine.h
--- a/Engine/Source/Runtime/Engine/Classes/Engine/Engine.h
+++ b/Engine/Source/Runtime/Engine/Classes/Engine/Engine.h
@@ -3753,9 +3753,7 @@ public:
 	}
 
 	/**
-	 * Get all Subsystem of specified type, this is only necessary for interfaces that can have multiple implementations instanced at a time.
-	 *
-	 * Do not hold onto this Array reference unless you are sure the lifetime is less than that of UGameInstance
+	 * Get all Subsystems of specified type, this is only necessary for interfaces that can have multiple implementations instanced at a time.
 	 */
 	template <typename TSubsystemClass>
 	TArray<TSubsystemClass*> GetEngineSubsystemArrayCopy() const
diff --git a/Engine/Source/Runtime/Engine/Classes/Engine/GameInstance.h b/Engine/Source/Runtime/Engine/Classes/Engine/GameInstance.h
--- a/Engine/Source/Runtime/Engine/Classes/Engine/GameInstance.h
+++ b/Engine/Source/Runtime/Engine/Classes/Engine/GameInstance.h
@@ -457,9 +457,7 @@ public:
 	}
 
 	/**
-	 * Get all Subsystem of specified type, this is only necessary for interfaces that can have multiple implementations instanced at a time.
-	 *
-	 * Do not hold onto this Array reference unless you are sure the lifetime is less than that of UGameInstance
+	 * Get all Subsystems of specified type, this is only necessary for interfaces that can have multiple implementations instanced at a time.
 	 */
 	template <typename TSubsystemClass>
 	TArray<TSubsystemClass*> GetSubsystemArrayCopy() const
@@ -468,7 +466,7 @@ public:
 	}
 
 	/**
-	 * Performs an operation on all all Subsystem of specified type, this is only necessary for interfaces that can have multiple implementations instanced at a time.
+	 * Performs an operation on all Subsystems of specified type, this is only necessary for interfaces that can have multiple implementations instanced at a time.
 	 */
 	template <typename TSubsystemClass>
 	void ForEachSubsystem(TFunctionRef<void(TSubsystemClass*)> Operation) const
diff --git a/Engine/Source/Runtime/Engine/Classes/Engine/LocalPlayer.h b/Engine/Source/Runtime/Engine/Classes/Engine/LocalPlayer.h
--- a/Engine/Source/Runtime/Engine/Classes/Engine/LocalPlayer.h
+++ b/Engine/Source/Runtime/Engine/Classes/Engine/LocalPlayer.h
@@ -387,9 +387,7 @@ public:
 	}
 	
 	/**
-	 * Get all Subsystem of specified type, this is only necessary for interfaces that can have multiple implementations instanced at a time.
-	 *
-	 * Do not hold onto this Array reference unless you are sure the lifetime is less than that of ULocalPlayer
+	 * Get all Subsystems of specified type, this is only necessary for interfaces that can have multiple implementations instanced at a time.
 	 */
 	template <typename TSubsystemClass>
 	TArray<TSubsystemClass*> GetSubsystemArrayCopy() const
diff --git a/Engine/Source/Runtime/Engine/Classes/Engine/World.h b/Engine/Source/Runtime/Engine/Classes/Engine/World.h
--- a/Engine/Source/Runtime/Engine/Classes/Engine/World.h
+++ b/Engine/Source/Runtime/Engine/Classes/Engine/World.h
@@ -4235,9 +4235,7 @@ public:
 	}
 
 	/**
-	 * Get all Subsystem of specified type, this is only necessary for interfaces that can have multiple implementations instanced at a time.
-	 *
-	 * Do not hold onto this Array reference unless you are sure the lifetime is less than that of the World
+	 * Get all Subsystems of specified type, this is only necessary for interfaces that can have multiple implementations instanced at a time.
 	 */
 	template <typename TSubsystemClass>
 	TArray<TSubsystemClass*> GetSubsystemArrayCopy() const
@@ -4246,7 +4244,7 @@ public:
 	}
 
 	/**
-	 * Performs an operation on all all Subsystem of specified type, this is only necessary for interfaces that can have multiple implementations instanced at a time.
+	 * Performs an operation on all Subsystems of specified type, this is only necessary for interfaces that can have multiple implementations instanced at a time.
 	 */
 	template <typename TSubsystemClass>
 	void ForEachSubsystem(TFunctionRef<void(TSubsystemClass*)> Operation) const
diff --git a/Engine/Source/Runtime/Engine/Public/AudioDevice.h b/Engine/Source/Runtime/Engine/Public/AudioDevice.h
--- a/Engine/Source/Runtime/Engine/Public/AudioDevice.h
+++ b/Engine/Source/Runtime/Engine/Public/AudioDevice.h
@@ -1958,7 +1958,6 @@ public:
 
 	/**
 	 * Gets all Subsystems of specified type, this is only necessary for interfaces that can have multiple implementations instanced at a time.
-	 * Do not hold onto this Array reference unless you are sure the lifetime is less than that of the audio device
 	 */
 	template <typename TSubsystemClass>
 	TArray<TSubsystemClass*> GetSubsystemArrayCopy() const
