Clean up Emplace usage in ImportObjectProperties

diff --git a/Engine/Source/Editor/UnrealEd/Private/EditorObject.cpp b/Engine/Source/Editor/UnrealEd/Private/EditorObject.cpp
--- a/Engine/Source/Editor/UnrealEd/Private/EditorObject.cpp
+++ b/Engine/Source/Editor/UnrealEd/Private/EditorObject.cpp
@@ -923,6 +923,7 @@ const TCHAR* ImportObjectProperties( FImportObjectParams& InParams )
 	FObjectInstancingGraph TempGraph;
 	FObjectInstancingGraph& InstanceGraph = InParams.InInstanceGraph && !bShouldUseDedicatedInstancingGraph ? *InParams.InInstanceGraph : TempGraph;
 
+	// Use TOptional to avoid heap allocations and only pay for construction/destruction if Emplace is called.
 	TOptional<TMap<FSoftObjectPath, UObject*>> ExportedPathToObjectCreated;
 	if (bShouldUseDedicatedInstancingGraph)
 	{
@@ -931,8 +932,7 @@ const TCHAR* ImportObjectProperties( FImportObjectParams& InParams )
 		// When using override T3D, always map exported subobject paths so they can be resolved at redefinition time
 		if (!InParams.ObjectRemapper)
 		{
-			ExportedPathToObjectCreated.Emplace();
-			InParams.ObjectRemapper = ExportedPathToObjectCreated.GetPtrOrNull();
+			InParams.ObjectRemapper = &ExportedPathToObjectCreated.Emplace();
 		}
 	}
 	else if ( InParams.SubobjectRoot && InParams.SubobjectRoot != UObject::StaticClass()->GetDefaultObject() )
