[ComponentTypeRegistry]
- Suppress TSAN warning on bNeedsRefreshNextTick as it doesn't matter, invalidation will be picked up on a subsequent tick at some point.

diff --git a/Engine/Source/Editor/UnrealEd/Private/ComponentTypeRegistry.cpp b/Engine/Source/Editor/UnrealEd/Private/ComponentTypeRegistry.cpp
--- a/Engine/Source/Editor/UnrealEd/Private/ComponentTypeRegistry.cpp
+++ b/Engine/Source/Editor/UnrealEd/Private/ComponentTypeRegistry.cpp
@@ -298,7 +298,7 @@ public:
 	TMap<FTopLevelAssetPath, int32> ClassPathToClassListIndexMap;
 	FOnComponentTypeListChanged ComponentListChanged;
 	FAssetRegistryModule& AssetRegistryModule;
-	bool bNeedsRefreshNextTick;
+	TSAN_ATOMIC(bool) bNeedsRefreshNextTick;
 };
 
 static const FString CommonClassGroup(TEXT("Common"));
