CookCustomization: When calculating the initial set of requests for the cook from various sources, run the ModifyCook delegates (AssetManager and others) last, so that they can make complicated decisions based on the full set of data from the other relatively-simple sources.
#rnx

diff --git a/Engine/Source/Editor/UnrealEd/Private/CookOnTheFlyServer.cpp b/Engine/Source/Editor/UnrealEd/Private/CookOnTheFlyServer.cpp
--- a/Engine/Source/Editor/UnrealEd/Private/CookOnTheFlyServer.cpp
+++ b/Engine/Source/Editor/UnrealEd/Private/CookOnTheFlyServer.cpp
@@ -8562,7 +8562,9 @@
 	bool bMapsOnly = (!!(FilesToCookFlags & ECookByTheBookOptions::MapsOnly)) || PackagingSettings->bCookMapsOnly;
 	bool bNoDev = !!(FilesToCookFlags & ECookByTheBookOptions::NoDevContent);
 
-	int32 InitialNum = FilesInPath.Num();
+	int32 PreviousFilesInPathNum = FilesInPath.Num();
+	int32 NumFilesAddedByCommandLineOrGameCallback = 0;
+
 	struct FNameWithInstigator
 	{
 		FInstigator Instigator;
@@ -8658,7 +8660,8 @@
 
 	TSet<FName> ScratchNewFiles;
 	TArray<FName> ScratchRemoveFiles;
-	auto UpdateInstigators = [&FilesInPath, &Instigators, &ScratchNewFiles, &ScratchRemoveFiles](const FInstigator& InInstigator)
+	auto UpdateInstigators = [&FilesInPath, &Instigators, &ScratchNewFiles, &ScratchRemoveFiles]
+	(const FInstigator& InInstigator, int32* NumFilesCounter)
 	{
 		ScratchNewFiles.Reset();
 		ScratchNewFiles.Reserve(FilesInPath.Num());
@@ -8669,6 +8672,10 @@
 			if (Existing.Category == EInstigator::InvalidCategory)
 			{
 				Existing = InInstigator;
+				if (NumFilesCounter)
+				{
+					++(*NumFilesCounter);
+				}
 			}
 		}
 		ScratchRemoveFiles.Reset();
@@ -8685,6 +8692,120 @@
 		}
 	};
 
+	for (const FString& CurrEntry : CookMaps)
+	{
+		AddFlexPathToCook(FilesInPath, Instigators, CurrEntry, EInstigator::CommandLinePackage);
+	}
+
+	if (!(FilesToCookFlags & ECookByTheBookOptions::SkipSoftReferences)
+		&& !(FilesToCookFlags & ECookByTheBookOptions::NoGameAlwaysCookPackages))
+	{
+		for (const ITargetPlatform* TargetPlatform : TargetPlatforms)
+		{
+			TargetPlatform->GetExtraPackagesToCook(FilesInPath);
+		}
+		UpdateInstigators(EInstigator::TargetPlatformExtraPackagesToCook, nullptr);
+	}
+
+	if (!(FilesToCookFlags & ECookByTheBookOptions::SkipSoftReferences))
+	{
+		const FString ExternalMountPointName(TEXT("/Game/"));
+		for (const FNameWithInstigator& CurrEntry : CookDirectories)
+		{
+			TArray<FString> Files;
+			FString DirectoryName = CurrEntry.Name.ToString();
+			IFileManager::Get().FindFilesRecursive(Files, *DirectoryName, *(FString(TEXT("*")) + FPackageName::GetAssetPackageExtension()), true, false);
+			for (int32 Index = 0; Index < Files.Num(); Index++)
+			{
+				FString StdFile = Files[Index];
+				FPaths::MakeStandardFilename(StdFile);
+				AddFileToCook(FilesInPath, Instigators, StdFile, CurrEntry.Instigator);
+
+				// this asset may not be in our currently mounted content directories, so try to mount a new one now
+				FString LongPackageName;
+				if (!FPackageName::IsValidLongPackageName(StdFile) && !FPackageName::TryConvertFilenameToLongPackageName(StdFile, LongPackageName))
+				{
+					FPackageName::RegisterMountPoint(ExternalMountPointName, DirectoryName);
+				}
+			}
+		}
+	}
+
+	NumFilesAddedByCommandLineOrGameCallback += FilesInPath.Num() - PreviousFilesInPathNum;
+	PreviousFilesInPathNum = FilesInPath.Num();
+
+	if (!(FilesToCookFlags & ECookByTheBookOptions::NoDefaultMaps))
+	{
+		for (const auto& GameDefaultSet : GameDefaultObjects)
+		{
+			if (GameDefaultSet.Key == FName(TEXT("ServerDefaultMap")) && !IsCookFlagSet(ECookInitializationFlags::IncludeServerMaps))
+			{
+				continue;
+			}
+
+			for (FName PackagePath : GameDefaultSet.Value)
+			{
+				TArray<FAssetData> Assets;
+				if (!AssetRegistry->GetAssetsByPackageName(PackagePath, Assets))
+				{
+					const FText ErrorMessage = FText::Format(LOCTEXT("GameMapSettingsMissing",
+						"{0} contains a path to a missing asset '{1}'. "
+						"The intended asset will fail to load in a packaged build. "
+						"Select the intended asset again in Project Settings to fix this issue."),
+						FText::FromName(GameDefaultSet.Key), FText::FromName(PackagePath));
+					LogCookerMessage(ErrorMessage.ToString(), EMessageSeverity::Error);
+				}
+				else
+				{
+					TArray<const FAssetData*, TInlineAllocator<1>> AssetPtrs;
+					AssetPtrs.Reserve(Assets.Num());
+					for (const FAssetData& AssetData : Assets)
+					{
+						AssetPtrs.Add(&AssetData);
+					}
+					const FAssetData* PrimaryAssetData = UE::AssetRegistry::GetMostImportantAsset(AssetPtrs);
+					if (PrimaryAssetData && PrimaryAssetData->IsRedirector())
+					{
+						const FText ErrorMessage = FText::Format(LOCTEXT("GameMapSettingsRedirectorDetected",
+							"{0} contains a redirected reference '{1}'. "
+							"The intended asset will fail to load in a packaged build. "
+							"Select the intended asset again in Project Settings to fix this issue."),
+							FText::FromName(GameDefaultSet.Key), FText::FromName(PackagePath));
+						LogCookerMessage(ErrorMessage.ToString(), EMessageSeverity::Error);
+					}
+				}
+
+				AddFileToCook(FilesInPath, Instigators, PackagePath.ToString(),
+					FInstigator(EInstigator::GameDefaultObject, GameDefaultSet.Key));
+			}
+		}
+	}
+
+	if (!(FilesToCookFlags & ECookByTheBookOptions::NoInputPackages))
+	{
+		// make sure we cook any extra assets for the default touch interface
+		// @todo need a better approach to cooking assets which are dynamically loaded by engine code based on settings
+		FConfigFile InputIni;
+		FString InterfaceFile;
+		FConfigCacheIni::LoadLocalIniFile(InputIni, TEXT("Input"), true);
+		if (InputIni.GetString(TEXT("/Script/Engine.InputSettings"), TEXT("DefaultTouchInterface"), InterfaceFile))
+		{
+			if (InterfaceFile != TEXT("None") && InterfaceFile != TEXT(""))
+			{
+				AddFileToCook(FilesInPath, Instigators, InterfaceFile, EInstigator::InputSettingsIni);
+			}
+		}
+	}
+
+	// Preserve legacy behavior: the legacy definition of "packages added by commandline or game callback" - to decide
+	// whether we need to apply bCookAll below - does not include GameDefaultObjects or InputPackages.
+	// So update PreviousFilesInPathNum without adding the delta to NumFilesAddedByCommandLineOrGameCallback.
+	PreviousFilesInPathNum = FilesInPath.Num();
+
+	// Call the ModifyCook delegates (almost) last, after adding all other sources of packages, so that they can see
+	// and modify the complete list of what the ProjectPackagingSettings and CommandLine specified should be cooked.
+	// The one exception is the legacy bCookAll behaivor, which needs to take into account whether the delegates
+	// added anything to be cooked.
 	if (!(FilesToCookFlags & ECookByTheBookOptions::NoGameAlwaysCookPackages))
 	{
 		UE_SCOPED_HIERARCHICAL_COOKTIMER_AND_DURATION(CookModificationDelegate, DetailedCookStats::GameCookModificationDelegateTimeSec);
@@ -8694,13 +8815,13 @@
 
 		// allow the AssetManager to fill out the asset registry, as well as get a list of objects to always cook
 		UAssetManager::Get().ModifyCook(TargetPlatforms, FilesInPath, PackagesToNeverCook);
-		UpdateInstigators(EInstigator::AssetManagerModifyCook);
+		UpdateInstigators(EInstigator::AssetManagerModifyCook, &NumFilesAddedByCommandLineOrGameCallback);
 
 		if (ModifyCookDelegate.IsBound())
 		{
 			// allow game or plugins to fill out the asset registry, as well as get a list of objects to always cook
 			ModifyCookDelegate.Broadcast(TargetPlatforms, FilesInPath, PackagesToNeverCook);
-			UpdateInstigators(EInstigator::ModifyCookDelegate);
+			UpdateInstigators(EInstigator::ModifyCookDelegate, &NumFilesAddedByCommandLineOrGameCallback);
 		}
 
 		FCookInfoModifyCookDelegate& CookInfoModifyCook = UE::Cook::FDelegates::ModifyCook;
@@ -8726,6 +8847,7 @@
 						Existing = FInstigator(EInstigator::ModifyCookDelegate, CookRule.InstigatorName);
 					}
 					FilesInPath.Add(CookRule.PackageName);
+					++NumFilesAddedByCommandLineOrGameCallback;
 					break;
 				}
 				case EPackageCookRule::NeverCook:
@@ -8739,7 +8861,7 @@
 					break;
 				}
 			}
-			UpdateInstigators(EInstigator::ModifyCookDelegate);
+			UpdateInstigators(EInstigator::ModifyCookDelegate, nullptr);
 		}
 
 		for (FName NeverCookPackage : PackagesToNeverCook)
@@ -8752,15 +8874,11 @@
 		}
 	}
 
-	for (const FString& CurrEntry : CookMaps)
-	{
-		AddFlexPathToCook(FilesInPath, Instigators, CurrEntry, EInstigator::CommandLinePackage);
-	}
 	if (IsCookingDLC())
 	{
 		TArray<FName> PackagesToNeverCook;
 		UAssetManager::Get().ModifyDLCCook(CookByTheBookOptions->DlcName, TargetPlatforms, FilesInPath, PackagesToNeverCook);
-		UpdateInstigators(EInstigator::AssetManagerModifyDLCCook);
+		UpdateInstigators(EInstigator::AssetManagerModifyDLCCook, &NumFilesAddedByCommandLineOrGameCallback);
 
 		for (FName NeverCookPackage : PackagesToNeverCook)
 		{
@@ -8772,45 +8890,15 @@
 		}
 	}
 
-	if (!(FilesToCookFlags & ECookByTheBookOptions::SkipSoftReferences)
-		&& !(FilesToCookFlags & ECookByTheBookOptions::NoGameAlwaysCookPackages))
-	{
-		for (const ITargetPlatform* TargetPlatform : TargetPlatforms)
-		{
-			TargetPlatform->GetExtraPackagesToCook(FilesInPath);
-		}
-		UpdateInstigators(EInstigator::TargetPlatformExtraPackagesToCook);
-	}
+	// NumFilesAddedByCommandLineOrGameCallback for the ModifyCook delegates was incremented by UpdateInstigators,
+	// so don't update it here, but do update the other accumulator PreviousFilesInPathNum.
+	PreviousFilesInPathNum = FilesInPath.Num();
 
-	if (!(FilesToCookFlags & ECookByTheBookOptions::SkipSoftReferences))
-	{
-		const FString ExternalMountPointName(TEXT("/Game/"));
-		for (const FNameWithInstigator& CurrEntry : CookDirectories)
-		{
-			TArray<FString> Files;
-			FString DirectoryName = CurrEntry.Name.ToString();
-			IFileManager::Get().FindFilesRecursive(Files, *DirectoryName, *(FString(TEXT("*")) + FPackageName::GetAssetPackageExtension()), true, false);
-			for (int32 Index = 0; Index < Files.Num(); Index++)
-			{
-				FString StdFile = Files[Index];
-				FPaths::MakeStandardFilename(StdFile);
-				AddFileToCook(FilesInPath, Instigators, StdFile, CurrEntry.Instigator);
-
-				// this asset may not be in our currently mounted content directories, so try to mount a new one now
-				FString LongPackageName;
-				if (!FPackageName::IsValidLongPackageName(StdFile) && !FPackageName::TryConvertFilenameToLongPackageName(StdFile, LongPackageName))
-				{
-					FPackageName::RegisterMountPoint(ExternalMountPointName, DirectoryName);
-				}
-			}
-		}
-	}
-
 	// If no packages were explicitly added by command line or game callback, add all maps
 	if (!(FilesToCookFlags & ECookByTheBookOptions::SkipSoftReferences)
 		&& !(FilesToCookFlags & ECookByTheBookOptions::NoGameAlwaysCookPackages))
 	{
-		if (bCookAll || (UE::Cook::bCookAllByDefault && FilesInPath.Num() == InitialNum))
+		if (bCookAll || (UE::Cook::bCookAllByDefault && NumFilesAddedByCommandLineOrGameCallback == 0))
 		{
 			TArray<FString> Tokens;
 			Tokens.Empty(2);
@@ -8845,76 +8933,13 @@
 				}
 			}
 		}
-		else if (FilesInPath.Num() == InitialNum)
+		else if (NumFilesAddedByCommandLineOrGameCallback == 0)
 		{
 			LogCookerMessage(TEXT("No package requests specified on -run=Cook commandline or ini. ")
 				TEXT("Set the flag 'Edit->Project Settings->Project/Packaging->Packaging/Advanced->Cook Everything in the Project Content Directory'. ")
 				TEXT("Or launch 'UnrealEditor -run=cook -helpcookusage' to see all package request options."), EMessageSeverity::Warning);
 		}
 	}
-
-	if (!(FilesToCookFlags & ECookByTheBookOptions::NoDefaultMaps))
-	{
-		for (const auto& GameDefaultSet : GameDefaultObjects)
-		{
-			if (GameDefaultSet.Key == FName(TEXT("ServerDefaultMap")) && !IsCookFlagSet(ECookInitializationFlags::IncludeServerMaps))
-			{
-				continue;
-			}
-
-			for (FName PackagePath : GameDefaultSet.Value)
-			{
-				TArray<FAssetData> Assets;
-				if (!AssetRegistry->GetAssetsByPackageName(PackagePath, Assets))
-				{
-					const FText ErrorMessage = FText::Format(LOCTEXT("GameMapSettingsMissing",
-						"{0} contains a path to a missing asset '{1}'. "
-						"The intended asset will fail to load in a packaged build. "
-						"Select the intended asset again in Project Settings to fix this issue."),
-						FText::FromName(GameDefaultSet.Key), FText::FromName(PackagePath));
-					LogCookerMessage(ErrorMessage.ToString(), EMessageSeverity::Error);
-				}
-				else
-				{
-					TArray<const FAssetData*, TInlineAllocator<1>> AssetPtrs;
-					AssetPtrs.Reserve(Assets.Num());
-					for (const FAssetData& AssetData : Assets)
-					{
-						AssetPtrs.Add(&AssetData);
-					}
-					const FAssetData* PrimaryAssetData = UE::AssetRegistry::GetMostImportantAsset(AssetPtrs);
-					if (PrimaryAssetData && PrimaryAssetData->IsRedirector())
-					{
-						const FText ErrorMessage = FText::Format(LOCTEXT("GameMapSettingsRedirectorDetected",
-							"{0} contains a redirected reference '{1}'. "
-							"The intended asset will fail to load in a packaged build. "
-							"Select the intended asset again in Project Settings to fix this issue."),
-							FText::FromName(GameDefaultSet.Key), FText::FromName(PackagePath));
-						LogCookerMessage(ErrorMessage.ToString(), EMessageSeverity::Error);
-					}
-				}
-
-				AddFileToCook(FilesInPath, Instigators, PackagePath.ToString(),
-					FInstigator(EInstigator::GameDefaultObject, GameDefaultSet.Key));
-			}
-		}
-	}
-
-	if (!(FilesToCookFlags & ECookByTheBookOptions::NoInputPackages))
-	{
-		// make sure we cook any extra assets for the default touch interface
-		// @todo need a better approach to cooking assets which are dynamically loaded by engine code based on settings
-		FConfigFile InputIni;
-		FString InterfaceFile;
-		FConfigCacheIni::LoadLocalIniFile(InputIni, TEXT("Input"), true);
-		if (InputIni.GetString(TEXT("/Script/Engine.InputSettings"), TEXT("DefaultTouchInterface"), InterfaceFile))
-		{
-			if (InterfaceFile != TEXT("None") && InterfaceFile != TEXT(""))
-			{
-				AddFileToCook(FilesInPath, Instigators, InterfaceFile, EInstigator::InputSettingsIni);
-			}
-		}
-	}
 }
 
 void UCookOnTheFlyServer::GetGameDefaultObjects(const TArray<ITargetPlatform*>& TargetPlatforms, TMap<FName, TArray<FName>>& OutGameDefaultObjects)
