Fix a bug where viewports could control the movement of a locked actor in a different viewport.

[RN] Minor Editor

diff --git a/Engine/Source/Editor/UnrealEd/Private/LevelEditorViewport.cpp b/Engine/Source/Editor/UnrealEd/Private/LevelEditorViewport.cpp
--- a/Engine/Source/Editor/UnrealEd/Private/LevelEditorViewport.cpp
+++ b/Engine/Source/Editor/UnrealEd/Private/LevelEditorViewport.cpp
@@ -4190,7 +4190,7 @@ void FLevelEditorViewportClient::MoveLockedActorToCamera()
 			}
 
 			// Need to disable orbit camera before setting actor position so that the viewport camera location is converted back
-			GCurrentLevelEditingViewportClient->ToggleOrbitCamera(false);
+			ToggleOrbitCamera(false);
 
 			USceneComponent* ActiveActorLockComponent = ActiveActorLock->GetRootComponent();
 			TOptional<FRotator> PreviousRotator;
@@ -4209,20 +4209,20 @@ void FLevelEditorViewportClient::MoveLockedActorToCamera()
 					float AdditiveFOV;
 					CameraViewComponent->GetAdditiveOffset(AdditiveOffset, AdditiveFOV);
 					const FTransform RelativeTransform = (AdditiveOffset * CameraViewComponent->GetComponentTransform()).Inverse();
-					const FTransform DesiredTransform = FTransform(GCurrentLevelEditingViewportClient->GetViewRotation(), GCurrentLevelEditingViewportClient->GetViewLocation());
+					const FTransform DesiredTransform = FTransform(GetViewRotation(), GetViewLocation());
 					ActiveActorLock->SetActorTransform(ActiveActorLock->GetActorTransform() * RelativeTransform * DesiredTransform);
 				}
 				else if (const USceneComponent* SceneViewComponent = Cast<USceneComponent>(ViewComponent))
 				{
 					const FTransform RelativeTransform = SceneViewComponent->GetComponentTransform().Inverse();
-					const FTransform DesiredTransform = FTransform(GCurrentLevelEditingViewportClient->GetViewRotation(), GCurrentLevelEditingViewportClient->GetViewLocation());
+					const FTransform DesiredTransform = FTransform(GetViewRotation(), GetViewLocation());
 					ActiveActorLock->SetActorTransform(ActiveActorLock->GetActorTransform() * RelativeTransform * DesiredTransform);
 				}
 			}
 			else
 			{
-				ActiveActorLock->SetActorLocation(GCurrentLevelEditingViewportClient->GetViewLocation(), false);
-				ActiveActorLock->SetActorRotation(GCurrentLevelEditingViewportClient->GetViewRotation());
+				ActiveActorLock->SetActorLocation(GetViewLocation(), false);
+				ActiveActorLock->SetActorRotation(GetViewRotation());
 			}
 
 			if (ActiveActorLockComponent)
