Invalidate and recreate the model geometry before trying to propagate lighting changes.  This ensures that the level model components are reset to a valid state and eliminates any nulled out references because of the hot reload.

#rnx

diff --git a/Engine/Source/Editor/UnrealEd/Private/PackageTools.cpp b/Engine/Source/Editor/UnrealEd/Private/PackageTools.cpp
--- a/Engine/Source/Editor/UnrealEd/Private/PackageTools.cpp
+++ b/Engine/Source/Editor/UnrealEd/Private/PackageTools.cpp
@@ -1028,6 +1028,14 @@ UPackageTools::UPackageTools(const FObjectInitializer& ObjectInitializer)
 				}
 			}
 
+			CurrentWorldPtr->InvalidateModelGeometry(nullptr);
+
+			const TArray<ULevel*>& Levels = CurrentWorldPtr->GetLevels();
+			for (ULevel* Level : Levels)
+			{
+				Level->UpdateModelComponents();
+			}
+			
 			CurrentWorldPtr->PropagateLightingScenarioChange();
 		}
 
