Added null checks to gameplay tag customization and a few editor blueprint utility functions.

diff --git a/Engine/Plugins/Editor/GameplayTagsEditor/Source/GameplayTagsEditor/Private/GameplayTagContainerCustomization.cpp b/Engine/Plugins/Editor/GameplayTagsEditor/Source/GameplayTagsEditor/Private/GameplayTagContainerCustomization.cpp
--- a/Engine/Plugins/Editor/GameplayTagsEditor/Source/GameplayTagsEditor/Private/GameplayTagContainerCustomization.cpp
+++ b/Engine/Plugins/Editor/GameplayTagsEditor/Source/GameplayTagsEditor/Private/GameplayTagContainerCustomization.cpp
@@ -96,7 +96,7 @@ void FGameplayTagContainerCustomization::CustomizeChildren(TSharedRef<IPropertyH
 	if (const UBlueprintGeneratedClass* GeneratedClass = Cast<UBlueprintGeneratedClass>(StructPropertyHandle->GetOuterBaseClass()))
 	{
 		CreatorBP = Cast<UBlueprint>(GeneratedClass->ClassGeneratedBy);
-		if (!FBlueprintEditorUtils::IsVariableCreatedByBlueprint(CreatorBP.Get(), StructPropertyHandle->GetProperty()))
+		if (!CreatorBP.IsValid() || !FBlueprintEditorUtils::IsVariableCreatedByBlueprint(CreatorBP.Get(), StructPropertyHandle->GetProperty()))
 		{
 			return;
 		}
diff --git a/Engine/Plugins/Editor/GameplayTagsEditor/Source/GameplayTagsEditor/Private/GameplayTagCustomization.cpp b/Engine/Plugins/Editor/GameplayTagsEditor/Source/GameplayTagsEditor/Private/GameplayTagCustomization.cpp
--- a/Engine/Plugins/Editor/GameplayTagsEditor/Source/GameplayTagsEditor/Private/GameplayTagCustomization.cpp
+++ b/Engine/Plugins/Editor/GameplayTagsEditor/Source/GameplayTagsEditor/Private/GameplayTagCustomization.cpp
@@ -99,7 +99,7 @@ void FGameplayTagCustomization::CustomizeChildren(TSharedRef<IPropertyHandle> In
 	if (const UBlueprintGeneratedClass* GeneratedClass = Cast<UBlueprintGeneratedClass>(StructPropertyHandle->GetOuterBaseClass()))
 	{
 		CreatorBP = Cast<UBlueprint>(GeneratedClass->ClassGeneratedBy);
-		if (!FBlueprintEditorUtils::IsVariableCreatedByBlueprint(CreatorBP.Get(), StructPropertyHandle->GetProperty()))
+		if (!CreatorBP.IsValid() || !FBlueprintEditorUtils::IsVariableCreatedByBlueprint(CreatorBP.Get(), StructPropertyHandle->GetProperty()))
 		{
 			return;
 		}
diff --git a/Engine/Source/Editor/UnrealEd/Private/Kismet2/BlueprintEditorUtils.cpp b/Engine/Source/Editor/UnrealEd/Private/Kismet2/BlueprintEditorUtils.cpp
--- a/Engine/Source/Editor/UnrealEd/Private/Kismet2/BlueprintEditorUtils.cpp
+++ b/Engine/Source/Editor/UnrealEd/Private/Kismet2/BlueprintEditorUtils.cpp
@@ -3481,18 +3481,23 @@ FString FBlueprintEditorUtils::GetBlueprintTypeDescription(const UBlueprint* Blu
 bool FBlueprintEditorUtils::IsVariableCreatedByBlueprint(UBlueprint* InBlueprint, FProperty* InVariableProperty)
 {
 	bool bIsVariableCreatedByBlueprint = false;
-	if (UBlueprintGeneratedClass* GeneratedClass = Cast<UBlueprintGeneratedClass>(InVariableProperty->GetOwnerClass()))
+
+	if (InBlueprint != nullptr && InVariableProperty != nullptr)
 	{
-		UBlueprint* OwnerBlueprint = Cast<UBlueprint>(GeneratedClass->ClassGeneratedBy);
-		bIsVariableCreatedByBlueprint = (OwnerBlueprint == InBlueprint && FBlueprintEditorUtils::FindNewVariableIndex(InBlueprint, InVariableProperty->GetFName()) != INDEX_NONE);
+		if (UBlueprintGeneratedClass* GeneratedClass = Cast<UBlueprintGeneratedClass>(InVariableProperty->GetOwnerClass()))
+		{
+			UBlueprint* OwnerBlueprint = Cast<UBlueprint>(GeneratedClass->ClassGeneratedBy);
+			bIsVariableCreatedByBlueprint = (OwnerBlueprint == InBlueprint && FBlueprintEditorUtils::FindNewVariableIndex(InBlueprint, InVariableProperty->GetFName()) != INDEX_NONE);
+		}
 	}
+
 	return bIsVariableCreatedByBlueprint;
 }
 
 // Find the index of a variable first declared in this blueprint. Returns INDEX_NONE if not found.
 int32 FBlueprintEditorUtils::FindNewVariableIndex(const UBlueprint* Blueprint, const FName& InName) 
 {
-	if(InName != NAME_None)
+	if(InName != NAME_None && Blueprint != nullptr)
 	{
 		for(int32 i=0; i<Blueprint->NewVariables.Num(); i++)
 		{
