[UBT][CppCompileWarnings] Engine compliance for wundef

diff --git a/Engine/Source/Editor/UnrealEd/Private/Kismet2/DebuggerCommands.cpp b/Engine/Source/Editor/UnrealEd/Private/Kismet2/DebuggerCommands.cpp
--- a/Engine/Source/Editor/UnrealEd/Private/Kismet2/DebuggerCommands.cpp
+++ b/Engine/Source/Editor/UnrealEd/Private/Kismet2/DebuggerCommands.cpp
@@ -29,6 +29,7 @@
 #include "Interfaces/TargetDeviceId.h"
 #include "Interfaces/ITargetPlatform.h"
 #include "Interfaces/ITargetPlatformManagerModule.h"
+#include "Interfaces/ITurnkeySupportModule.h"
 #include "ITargetDeviceProxy.h"
 #include "ITargetDeviceServicesModule.h"
 #include "ISettingsModule.h"
