[HLOD] Fix distributed HLOD builds running on non-WP worlds not doing any work

diff --git a/Engine/Source/Editor/UnrealEd/Private/WorldPartition/WorldPartitionHLODsBuilder.cpp b/Engine/Source/Editor/UnrealEd/Private/WorldPartition/WorldPartitionHLODsBuilder.cpp
--- a/Engine/Source/Editor/UnrealEd/Private/WorldPartition/WorldPartitionHLODsBuilder.cpp
+++ b/Engine/Source/Editor/UnrealEd/Private/WorldPartition/WorldPartitionHLODsBuilder.cpp
@@ -929,7 +929,9 @@ TArray<UWorldPartitionHLODsBuilder::FHLODWorkload> UWorldPartitionHLODsBuilder::
 {
 	if (!WorldPartition)
 	{
-		return { FHLODWorkload() };
+		FHLODWorkload HLODWorkload;
+		HLODWorkload.PerWorldHLODWorkloads.Add( { FGuid() } );
+		return { HLODWorkload };
 	}
 
 	// Build a mapping of HLODActor to WorldPartition Index to be used when splitting actors into workloads
