Avoid creating a feedback dialog during a drag drop operation
[RN] Fixed crash that occurs when a slow task dialog is created while dragging an actor blueprint into the level

diff --git a/Engine/Source/Editor/UnrealEd/Private/FeedbackContextEditor.cpp b/Engine/Source/Editor/UnrealEd/Private/FeedbackContextEditor.cpp
--- a/Engine/Source/Editor/UnrealEd/Private/FeedbackContextEditor.cpp
+++ b/Engine/Source/Editor/UnrealEd/Private/FeedbackContextEditor.cpp
@@ -450,6 +450,14 @@ void FFeedbackContextEditor::StartSlowTask( const FText& Task, bool bShowCancelB
 
 	if (GEditor)
 	{
+		// Don't show a dialog in the middle of a drag/drop operation, as it will cause reentrancy into user
+		// drag operations:
+		if(FSlateApplication::Get().IsDragDropping())
+		{
+			UE_LOG(LogSlate, Log, TEXT("Slow task dialog suppressed because a drag operation is in progress"));
+			return;
+		}
+
 		// reset the cancellation flag
 		HasTaskBeenCancelled = false;
 
