Fix package lifetime on ODSC server : TStrongObjectPtr is necessary to make sure the package does not get collected by GC
CleanupWorldPartitionAssets stills works because we first null the TStrongObjectPtr before calling UPackageTools::UnloadPackages
This fixes some crashes due to invalid pointer to packages when calling UPackageTools::UnloadPackages

diff --git a/Engine/Source/Editor/UnrealEd/Private/Cooker/OnDemandShaderCompilation.cpp b/Engine/Source/Editor/UnrealEd/Private/Cooker/OnDemandShaderCompilation.cpp
--- a/Engine/Source/Editor/UnrealEd/Private/Cooker/OnDemandShaderCompilation.cpp
+++ b/Engine/Source/Editor/UnrealEd/Private/Cooker/OnDemandShaderCompilation.cpp
@@ -130,8 +130,11 @@ void FODSCClientData::CleanupWorldPartitionAssets()
 	for (auto Iter = WorldPartitionAssets.CreateIterator(); Iter; ++Iter)
 	{
 		FWorldPartitionAssets& DynamicMaterialData = Iter.Value();
-		PackagesToUnloadSet.FindOrAdd(DynamicMaterialData.PackagePtr);
-		DynamicMaterialData.PackagePtr = nullptr;
+		if (DynamicMaterialData.PackagePtr.IsValid())
+		{
+			PackagesToUnloadSet.FindOrAdd(DynamicMaterialData.PackagePtr.Get());
+			DynamicMaterialData.PackagePtr = nullptr;
+		}
 	}
 
 	TArray<UPackage*> PackagesToUnloadArray;
@@ -322,7 +325,7 @@ UMaterialInterface* FODSCClientData::TryFindWorldPartitionMaterial(const FSoftOb
 	if (!DynamicMaterialData->PackagePtr)
 	{
 		TRACE_CPUPROFILER_EVENT_SCOPE_TEXT(*FString::Printf(TEXT("FODSCClientData_LoadPackage %s"), *DynamicMaterialData->PackageName));
-		DynamicMaterialData->PackagePtr = LoadPackage(nullptr, *DynamicMaterialData->PackageName, LOAD_None);
+		DynamicMaterialData->PackagePtr = TStrongObjectPtr<UPackage>(LoadPackage(nullptr, *DynamicMaterialData->PackageName, LOAD_None));
 	}
 
 	TArray<FString> MaterialPathsToTry;
diff --git a/Engine/Source/Editor/UnrealEd/Private/Cooker/OnDemandShaderCompilation.h b/Engine/Source/Editor/UnrealEd/Private/Cooker/OnDemandShaderCompilation.h
--- a/Engine/Source/Editor/UnrealEd/Private/Cooker/OnDemandShaderCompilation.h
+++ b/Engine/Source/Editor/UnrealEd/Private/Cooker/OnDemandShaderCompilation.h
@@ -33,7 +33,7 @@ private:
 	struct FWorldPartitionAssets
 	{
 		FString PackageName;
-		TObjectPtr<UPackage> PackagePtr;
+		TStrongObjectPtr<UPackage> PackagePtr;
 	};
 
 	static TMap<FString, FWorldPartitionAssets> WorldPartitionAssets;
