Fix camera speed text.

#rnx

diff --git a/Engine/Source/Editor/UnrealEd/Private/ViewportToolbar/UnrealEdViewportToolbar.cpp b/Engine/Source/Editor/UnrealEd/Private/ViewportToolbar/UnrealEdViewportToolbar.cpp
--- a/Engine/Source/Editor/UnrealEd/Private/ViewportToolbar/UnrealEdViewportToolbar.cpp
+++ b/Engine/Source/Editor/UnrealEd/Private/ViewportToolbar/UnrealEdViewportToolbar.cpp
@@ -3079,7 +3079,7 @@ TSharedRef<SWidget> CreateCameraSpeedWidget(const TSharedPtr<FEditorViewportClie
 		[
 			SNew(STextBlock)
 			.ColorAndOpacity(FSlateColor::UseSubduedForeground())
-			.Text(LOCTEXT("CameraSpeedSliderDescription", "Hold either mouse button use the scroll wheel to adjust camera speed on the fly."))
+			.Text(LOCTEXT("CameraSpeedSliderDescription", "Hold either mouse button and use the scroll wheel to adjust the speed on the fly."))
 			.WrapTextAt(260.f)
 		]
 		;
@@ -4583,7 +4583,7 @@ FToolMenuEntry CreateCameraSpeedEntry(const TWeakPtr<::SEditorViewport>& InEdito
 		FText::GetEmpty()
 	);
 	
-	CameraSpeedSliderEntry.ToolTip = LOCTEXT("CameraSpeedSliderTooltip", "Sets the speed of the camera in FPS mode.\n\nCan be changed on the fly by scrolling the mouse wheel\nwhile holding the left or right mouse button.");
+	CameraSpeedSliderEntry.ToolTip = LOCTEXT("CameraSpeedSliderTooltip", "The speed of the camera when moving around using the keyboard (WASD).");
 
 	return CameraSpeedSliderEntry;
 }
