Cooker: Suppress static analysis error: "Dereferencing NULL pointer 'KeyInQueue'.
#rnx

diff --git a/Engine/Source/Editor/UnrealEd/Private/CookOnTheFlyServer.cpp b/Engine/Source/Editor/UnrealEd/Private/CookOnTheFlyServer.cpp
--- a/Engine/Source/Editor/UnrealEd/Private/CookOnTheFlyServer.cpp
+++ b/Engine/Source/Editor/UnrealEd/Private/CookOnTheFlyServer.cpp
@@ -2086,6 +2086,8 @@ void UCookOnTheFlyServer::FPollable::TriggerInternal(UCookOnTheFlyServer& COTFS)
 		[this](const FPollableQueueKey& Existing) { return Existing.Pollable.GetReference() == this; });
 	if (ensure(KeyInQueue))
 	{
+		// CA_ASSUME is temporarily required because static analysis is not recognizing the if condition.
+		CA_ASSUME(KeyInQueue != nullptr);
 		FPollableQueueKey LocalQueueKey;
 		LocalQueueKey.Pollable = MoveTemp(KeyInQueue->Pollable);
 		// If the top of the heap is already triggered, put this after the top of the heap to
@@ -2130,6 +2132,8 @@ void UCookOnTheFlyServer::FPollable::RunNowInternal(UCookOnTheFlyServer & COTFS,
 		[this](const FPollableQueueKey& Existing) { return Existing.Pollable.GetReference() == this; });
 	if (ensure(KeyInQueue))
 	{
+		// CA_ASSUME is temporarily required because static analysis is not recognizing the if condition.
+		CA_ASSUME(KeyInQueue != nullptr);
 		FPollableQueueKey LocalQueueKey;
 		LocalQueueKey.Pollable = MoveTemp(KeyInQueue->Pollable);
 		LocalQueueKey.NextTimeSeconds = TimeLastRun + this->PeriodSeconds;
