BuildLayers and MultiProductBuild: Remove the temporary API function ExperimentalDoNotUse_IsCookingDLCBaseGamePlugin; and instead change GetCookingDLC to return ECookingDLC::No for MultiProductBuild. This is the second temporary implementation and is just to avoid the need for the temporary experimental API function. We plan to modify MultiProductBuild into the BuildLayers concept and will change the GetCookingDLC API more extensively when that happens.
#rnx

diff --git a/Engine/Source/Editor/UnrealEd/Classes/CookOnTheSide/CookOnTheFlyServer.h b/Engine/Source/Editor/UnrealEd/Classes/CookOnTheSide/CookOnTheFlyServer.h
--- a/Engine/Source/Editor/UnrealEd/Classes/CookOnTheSide/CookOnTheFlyServer.h
+++ b/Engine/Source/Editor/UnrealEd/Classes/CookOnTheSide/CookOnTheFlyServer.h
@@ -698,7 +698,6 @@ public:
 	UNREALED_API virtual TArray<UE::Cook::FInstigator> GetInstigatorChain(FName PackageName) override;
 	UNREALED_API virtual UE::Cook::ECookType GetCookType() override;
 	UNREALED_API virtual UE::Cook::ECookingDLC GetCookingDLC() override;
-	UNREALED_API virtual bool ExperimentalDoNotUse_IsCookingDLCBaseGamePlugin() override;
 	UNREALED_API virtual FString GetDLCName() override;
 	UNREALED_API virtual UE::Cook::EProcessType GetProcessType() override;
 	UNREALED_API virtual UE::Cook::ECookValidationOptions GetCookValidationOptions() override;
diff --git a/Engine/Source/Editor/UnrealEd/Private/CookOnTheFlyServer.cpp b/Engine/Source/Editor/UnrealEd/Private/CookOnTheFlyServer.cpp
--- a/Engine/Source/Editor/UnrealEd/Private/CookOnTheFlyServer.cpp
+++ b/Engine/Source/Editor/UnrealEd/Private/CookOnTheFlyServer.cpp
@@ -6809,12 +6809,19 @@ UE::Cook::ECookType UCookOnTheFlyServer::GetCookType()
 UE::Cook::ECookingDLC UCookOnTheFlyServer::GetCookingDLC()
 {
 	using namespace UE::Cook;
-	return IsCookingDLC() ? ECookingDLC::Yes : ECookingDLC::No;
-}
-
-bool UCookOnTheFlyServer::ExperimentalDoNotUse_IsCookingDLCBaseGamePlugin()
-{
-	return IsCookingDLC() && bCookingDLCBaseGamePlugin;
+	if (!IsCookingDLC())
+	{
+		return ECookingDLC::No;
+	}
+	if (bCookingDLCBaseGamePlugin)
+	{
+		// Temporary solution to support MultiProductBuild: separate cooks that will be welded together into the
+		// basegame BuildLayer should not be treated as the legacy idea of DLC. We plan to change this later into a
+		// new concept - BuildLayers - and replace the ECookingDLC API with functions that read properties about
+		// the BuildLayer being cooked from project-specific code in AssetManager.
+		return ECookingDLC::No;
+	}
+	return ECookingDLC::Yes;
 }
 
 FString UCookOnTheFlyServer::GetDLCName()
diff --git a/Engine/Source/Runtime/CoreUObject/Public/UObject/ICookInfo.h b/Engine/Source/Runtime/CoreUObject/Public/UObject/ICookInfo.h
--- a/Engine/Source/Runtime/CoreUObject/Public/UObject/ICookInfo.h
+++ b/Engine/Source/Runtime/CoreUObject/Public/UObject/ICookInfo.h
@@ -213,7 +213,6 @@ public:
 	virtual UE::Cook::ECookType GetCookType() = 0;
 	/** Whether DLC is being cooked (e.g. via "-dlcname=<PluginName>"). This function will not return ECookingDLC::Unknown. */
 	virtual UE::Cook::ECookingDLC GetCookingDLC() = 0;
-	virtual bool ExperimentalDoNotUse_IsCookingDLCBaseGamePlugin() = 0;
 	/** When DLC is being cooked (@see GetCookingDLC), this will return the name of the DLC plugin. */
 	virtual FString GetDLCName() = 0;
 	/** The role the current process plays in its MPCook session, or EProcessType::SingleProcess if it is running standalone. */
diff --git a/Engine/Source/Runtime/Engine/Private/Materials/MaterialInterface.cpp b/Engine/Source/Runtime/Engine/Private/Materials/MaterialInterface.cpp
--- a/Engine/Source/Runtime/Engine/Private/Materials/MaterialInterface.cpp
+++ b/Engine/Source/Runtime/Engine/Private/Materials/MaterialInterface.cpp
@@ -212,7 +212,7 @@ void UMaterialInterface::Serialize(FArchive& Ar)
 	{
 		// Mark whether this material is part of the base game or of a DLC shipped with the base game.
 		// This provides information on whether it is safe to be used as a parent safely in child modules.
-		bIncludedInBaseGame = Ar.GetCookContext()->GetCookingDLC() == UE::Cook::ECookingDLC::No || CookContext->GetCookInfo()->ExperimentalDoNotUse_IsCookingDLCBaseGamePlugin();
+		bIncludedInBaseGame = Ar.GetCookContext()->GetCookingDLC() == UE::Cook::ECookingDLC::No;
 	}
 #endif
 
