[Depth Bar] Separates the clipping behavior of focus to its own command.

#rnx

diff --git a/Engine/Source/Editor/LevelEditor/Private/ViewportToolbar/LevelEditorSubmenus.cpp b/Engine/Source/Editor/LevelEditor/Private/ViewportToolbar/LevelEditorSubmenus.cpp
--- a/Engine/Source/Editor/LevelEditor/Private/ViewportToolbar/LevelEditorSubmenus.cpp
+++ b/Engine/Source/Editor/LevelEditor/Private/ViewportToolbar/LevelEditorSubmenus.cpp
@@ -627,6 +627,9 @@ FToolMenuEntry CreateCameraMovementSubmenu(const TWeakPtr<::SLevelViewport>& Lev
 							FToolMenuEntry& FrameEntry =
 								CameraMovementSection.AddMenuEntry(FEditorViewportCommands::Get().FocusViewportToSelection);
 							FrameEntry.Icon = FSlateIcon(FAppStyle::GetAppStyleSetName(), "Icons.FrameActor");
+							
+							FToolMenuEntry& FrameAndClip = CameraMovementSection.AddMenuEntry(FEditorViewportCommands::Get().FocusAndClipViewportToSelection);
+							FrameAndClip.Icon = FSlateIcon(FAppStyle::GetAppStyleSetName(), "Icons.FrameActor");
 						}
 
 						// The following entries are perspective view only
diff --git a/Engine/Source/Editor/UnrealEd/Private/EditorViewportClient.cpp b/Engine/Source/Editor/UnrealEd/Private/EditorViewportClient.cpp
--- a/Engine/Source/Editor/UnrealEd/Private/EditorViewportClient.cpp
+++ b/Engine/Source/Editor/UnrealEd/Private/EditorViewportClient.cpp
@@ -959,15 +959,45 @@ void FEditorViewportClient::FocusViewportOnBox( const FBox& BoundingBox, bool bI
 		}
 		else
 		{
-			// For ortho viewports just set the camera position to the center of the bounding volume.
-			//SetViewLocation( Position );
-			ViewTransform.TransitionToLocation(Position, EditorViewportWidget, bInstant);
-			
 			if (SupportsManualOrthographicDepthControl())
 			{
-				// Shift the orthographic near view plane back so that it is behind the object
-				OrthographicNearPlane = -Radius;
+				const FVector StartingLocation = ViewTransform.GetLocation();
+				const FVector Forward = GetForwardVector();
+				
+				// Adjust the camera by the current forward instantly
+				const FVector LocationDelta = Position - StartingLocation;
+				const FVector ForwardLocationDelta = LocationDelta.ProjectOnToNormal(Forward);
+				ViewTransform.SetLocation(StartingLocation + ForwardLocationDelta);
+				
+				// Adjust the near & far planes so that they remain fixed in world space
+				const double ForwardDelta = ForwardLocationDelta.Dot(Forward);
+				
+				if (OrthographicNearPlane.IsSet())
+				{
+					OrthographicNearPlane = OrthographicNearPlane.GetValue() - ForwardDelta;
+				}
+				
+				if (OrthographicFarPlane.IsSet())
+				{
+					OrthographicFarPlane = OrthographicFarPlane.GetValue() - ForwardDelta;
+				}
+				
+				if (GetClipToSelectionOnFocus())
+				{
+					// Shift the orthographic near view plane back so that it is just before the object
+					OrthographicNearPlane = -Radius;
+					
+					// Ensure the orthographic far view plane can display the object
+					if (OrthographicFarPlane.IsSet())
+					{
+						OrthographicFarPlane = FMath::Max(OrthographicFarPlane.GetValue(), Radius);
+					}
+				}
 			}
+		
+			// For ortho viewports just set the camera position to the center of the bounding volume.
+			//SetViewLocation( Position );
+			ViewTransform.TransitionToLocation(Position, EditorViewportWidget, bInstant);
 			
 			if (!(Viewport->KeyState(EKeys::LeftControl) || Viewport->KeyState(EKeys::RightControl)))
 			{
@@ -995,6 +1025,10 @@ void FEditorViewportClient::FocusViewportOnBox( const FBox& BoundingBox, bool bI
 	Invalidate();
 }
 
+void FEditorViewportClient::SetClipToSelectionOnFocus(bool bInClipToSelection)
+{
+	bClipToSelectionOnFocus = bInClipToSelection;
+}
 
 void FEditorViewportClient::CenterViewportAtPoint(const FVector& NewLookAt, bool bInstant /* = false */)
 {
@@ -1386,9 +1420,13 @@ FSceneView* FEditorViewportClient::CalcSceneView(FSceneViewFamily* ViewFamily, c
 					CalculatePlanesViewInfo.OrthoFarClipPlane = DepthSpace.GetAlignedMax();
 				}
 				
-				ViewInitOptions.UpdateOrthoPlanes(CalculatePlanesViewInfo);	
-				ZScale = 1.0f / (CalculatePlanesViewInfo.OrthoFarClipPlane - CalculatePlanesViewInfo.OrthoNearClipPlane);
-				ZOffset = -CalculatePlanesViewInfo.OrthoNearClipPlane;
+				ViewInitOptions.UpdateOrthoPlanes(CalculatePlanesViewInfo);
+				const float Denominator = CalculatePlanesViewInfo.OrthoFarClipPlane - CalculatePlanesViewInfo.OrthoNearClipPlane;
+				if (Denominator > 0.0f)
+				{
+					ZScale = 1.0f / Denominator;
+					ZOffset = -CalculatePlanesViewInfo.OrthoNearClipPlane;
+				}
 			}
 			else if(!ViewFamily->EngineShowFlags.Wireframe)
 			{
diff --git a/Engine/Source/Editor/UnrealEd/Private/EditorViewportCommands.cpp b/Engine/Source/Editor/UnrealEd/Private/EditorViewportCommands.cpp
--- a/Engine/Source/Editor/UnrealEd/Private/EditorViewportCommands.cpp
+++ b/Engine/Source/Editor/UnrealEd/Private/EditorViewportCommands.cpp
@@ -193,6 +193,7 @@ void FEditorViewportCommands::RegisterCommands()
 	
 	UI_COMMAND( FocusAllViewportsToSelection, "Frame Selected Actors in All Viewports", "Centers all open viewports on the selected actors", EUserInterfaceActionType::Button, FInputChord(EKeys::F, EModifierKey::Shift) );
 	UI_COMMAND( FocusViewportToSelection, "Frame Selected", "Centers the viewport on the selected actors", EUserInterfaceActionType::Button, FInputChord( EKeys::F ) );
+	UI_COMMAND( FocusAndClipViewportToSelection, "Frame & Clip To Selected", "Centers the viewport on the selected actors and clips the near view plane to them", EUserInterfaceActionType::Button, FInputChord(EKeys::F, EModifierKey::Control | EModifierKey::Alt) );
 	UI_COMMAND( FocusOutlinerToSelection, "Frame Selected in Outliner", "Focuses the selected actors in all open outliners", EUserInterfaceActionType::Button, FInputChord() );
 	UI_COMMAND( FocusOutlinerToContextFolder, "Frame Current Folder in Outliner", "Focuses the current folder (if it's set) in the outliner", EUserInterfaceActionType::Button, FInputChord() );
 
diff --git a/Engine/Source/Editor/UnrealEd/Private/SEditorViewport.cpp b/Engine/Source/Editor/UnrealEd/Private/SEditorViewport.cpp
--- a/Engine/Source/Editor/UnrealEd/Private/SEditorViewport.cpp
+++ b/Engine/Source/Editor/UnrealEd/Private/SEditorViewport.cpp
@@ -574,6 +574,11 @@ void SEditorViewport::BindCommands()
 		//FExecuteAction::CreateStatic( &FLevelEditorActionCallbacks::ExecuteExecCommand, FString( TEXT("CAMERA ALIGN ACTIVEVIEWPORTONLY") ) )
 		);
 		
+	FUIAction FocusToAndClipAction = FExecuteAction::CreateSP( this, &SEditorViewport::OnFocusAndClipViewportToSelection );
+	FocusToAndClipAction.IsActionVisibleDelegate.BindSP( this, &SEditorViewport::CanUseDepthBar );
+		
+	CommandListRef.MapAction(Commands.FocusAndClipViewportToSelection, FocusToAndClipAction);
+		
 	CommandListRef.MapAction(
 		Commands.UseOrthographicClippingPlanes,
 		FExecuteAction::CreateSP( this, &SEditorViewport::OnToggleUseDepthBar ),
@@ -731,6 +736,14 @@ EVisibility SEditorViewport::OnGetViewportContentVisibility() const
 	return bIsViewportUIHidden ? EVisibility::Collapsed : EVisibility::SelfHitTestInvisible;
 }
 
+void SEditorViewport::OnFocusAndClipViewportToSelection()
+{
+	const bool bPrevious = Client->GetClipToSelectionOnFocus(); 
+	Client->SetClipToSelectionOnFocus(true);
+	OnFocusViewportToSelection();
+	Client->SetClipToSelectionOnFocus(bPrevious);
+}
+
 EVisibility SEditorViewport::OnGetDepthBarVisibility() const
 {
 	if (bHasSpaceForDepthBar && Client->SupportsManualOrthographicDepthControl() && Client->IsOrtho() && bShowDepthBar)
@@ -1161,6 +1174,10 @@ TSharedPtr<SWidget> SEditorViewport::BuildViewportDepthBar()
 	
 		FToolMenuSection& Actions = Menu->AddSection("Actions");
 		Actions.AddSeparator("Start");
+		
+		FToolMenuEntry& FocusEntry = Actions.AddMenuEntry(FEditorViewportCommands::Get().FocusAndClipViewportToSelection);
+		FocusEntry.Icon = FSlateIcon(FAppStyle::GetAppStyleSetName(), "Icons.FrameActor");
+		
 		Actions.AddEntry(CreateOrthographicClippingPlanesSubmenu(FOrthographicClippingPlanesSubmenuOptions { false } ));
 		
 		FToolMenuSection& EndSection = Menu->AddSection("End", {}, FToolMenuInsert(NAME_None, EToolMenuInsertType::Last));
diff --git a/Engine/Source/Editor/UnrealEd/Private/ViewportToolbar/UnrealEdViewportToolbar.cpp b/Engine/Source/Editor/UnrealEd/Private/ViewportToolbar/UnrealEdViewportToolbar.cpp
--- a/Engine/Source/Editor/UnrealEd/Private/ViewportToolbar/UnrealEdViewportToolbar.cpp
+++ b/Engine/Source/Editor/UnrealEd/Private/ViewportToolbar/UnrealEdViewportToolbar.cpp
@@ -3579,6 +3579,9 @@ void PopulateCameraMenu(UToolMenu* InMenu, const FViewportCameraMenuOptions& InO
 
 			FToolMenuEntry& FrameEntry = MovementSection.AddMenuEntry(FEditorViewportCommands::Get().FocusViewportToSelection);
 			FrameEntry.Icon = FSlateIcon(FAppStyle::GetAppStyleSetName(), "Icons.FrameActor");
+			
+			FToolMenuEntry& FrameAndClip = MovementSection.AddMenuEntry(FEditorViewportCommands::Get().FocusAndClipViewportToSelection);
+			FrameAndClip.Icon = FSlateIcon(FAppStyle::GetAppStyleSetName(), "Icons.FrameActor");
 		}
 	}
 
diff --git a/Engine/Source/Editor/UnrealEd/Public/EditorViewportClient.h b/Engine/Source/Editor/UnrealEd/Public/EditorViewportClient.h
--- a/Engine/Source/Editor/UnrealEd/Public/EditorViewportClient.h
+++ b/Engine/Source/Editor/UnrealEd/Public/EditorViewportClient.h
@@ -1093,6 +1093,16 @@ public:
 	 */
 	UNREALED_API void FocusViewportOnBox( const FBox& BoundingBox, bool bInstant = false );
 
+	/**
+	 * Gets whether focusing on a box should clip to selection.
+	 */
+	bool GetClipToSelectionOnFocus() const { return bClipToSelectionOnFocus; }
+
+	/**
+	 * Sets whether focusing on a box should clip to selection.
+	 */
+	void SetClipToSelectionOnFocus(bool bInClipToSelection);
+
 	/**
 	 * Translates the viewport so that the given LookAt point is at the center of viewport, while maintaining current Location/LookAt distance
 	 *
@@ -2220,6 +2230,9 @@ private:
 	/** The location of the far clipping plane in orthographic view. */
 	TOptional<double> OrthographicFarPlane;
 	
+	/** Whether focusing on a box should clip the near plane to that selection. */
+	bool bClipToSelectionOnFocus = false;
+	
 	/** The cached bounds of the world. */
 	FBox CachedWorldBounds;
 	
diff --git a/Engine/Source/Editor/UnrealEd/Public/EditorViewportCommands.h b/Engine/Source/Editor/UnrealEd/Public/EditorViewportCommands.h
--- a/Engine/Source/Editor/UnrealEd/Public/EditorViewportCommands.h
+++ b/Engine/Source/Editor/UnrealEd/Public/EditorViewportCommands.h
@@ -226,6 +226,7 @@ public:
 
 	TSharedPtr< FUICommandInfo > FocusAllViewportsToSelection;
 	TSharedPtr< FUICommandInfo > FocusViewportToSelection;
+	TSharedPtr< FUICommandInfo > FocusAndClipViewportToSelection;
 	TSharedPtr< FUICommandInfo > FocusOutlinerToSelection;
 	TSharedPtr< FUICommandInfo > FocusOutlinerToContextFolder;
 	
diff --git a/Engine/Source/Editor/UnrealEd/Public/SEditorViewport.h b/Engine/Source/Editor/UnrealEd/Public/SEditorViewport.h
--- a/Engine/Source/Editor/UnrealEd/Public/SEditorViewport.h
+++ b/Engine/Source/Editor/UnrealEd/Public/SEditorViewport.h
@@ -367,6 +367,9 @@ private:
 	 */
 	void RegisterViewportInteractions();
 	
+	/** Called when the user wants to focus the viewport to the current selection and clip the near plane to the bounds of that selection */
+	void OnFocusAndClipViewportToSelection();
+
 	/**
 	 * Creates the viewport depth bar and any containing widgets
 	 */
