Re-adds UE_EXPERIMENTAL to SDepthBar and FDepthBarStyle.

diff --git a/Engine/Source/Editor/EditorFramework/Public/Styling/DepthBarStyle.h b/Engine/Source/Editor/EditorFramework/Public/Styling/DepthBarStyle.h
--- a/Engine/Source/Editor/EditorFramework/Public/Styling/DepthBarStyle.h
+++ b/Engine/Source/Editor/EditorFramework/Public/Styling/DepthBarStyle.h
@@ -9,6 +9,7 @@
 
 #include "DepthBarStyle.generated.h"
 
+struct UE_EXPERIMENTAL(5.7, "Represents the style of an SDepthBar") FDepthBarStyle;
 USTRUCT()
 struct FDepthBarStyle : public FSlateWidgetStyle
 {
diff --git a/Engine/Source/Editor/EditorFramework/Public/Viewports/SDepthBar.h b/Engine/Source/Editor/EditorFramework/Public/Viewports/SDepthBar.h
--- a/Engine/Source/Editor/EditorFramework/Public/Viewports/SDepthBar.h
+++ b/Engine/Source/Editor/EditorFramework/Public/Viewports/SDepthBar.h
@@ -13,6 +13,7 @@
 /**
  * A bar that allows for visualization and manipulation of a viewport's depth
  */
+class UE_EXPERIMENTAL(5.7, "SDepthBar represents and controls view depth within viewports.") SDepthBar;
 class SDepthBar : public SPanel
 {
 public:
@@ -108,7 +109,9 @@ public:
 		
 		SLATE_SLOT_ARGUMENT(SDepthBar::FDepthIndicatorSlot, Slots)
 		
+		PRAGMA_DISABLE_EXPERIMENTAL_WARNINGS
 		SLATE_STYLE_ARGUMENT(FDepthBarStyle, Style)
+		PRAGMA_ENABLE_EXPERIMENTAL_WARNINGS
 	SLATE_END_ARGS()
 	
 	using FScopedWidgetDepthIndicatorSlotArguments = TPanelChildren<FDepthIndicatorSlot>::FScopedWidgetSlotArguments;
@@ -258,6 +261,8 @@ private:
 	FOnDepthChanged OnNearPlaneChanged;
 	FOnDepthChanged OnFarPlaneChanged;
 	
+	PRAGMA_DISABLE_EXPERIMENTAL_WARNINGS
 	const FDepthBarStyle* Style = nullptr;
+	PRAGMA_ENABLE_EXPERIMENTAL_WARNINGS
 };
 
diff --git a/Engine/Source/Editor/UnrealEd/Public/SEditorViewport.h b/Engine/Source/Editor/UnrealEd/Public/SEditorViewport.h
--- a/Engine/Source/Editor/UnrealEd/Public/SEditorViewport.h
+++ b/Engine/Source/Editor/UnrealEd/Public/SEditorViewport.h
@@ -340,9 +340,11 @@ protected:
 	/** The client responsible for setting up the scene */
 	TSharedPtr<FEditorViewportClient> Client;
 	
+	PRAGMA_DISABLE_EXPERIMENTAL_WARNINGS
 	/** The depth bar used by this viewport */
 	UE_EXPERIMENTAL(5.7, "Experimental interface to the Depth bar controls")
 	TSharedPtr<SDepthBar> DepthBar;
+	PRAGMA_ENABLE_EXPERIMENTAL_WARNINGS
 	
 	/** Commandlist used in the viewport (Maps commands to viewport specific actions) */
 	TSharedPtr<FUICommandList> CommandList;
@@ -392,12 +394,14 @@ private:
 	
 	/** @return The visibility of the viewport depth bar. */
 	EVisibility OnGetDepthBarVisibility() const;
-	
+
+	PRAGMA_DISABLE_EXPERIMENTAL_WARNINGS	
 	/** @return The current mode of the depth bar. */
 	SDepthBar::EMode GetDepthBarMode() const;
-	
+
 	/** @return The current depth space of the depth bar. */
 	TOptional<SDepthBar::FDepthSpace> GetDepthBarSpace() const;
+	PRAGMA_ENABLE_EXPERIMENTAL_WARNINGS
 	
 	/** @return The current near plane for the depth bar. */
 	TOptional<double> GetDepthBarNearPlane() const;
