Cooker: SkipOnlyEditorOnly: Add a configurable category for UnexpectedLoad log messages so that we can leave it on by default but turn it off for large projects that don't want the spam yet. Change the log message for UnexpectedLoads to be clearer about what the problem is and how to fix it.
#rnx

diff --git a/Engine/Config/BaseEditor.ini b/Engine/Config/BaseEditor.ini
--- a/Engine/Config/BaseEditor.ini
+++ b/Engine/Config/BaseEditor.ini
@@ -368,6 +368,14 @@ CookContentMissingSeverity=Error
 ; But this warning can be spurious on farm machines so the project may want to downgrade it to display.
 CookerIdleWarningSeverity=Warning
 
+; Unexpected loads occur when loading a source package during cooking and system-specific c++ code in hooks called for
+; objects within that package (Serialize, PostLoad, BeginCacheForCookedPlatformData, PreSave) calls LoadPackage on a
+; target package that was not declared as a hard dependency (aka Import) of the source package when it was saved in the
+; editor. Unexpected loads are handled conservatively and are added to the cooked output, which might be incorrect
+; because often they are only needed during cooking and not at runtime. We thus log them out. The severity of the log
+; message is configurable because a project might not want to take the time to fix them all yet.
+UnexpectedLoadWarningSeverity=Display
+
 ; When Incremental Cook is enabled, either by CookIncrementalDefault=true or by -forcerecook=false on commandline,
 ; Incremental Cook causes every cooked package to test its dependencies (its package, the schemas of
 ; its classes, its referenced config settings, its referenced other packages, and others) for whether it has changed.
diff --git a/Engine/Source/Editor/UnrealEd/Classes/CookOnTheSide/CookOnTheFlyServer.h b/Engine/Source/Editor/UnrealEd/Classes/CookOnTheSide/CookOnTheFlyServer.h
--- a/Engine/Source/Editor/UnrealEd/Classes/CookOnTheSide/CookOnTheFlyServer.h
+++ b/Engine/Source/Editor/UnrealEd/Classes/CookOnTheSide/CookOnTheFlyServer.h
@@ -1574,6 +1574,7 @@ private:
 	TSet<FName> CookFilterIncludedAssetClasses;
 
 	ELogVerbosity::Type CookerIdleWarningSeverity = ELogVerbosity::Warning;
+	ELogVerbosity::Type UnexpectedLoadWarningSeverity = ELogVerbosity::Warning;
 	UE::Cook::EMPCookGeneratorSplit MPCookGeneratorSplit = UE::Cook::EMPCookGeneratorSplit::AnyWorker;
 	/** True when PumpLoads has detected it is blocked on async work and CookOnTheFlyServer should do work elsewhere. */
 	bool bLoadBusy = false;
diff --git a/Engine/Source/Editor/UnrealEd/Private/CookOnTheFlyServer.cpp b/Engine/Source/Editor/UnrealEd/Private/CookOnTheFlyServer.cpp
--- a/Engine/Source/Editor/UnrealEd/Private/CookOnTheFlyServer.cpp
+++ b/Engine/Source/Editor/UnrealEd/Private/CookOnTheFlyServer.cpp
@@ -4795,6 +4795,7 @@ void UCookOnTheFlyServer::ProcessUnsolicitedPackages(TArray<FName>* OutDiscovere
 		}
 
 		FPackageData* Referencer = nullptr;
+		FName InstancedLeafLoadedName;
 		if (!Instigator.Referencer.IsNone())
 		{
 			// Packages loaded by an instanced package referencer need to redirect the Instigator to be the Instigator of
@@ -4804,6 +4805,7 @@ void UCookOnTheFlyServer::ProcessUnsolicitedPackages(TArray<FName>* OutDiscovere
 				= PackageTracker->FindInstancedPackage(Instigator.Referencer);
 			if (Instance)
 			{
+				InstancedLeafLoadedName = Instance->PackageName;
 				Instigator.Referencer = Instance->Instigator.Referencer;
 				if (Instigator.Category == EInstigator::Unsolicited)
 				{
@@ -4846,18 +4848,12 @@ void UCookOnTheFlyServer::ProcessUnsolicitedPackages(TArray<FName>* OutDiscovere
 				Instigator.Category != EInstigator::StartupPackageCookLoadScope &&
 				Instigator.Category != EInstigator::StartupSoftObjectPath)
 			{
-				FString Message = FString::Printf(
-					TEXT("Unsolicited package %s was loaded outside of any cook operation on a referencer package; we don't know why it was loaded. ")
+#if !NO_LOGGING
+				FMsg::Logf(__FILE__, __LINE__, LogCook.GetCategoryName(), UnexpectedLoadWarningSeverity,
+					TEXT("UnexpectedLoad: Package %s was loaded outside of any cook operation on a referencer package; we don't know why it was loaded. ")
 					TEXT("Adding it to the current cook, but it will possibly not be found in future incremental cooks."),
 					*PackageNameStr);
-				if (bHiddenDependenciesDebug)
-				{
-					UE_LOG(LogCook, Error, TEXT("%s"), *Message);
-				}
-				else
-				{
-					UE_LOG(LogCook, Display, TEXT("%s"), *Message);
-				}
+#endif
 			}
 			QueueDiscoveredPackage(*PackageData, FInstigator(Instigator), EDiscoveredPlatformSet::CopyFromInstigator);
 			continue;
@@ -4906,17 +4902,23 @@ void UCookOnTheFlyServer::ProcessUnsolicitedPackages(TArray<FName>* OutDiscovere
 
 			// Otherwise, it is a bug; unsolicited loads are not allowed. Report it, and add it as a
 			// discovered dependency.
-			FString Message = FString::Printf(
-				TEXT("Unsolicited package %s was loaded by package %s. This is a hidden dependency, causes poor cook performance, and it might be a bug to add it to the cook. Declare this package in the AssetRegistry dependencies of the loading package, or mark up its load with FCookLoadScope to specify whether it is runtime or editoronly."),
-				*PackageNameStr, *WriteToString<256>(Instigator.Referencer));
-			if (bHiddenDependenciesDebug)
+			FString DescriptionOfReferencer;
+			if (InstancedLeafLoadedName.IsNone())
 			{
-				UE_LOG(LogCook, Error, TEXT("%s"), *Message);
+				DescriptionOfReferencer = Instigator.Referencer.ToString();
 			}
 			else
 			{
-				UE_LOG(LogCook, Verbose, TEXT("%s"), *Message);
+				DescriptionOfReferencer = FString::Printf(TEXT("%s (an instanced package load triggered by %s)"),
+					*InstancedLeafLoadedName.ToString(), *Instigator.Referencer.ToString());
 			}
+#if !NO_LOGGING
+			FMsg::Logf(__FILE__, __LINE__, LogCook.GetCategoryName(), UnexpectedLoadWarningSeverity,
+				TEXT("UnexpectedLoad: TargetPackage %s was unexpectededly loaded during the cook of SourcePackage %s; it was not declared as a hard dependency (aka Import) of the SourcePackage. ")
+				TEXT("We have to conservatively add this package to the cook, but it might not be needed at runtime. ")
+				TEXT("Declare the TargetPackage as an import during editor save of the SourcePackage, or mark up its load with FCookLoadScope to specify whether it is runtime or editoronly."),
+				*PackageNameStr, *DescriptionOfReferencer);
+#endif
 		}
 
 		Referencer->AddDiscoveredDependency(EDiscoveredPlatformSet::CopyFromInstigator, PackageData, Instigator.Category);
@@ -6502,6 +6504,8 @@ void UCookOnTheFlyServer::SetInitializeConfigSettings(UE::Cook::FInitializeConfi
 	FString Severity;
 	GConfig->GetString(TEXT("CookSettings"), TEXT("CookerIdleWarningSeverity"), Severity, GEditorIni);
 	CookerIdleWarningSeverity = ParseLogVerbosityFromString(Severity);
+	GConfig->GetString(TEXT("CookSettings"), TEXT("UnexpectedLoadWarningSeverity"), Severity, GEditorIni);
+	UnexpectedLoadWarningSeverity = ParseLogVerbosityFromString(Severity);
 
 	bCookFastStartup = FParse::Param(CommandLine, TEXT("cookfaststartup"));
 
