Fix a LOCTEXT id collision.

#rnx

diff --git a/Engine/Source/Editor/UnrealEd/Private/ViewportToolbar/UnrealEdViewportToolbar.cpp b/Engine/Source/Editor/UnrealEd/Private/ViewportToolbar/UnrealEdViewportToolbar.cpp
--- a/Engine/Source/Editor/UnrealEd/Private/ViewportToolbar/UnrealEdViewportToolbar.cpp
+++ b/Engine/Source/Editor/UnrealEd/Private/ViewportToolbar/UnrealEdViewportToolbar.cpp
@@ -2992,7 +2992,7 @@ TSharedRef<SWidget> CreateCameraSpeedWidget(const TSharedPtr<FEditorViewportClie
 			[
 				SNew(STextBlock)
 				.ColorAndOpacity(FLinearColor::White)
-				.Text(LOCTEXT("CameraSpeedLabel", "Speed"))
+				.Text(LOCTEXT("CameraSpeedWidgetLabel", "Speed"))
 				.TextStyle(FAppStyle::Get(), "Menu.Label")
 			]
 			+ SHorizontalBox::Slot()
