Cooker: Write the CookSettings.txt file using FConfigFile::WriteToString rather than Write, to avoid suppression of the write when -Multiprocess is on the commandline. Multiprocess is used in some DLC cooks when multiple DLCs can be cooked together. -Multiprocess is separate from the -MultiprocessId used by MPCook CookWorkers, and we also don't need it because we only attempt to write CookSettings.txt from the director process.
#rnx
#incrementalcook

diff --git a/Engine/Source/Editor/UnrealEd/Private/CookOnTheFlyServer.cpp b/Engine/Source/Editor/UnrealEd/Private/CookOnTheFlyServer.cpp
--- a/Engine/Source/Editor/UnrealEd/Private/CookOnTheFlyServer.cpp
+++ b/Engine/Source/Editor/UnrealEd/Private/CookOnTheFlyServer.cpp
@@ -7727,8 +7727,10 @@ void UCookOnTheFlyServer::SaveCookSettings(const TMap<FName, FString>& CurrentCo
 		ConfigFile.AddToSection(TEXT_CookSettings, CurrentSetting.Key, CurrentSetting.Value);
 	}
 	ConfigFile.AddToSection(TEXT_CookSettings, TEXT_CookInProgress, TEXT("true"));
-	ConfigFile.Dirty = true; // Writing to a section does not set the dirty flag, so set it manually to make Write work
-	ConfigFile.Write(GetCookSettingsFileName(TargetPlatform));
+	// Use ConfigFile.WriteToString instead of Write because Write applies suppressions specific to GConfig and Dirty.
+	FString Text;
+	ConfigFile.WriteToString(Text);
+	FFileHelper::SaveStringToFile(Text, *GetCookSettingsFileName(TargetPlatform));
 }
 
 void UCookOnTheFlyServer::ClearCookInProgressFlagFromCookSettings(const ITargetPlatform* TargetPlatform) const
@@ -7737,8 +7739,10 @@ void UCookOnTheFlyServer::ClearCookInProgressFlagFromCookSettings(const ITargetP
 	FString Filename = GetCookSettingsFileName(TargetPlatform);
 	ConfigFile.Read(Filename);
 	ConfigFile.RemoveKeyFromSection(TEXT_CookSettings, TEXT_CookInProgress);
-	ConfigFile.Dirty = true; // Writing to a section does not set the dirty flag, so set it manually to make Write work
-	ConfigFile.Write(Filename);
+	// Use ConfigFile.WriteToString instead of Write because Write applies suppressions specific to GConfig and Dirty.
+	FString Text;
+	ConfigFile.WriteToString(Text);
+	FFileHelper::SaveStringToFile(Text, *Filename);
 }
 
 bool UCookOnTheFlyServer::IniSettingsOutOfDate(const ITargetPlatform* TargetPlatform) const
