Ensure we load an animation instead of trying to just find the package before we "Bake To Animation" so that it doesn't try to create a new asset.
PR #13751: Update AnimationEditorUtils.cpp

diff --git a/Engine/Source/Editor/UnrealEd/Private/AnimationEditorUtils.cpp b/Engine/Source/Editor/UnrealEd/Private/AnimationEditorUtils.cpp
--- a/Engine/Source/Editor/UnrealEd/Private/AnimationEditorUtils.cpp
+++ b/Engine/Source/Editor/UnrealEd/Private/AnimationEditorUtils.cpp
@@ -524,7 +524,7 @@ namespace AnimationEditorUtils
 
 				if (bAllowReplaceExisting)
 				{
-					UPackage* ExistingPackage = FindPackage(nullptr, *PackageName);
+					UPackage* ExistingPackage = LoadPackage(nullptr, *PackageName, LOAD_None);
 					UObject* ExistingObject = StaticFindObject(AssetClass.Get(), ExistingPackage, *Name);
 					if (ExistingObject)
 					{
