Fix ODSC persistent data not being able to send incomplete shadermaps back to the client. This was preventing cooked clients from being able to retrieve shaders when running with -odschost=XXXX

diff --git a/Engine/Source/Editor/UnrealEd/Private/CookOnTheFlyServer.cpp b/Engine/Source/Editor/UnrealEd/Private/CookOnTheFlyServer.cpp
--- a/Engine/Source/Editor/UnrealEd/Private/CookOnTheFlyServer.cpp
+++ b/Engine/Source/Editor/UnrealEd/Private/CookOnTheFlyServer.cpp
@@ -6497,6 +6497,7 @@ void UCookOnTheFlyServer::SetInitializeConfigSettings(UE::Cook::FInitializeConfi
 	if (bRunningAsShaderServer)
 	{
 		ODSCClientData = MakeUnique<UE::Cook::FODSCClientData>();
+		bAllowIncompleteShaderMapsDuringCookOnTheFly = true;
 	}
 
 	bSkipSave = FParse::Param(CommandLine, TEXT("CookSkipSave"));
