[FPackageSourceControlHelper] Checkout now handles opened for delete files (revert then edit)
- use SetSoftRevert(true) to avoid losing local changes on the reverted files

#rnx

diff --git a/Engine/Source/Editor/UnrealEd/Private/PackageSourceControlHelper.cpp b/Engine/Source/Editor/UnrealEd/Private/PackageSourceControlHelper.cpp
--- a/Engine/Source/Editor/UnrealEd/Private/PackageSourceControlHelper.cpp
+++ b/Engine/Source/Editor/UnrealEd/Private/PackageSourceControlHelper.cpp
@@ -339,6 +339,7 @@ bool FPackageSourceControlHelper::Checkout(const TArray<FString>& PackageNames,
 
 	// Two-pass checkout mechanism
 	TArray<FString> PackagesToCheckout;
+	TArray<FString> PackagesToRevert;
 	PackagesToCheckout.Reserve(PackageFilenames.Num());
 	bool bSuccess = true;
 
@@ -376,6 +377,12 @@ bool FPackageSourceControlHelper::Checkout(const TArray<FString>& PackageNames,
 			}
 			else if (SourceControlState->IsSourceControlled())
 			{
+				// Package needs to be checked out but if it is currently opened for delete, it needs to be reverted first.
+				if (SourceControlState->IsDeleted())
+				{
+					PackagesToRevert.Add(PackageFilename);
+				}
+
 				PackagesToCheckout.Add(PackageFilename);
 			}
 		}
@@ -409,6 +416,20 @@ bool FPackageSourceControlHelper::Checkout(const TArray<FString>& PackageNames,
 	}
 	else if (bUseSourceControl)
 	{
+		if (!PackagesToRevert.IsEmpty())
+		{
+			TSharedRef<FRevert, ESPMode::ThreadSafe> RevertOperation = ISourceControlOperation::Create<FRevert>();
+			
+			// Avoid losing local changes
+			RevertOperation->SetSoftRevert(true);
+
+			if (GetSourceControlProvider().Execute(RevertOperation, PackagesToRevert) != ECommandResult::Succeeded)
+			{
+				FPackageSourceControlHelperLog::Error(bErrorsAsWarnings, TEXT("Error reverting opened for delete file packages prior to checking out"));
+				return false;
+			}
+		}
+
 		if (GetSourceControlProvider().Execute(ISourceControlOperation::Create<FCheckOut>(), PackagesToCheckout) != ECommandResult::Succeeded)
 		{
 			// If operation didn't succeed. Get the list of invalid states when provided with a OutFailedPackages
