[FBX Exporter] Crash fix for Missing Bone in Reference Skeleton exports.

#rnx

diff --git a/Engine/Source/Editor/UnrealEd/Private/Fbx/FbxAnimationExport.cpp b/Engine/Source/Editor/UnrealEd/Private/Fbx/FbxAnimationExport.cpp
--- a/Engine/Source/Editor/UnrealEd/Private/Fbx/FbxAnimationExport.cpp
+++ b/Engine/Source/Editor/UnrealEd/Private/Fbx/FbxAnimationExport.cpp
@@ -106,6 +106,14 @@ void FFbxExporter::ExportAnimSequenceToFbx(const UAnimSequence* AnimSeq, const U
 	{
 		FbxNode* CurrentBoneNode = BoneNodes[BoneIndex];
 		const int32 BoneTreeIndex = Skeleton->GetSkeletonBoneIndexFromMeshBoneIndex(SkelMesh, BoneIndex);
+
+		if (!Skeleton->GetReferenceSkeleton().IsValidIndex(BoneTreeIndex))
+		{
+			ensureMsgf(false, TEXT("Failed to find Bone [%hs] in Reference Skeleton, when exporting AnimSequence [%s]."), CurrentBoneNode->GetName(), *AnimSeq->GetName());
+
+			continue;
+		}
+
 		const FName BoneName = Skeleton->GetReferenceSkeleton().GetBoneName(BoneTreeIndex);
 		const IAnimationDataModel* DataModel = AnimSeq->GetDataModel();
 
