Preview Curve Editor: Can now be resized. Visibility and lock icons now square instead of stretched.

#13602

diff --git a/Engine/Source/Editor/DetailCustomizations/Private/CurveColorCustomization.cpp b/Engine/Source/Editor/DetailCustomizations/Private/CurveColorCustomization.cpp
--- a/Engine/Source/Editor/DetailCustomizations/Private/CurveColorCustomization.cpp
+++ b/Engine/Source/Editor/DetailCustomizations/Private/CurveColorCustomization.cpp
@@ -496,9 +496,9 @@ FReply FCurveColorCustomization::OnCurvePreviewDoubleClick(const FGeometry& InMy
 				.ClientSize( FCurveColorCustomization::DEFAULT_WINDOW_SIZE )
 				.ScreenPosition(AdjustedSummonLocation)
 				.AutoCenter(EAutoCenter::None)
-				.SupportsMaximize(false)
+				.SupportsMaximize(true)
 				.SupportsMinimize(false)
-				.SizingRule( ESizingRule::FixedSize );
+				.SizingRule( ESizingRule::UserSized );
 
 			// init the mini curve editor widget
 			TSharedRef<SMiniCurveEditor> MiniCurveEditor =
diff --git a/Engine/Source/Editor/UnrealEd/Private/SCurveEditor.cpp b/Engine/Source/Editor/UnrealEd/Private/SCurveEditor.cpp
--- a/Engine/Source/Editor/UnrealEd/Private/SCurveEditor.cpp
+++ b/Engine/Source/Editor/UnrealEd/Private/SCurveEditor.cpp
@@ -613,6 +613,7 @@ TSharedRef<SWidget> SCurveEditor::CreateCurveSelectionWidget() const
 					.UncheckedImage(FAppStyle::Get().GetBrush("Icons.Hidden"))
 					.UncheckedHoveredImage(FAppStyle::Get().GetBrush("Icons.Hidden"))
 					.UncheckedPressedImage(FAppStyle::Get().GetBrush("Icons.Hidden"))
+					.Padding(8.0f)
 				]
 
 				+ SHorizontalBox::Slot()
@@ -631,6 +632,7 @@ TSharedRef<SWidget> SCurveEditor::CreateCurveSelectionWidget() const
 					.UncheckedImage(FAppStyle::Get().GetBrush("Icons.Unlock"))
 					.UncheckedHoveredImage(FAppStyle::Get().GetBrush("Icons.Unlock"))
 					.UncheckedPressedImage(FAppStyle::Get().GetBrush("Icons.Unlock"))
+					.Padding(8.0f)
 					.Visibility(bCanEditTrack ? EVisibility::Visible : EVisibility::Collapsed)
 				]
 			];
