Proper setup for FPreviewPlatformInfo Default Constructor

diff --git a/Engine/Source/Editor/UnrealEd/Classes/Editor/EditorEngine.h b/Engine/Source/Editor/UnrealEd/Classes/Editor/EditorEngine.h
--- a/Engine/Source/Editor/UnrealEd/Classes/Editor/EditorEngine.h
+++ b/Engine/Source/Editor/UnrealEd/Classes/Editor/EditorEngine.h
@@ -213,7 +213,7 @@ struct FSelectionStateOfLevel
 struct FPreviewPlatformInfo
 {
 	FPreviewPlatformInfo()
-	:	PreviewFeatureLevel(ERHIFeatureLevel::SM5)
+	:	PreviewFeatureLevel(ERHIFeatureLevel::Num)
 	,	ShaderPlatform(EShaderPlatform::SP_NumPlatforms)
 	,	PreviewPlatformName(NAME_None)
 	,	PreviewShaderFormatName(NAME_None)
@@ -222,9 +222,7 @@ struct FPreviewPlatformInfo
 	,	PreviewShaderPlatformFriendlyName()
 	,	ConstrainedAspectRatio(0.0f)
 	,	SafeZones(FVector4f(0))
-	{
-		InternalSetFriendlyName();
-	}
+	{}
 
 	FPreviewPlatformInfo(ERHIFeatureLevel::Type InFeatureLevel, 
 		EShaderPlatform InShaderPlatform = EShaderPlatform::SP_NumPlatforms, 
