Bump cook incremental version to clear out old shader cruft.

diff --git a/Engine/Source/Editor/UnrealEd/Private/Cooker/CookIncrementalVersion.cpp b/Engine/Source/Editor/UnrealEd/Private/Cooker/CookIncrementalVersion.cpp
--- a/Engine/Source/Editor/UnrealEd/Private/Cooker/CookIncrementalVersion.cpp
+++ b/Engine/Source/Editor/UnrealEd/Private/Cooker/CookIncrementalVersion.cpp
@@ -8,6 +8,6 @@ namespace UE::Cook
 {
 
 // Change CookIncrementalVersion to a new guid when all packages in an incremental cook need to be invalidated
-FGuid CookIncrementalVersion(0x2cfa9302, 0x16374f05, 0x872b1d7e, 0x6af13761);
+FGuid CookIncrementalVersion( 0x9F3A7C12, 0x4B8E2F90, 0xC1D47A6B, 0x6E9B0FD3 );
 
 }
