When UEditorEngine::Map_Rebuild execute, first check if PIE is running and ask to stop before continuing.
This prevents multiple checks that can happen.

diff --git a/Engine/Source/Editor/UnrealEd/Private/EditorServer.cpp b/Engine/Source/Editor/UnrealEd/Private/EditorServer.cpp
--- a/Engine/Source/Editor/UnrealEd/Private/EditorServer.cpp
+++ b/Engine/Source/Editor/UnrealEd/Private/EditorServer.cpp
@@ -1420,6 +1420,11 @@ bool UEditorEngine::Map_Sendto(UWorld* InWorld, const TCHAR* Str, FOutputDevice&
 
 bool UEditorEngine::Map_Rebuild(UWorld* InWorld, const TCHAR* Str, FOutputDevice& Ar)
 {
+	if (ShouldAbortBecauseOfPIEWorld())
+	{
+		return false;
+	}
+
 	TMap<AActor*, TArray<int32>> VisibleBSPSurfaceMap;
 	bool bAllVisible = false;
 
