[AutoViz] AutoVizPrototyping Plugin
- Created new Viewport modes for Geometry Inspection: Clay, Zebra, Front/Back Face and Random Color
- Created Wireframe swich and opacity control
- Created Zebra orientation and Zebra stripe count controls
- Fixed 5 spelling issues
#virtualized
#was-backout 45903479

diff --git a/Engine/Config/BaseEngine.ini b/Engine/Config/BaseEngine.ini
--- a/Engine/Config/BaseEngine.ini
+++ b/Engine/Config/BaseEngine.ini
@@ -212,6 +212,10 @@ DefaultFlattenMaterialName=/Engine/EngineMaterials/FlattenMaterial.FlattenMateri
 DefaultHLODFlattenMaterialName=/Engine/EngineMaterials/FlattenMaterial_VT.FlattenMaterial_VT
 DefaultLandscapeFlattenMaterialName=/Engine/EngineMaterials/FlattenMaterial_WS_Normal_VT.FlattenMaterial_WS_Normal_VT
 DefaultZenStreamingTextureName=/Engine/EngineResources/Development/ZenLogo_64.ZenLogo_64
+ClayMaterialName=/Engine/EngineDebugMaterials/M_GeoInsp_Clay.M_GeoInsp_Clay
+ZebraMaterialName=/Engine/EngineDebugMaterials/M_GeoInsp_Zebra.M_GeoInsp_Zebra
+FrontBackFaceMaterialName=/Engine/EngineDebugMaterials/M_GeoInsp_BackFace.M_GeoInsp_BackFace
+RandomColorMaterialName=/Engine/EngineDebugMaterials/M_GeoInsp_Clay.M_GeoInsp_Clay
 TimecodeProviderClassName=None
 ArrowMaterialName=/Engine/EngineMaterials/GizmoMaterial.GizmoMaterial
 ParticleEventManagerClassPath="/Script/Engine.ParticleEventManager"
diff --git a/Engine/Plugins/Runtime/MeshModelingToolset/Source/ModelingComponents/Private/PreviewMesh.cpp b/Engine/Plugins/Runtime/MeshModelingToolset/Source/ModelingComponents/Private/PreviewMesh.cpp
--- a/Engine/Plugins/Runtime/MeshModelingToolset/Source/ModelingComponents/Private/PreviewMesh.cpp
+++ b/Engine/Plugins/Runtime/MeshModelingToolset/Source/ModelingComponents/Private/PreviewMesh.cpp
@@ -55,6 +55,10 @@ void UPreviewMesh::CreateInWorld(UWorld* World, const FTransform& WithTransform)
 
 	TemporaryParentActor->SetActorTransform(WithTransform);
 	//Builder.NewMeshComponent->SetWorldTransform(PlaneFrame.ToFTransform());
+
+	// Disable View Mode Overrides because we don't want editor mode visualizations 
+	// to be removed/replaced by View Mode Overrides
+	DynamicMeshComponent->SetViewModeOverridesEnabled(false);
 }
 
 
diff --git a/Engine/Source/Editor/UnrealEd/Private/EditorViewportClient.cpp b/Engine/Source/Editor/UnrealEd/Private/EditorViewportClient.cpp
--- a/Engine/Source/Editor/UnrealEd/Private/EditorViewportClient.cpp
+++ b/Engine/Source/Editor/UnrealEd/Private/EditorViewportClient.cpp
@@ -210,7 +210,7 @@ namespace PreviewLightConstants
 	const float MinMouseRadius = 100.0f;
 	const float MinArrowLength = 10.0f;
 	const float ArrowLengthToSizeRatio = 0.1f;
-	const float MouseLengthToArrowLenghtRatio = 0.2f;
+	const float MouseLengthToArrowLengthRatio = 0.2f;
 
 	const float ArrowLengthToThicknessRatio = 0.05f;
 	const float MinArrowThickness = 2.0f;
@@ -218,6 +218,12 @@ namespace PreviewLightConstants
 	// Note: MinMouseRadius must be greater than MinArrowLength
 }
 
+namespace GeometryInspectionZebraMaterialParameterNames
+{
+	const FName Orientation = TEXT("XYZ");
+	const FName StripeCount = TEXT("Stripe Counts");
+}
+
 /**
  * Cached off joystick input state
  */
@@ -4900,7 +4906,7 @@ void FEditorViewportClient::DrawPreviewLightVisualization(const FSceneView* View
 		const FVector LightToMousePos = LastMouseWorldPos - LightLocalToWorld.GetTranslation();
 		const float LightToMouseRadius = FMath::Max<FVector::FReal>(LightToMousePos.Size(), PreviewLightConstants::MinMouseRadius);
 
-		const float ArrowLength = FMath::Max(PreviewLightConstants::MinArrowLength, LightToMouseRadius * PreviewLightConstants::MouseLengthToArrowLenghtRatio);
+		const float ArrowLength = FMath::Max(PreviewLightConstants::MinArrowLength, LightToMouseRadius * PreviewLightConstants::MouseLengthToArrowLengthRatio);
 		const float ArrowSize = PreviewLightConstants::ArrowLengthToSizeRatio * ArrowLength;
 		const float ArrowThickness = FMath::Max(PreviewLightConstants::ArrowLengthToThicknessRatio * ArrowLength, PreviewLightConstants::MinArrowThickness);
 
@@ -6182,6 +6188,73 @@ void FEditorViewportClient::SetViewMode(EViewModeIndex InViewModeIndex)
 	Invalidate();
 }
 
+void FEditorViewportClient::SetZebraViewMode(bool bHorizontalZebraMode)
+{
+	if (UMaterialInstanceDynamic* MaterialInstance = Cast<UMaterialInstanceDynamic>(GEngine->ZebraMaterialInstance))
+	{
+		FLinearColor Color  = bHorizontalZebraMode ? FLinearColor(0, 0, 1) : FLinearColor(1, 0, 0);
+
+		MaterialInstance->SetVectorParameterValue(GeometryInspectionZebraMaterialParameterNames::Orientation, Color);
+	}
+
+	SetViewMode(VMI_Zebra);
+}
+
+bool FEditorViewportClient::IsZebraViewModeEnabled(bool bHorizontal) const
+{
+	if (IsViewModeEnabled(VMI_Zebra))
+	{
+		if (UMaterialInstanceDynamic* MaterialInstance = Cast<UMaterialInstanceDynamic>(GEngine->ZebraMaterialInstance))
+		{
+			FLinearColor Color;
+			if (MaterialInstance->GetVectorParameterValue(GeometryInspectionZebraMaterialParameterNames::Orientation, Color))
+			{
+				if (Color.B == 1.0f && bHorizontal)
+				{
+					return true;
+				}
+
+				if (Color.R == 1.0f && !bHorizontal)
+				{
+					return true;
+				}
+
+				return false;
+			}
+		}
+	}
+
+	return false;
+}
+
+void FEditorViewportClient::SetZebraViewModeStripeCount(uint8 InStripeCount)
+{
+	if (InStripeCount>0 && InStripeCount<32)
+	{
+		if (UMaterialInstanceDynamic* MaterialInstance = Cast<UMaterialInstanceDynamic>(GEngine->ZebraMaterialInstance))
+		{
+			MaterialInstance->SetScalarParameterValue(GeometryInspectionZebraMaterialParameterNames::StripeCount, InStripeCount);
+
+			Invalidate();
+		}
+	}
+}
+
+uint8 FEditorViewportClient::GetZebraViewModeStripeCount()
+{
+	if (UMaterialInstanceDynamic* MaterialInstance = Cast<UMaterialInstanceDynamic>(GEngine->ZebraMaterialInstance))
+	{
+		FHashedMaterialParameterInfo ParameterInfo(GeometryInspectionZebraMaterialParameterNames::StripeCount);
+		float Value;
+		if (MaterialInstance->GetScalarParameterValue(ParameterInfo, Value))
+		{
+			return Value;
+		}
+	}
+
+	return 1;
+}
+
 void FEditorViewportClient::SetViewModes(const EViewModeIndex InPerspViewModeIndex, const EViewModeIndex InOrthoViewModeIndex)
 {
 	PerspViewModeIndex = InPerspViewModeIndex;
diff --git a/Engine/Source/Editor/UnrealEd/Private/EditorViewportCommands.cpp b/Engine/Source/Editor/UnrealEd/Private/EditorViewportCommands.cpp
--- a/Engine/Source/Editor/UnrealEd/Private/EditorViewportCommands.cpp
+++ b/Engine/Source/Editor/UnrealEd/Private/EditorViewportCommands.cpp
@@ -133,6 +133,13 @@ void FEditorViewportCommands::RegisterCommands()
 	UI_COMMAND( CollisionPawn, "Player Collision", "Renders player collision visualization", EUserInterfaceActionType::RadioButton, FInputChord() );
 	UI_COMMAND( CollisionVisibility, "Visibility Collision", "Renders visibility collision visualization", EUserInterfaceActionType::RadioButton, FInputChord() );
 
+	UI_COMMAND( GeometryInspectionClayMode, "Geometry Inspection Clay Mode", "Displays using clay material to evaluate surface curvature and smoothness.", EUserInterfaceActionType::RadioButton, FInputChord());
+	UI_COMMAND( GeometryInspectionZebraHorizontalMode, "Horizontal", "Displays simulated light reflections using stripes to evaluate surface curvature and smoothness.", EUserInterfaceActionType::RadioButton, FInputChord());
+	UI_COMMAND( GeometryInspectionZebraVerticalMode, "Vertical", "Displays simulated light reflections using stripes to evaluate surface curvature and smoothness.", EUserInterfaceActionType::RadioButton, FInputChord());
+	UI_COMMAND( GeometryInspectionFrontBackFaceMode, "Geometry Inspection Front/Back Face Mode", "Highlights front-facing and back-facing polygons differently to evaluate face normals.", EUserInterfaceActionType::RadioButton, FInputChord());
+	UI_COMMAND( GeometryInspectionRandomColorMode, "Geometry Inspection Random Color Mode", "Assigns random colors to individual elements or components to improve visual distinction between overlapping or adjacent parts.", EUserInterfaceActionType::RadioButton, FInputChord());
+	UI_COMMAND( GeometryInspectionShowWireframe, "Use Wireframe", "Use Wireframe in Geometry Inspection Modes.", EUserInterfaceActionType::Check, FInputChord());
+
 	UI_COMMAND( ToggleRealTime, "Realtime", "Toggles real time rendering in this viewport", EUserInterfaceActionType::ToggleButton, FInputChord( EModifierKey::Control, EKeys::R ) );
 	UI_COMMAND( ToggleStats, "Show Stats", "Toggles the ability to show stats in this viewport (enables realtime)", EUserInterfaceActionType::ToggleButton, FInputChord( EModifierKey::Shift, EKeys::L ) );
 	UI_COMMAND( ToggleFPS, "Show FPS", "Toggles showing frames per second in this viewport (enables realtime)", EUserInterfaceActionType::ToggleButton, FInputChord( EModifierKey::Control|EModifierKey::Shift, EKeys::H ) );
diff --git a/Engine/Source/Editor/UnrealEd/Private/SEditorViewport.cpp b/Engine/Source/Editor/UnrealEd/Private/SEditorViewport.cpp
--- a/Engine/Source/Editor/UnrealEd/Private/SEditorViewport.cpp
+++ b/Engine/Source/Editor/UnrealEd/Private/SEditorViewport.cpp
@@ -155,7 +155,7 @@ void SEditorViewport::Construct( const FArguments& InArgs )
 		PRAGMA_ENABLE_DEPRECATION_WARNINGS
 		TSharedPtr<SWidget> NewViewportToolbar = BuildViewportToolbar();
 		
-		// Allow programatically-migrated viewports (e.g. SCommonEditorViewportToolbarBase) to register themselves on construction as supporting either position
+		// Allow programmatically-migrated viewports (e.g. SCommonEditorViewportToolbarBase) to register themselves on construction as supporting either position
 		// This indirection is done to minimize API changes (e.g. using a specific toolbar interface). The goal is that clients do nothing and get upgrades.
 		// Overrides of SCommonEditorViewportToolbarBase can construct the full base toolbar and then overwrite the child slot with their own content.
 		// This check ensures that the automatic upgrade is only allowed if we end up with a widget containing the child that we expected to have
@@ -720,6 +720,30 @@ void SEditorViewport::BindCommands()
 		MAP_VIEWMODEPARAM_ACTION( Commands.RequiredTextureResolutionSingle[TextureIndex], TextureIndex );
 	}
 
+	MAP_VIEWMODE_ACTION( Commands.GeometryInspectionClayMode, VMI_Clay );
+	MAP_VIEWMODE_ACTION( Commands.GeometryInspectionFrontBackFaceMode, VMI_FrontBackFace );
+	MAP_VIEWMODE_ACTION( Commands.GeometryInspectionRandomColorMode, VMI_RandomColor );
+	CommandListRef.MapAction( Commands.GeometryInspectionZebraHorizontalMode,
+		FExecuteAction::CreateSP( ClientRef, &FEditorViewportClient::SetZebraViewMode, true ),
+		FCanExecuteAction(),
+		FIsActionChecked::CreateSP( ClientRef, &FEditorViewportClient::IsZebraViewModeEnabled, true ) );
+	CommandListRef.MapAction( Commands.GeometryInspectionZebraVerticalMode,
+		FExecuteAction::CreateSP( ClientRef, &FEditorViewportClient::SetZebraViewMode, false ),
+		FCanExecuteAction(),
+		FIsActionChecked::CreateSP( ClientRef, &FEditorViewportClient::IsZebraViewModeEnabled, false ) );
+	
+	IConsoleVariable* ShowWireframeCVar = IConsoleManager::Get().FindConsoleVariable(TEXT("r.GeometryInspection.DrawWireframe"));
+	check(ShowWireframeCVar);
+	CommandListRef.MapAction(
+		Commands.GeometryInspectionShowWireframe,
+		FExecuteAction::CreateLambda([=]()
+		{
+			ShowWireframeCVar->Set(!ShowWireframeCVar->GetBool());
+		}),
+		FCanExecuteAction::CreateLambda([]() { return true; }),
+		FIsActionChecked::CreateLambda([=]() { return ShowWireframeCVar->GetBool(); })
+	);
+
 	BindShowCommands( CommandListRef );
 }
 
diff --git a/Engine/Source/Editor/UnrealEd/Private/ViewportToolbar/UnrealEdViewportToolbar.cpp b/Engine/Source/Editor/UnrealEd/Private/ViewportToolbar/UnrealEdViewportToolbar.cpp
--- a/Engine/Source/Editor/UnrealEd/Private/ViewportToolbar/UnrealEdViewportToolbar.cpp
+++ b/Engine/Source/Editor/UnrealEd/Private/ViewportToolbar/UnrealEdViewportToolbar.cpp
@@ -2136,6 +2136,148 @@ void PopulateViewModesMenu(UToolMenu* InMenu)
 				);
 			}
 		}
+		
+		{
+			struct Local
+			{
+				static void BuildGeometryInspectionMenu(UToolMenu* Menu,
+					IsViewModeSupportedDelegate IsViewModeSupported, 
+					const TWeakPtr<SEditorViewport> EditorViewportWeakPtr)
+				{
+					{
+						FToolMenuSection& Section = Menu->AddSection(
+							"GeometryInspection", LOCTEXT("GeometryInspectionModesHeader", "Geometry Inspection")
+						);
+
+						UE::UnrealEd::Private::ViewModesSubmenu::AddModeIfSupported(
+							IsViewModeSupported, Section, FEditorViewportCommands::Get().GeometryInspectionClayMode, VMI_Clay
+						);
+
+						FUIAction MenuAction;
+						MenuAction.GetActionCheckState = BuildParentCheckStateForCommands(Menu->Context, {
+							FEditorViewportCommands::Get().GeometryInspectionZebraHorizontalMode,
+							FEditorViewportCommands::Get().GeometryInspectionZebraVerticalMode
+						});
+
+						Section.AddSubMenu(
+							"Zebra",
+							LOCTEXT("GeometryInspectionZebraDisplayName", "Zebra"),
+							LOCTEXT("GeometryInspectionZebra_ToolTip", "Geometry Inspection Zebra View Modes"),
+							FNewToolMenuDelegate::CreateStatic(&BuildZebraMenu, IsViewModeSupported, EditorViewportWeakPtr),
+							MenuAction,
+							EUserInterfaceActionType::RadioButton,
+							/* bInOpenSubMenuOnClick = */ false,
+							FSlateIcon()
+						);
+
+						UE::UnrealEd::Private::ViewModesSubmenu::AddModeIfSupported(
+							IsViewModeSupported, Section, FEditorViewportCommands::Get().GeometryInspectionFrontBackFaceMode, VMI_FrontBackFace
+						);
+
+						UE::UnrealEd::Private::ViewModesSubmenu::AddModeIfSupported(
+							IsViewModeSupported, Section, FEditorViewportCommands::Get().GeometryInspectionRandomColorMode, VMI_RandomColor
+						);
+					}
+
+					{
+						FToolMenuSection& Section = Menu->AddSection(
+							"GeometryInspectionWireframe", LOCTEXT("GeometryInspectionWireframeHeader", "Wireframe")
+						);
+
+						if (!IsViewModeSupported.IsBound() || IsViewModeSupported.Execute(VMI_Lit_Wireframe))
+						{
+							Section.AddMenuEntry(FEditorViewportCommands::Get().GeometryInspectionShowWireframe);
+
+							if (TSharedPtr<SEditorViewport> EditorViewport = EditorViewportWeakPtr.Pin())
+							{
+								Section.AddEntry(FToolMenuEntry::InitWidget(
+							"WireframeOpacity", EditorViewport->BuildWireframeMenu(), LOCTEXT("WireframeOpacity", "Opacity")
+								));
+							}
+						}
+					}
+				}
+
+				static void BuildZebraMenu(UToolMenu* Menu, IsViewModeSupportedDelegate IsViewModeSupported, const TWeakPtr<SEditorViewport> EditorViewportWeakPtr)
+				{
+					FToolMenuSection& Section = Menu->AddSection(
+						"GeometryInspectionZebra", LOCTEXT("GeometryInspectionZebraModesHeader", "Zebra Mode")
+					);
+
+					Section.AddMenuEntry(FEditorViewportCommands::Get().GeometryInspectionZebraHorizontalMode);
+					Section.AddMenuEntry(FEditorViewportCommands::Get().GeometryInspectionZebraVerticalMode);
+
+					Menu->AddMenuEntry(TEXT("Zebra Configuration"), 
+						FToolMenuEntry::InitWidget(
+						"StripeCount",
+						BuildZebraStripeCountWidget(Menu, EditorViewportWeakPtr),
+					LOCTEXT("GeometryInspectionStripeCountLabel", "Stripe Count")));
+				}
+
+				static TSharedRef<SWidget> BuildZebraStripeCountWidget(UToolMenu* Menu, const TWeakPtr<SEditorViewport>& EditorViewportWeakPtr)
+				{
+					return SNew( SBox )
+					.HAlign( HAlign_Right )
+					[
+						SNew( SBox )
+						.Padding( FMargin(0.0f, 0.0f, 0.0f, 0.0f) )
+						.WidthOverride( 100.0f )
+						[
+							SNew ( SBorder )
+							.BorderImage(FAppStyle::Get().GetBrush("Menu.WidgetBorder"))
+							.Padding(FMargin(1.0f))
+							[
+								SNew(SSpinBox<uint8>)
+								.Style(&FAppStyle::Get(), "Menu.SpinBox")
+								.Font( FAppStyle::GetFontStyle( TEXT( "MenuItem.Font" ) ) )
+								.MinValue(1.f)
+								.MaxValue(16.f)
+								.SupportDynamicSliderMaxValue(false)
+								.Value_Static ( &FEditorViewportClient::GetZebraViewModeStripeCount )
+								.OnValueChanged_Lambda( [EditorViewportWeakPtr](uint16 NewValue)
+								{
+									if (TSharedPtr<SEditorViewport> EditorViewport = EditorViewportWeakPtr.Pin())
+									{
+										if (TSharedPtr<FEditorViewportClient> Client = StaticCastSharedPtr<FEditorViewportClient>(
+											EditorViewport->GetViewportClient()))
+										{
+											Client->SetZebraViewModeStripeCount(NewValue);
+										}
+									}
+								})
+								.ToolTipText(LOCTEXT("WireframeOpacity_ToolTip", "Adjust number of the stripes in zebra mode."))
+							]
+						]
+					];
+				}
+
+			};
+
+			if (!IsViewModeSupported.IsBound()
+				|| (IsViewModeSupported.Execute(VMI_Clay) || IsViewModeSupported.Execute(VMI_Zebra)
+				|| (IsViewModeSupported.Execute(VMI_FrontBackFace) || IsViewModeSupported.Execute(VMI_RandomColor))))
+			{
+				FUIAction MenuAction;
+				MenuAction.GetActionCheckState = BuildParentCheckStateForCommands(InMenu->Context, {
+					FEditorViewportCommands::Get().GeometryInspectionClayMode,
+					FEditorViewportCommands::Get().GeometryInspectionFrontBackFaceMode,
+					FEditorViewportCommands::Get().GeometryInspectionRandomColorMode,
+					FEditorViewportCommands::Get().GeometryInspectionZebraHorizontalMode,
+					FEditorViewportCommands::Get().GeometryInspectionZebraVerticalMode
+				});
+
+				Section.AddSubMenu(
+					"GeometryInspection",
+					LOCTEXT("GeometryInspectionDisplayName", "Geometry Inspection"),
+					LOCTEXT("GeometryInspection_ToolTip", "Geometry Inspection View Modes"),
+					FNewToolMenuDelegate::CreateStatic(&Local::BuildGeometryInspectionMenu, IsViewModeSupported, EditorViewport.ToWeakPtr()),
+					MenuAction,
+					EUserInterfaceActionType::RadioButton,
+					/* bInOpenSubMenuOnClick = */ false,
+					FSlateIcon(FAppStyle::GetAppStyleSetName(), "EditorViewport.CollisionVisibility")
+				);
+			}
+		}
 
 		if (GEnableGPUSkinCache
 			&& UE::UnrealEd::Private::ViewModesSubmenu::IsMenuSectionAvailable(
diff --git a/Engine/Source/Editor/UnrealEd/Public/EditorViewportClient.h b/Engine/Source/Editor/UnrealEd/Public/EditorViewportClient.h
--- a/Engine/Source/Editor/UnrealEd/Public/EditorViewportClient.h
+++ b/Engine/Source/Editor/UnrealEd/Public/EditorViewportClient.h
@@ -989,6 +989,20 @@ public:
 	 * @param	InViewModeIndex				View mode to set for the current viewport type
 	 */
 	UNREALED_API virtual void SetViewMode(EViewModeIndex InViewModeIndex);
+	
+	/**
+	 * Sets the viewport type to Geometry Inspection Zebra mode, horizontal or vertical type
+	 * 
+	 * @param	bHorizontalZebraMode		Controls if this is horizontal or vertical stripe type
+	 */
+	UNREALED_API virtual void SetZebraViewMode(bool bHorizontalZebraMode);
+	
+	/**
+	 * Sets Geometry Inspection Zebra mode stripe count
+	 * 
+	 * @param	InStripeCount				Number of stripes Zebra mode should have
+	 */
+	UNREALED_API virtual void SetZebraViewModeStripeCount(uint8 InStripeCount);
 
 	/**
 	 * Sets how the viewport is displayed (lit, wireframe, etc)
@@ -1028,6 +1042,12 @@ public:
 		return GetViewMode() == InViewModeIndex;
 	}
 
+	/** @return True if ViewMode is Geometry Inspection Zebra in specified orientation */
+	UNREALED_API bool IsZebraViewModeEnabled(bool bHorizontal) const;
+
+	/** @return Configured number of stripes in Geometry Inspection Zebra view mode */
+	UNREALED_API static uint8 GetZebraViewModeStripeCount();
+
 	/** @return True if InViewModeIndex is the current view mode param */
 	UNREALED_API bool IsViewModeParam(int32 InViewModeParam) const;
 
diff --git a/Engine/Source/Editor/UnrealEd/Public/EditorViewportCommands.h b/Engine/Source/Editor/UnrealEd/Public/EditorViewportCommands.h
--- a/Engine/Source/Editor/UnrealEd/Public/EditorViewportCommands.h
+++ b/Engine/Source/Editor/UnrealEd/Public/EditorViewportCommands.h
@@ -158,6 +158,14 @@ public:
 	TSharedPtr< FUICommandInfo > CollisionPawn;
 	TSharedPtr< FUICommandInfo > CollisionVisibility;
 
+	/** Geometry Inspection Modes */
+	TSharedPtr< FUICommandInfo > GeometryInspectionClayMode;
+	TSharedPtr< FUICommandInfo > GeometryInspectionZebraHorizontalMode;
+	TSharedPtr< FUICommandInfo > GeometryInspectionZebraVerticalMode;
+	TSharedPtr< FUICommandInfo > GeometryInspectionFrontBackFaceMode;
+	TSharedPtr< FUICommandInfo > GeometryInspectionRandomColorMode;
+	TSharedPtr< FUICommandInfo > GeometryInspectionShowWireframe;
+
 	/** Toggles realtime rendering in the viewport */
 	TSharedPtr< FUICommandInfo > ToggleRealTime;
 
diff --git a/Engine/Source/Runtime/Core/Private/HAL/ConsoleManager.cpp b/Engine/Source/Runtime/Core/Private/HAL/ConsoleManager.cpp
--- a/Engine/Source/Runtime/Core/Private/HAL/ConsoleManager.cpp
+++ b/Engine/Source/Runtime/Core/Private/HAL/ConsoleManager.cpp
@@ -249,7 +249,7 @@ public:
 	
 	}
 	
-	/** Legacy funciton to add old single delegates to the new multicast delegate. */
+	/** Legacy function to add old single delegates to the new multicast delegate. */
 	virtual void SetOnChangedCallback(const FConsoleVariableDelegate& Callback) 
 	{
 		OnChangedCallback.Remove(LegacyDelegateHandle);
@@ -4575,6 +4575,11 @@ static TAutoConsoleVariable<int32> CVarLegacyLuminanceFactors(
 		"1: Legacy coefficients (inaccurate). Activate to retain the 5.4 and earlier versions behavior.\n"),
 	ECVF_ReadOnly | ECVF_RenderThreadSafe);
 
+static TAutoConsoleVariable<bool> CVarShowWireframe(
+		TEXT("r.GeometryInspection.DrawWireframe"),
+		true,
+		TEXT("Controls if wireframe is drawn in Geometry Inspection modes"),
+		ECVF_RenderThreadSafe);
 
 class FConsoleObjectShadowData
 {
diff --git a/Engine/Source/Runtime/Engine/Classes/Engine/Engine.h b/Engine/Source/Runtime/Engine/Classes/Engine/Engine.h
--- a/Engine/Source/Runtime/Engine/Classes/Engine/Engine.h
+++ b/Engine/Source/Runtime/Engine/Classes/Engine/Engine.h
@@ -1159,6 +1159,38 @@ public:
 	/** A material used to render the debug texture painting mask on mesh. */
 	UPROPERTY(globalconfig)
 	FSoftObjectPath TexturePaintingMaskMaterialName;
+
+	UPROPERTY()
+	TObjectPtr<class UMaterialInterface> ClayMaterial;
+
+	/** A material used to render the geometry inspection clay mode. */
+	UPROPERTY(globalconfig)
+	FSoftObjectPath ClayMaterialName;
+
+	UPROPERTY()
+	TObjectPtr<class UMaterial> ZebraMaterial;
+
+	UPROPERTY()
+	TObjectPtr<class UMaterialInstanceDynamic> ZebraMaterialInstance;
+
+	/** A material used to render the geometry inspection zebra mode. */
+	UPROPERTY(globalconfig)
+	FSoftObjectPath ZebraMaterialName;
+	
+	UPROPERTY()
+	TObjectPtr<class UMaterialInterface> FrontBackFaceMaterial;
+
+	/** A material used to render the geometry inspection front/back face mode. */
+	UPROPERTY(globalconfig)
+	FSoftObjectPath FrontBackFaceMaterialName;
+
+	UPROPERTY()
+	TObjectPtr<class UMaterialInterface> RandomColorMaterial;
+
+	/** A material used to render the geometry inspection random color mode. */
+	UPROPERTY(globalconfig)
+	FSoftObjectPath RandomColorMaterialName;
+
 #endif
 
 	/** A material used to render debug opaque material. Used in various animation editor viewport features. */
diff --git a/Engine/Source/Runtime/Engine/Classes/Engine/EngineBaseTypes.h b/Engine/Source/Runtime/Engine/Classes/Engine/EngineBaseTypes.h
--- a/Engine/Source/Runtime/Engine/Classes/Engine/EngineBaseTypes.h
+++ b/Engine/Source/Runtime/Engine/Classes/Engine/EngineBaseTypes.h
@@ -1047,6 +1047,18 @@ enum EViewModeIndex : int
 	/** Visualize Shadow Casters */
 	VMI_ShadowCasters = 39 UMETA(DisplayName = "Shadow Casters Visualization"),
 
+	/** Displays using clay material to evaluate surface curvature and smoothness. */
+	VMI_Clay = 40 UMETA(DisplayName = "Clay"),
+
+	/** Displays simulated light reflections using stripes to evaluate surface curvature and smoothness. */
+	VMI_Zebra = 41 UMETA(DisplayName = "Zebra"),
+
+	/** Highlights front-facing and back-facing polygons differently to evaluate face normals. */
+	VMI_FrontBackFace = 42 UMETA(DisplayName = "Front/Back Face"),
+
+	/** Assigns random colors to individual elements or components to improve visual distinction between overlapping or adjacent parts. */
+	VMI_RandomColor = 43 UMETA(DisplayName = "Random"),
+
 	VMI_Max UMETA(Hidden),
 
 	// VMI_Unknown - The value assigned to VMI_Unknown must be the highest possible of any member of EViewModeIndex, or GetViewModeName might seg-fault
diff --git a/Engine/Source/Runtime/Engine/Private/PrimitiveDrawingUtils.cpp b/Engine/Source/Runtime/Engine/Private/PrimitiveDrawingUtils.cpp
--- a/Engine/Source/Runtime/Engine/Private/PrimitiveDrawingUtils.cpp
+++ b/Engine/Source/Runtime/Engine/Private/PrimitiveDrawingUtils.cpp
@@ -17,6 +17,9 @@
 #include "UnrealClient.h"
 #include "PrimitiveDrawingUtils.h"
 
+#include "Materials/MaterialInstance.h"
+#include "Materials/MaterialInstanceDynamic.h"
+
 void DrawPlane10x10(class FPrimitiveDrawInterface* PDI,const FMatrix& ObjectToWorld, float Radii, FVector2D UVMin, FVector2D UVMax, const FMaterialRenderProxy* MaterialRenderProxy, uint8 DepthPriorityGroup)
 {
 	// -> TileCount * TileCount * 2 triangles
@@ -1731,6 +1734,74 @@ void ApplyViewModeOverrides(
 			Mesh.MaterialRenderProxy = LevelColorationMaterialInstance;
 			Collector.RegisterOneFrameMaterialProxy(LevelColorationMaterialInstance);
 		}
+#if WITH_EDITOR
+		else if (EngineShowFlags.RandomColor)
+		{
+			FRandomStream RandomStream(GetTypeHash(PrimitiveSceneProxy->GetPrimitiveComponentId().PrimIDValue) + GetTypeHash(Mesh.MeshIdInPrimitive));
+
+			const uint8 R = static_cast<uint8>(RandomStream.FRandRange(64, 255));
+			const uint8 G = static_cast<uint8>(RandomStream.FRandRange(64, 255));
+			const uint8 B = static_cast<uint8>(RandomStream.FRandRange(64, 255));
+			FLinearColor RandomColor(FColor(R, G, B, 255));
+			FMaterialRenderProxy* LevelColorationMaterialInstance = nullptr;
+
+			if (TObjectPtr<UMaterialInterface> RandomMaterial = GEngine->RandomColorMaterial)
+			{
+				LevelColorationMaterialInstance = new FColoredMaterialRenderProxy(
+					RandomMaterial->GetRenderProxy(), 
+					GetSelectionColor(RandomColor, bSelected, PrimitiveSceneProxy->IsHovered(), /*bUseOverlayIntensity=*/true), 
+					TEXT("Diffuse"));
+			}
+			else
+			{
+				const UMaterial* LevelColorationMaterial = EngineShowFlags.Lighting ? GEngine->LevelColorationLitMaterial : GEngine->LevelColorationUnlitMaterial;
+				
+				// Draw the mesh with level coloration.
+				LevelColorationMaterialInstance = new FColoredMaterialRenderProxy(
+					LevelColorationMaterial->GetRenderProxy(), 
+					GetSelectionColor(RandomColor, bSelected, PrimitiveSceneProxy->IsHovered(), /*bUseOverlayIntensity=*/true));
+			}
+
+			Mesh.MaterialRenderProxy = LevelColorationMaterialInstance;
+			Collector.RegisterOneFrameMaterialProxy(LevelColorationMaterialInstance);
+		}
+		else if (EngineShowFlags.Clay)
+		{
+			if (TObjectPtr<UMaterialInterface> ClayMaterial = GEngine->ClayMaterial)
+			{
+				FMaterialRenderProxy* LevelColorationMaterialInstance = new FColoredMaterialRenderProxy(
+					ClayMaterial->GetRenderProxy(), 
+					GetSelectionColor(FLinearColor(1, 1, 1, 1), bSelected, PrimitiveSceneProxy->IsHovered(), /*bUseOverlayIntensity=*/true));
+
+				Mesh.MaterialRenderProxy = LevelColorationMaterialInstance;
+				Collector.RegisterOneFrameMaterialProxy(LevelColorationMaterialInstance);
+			}
+		}
+		else if (EngineShowFlags.Zebra)
+		{
+			if (TObjectPtr<UMaterialInstanceDynamic> ZebraMaterial = GEngine->ZebraMaterialInstance)
+			{
+				FMaterialRenderProxy* LevelColorationMaterialInstance = new FColoredMaterialRenderProxy(
+					ZebraMaterial->GetRenderProxy(), 
+					GetSelectionColor(FLinearColor(0, 0, 0, 0), bSelected, PrimitiveSceneProxy->IsHovered(), /*bUseOverlayIntensity=*/true));
+
+				Mesh.MaterialRenderProxy = LevelColorationMaterialInstance;
+				Collector.RegisterOneFrameMaterialProxy(LevelColorationMaterialInstance);
+			}
+		}
+		else if (EngineShowFlags.FrontBackFace)
+		{
+			if (TObjectPtr<UMaterialInterface> FrontBackFaceMaterial = GEngine->FrontBackFaceMaterial)
+			{
+				FMaterialRenderProxy* LevelColorationMaterialInstance = new FColoredMaterialRenderProxy(
+					FrontBackFaceMaterial->GetRenderProxy(),
+					GetSelectionColor(FLinearColor(0, 0, 0, 0), bSelected, PrimitiveSceneProxy->IsHovered(), /*bUseOverlayIntensity=*/false));
+
+				Mesh.MaterialRenderProxy = LevelColorationMaterialInstance;
+				Collector.RegisterOneFrameMaterialProxy(LevelColorationMaterialInstance);
+			}
+		}
+#endif // WITH_EDITOR
 		else if (EngineShowFlags.BSPSplit
 			&& PrimitiveSceneProxy->ShowInBSPSplitViewmode())
 		{
diff --git a/Engine/Source/Runtime/Engine/Private/ShowFlags.cpp b/Engine/Source/Runtime/Engine/Private/ShowFlags.cpp
--- a/Engine/Source/Runtime/Engine/Private/ShowFlags.cpp
+++ b/Engine/Source/Runtime/Engine/Private/ShowFlags.cpp
@@ -378,6 +378,12 @@ void ApplyViewMode(EViewModeIndex ViewModeIndex, bool bPerspective, FEngineShowF
 		case VMI_PathTracing:
 			bPostProcessing = true;
 			break;
+		case VMI_Zebra:
+		case VMI_Clay:
+		case VMI_FrontBackFace:
+		case VMI_RandomColor:
+			bPostProcessing = true;
+			break;
 	}
 
 	// set the EngineShowFlags:
@@ -420,6 +426,10 @@ void ApplyViewMode(EViewModeIndex ViewModeIndex, bool bPerspective, FEngineShowF
 	EngineShowFlags.SetVisualizeShadowCasters(ViewModeIndex == VMI_ShadowCasters);
 	EngineShowFlags.SetMeshEdges(ViewModeIndex == VMI_Lit_Wireframe);
 	EngineShowFlags.SetActorColoration(ViewModeIndex == VMI_VisualizeActorColoration);
+	EngineShowFlags.SetClay(ViewModeIndex == VMI_Clay);
+	EngineShowFlags.SetZebra(ViewModeIndex == VMI_Zebra);
+	EngineShowFlags.SetFrontBackFace(ViewModeIndex == VMI_FrontBackFace);
+	EngineShowFlags.SetRandomColor(ViewModeIndex == VMI_RandomColor);
 }
 
 void EngineShowFlagOverride(EShowFlagInitMode ShowFlagInitMode, EViewModeIndex ViewModeIndex, FEngineShowFlags& EngineShowFlags, bool bCanDisableTonemapper)
@@ -562,6 +572,43 @@ void EngineShowFlagOverride(EShowFlagInitMode ShowFlagInitMode, EViewModeIndex V
 			EngineShowFlags.SetFog(false);
 		}
 
+
+		static TOptional<bool> NanitePreviousValue;
+		static const auto ICNaniteVar = IConsoleManager::Get().FindConsoleVariable(TEXT("r.Nanite"));
+
+		if (ViewModeIndex == VMI_Clay ||
+			ViewModeIndex == VMI_Zebra ||
+			ViewModeIndex == VMI_FrontBackFace ||
+			ViewModeIndex == VMI_RandomColor)
+		{
+			static const auto ICVar = IConsoleManager::Get().FindTConsoleVariableDataBool(TEXT("r.GeometryInspection.DrawWireframe"));
+			const bool bDrawWireframe = (ICVar && ICVar->GetValueOnGameThread());
+
+			EngineShowFlags.SetMeshEdges(bDrawWireframe);
+
+			EngineShowFlags.SetLighting(false);
+			EngineShowFlags.SetCapsuleShadows(false);
+			EngineShowFlags.SetContactShadows(false);
+			EngineShowFlags.SetDynamicShadows(false);
+
+			EngineShowFlags.SetAtmosphere(false);
+			EngineShowFlags.SetFog(false);
+
+			if (ICNaniteVar)
+			{
+				if (!NanitePreviousValue.IsSet())
+				{
+					NanitePreviousValue = ICNaniteVar->GetBool();
+				}
+				ICNaniteVar->Set(false);
+			}
+		}
+		else if (ICNaniteVar && NanitePreviousValue.IsSet())
+		{
+			ICNaniteVar->Set(NanitePreviousValue.GetValue());
+		}
+
+
 		if (ViewModeIndex == VMI_ShadowCasters)
 		{
 			EngineShowFlags.SetTemporalAA(false);
@@ -604,6 +651,7 @@ void EngineShowFlagOverride(EShowFlagInitMode ShowFlagInitMode, EViewModeIndex V
 		if (ViewModeIndex == VMI_LODColoration || 
 			ViewModeIndex == VMI_HLODColoration ||
 			ViewModeIndex == VMI_VisualizeGPUSkinCache)
+
 		{
 			EngineShowFlags.SetDecals(false); // Decals require the use of FDebugPSInLean.
 		}
@@ -876,6 +924,22 @@ EViewModeIndex FindViewMode(const FEngineShowFlags& EngineShowFlags)
 	{
 		return VMI_ShadowCasters;
 	}
+	else if (EngineShowFlags.Clay)
+	{
+		return VMI_Clay;
+	}
+	else if (EngineShowFlags.Zebra)
+	{
+		return VMI_Zebra;
+	}
+	else if (EngineShowFlags.FrontBackFace)
+	{
+		return VMI_FrontBackFace;
+	}
+	else if (EngineShowFlags.RandomColor)
+	{
+		return VMI_RandomColor;
+	}
 
 	return EngineShowFlags.Lighting ? VMI_Lit : VMI_Unlit;
 }
@@ -920,6 +984,10 @@ const TCHAR* GetViewModeName(EViewModeIndex ViewModeIndex)
 		case VMI_VisualizeGPUSkinCache:		return TEXT("VisualizeGPUSkinCache");
 		case VMI_LWCComplexity:				return TEXT("LWCComplexity");
 		case VMI_ShadowCasters:				return TEXT("ShadowCasters");
+		case VMI_Clay:						return TEXT("Clay");
+		case VMI_Zebra:						return TEXT("Zebra");
+		case VMI_FrontBackFace:				return TEXT("FrontBackFace");
+		case VMI_RandomColor:				return TEXT("RandomColor");
 	}
 	return TEXT("");
 }
diff --git a/Engine/Source/Runtime/Engine/Private/StaticMeshSceneProxy.cpp b/Engine/Source/Runtime/Engine/Private/StaticMeshSceneProxy.cpp
--- a/Engine/Source/Runtime/Engine/Private/StaticMeshSceneProxy.cpp
+++ b/Engine/Source/Runtime/Engine/Private/StaticMeshSceneProxy.cpp
@@ -1758,6 +1758,11 @@ void FStaticMeshSceneProxy::GetDynamicMeshElements(const TArray<const FSceneView
 	#if WITH_EDITOR
 									if (GIsEditor)
 									{
+										if (EngineShowFlags.FrontBackFace)
+										{
+											MeshElement.bDisableBackfaceCulling = true;
+										}
+
 										const FLODInfo::FSectionInfo& Section = LODs[LODIndex].Sections[SectionIndex];
 
 										bSectionIsSelected = Section.bSelected || (bIsWireframeView && bProxyIsSelected);
diff --git a/Engine/Source/Runtime/Engine/Private/UnrealEngine.cpp b/Engine/Source/Runtime/Engine/Private/UnrealEngine.cpp
--- a/Engine/Source/Runtime/Engine/Private/UnrealEngine.cpp
+++ b/Engine/Source/Runtime/Engine/Private/UnrealEngine.cpp
@@ -3061,6 +3061,30 @@ void LoadSpecialMaterial(const FString& MaterialName, const FString& MaterialNam
 	}
 }
 
+/**
+* Loads a special material interface and verifies that it can be loaded
+*
+* @param MaterialName Variable name of the material in the engine to identify which material can't be loaded
+* @param MaterialNamePath Fully qualified name of a material to load/find
+*/
+UMaterialInterface* LoadSpecialMaterialInterface(const FString& MaterialName, const FString& MaterialNamePath)
+{
+	UMaterialInterface* Material = nullptr;
+
+	// find or load the object
+	Material = LoadObject<UMaterialInterface>(NULL, *MaterialNamePath, NULL, LOAD_None, NULL);
+
+	if (!Material)
+	{
+#if !WITH_EDITORONLY_DATA
+		UE_LOG(LogEngine, Log, TEXT("ERROR: Failed to load special material '%s' from path '%s'. This will probably have bad consequences (depending on its use)."), *MaterialName, *MaterialNamePath);
+#else
+		UE_LOG(LogEngine, Fatal, TEXT("Failed to load special material '%s' from path '%s'."), *MaterialName, *MaterialNamePath);
+#endif
+	}
+
+	return Material;
+}
 
 template<typename ClassType>
 void LoadEngineClass(FSoftClassPath& ClassName, TSubclassOf<ClassType>& EngineClassRef)
@@ -3287,6 +3311,12 @@ void UEngine::InitializeObjectReferences()
 		LoadSpecialMaterial(TEXT("PhysicalMaterialMaskMaterialName"), PhysicalMaterialMaskMaterialName.ToString(), PhysicalMaterialMaskMaterial, false);
 		LoadSpecialMaterial(TEXT("TexturePaintingMaskMaterialName"), TexturePaintingMaskMaterialName.ToString(), TexturePaintingMaskMaterial, false);
 
+		ClayMaterial = LoadSpecialMaterialInterface(TEXT("ClayMaterialName"), ClayMaterialName.ToString());
+		LoadSpecialMaterial(TEXT("ZebraMaterialName"), ZebraMaterialName.ToString(), ZebraMaterial, false);
+		ZebraMaterialInstance = UMaterialInstanceDynamic::Create(ZebraMaterial, nullptr);
+		FrontBackFaceMaterial = LoadSpecialMaterialInterface(TEXT("FrontBackFaceMaterialName"), FrontBackFaceMaterialName.ToString());
+		RandomColorMaterial = LoadSpecialMaterialInterface(TEXT("RandomColorMaterialName"), RandomColorMaterialName.ToString());
+
 		ClothPaintMaterialInstance = UMaterialInstanceDynamic::Create(ClothPaintMaterial, nullptr);
 		ClothPaintMaterialWireframeInstance = UMaterialInstanceDynamic::Create(ClothPaintMaterialWireframe, nullptr);
 		ClothPaintOpaqueMaterialInstance = UMaterialInstanceDynamic::Create(ClothPaintOpaqueMaterial, nullptr);
@@ -15402,7 +15432,7 @@ EBrowseReturnVal::Type UEngine::Browse( FWorldContext& WorldContext, FURL URL, F
 
 		if( !WorldContext.PendingNetGame )
 		{
-			// If the inital packet sent in InitNetDriver results in a socket error, HandleDisconnect() and CancelPending() may be called, which will null the PendingNetGame.
+			// If the initial packet sent in InitNetDriver results in a socket error, HandleDisconnect() and CancelPending() may be called, which will null the PendingNetGame.
 			Error = NSLOCTEXT("Engine", "PendingNetGameInitFailure", "Error initializing the network driver.").ToString();
 			BroadcastTravelFailure(WorldContext.World(), ETravelFailure::PendingNetGameCreateFailure, Error);
 			return EBrowseReturnVal::Failure;
@@ -16957,8 +16987,8 @@ void UEngine::CheckAndHandleStaleWorldObjectReferences(FWorldContext* WorldConte
 	for( TObjectIterator<UWorld> It; It; ++It )
 	{
 		UWorld* World = *It;
-		const bool bIsPersistantWorldType = (World->WorldType == EWorldType::Inactive) || (World->WorldType == EWorldType::EditorPreview);
-		if (!bIsPersistantWorldType && !WorldHasValidContext(World))
+		const bool bPersistentWorldType = (World->WorldType == EWorldType::Inactive) || (World->WorldType == EWorldType::EditorPreview);
+		if (!bPersistentWorldType && !WorldHasValidContext(World))
 		{
 			if ((World->PersistentLevel == nullptr || !WorldHasValidContext(World->PersistentLevel->OwningWorld)) && !IsWorldDuplicate(World))
 			{
diff --git a/Engine/Source/Runtime/Engine/Private/ViewModeNames.cpp b/Engine/Source/Runtime/Engine/Private/ViewModeNames.cpp
--- a/Engine/Source/Runtime/Engine/Private/ViewModeNames.cpp
+++ b/Engine/Source/Runtime/Engine/Private/ViewModeNames.cpp
@@ -192,6 +192,24 @@ TArray<FText> FillViewModeDisplayNames()
 			ViewModeDisplayNames.Emplace(LOCTEXT("UViewModeUtils_VMI_VisualizeActorColoration", "Actor Coloration Visualization"));
 		}
 
+		// Geometry Inspection Modes
+		else if (ViewModeIndex == VMI_Clay)
+		{
+			ViewModeDisplayNames.Emplace(LOCTEXT("UViewModeUtils_VMI__Clay", "Clay"));
+		}
+		else if (ViewModeIndex == VMI_Zebra)
+		{
+			ViewModeDisplayNames.Emplace(LOCTEXT("UViewModeUtils_VMI__Zebra", "Zebra"));
+		}
+		else if (ViewModeIndex == VMI_FrontBackFace)
+		{
+			ViewModeDisplayNames.Emplace(LOCTEXT("UViewModeUtils_VMI__FrontBackFace", "Front/Back Face"));
+		}
+		else if (ViewModeIndex == VMI_RandomColor)
+		{
+			ViewModeDisplayNames.Emplace(LOCTEXT("UViewModeUtils_VMI__RandomColor", "Random Color"));
+		}
+
 		// VMI_Max
 		else if (ViewModeIndex == VMI_Max)
 		{
@@ -256,6 +274,7 @@ FName GetViewModeIconName(const EViewModeIndex InViewMode)
 	static FName RayTracingDebugIcon("EditorViewport.RayTracingDebugMode");
 	static FName VisualizeVirtualTextureIcon("EditorViewport.VisualizeVirtualTextureMode");
 	static FName VisualizeActorColorationIcon("EditorViewport.VisualizeActorColorationMode");
+	static FName GeometryInspectionIcon("EditorViewport.CollisionVisibility");
 
 	FName Icon = NAME_None;
 
@@ -364,6 +383,12 @@ FName GetViewModeIconName(const EViewModeIndex InViewMode)
 		break;
 	case VMI_Max:
 		break;
+	case VMI_Clay:
+	case VMI_Zebra:
+	case VMI_FrontBackFace:
+	case VMI_RandomColor:
+		Icon = GeometryInspectionIcon;
+		break;
 	case VMI_Unknown:
 		break;
 	}
diff --git a/Engine/Source/Runtime/Engine/Public/ShowFlags.h b/Engine/Source/Runtime/Engine/Public/ShowFlags.h
--- a/Engine/Source/Runtime/Engine/Public/ShowFlags.h
+++ b/Engine/Source/Runtime/Engine/Public/ShowFlags.h
@@ -531,6 +531,11 @@ private:
 		SetLumenFarFieldTraces(true);
 		SetLumenSecondaryBounces(true);
 		SetLumenShortRangeAmbientOcclusion(true);
+
+		SetClay(false);
+		SetZebra(false);
+		SetFrontBackFace(false);
+		SetRandomColor(false);
 	}
 
 
diff --git a/Engine/Source/Runtime/Engine/Public/ShowFlagsValues.inl b/Engine/Source/Runtime/Engine/Public/ShowFlagsValues.inl
--- a/Engine/Source/Runtime/Engine/Public/ShowFlagsValues.inl
+++ b/Engine/Source/Runtime/Engine/Public/ShowFlagsValues.inl
@@ -485,5 +485,11 @@ SHOWFLAG_ALWAYS_ACCESSIBLE(AlphaInvert, SFG_PostProcess, NSLOCTEXT("UnrealEd", "
 /** Directly render motion vectors and motion vector depth for OpenXR extensions when normal velocity rendering is unsupported. Motion vectors/depth must be in Vulkan NDC-space and can be requested at any size. */
 SHOWFLAG_ALWAYS_ACCESSIBLE(StereoMotionVectors, SFG_Normal, NSLOCTEXT("UnrealEd", "MotionVectorsSF", "Render Stereo Motion Vectors for OpenXR"))
 
+/** Geometry Inspection modes. */
+SHOWFLAG_FIXED_IN_SHIPPING(0, Clay, SFG_Hidden, NSLOCTEXT("UnrealEd", "Clay", "Clay"))
+SHOWFLAG_FIXED_IN_SHIPPING(0, Zebra, SFG_Hidden, NSLOCTEXT("UnrealEd", "Zebra", "Zebra"))
+SHOWFLAG_FIXED_IN_SHIPPING(0, FrontBackFace, SFG_Hidden, NSLOCTEXT("UnrealEd", "FrontBackFace", "Front/Back Face"))
+SHOWFLAG_FIXED_IN_SHIPPING(0, RandomColor, SFG_Hidden, NSLOCTEXT("UnrealEd", "RandomColor", "RandomColor"))
+
 #undef SHOWFLAG_ALWAYS_ACCESSIBLE
 #undef SHOWFLAG_FIXED_IN_SHIPPING
