Prevent navigation from being force enabled while building HLODs

diff --git a/Engine/Source/Editor/UnrealEd/Private/WorldPartition/WorldPartitionHLODsBuilder.cpp b/Engine/Source/Editor/UnrealEd/Private/WorldPartition/WorldPartitionHLODsBuilder.cpp
--- a/Engine/Source/Editor/UnrealEd/Private/WorldPartition/WorldPartitionHLODsBuilder.cpp
+++ b/Engine/Source/Editor/UnrealEd/Private/WorldPartition/WorldPartitionHLODsBuilder.cpp
@@ -48,6 +48,8 @@
 #include "WorldPartition/IWorldPartitionEditorModule.h"
 #include "WorldPartition/RuntimeHashSet/WorldPartitionRuntimeHashSet.h"
 
+#include "NavSystemConfigOverride.h"
+
 #include UE_INLINE_GENERATED_CPP_BY_NAME(WorldPartitionHLODsBuilder)
 
 #define LOCTEXT_NAMESPACE "WorldPartitionHLODsBuilder"
@@ -590,6 +592,12 @@ bool UWorldPartitionHLODsBuilder::BuildHLODActors()
 		WPHLODUtilities->SetHLODBuildEvaluator(IWorldPartitionHLODUtilities::FHLODBuildEvaluator::CreateUObject(this, &UWorldPartitionHLODsBuilder::EvaluateHLODBuildConditions));
 	}
 
+	// Disable nav system config overrides which can possible reenable navigation on our world. This leads to log spamming and extra overhead during our actors loading.
+	for (TActorIterator<ANavSystemConfigOverride> NavSystemConfigOverrideActor(World); NavSystemConfigOverrideActor; ++NavSystemConfigOverrideActor)
+	{
+		NavSystemConfigOverrideActor->UnregisterAllComponents();
+	}
+
 	ON_SCOPE_EXIT
 	{
 		if (bReuseParentBranchHLODs)
@@ -597,6 +605,12 @@ bool UWorldPartitionHLODsBuilder::BuildHLODActors()
 			WPHLODUtilities->SetHLODBuildEvaluator(nullptr);
 			BuildEvaluationStats::PrintStats();
 		}
+
+		// Reenable nav system config overrides
+		for (TActorIterator<ANavSystemConfigOverride> NavSystemConfigOverrideActor(World); NavSystemConfigOverrideActor; ++NavSystemConfigOverrideActor)
+		{
+			NavSystemConfigOverrideActor->RegisterAllComponents();
+		}
 	};
 
 	auto SaveHLODActor = [this](AWorldPartitionHLOD* HLODActor)
