Implemented a new scheduler for GatherText, that allows running certain steps in parallel

Legacy Scheduler:
The legacy scheduler just processed each step in each gather config in sequence.
This works fine for small sets of localization targets, but starts to have scaling issues when you have hundreds/thousands of localization targets.

New Scheduler:
The new scheduler does a pre-pass of the steps in each gather config, and re-arranges them into pre-defined sequential phases (see EGatherTextCommandletPhase), where each phase may contain separate sets of sequential or parallel steps.
This allows tasks that are self-contained (such as PO export/import and LocRes compilation) to run in parallel with the same tasks for other localization targets, which can massively reduce GatherText time for large sets (hundreds/thousands) of localization targets.

Notes for custom gather commandlets/configs:
The new scheduler will produce an equivalent result to the legacy scheduler for gather configs that comply with our reference implementation (such as those shipped with UE, and those generated by the localization dashboard).
If you have custom gather commandlets or gather configs, you may find that the new scheduler re-arranges your tasks in an undesirable way. If you find this happens, you can try the following:
1. Check that your custom gather commandlets define a phase to run in, by overriding UGatherTextCommandletBase::GetPhase(). There are several additional pre/post phases to allow steps to be scheduled around the commandlets shipped with UE.
2. If you cannot represent your gather config in the new scheduler, you can re-enable the legacy scheduler by either passing -LegacyScheduler on the command line, or setting the "Localization.ForceGatherTextLegacyScheduler" CVar to true. Please let us know if you have to do this, as we may be able to tweak the new scheduler to meet your use-case.

diff --git a/Engine/Source/Developer/Localization/Private/LocalizationSourceControlUtil.cpp b/Engine/Source/Developer/Localization/Private/LocalizationSourceControlUtil.cpp
--- a/Engine/Source/Developer/Localization/Private/LocalizationSourceControlUtil.cpp
+++ b/Engine/Source/Developer/Localization/Private/LocalizationSourceControlUtil.cpp
@@ -1,11 +1,14 @@
 // Copyright Epic Games, Inc. All Rights Reserved.
 
 #include "LocalizationSourceControlUtil.h"
+
+#include "HAL/FileManager.h"
+#include "HAL/PlatformFileManager.h"
 #include "ISourceControlModule.h"
 #include "ISourceControlProvider.h"
-#include "SourceControlOperations.h"
-#include "SourceControlHelpers.h"
 #include "Misc/Paths.h"
+#include "SourceControlHelpers.h"
+#include "SourceControlOperations.h"
 
 DEFINE_LOG_CATEGORY_STATIC(LogLocalizationSourceControl, Log, All);
 
@@ -13,11 +16,15 @@ DEFINE_LOG_CATEGORY_STATIC(LogLocalizationSourceControl, Log, All);
 
 FLocalizationSCC::FLocalizationSCC()
 {
+	checkf(IsInGameThread(), TEXT("FLocalizationSCC must be created on the game-thread"));
+
 	ISourceControlModule::Get().GetProvider().Init();
 }
 
 FLocalizationSCC::~FLocalizationSCC()
 {
+	checkf(IsInGameThread(), TEXT("FLocalizationSCC must be destroyed on the game-thread"));
+
 	if (CheckedOutFiles.Num() > 0)
 	{
 		UE_LOG(LogLocalizationSourceControl, Log, TEXT("Revision Control wrapper shutting down with checked out files."));
@@ -26,129 +33,197 @@ FLocalizationSCC::~FLocalizationSCC()
 	ISourceControlModule::Get().GetProvider().Close();
 }
 
-bool FLocalizationSCC::CheckOutFile(const FString& InFile, FText& OutError)
+void FLocalizationSCC::BeginParallelTasks()
 {
-	if (InFile.IsEmpty())
-	{
-		OutError = LOCTEXT("InvalidFileSpecified", "Could not checkout file at invalid path.");
-		return false;
-	}
+	checkf(IsInGameThread(), TEXT("FLocalizationSCC::BeginParallelTasks must be called on the game-thread"));
+	checkf(!bRunningParallelTasks, TEXT("FLocalizationSCC::BeginParallelTasks was called while parallel tasks are already running"));
 
-	if (InFile.StartsWith(TEXT("\\\\")))
+	bRunningParallelTasks = true;
+}
+
+bool FLocalizationSCC::EndParallelTasks(FText& OutError)
+{
+	checkf(IsInGameThread(), TEXT("FLocalizationSCC::EndParallelTasks must be called on the game-thread"));
+	checkf(bRunningParallelTasks, TEXT("FLocalizationSCC::EndParallelTasks was called while no parallel tasks are running"));
+
+	bRunningParallelTasks = false;
+
+	if (DeferredCheckedOutFiles.IsEmpty())
 	{
-		// We can't check out a UNC path, but don't say we failed
 		return true;
 	}
 
-	FText SCCError;
-	if (!IsReady(SCCError))
+	TArray<FString> DeferredCheckedOutFilesArray = DeferredCheckedOutFiles.Array();
+	DeferredCheckedOutFiles.Reset();
+
+	if (!IsReady(OutError))
 	{
-		OutError = SCCError;
 		return false;
 	}
 
-	FString AbsoluteFilename = FPaths::ConvertRelativePathToFull(InFile);
-
-	if (CheckedOutFiles.Contains(AbsoluteFilename))
+	// Try and apply the deferred check-out requests
+	if (!USourceControlHelpers::CheckOutOrAddFiles(DeferredCheckedOutFilesArray))
 	{
-		return true;
-	}
+		ISourceControlProvider& SourceControlProvider = ISourceControlModule::Get().GetProvider();
 
-	bool bSuccessfullyCheckedOut = false;
-	TArray<FString> FilesToBeCheckedOut;
-	FilesToBeCheckedOut.Add(AbsoluteFilename);
+		// Figure out which files failed to check-out or add
+		if (TArray<FSourceControlStateRef> SourceControlStates;
+			SourceControlProvider.GetState(DeferredCheckedOutFilesArray, SourceControlStates, EStateCacheUsage::ForceUpdate) == ECommandResult::Succeeded)
+		{
+			check(DeferredCheckedOutFilesArray.Num() == SourceControlStates.Num());
 
-	FFormatNamedArguments Args;
-	Args.Add(TEXT("Filepath"), FText::FromString(InFile));
+			TArray<FString> FilesToForceSync;
+			for (int32 Index = 0; Index < DeferredCheckedOutFilesArray.Num(); ++Index)
+			{
+				FString& DeferredCheckedOutFile = DeferredCheckedOutFilesArray[Index];
+				const FSourceControlStateRef& SourceControlState = SourceControlStates[Index];
+
+				if (SourceControlState->IsCheckedOut() || SourceControlState->IsAdded())
+				{
+					CheckedOutFiles.Add(MoveTemp(DeferredCheckedOutFile));
+				}
+				else
+				{
+					FilesToForceSync.Add(MoveTemp(DeferredCheckedOutFile));
+				}
+			}
 
-	ISourceControlProvider& SourceControlProvider = ISourceControlModule::Get().GetProvider();
-	FSourceControlStatePtr SourceControlState = SourceControlProvider.GetState(AbsoluteFilename, EStateCacheUsage::ForceUpdate);
+			if (FilesToForceSync.Num() > 0)
+			{
+				// Failed the deferred check-out, so try and restore the on-disk file to the source controlled state
+				TSharedRef<FSync> ForceSyncOperation = ISourceControlOperation::Create<FSync>();
+				ForceSyncOperation->SetForce(true);
+				ForceSyncOperation->SetLastSyncedFlag(true);
+
+				SourceControlProvider.Execute(ForceSyncOperation, FilesToForceSync);
+			}
 
-	if (SourceControlState.IsValid() && SourceControlState->IsDeleted())
+			OutError = FText::Format(LOCTEXT("EndParallelTasksFailed.WithFiles", "FLocalizationSCC::EndParallelTasks failed to check-out or add: {0}"), FText::FromString(FString::Join(FilesToForceSync, TEXT(", "))));
+		}
+		else
+		{
+			OutError = LOCTEXT("EndParallelTasksFailed.Generic", "FLocalizationSCC::EndParallelTasksFailed failed to check-out or add some files, but querying the file states also failed");
+		}
+		return false;
+	}
+
+	return true;
+}
+
+bool FLocalizationSCC::CheckOutFile(const FString& InFile, FText& OutError)
+{
+	if (InFile.IsEmpty())
+	{
+		OutError = LOCTEXT("InvalidFileSpecified", "Could not checkout file at invalid path.");
+		return false;
+	}
+
+	if (InFile.StartsWith(TEXT("\\\\")))
 	{
-		// If it's deleted, we need to revert that first
-		SourceControlProvider.Execute(ISourceControlOperation::Create<FRevert>(), FilesToBeCheckedOut);
-		SourceControlState = SourceControlProvider.GetState(AbsoluteFilename, EStateCacheUsage::Use);
+		// We can't check out a UNC path, but don't say we failed
+		return true;
 	}
 
-	if (SourceControlState.IsValid())
+	FString AbsoluteFilename = FPaths::ConvertRelativePathToFull(InFile);
+
+	if (bRunningParallelTasks)
 	{
-		FString Other;
-		if (SourceControlState->IsAdded() ||
-			SourceControlState->IsCheckedOut())
+		UE::TScopeLock _(FilesMutex);
+
+		if (CheckedOutFiles.Contains(AbsoluteFilename) || DeferredCheckedOutFiles.Contains(AbsoluteFilename))
 		{
-			// Already checked out or opened for add
-			bSuccessfullyCheckedOut = true;
+			return true;
 		}
-		else if (SourceControlState->CanCheckout())
+
+		// Make the file writable on-disk and add it to the deferred set
+		if (IPlatformFile& PlatformFile = FPlatformFileManager::Get().GetPlatformFile();
+			PlatformFile.IsReadOnly(*AbsoluteFilename))
 		{
-			bSuccessfullyCheckedOut = (SourceControlProvider.Execute(ISourceControlOperation::Create<FCheckOut>(), FilesToBeCheckedOut) == ECommandResult::Succeeded);
-			if (!bSuccessfullyCheckedOut)
+			if (PlatformFile.SetReadOnly(*AbsoluteFilename, false))
 			{
-				OutError = FText::Format(LOCTEXT("FailedToCheckOutFile", "Failed to check out file '{Filepath}'."), Args);
+				DeferredCheckedOutFiles.Add(MoveTemp(AbsoluteFilename));
 			}
-		}
-		else if (!SourceControlState->IsSourceControlled() && SourceControlState->CanAdd())
-		{
-			bSuccessfullyCheckedOut = (SourceControlProvider.Execute(ISourceControlOperation::Create<FMarkForAdd>(), FilesToBeCheckedOut) == ECommandResult::Succeeded);
-			if (!bSuccessfullyCheckedOut)
+			else
 			{
-				OutError = FText::Format(LOCTEXT("FailedToAddFileToSourceControl", "Failed to add file '{Filepath}' to revision control."), Args);
+				OutError = FText::Format(LOCTEXT("FailedToCheckOutFile", "Failed to make file writable '{0}'."), FText::FromString(AbsoluteFilename));
+				return false;
 			}
 		}
-		else if (!SourceControlState->IsCurrent())
-		{
-			OutError = FText::Format(LOCTEXT("FileIsNotAtHeadRevision", "File '{Filepath}' is not at head revision."), Args);
-		}
-		else if (SourceControlState->IsCheckedOutOther(&(Other)))
-		{
-			Args.Add(TEXT("Username"), FText::FromString(Other));
-			OutError = FText::Format(LOCTEXT("FileIsAlreadyCheckedOutByAnotherUser", "File '{Filepath}' is checked out by another ('{Username}')."), Args);
-		}
 		else
 		{
-			// Improper or invalid SCC state
-			OutError = FText::Format(LOCTEXT("CouldNotGetStateOfFile", "Could not determine revision control state of file '{Filepath}'."), Args);
+			DeferredCheckedOutFiles.Add(MoveTemp(AbsoluteFilename));
 		}
+
+		return true;
 	}
-	else
+
+	if (!IsReady(OutError))
 	{
-		// Improper or invalid SCC state
-		OutError = FText::Format(LOCTEXT("CouldNotGetStateOfFile", "Could not determine revision control state of file '{Filepath}'."), Args);
+		return false;
 	}
 
-	if (bSuccessfullyCheckedOut)
+	if (CheckedOutFiles.Contains(AbsoluteFilename))
 	{
-		CheckedOutFiles.AddUnique(AbsoluteFilename);
+		return true;
+	}
+
+	if (!USourceControlHelpers::CheckOutOrAddFile(AbsoluteFilename))
+	{
+		OutError = USourceControlHelpers::LastErrorMsg();
+		return false;
 	}
 
-	return bSuccessfullyCheckedOut;
+	// Make sure the file is actually writable, as adding a read-only file to source control may leave it read-only
+	if (IPlatformFile& PlatformFile = FPlatformFileManager::Get().GetPlatformFile();
+		PlatformFile.IsReadOnly(*AbsoluteFilename))
+	{
+		PlatformFile.SetReadOnly(*AbsoluteFilename, false);
+	}
+
+	CheckedOutFiles.Add(MoveTemp(AbsoluteFilename));
+
+	return true;
 }
 
 bool FLocalizationSCC::CheckinFiles(const FText& InChangeDescription, FText& OutError)
 {
-	if (CheckedOutFiles.Num() == 0)
+	checkf(IsInGameThread(), TEXT("FLocalizationSCC::CheckinFiles must be called on the game-thread"));
+	checkf(!bRunningParallelTasks, TEXT("FLocalizationSCC::CheckinFiles was called while parallel tasks are running"));
+
+	if (CheckedOutFiles.IsEmpty())
 	{
 		return true;
 	}
 
-	FText SCCError;
-	if (!IsReady(SCCError))
+	if (!IsReady(OutError))
 	{
-		OutError = SCCError;
 		return false;
 	}
 
 	ISourceControlProvider& SourceControlProvider = ISourceControlModule::Get().GetProvider();
-	SourceControlHelpers::RevertUnchangedFiles(SourceControlProvider, CheckedOutFiles);
 
-	// remove unchanged files
-	for (int32 VerifyIndex = CheckedOutFiles.Num() - 1; VerifyIndex >= 0; --VerifyIndex)
+	// Revert any unchanged files
 	{
-		FSourceControlStatePtr SourceControlState = SourceControlProvider.GetState(CheckedOutFiles[VerifyIndex], EStateCacheUsage::ForceUpdate);
-		if (SourceControlState.IsValid() && !SourceControlState->IsCheckedOut() && !SourceControlState->IsAdded())
+		TArray<FString> CheckedOutFilesArray = CheckedOutFiles.Array();
+		USourceControlHelpers::RevertUnchangedFiles(SourceControlProvider, CheckedOutFilesArray);
+
+		// Update CheckedOutFiles with the files that are still actually checked-out or added
+		if (TArray<FSourceControlStateRef> SourceControlStates;
+			SourceControlProvider.GetState(CheckedOutFilesArray, SourceControlStates, EStateCacheUsage::ForceUpdate) == ECommandResult::Succeeded)
 		{
-			CheckedOutFiles.RemoveAt(VerifyIndex);
+			check(CheckedOutFilesArray.Num() == SourceControlStates.Num());
+			
+			CheckedOutFiles.Reset();
+			for (int32 Index = 0; Index < CheckedOutFilesArray.Num(); ++Index)
+			{
+				FString& CheckedOutFile = CheckedOutFilesArray[Index];
+				const FSourceControlStateRef& SourceControlState = SourceControlStates[Index];
+
+				if (SourceControlState->IsCheckedOut() || SourceControlState->IsAdded())
+				{
+					CheckedOutFiles.Add(MoveTemp(CheckedOutFile));
+				}
+			}
 		}
 	}
 
@@ -156,54 +231,79 @@ bool FLocalizationSCC::CheckinFiles(const FText& InChangeDescription, FText& Out
 	{
 		TSharedRef<FCheckIn, ESPMode::ThreadSafe> CheckInOperation = ISourceControlOperation::Create<FCheckIn>();
 		CheckInOperation->SetDescription(InChangeDescription);
-		if (!SourceControlProvider.Execute(CheckInOperation, CheckedOutFiles))
+		if (!SourceControlProvider.Execute(CheckInOperation, CheckedOutFiles.Array()))
 		{
 			OutError = LOCTEXT("FailedToCheckInFiles", "The checked out localization files could not be checked in.");
 			return false;
 		}
-		CheckedOutFiles.Empty();
+		CheckedOutFiles.Reset();
 	}
+
 	return true;
 }
 
 bool FLocalizationSCC::CleanUp(FText& OutError)
 {
-	if (CheckedOutFiles.Num() == 0)
+	checkf(IsInGameThread(), TEXT("FLocalizationSCC::CleanUp must be called on the game-thread"));
+	checkf(!bRunningParallelTasks, TEXT("FLocalizationSCC::CleanUp was called while parallel tasks are running"));
+
+	if (CheckedOutFiles.IsEmpty())
 	{
 		return true;
 	}
 
-	bool bCleanupSuccess = true;
-	FString AccumulatedErrorsStr;
-	const int32 FileCount = CheckedOutFiles.Num();
-	int32 FileIdx = 0;
-	for (int32 i = 0; i < FileCount; ++i)
+	if (!IsReady(OutError))
+	{
+		return false;
+	}
+
+	TArray<FString> CheckedOutFilesArray = CheckedOutFiles.Array();
+	CheckedOutFiles.Reset();
+
+	// Try and revert everything
+	if (!USourceControlHelpers::RevertFiles(CheckedOutFilesArray))
 	{
-		FText ErrorText;
-		bool bReverted = RevertFile(CheckedOutFiles[FileIdx], ErrorText);
-		bCleanupSuccess &= bReverted;
+		ISourceControlProvider& SourceControlProvider = ISourceControlModule::Get().GetProvider();
 
-		if (!bReverted)
+		// Update CheckedOutFiles with the files that are still actually checked-out or added
+		if (TArray<FSourceControlStateRef> SourceControlStates;
+			SourceControlProvider.GetState(CheckedOutFilesArray, SourceControlStates, EStateCacheUsage::ForceUpdate) == ECommandResult::Succeeded)
 		{
-			if (!AccumulatedErrorsStr.IsEmpty())
+			check(CheckedOutFilesArray.Num() == SourceControlStates.Num());
+
+			TArray<FString> FilesFailedToRevert;
+			for (int32 Index = 0; Index < CheckedOutFilesArray.Num(); ++Index)
 			{
-				AccumulatedErrorsStr += TEXT(", ");
+				FString& CheckedOutFile = CheckedOutFilesArray[Index];
+				const FSourceControlStateRef& SourceControlState = SourceControlStates[Index];
+
+				if (SourceControlState->IsCheckedOut() || SourceControlState->IsAdded())
+				{
+					CheckedOutFiles.Add(MoveTemp(CheckedOutFile));
+				}
+				else
+				{
+					FilesFailedToRevert.Add(MoveTemp(CheckedOutFile));
+				}
 			}
-			AccumulatedErrorsStr += FString::Printf(TEXT("%s : %s"), *CheckedOutFiles[FileIdx], *ErrorText.ToString());
-			++FileIdx;
-		}
-	}
 
-	if (!bCleanupSuccess)
-	{
-		OutError = FText::Format(LOCTEXT("CouldNotCompleteSourceControlCleanup", "Could not complete Revision Control cleanup.  {FailureReason}"), FText::FromString(AccumulatedErrorsStr));
+			OutError = FText::Format(LOCTEXT("CleanUpFailed.WithFiles", "FLocalizationSCC::CleanUp failed to revert: {0}"), FText::FromString(FString::Join(FilesFailedToRevert, TEXT(", "))));
+		}
+		else
+		{
+			OutError = LOCTEXT("CleanUpFailed.Generic", "FLocalizationSCC::CleanUp failed to revert some files, but querying the file states also failed");
+		}
+		return false;
 	}
 
-	return bCleanupSuccess;
+	return true;
 }
 
-bool FLocalizationSCC::IsReady(FText& OutError)
+bool FLocalizationSCC::IsReady(FText& OutError) const
 {
+	checkf(IsInGameThread(), TEXT("FLocalizationSCC::IsReady must be called on the game-thread"));
+	checkf(!bRunningParallelTasks, TEXT("FLocalizationSCC::IsReady was called while parallel tasks are running"));
+
 	if (!ISourceControlModule::Get().IsEnabled())
 	{
 		OutError = LOCTEXT("SourceControlNotEnabled", "Revision control is not enabled.");
@@ -215,51 +315,37 @@ bool FLocalizationSCC::IsReady(FText& OutError)
 		OutError = LOCTEXT("SourceControlNotAvailable", "Revision control server is currently not available.");
 		return false;
 	}
+
 	return true;
 }
 
 bool FLocalizationSCC::RevertFile(const FString& InFile, FText& OutError)
 {
+	checkf(IsInGameThread(), TEXT("FLocalizationSCC::RevertFile must be called on the game-thread"));
+	checkf(!bRunningParallelTasks, TEXT("FLocalizationSCC::RevertFile was called while parallel tasks are running"));
+
 	if (InFile.IsEmpty() || InFile.StartsWith(TEXT("\\\\")))
 	{
 		OutError = LOCTEXT("CouldNotRevertFile", "Could not revert file.");
 		return false;
 	}
 
-	FText SCCError;
-	if (!IsReady(SCCError))
+	if (!IsReady(OutError))
 	{
-		OutError = SCCError;
 		return false;
 	}
 
 	FString AbsoluteFilename = FPaths::ConvertRelativePathToFull(InFile);
-	ISourceControlProvider& SourceControlProvider = ISourceControlModule::Get().GetProvider();
-	FSourceControlStatePtr SourceControlState = SourceControlProvider.GetState(AbsoluteFilename, EStateCacheUsage::ForceUpdate);
-
-	bool bSuccessfullyReverted = false;
-
-	TArray<FString> FilesToRevert;
-	FilesToRevert.Add(AbsoluteFilename);
 
-	if (SourceControlState.IsValid() && !SourceControlState->IsCheckedOut() && !SourceControlState->IsAdded())
+	if (!USourceControlHelpers::RevertFile(AbsoluteFilename))
 	{
-		bSuccessfullyReverted = true;
-	}
-	else
-	{
-		bSuccessfullyReverted = (SourceControlProvider.Execute(ISourceControlOperation::Create<FRevert>(), FilesToRevert) == ECommandResult::Succeeded);
+		OutError = USourceControlHelpers::LastErrorMsg();
+		return false;
 	}
 
-	if (!bSuccessfullyReverted)
-	{
-		OutError = LOCTEXT("CouldNotRevertFile", "Could not revert file.");
-	}
-	else
-	{
-		CheckedOutFiles.Remove(AbsoluteFilename);
-	}
-	return bSuccessfullyReverted;
+	CheckedOutFiles.Remove(AbsoluteFilename);
+
+	return true;
 }
 
 void FLocFileSCCNotifies::PreFileWrite(const FString& InFilename)
diff --git a/Engine/Source/Developer/Localization/Public/LocalizationSourceControlUtil.h b/Engine/Source/Developer/Localization/Public/LocalizationSourceControlUtil.h
--- a/Engine/Source/Developer/Localization/Public/LocalizationSourceControlUtil.h
+++ b/Engine/Source/Developer/Localization/Public/LocalizationSourceControlUtil.h
@@ -2,6 +2,7 @@
 
 #pragma once 
 
+#include "Async/Mutex.h"
 #include "Containers/Array.h"
 #include "Containers/UnrealString.h"
 #include "CoreMinimal.h"
@@ -18,14 +19,31 @@ public:
 	UE_API FLocalizationSCC();
 	UE_API ~FLocalizationSCC();
 
+	/** Start a block of parallel tasks; this will defer check-out requests until EndParallelTasks is called */
+	UE_API void BeginParallelTasks();
+	/** Stop a block of parallel tasks; this will attempt to perform any check-outs deferred during the block */
+	UE_API bool EndParallelTasks(FText& OutError);
+
+	/** Attempt to check-out the given file; if called during a parallel tasks block, the file will be made writable and the check-out request deferred */
 	UE_API bool CheckOutFile(const FString& InFile, FText& OutError);
+	/** Attempt to check-in every file that is currently tracked as checked-out; this cannot be called during a parallel tasks block */
 	UE_API bool CheckinFiles(const FText& InChangeDescription, FText& OutError);
+	/** Attempt to revert every file that is currently tracked as checked-out; this cannot be called during a parallel tasks block */
 	UE_API bool CleanUp(FText& OutError);
+	/** Attempt to revert the given file; this cannot be called during a parallel tasks block */
 	UE_API bool RevertFile(const FString& InFile, FText& OutError);
-	UE_API bool IsReady(FText& OutError);
+	/** Check whether SCC is ready to be used; this cannot be called during a parallel tasks block */
+	UE_API bool IsReady(FText& OutError) const;
 
 private:
-	TArray<FString> CheckedOutFiles;
+	/** Set of files that were checked-out from SCC via CheckOutFile */
+	TSet<FString> CheckedOutFiles;
+	/** Set of files that will be checked-out from SCC via CheckOutFile, but had to be deferred due to running in parallel mode */
+	TSet<FString> DeferredCheckedOutFiles;
+	/** Mutex protecting CheckedOutFiles and DeferredCheckedOutFiles when bRunningParallelTasks is true */
+	mutable UE::FMutex FilesMutex;
+	/** True if this we're being used from a parallel tasks block, and need to defer check-out requests */
+	std::atomic<bool> bRunningParallelTasks = false;
 };
 
 class FLocFileSCCNotifies : public ILocFileNotifies
diff --git a/Engine/Source/Editor/UnrealEd/Classes/Commandlets/ExportDialogueScriptCommandlet.h b/Engine/Source/Editor/UnrealEd/Classes/Commandlets/ExportDialogueScriptCommandlet.h
--- a/Engine/Source/Editor/UnrealEd/Classes/Commandlets/ExportDialogueScriptCommandlet.h
+++ b/Engine/Source/Editor/UnrealEd/Classes/Commandlets/ExportDialogueScriptCommandlet.h
@@ -62,6 +62,8 @@ public:
 	virtual int32 Main(const FString& Params) override;
 	//~ End UCommandlet Interface
 
+	virtual EGatherTextCommandletPhase GetPhase() const override { return EGatherTextCommandletPhase::Export; }
+
 private:
 	static FString GenerateCSVHeader();
 	static FString GenerateCSVRow(const FDialogueScriptEntry& InDialogueScriptEntry);
diff --git a/Engine/Source/Editor/UnrealEd/Classes/Commandlets/GatherTextCommandlet.h b/Engine/Source/Editor/UnrealEd/Classes/Commandlets/GatherTextCommandlet.h
--- a/Engine/Source/Editor/UnrealEd/Classes/Commandlets/GatherTextCommandlet.h
+++ b/Engine/Source/Editor/UnrealEd/Classes/Commandlets/GatherTextCommandlet.h
@@ -7,10 +7,35 @@
 #include "Commandlets/GatherTextCommandletBase.h"
 #include "GatherTextCommandlet.generated.h"
 
-namespace EOutputJson
+USTRUCT()
+struct FGatherTextCommandletTask
 {
-	enum Format { Manifest, Archive };
-}
+	GENERATED_BODY();
+
+public:
+	/** Params to pass the the Main function of the commandlet */
+	UPROPERTY()
+	FString CommandletParams;
+
+	/** Commandlet to invoke the Main function of when running this task */
+	UPROPERTY()
+	TObjectPtr<UGatherTextCommandletBase> Commandlet;
+};
+
+USTRUCT()
+struct FGatherTextCommandletPhase
+{
+	GENERATED_BODY();
+
+public:
+	/** Array of tasks in this phase that need to run in sequence; these will run before ParallelTasks */
+	UPROPERTY()
+	TArray<FGatherTextCommandletTask> SequentialTasks;
+
+	/** Array of tasks in this phase that may run in parallel (if enabled); these will run after SequentialTasks */
+	UPROPERTY()
+	TArray<FGatherTextCommandletTask> ParallelTasks;
+};
 
 /**
  *	UGatherTextCommandlet: One commandlet to rule them all. This commandlet loads a config file and then calls other localization commandlets. Allows localization system to be easily extendable and flexible. 
@@ -27,8 +52,6 @@ public:
 	/** Internal implementation of Main that can be used to provide additional options when running embedded within another process */
 	UNREALED_API int32 Execute(const FString& Params, const TSharedPtr<const FGatherTextCommandletEmbeddedContext>& InEmbeddedContext);
 
-	int32 ProcessGatherConfig(const FString& GatherTextConfigPath, const TSharedPtr<FLocalizationSCC>& CommandletSourceControlInfo, const TArray<FString>& Tokens, const TArray<FString>& Switches, const TMap<FString, FString>& ParamVals);
-
 	static const FString UsageText;
 //~ Begin UGatherTextCommandletBase  Interface
 	virtual bool ShouldRunInPreview(const TArray<FString>& Switches, const TMap<FString, FString>& ParamVals) const override
@@ -39,6 +62,33 @@ public:
 	//~ End UGatherTextCommandletBase  Interface
 
 private:
+	/** Schedule the steps within all gather configs into GatherTextPhases */
+	bool ScheduleGatherConfigs(const TArray<FString>& GatherTextConfigPaths, const TSharedPtr<FLocalizationSCC>& CommandletSourceControlInfo, const TArray<FString>& Tokens, const TArray<FString>& Switches, const TMap<FString, FString>& ParamVals);
+
+	/** Schedule the steps within this gather config into GatherTextPhases */
+	bool ScheduleGatherConfig(const FString& GatherTextConfigPath, const TSharedPtr<FLocalizationSCC>& CommandletSourceControlInfo, const TArray<FString>& Tokens, const TArray<FString>& Switches, const TMap<FString, FString>& ParamVals);
+
+	/** Execute the work that was scheduled into GatherTextPhases */
+	bool ExecuteGatherTextPhases(const TSharedPtr<FLocalizationSCC>& CommandletSourceControlInfo);
+
+	/** Clean-up any stale per-platform data in SplitPlatformConfigs */
+	void CleanupStalePlatformData();
+
 	// Helper function to generate a changelist description
 	FText GetChangelistDescription(const TArray<FString>& GatherTextConfigPaths) const;
+
+	/** The set of tasks this commandlet will run, distributed via a sequential set of phases */
+	UPROPERTY()
+	FGatherTextCommandletPhase GatherTextPhases[(int32)EGatherTextCommandletPhase::NumPhases];
+
+	/** Set of split-platform root paths that we will process in CleanupStalePlatformData */
+	struct FSplitPlatformConfig
+	{
+		TSet<FString> PlatformsToSplit;
+		bool bShouldSplitPlatformData = false;
+	};
+	TMap<FString, FSplitPlatformConfig> SplitPlatformConfigs;
+
+	bool bLegacyScheduler = false;
+	bool bRunningInPreview = false;
 };
diff --git a/Engine/Source/Editor/UnrealEd/Classes/Commandlets/GatherTextCommandletBase.h b/Engine/Source/Editor/UnrealEd/Classes/Commandlets/GatherTextCommandletBase.h
--- a/Engine/Source/Editor/UnrealEd/Classes/Commandlets/GatherTextCommandletBase.h
+++ b/Engine/Source/Editor/UnrealEd/Classes/Commandlets/GatherTextCommandletBase.h
@@ -99,6 +99,14 @@ struct FGatherTextCommandletEmbeddedContext
 	/** Callback used to allow user termination of the gather process */
 	TFunction<bool()> WasAbortRequestedCallback;
 
+	/** If true, the provided TickCallback and WasAbortRequestedCallback are both thread-safe and we are allowed to run tasks in parallel */
+	bool bIsThreadSafe = false;
+
+	bool IsThreadSafe() const
+	{
+		return bIsThreadSafe || (!TickCallback && !WasAbortRequestedCallback);
+	}
+
 	void RunTick() const
 	{
 		if (TickCallback)
@@ -113,6 +121,77 @@ struct FGatherTextCommandletEmbeddedContext
 	}
 };
 
+/**
+ * The GatherText commandlet runs its steps over multiple sequential phases that may contain separate sets of sequential or parallel steps.
+ * This enum defines which phases are available, and allows the sub-commandlets to define which phase they need to run in.
+ * The phases run in the below order, up-to NumPhases. Sequential steps run before parallel steps within the same phase.
+ */
+enum class EGatherTextCommandletPhase : uint8
+{
+	/** Runs before the Initialize phase starts */
+	PreInitialize = 0,
+	/** The main Initialize phase (eg, to perform initialization needed by the other phases) */
+	Initialize,
+	/** Runs after the Initialize phase ends */
+	PostInitialize,
+
+	/** Runs before the UpdateManifests phase starts */
+	PreUpdateManifests,
+	/** The main UpdateManifests phase, where text is harvested and added to GatherManifestHelper, and the manifest files are updated on-disk */
+	UpdateManifests,
+	/** Runs after the UpdateManifests phase ends */
+	PostUpdateManifests,
+
+	/** Runs before the UpdateArchives phase starts */
+	PreUpdateArchives,
+	/** The main UpdateArchives phase, where the archive files are updated on-disk from the current manifest state */
+	UpdateArchives,
+	/** Runs after the UpdateArchives phase ends */
+	PostUpdateArchives,
+
+	/** Runs before the Import phase starts */
+	PreImport,
+	/** The main Import phase, where the archive files are updated on-disk from external file data (eg, PO files) */
+	Import,
+	/** Runs after the Import phase ends */
+	PostImport,
+
+	/** Runs before the Export phase starts */
+	PreExport,
+	/** The main Export phase, where external file data (eg, PO files), are updated from the archives */
+	Export,
+	/** Runs after the Export phase ends */
+	PostExport,
+
+	/** Runs before the Compile phase starts */
+	PreCompile,
+	/** The main Compile phase, where runtime files (eg, LocRes files), are updated from the archives */
+	Compile,
+	/** Runs after the Compile phase ends */
+	PostCompile,
+
+	/** Runs before the Finalize phase starts */
+	PreFinalize,
+	/** The main Finalize phase (eg, to perform clean-up or reporting of work done by the other phases) */
+	Finalize,
+	/** Runs after the Finalize phase ends */
+	PostFinalize,
+
+	/** Internal; do not use */
+	NumPhases,
+
+	/**
+	 * No specific phase has been specified, and this will be scheduled in whichever phase we are currently building.
+	 * You should upgrade your commandlet to use a specific phase to ensure it gets sequenced correctly!
+	 */
+	Undefined,
+
+	/** Internal; do not use */
+	FirstPhase = PreInitialize,
+};
+
+UNREALED_API const TCHAR* LexToString(EGatherTextCommandletPhase Phase);
+
 /**
  *	UGatherTextCommandletBase: Base class for localization commandlets. Just to force certain behaviors and provide helper functionality. 
  */
@@ -122,6 +201,13 @@ class UGatherTextCommandletBase : public UCommandlet
 	GENERATED_UCLASS_BODY()
 
 public:
+	/** Called once before GetPhase to allow the commandlet to compute which phase it needs to run it. Params are the same as are passed to Main. */
+	virtual bool ConfigurePhase(const FString& Params) { return true; }
+	/** Which phase does this task need to run in? */
+	virtual EGatherTextCommandletPhase GetPhase() const { return EGatherTextCommandletPhase::Undefined; }
+	/** Can this task run in parallel with other tasks in its phase? If true, a random thread may invoke Main so any external/shared resources used must be thread-safe. */
+	virtual bool CanParallelize() const { return false; }
+
 	UNREALED_API void SetEmbeddedContext(const TSharedPtr<const FGatherTextCommandletEmbeddedContext>& InEmbeddedContext);
 	UNREALED_API virtual void Initialize( const TSharedRef< FLocTextHelper >& InGatherManifestHelper, const TSharedPtr< FLocalizationSCC >& InSourceControlInfo );
 	UNREALED_API virtual void BeginDestroy() override;
diff --git a/Engine/Source/Editor/UnrealEd/Classes/Commandlets/GatherTextFromAssetsCommandlet.h b/Engine/Source/Editor/UnrealEd/Classes/Commandlets/GatherTextFromAssetsCommandlet.h
--- a/Engine/Source/Editor/UnrealEd/Classes/Commandlets/GatherTextFromAssetsCommandlet.h
+++ b/Engine/Source/Editor/UnrealEd/Classes/Commandlets/GatherTextFromAssetsCommandlet.h
@@ -30,6 +30,8 @@ public:
 	//~ Begin UCommandlet Interface
 	virtual int32 Main(const FString& Params) override;
 
+	virtual EGatherTextCommandletPhase GetPhase() const override { return EGatherTextCommandletPhase::UpdateManifests; }
+
 	bool GetConfigurationScript(const TMap<FString, FString>& InCommandLineParameters, FString& OutFilePath, FString& OutStepSectionName);
 	bool ConfigureFromScript(const FString& GatherTextConfigPath, const FString& SectionName);
 
diff --git a/Engine/Source/Editor/UnrealEd/Classes/Commandlets/GatherTextFromMetadataCommandlet.h b/Engine/Source/Editor/UnrealEd/Classes/Commandlets/GatherTextFromMetadataCommandlet.h
--- a/Engine/Source/Editor/UnrealEd/Classes/Commandlets/GatherTextFromMetadataCommandlet.h
+++ b/Engine/Source/Editor/UnrealEd/Classes/Commandlets/GatherTextFromMetadataCommandlet.h
@@ -21,6 +21,7 @@ public:
 	//~ End UCommandlet Interface
 
 	//~ Begin UGatherTextCommandletBase  Interface
+	virtual EGatherTextCommandletPhase GetPhase() const override { return EGatherTextCommandletPhase::UpdateManifests; }
 	virtual bool ShouldRunInPreview(const TArray<FString>& Switches, const TMap<FString, FString>& ParamVals) const override;
 	//~ End UGatherTextCommandletBase  Interface
 
diff --git a/Engine/Source/Editor/UnrealEd/Classes/Commandlets/GatherTextFromSourceCommandlet.h b/Engine/Source/Editor/UnrealEd/Classes/Commandlets/GatherTextFromSourceCommandlet.h
--- a/Engine/Source/Editor/UnrealEd/Classes/Commandlets/GatherTextFromSourceCommandlet.h
+++ b/Engine/Source/Editor/UnrealEd/Classes/Commandlets/GatherTextFromSourceCommandlet.h
@@ -625,6 +625,7 @@ public:
 	virtual int32 Main(const FString& Params) override;
 	//~ End UCommandlet Interface
 	//~ Begin UGatherTextCommandletBase  Interface
+	virtual EGatherTextCommandletPhase GetPhase() const override { return EGatherTextCommandletPhase::UpdateManifests; }
 	virtual bool ShouldRunInPreview(const TArray<FString>& Switches, const TMap<FString, FString>& ParamVals) const override;
 	//~ End UGatherTextCommandletBase  Interface
 
diff --git a/Engine/Source/Editor/UnrealEd/Classes/Commandlets/GenerateGatherArchiveCommandlet.h b/Engine/Source/Editor/UnrealEd/Classes/Commandlets/GenerateGatherArchiveCommandlet.h
--- a/Engine/Source/Editor/UnrealEd/Classes/Commandlets/GenerateGatherArchiveCommandlet.h
+++ b/Engine/Source/Editor/UnrealEd/Classes/Commandlets/GenerateGatherArchiveCommandlet.h
@@ -19,5 +19,7 @@ public:
 	//~ Begin UCommandlet Interface
 	virtual int32 Main( const FString& Params ) override;
 	//~ End UCommandlet Interface
+	virtual EGatherTextCommandletPhase GetPhase() const override { return EGatherTextCommandletPhase::UpdateArchives; }
+	virtual bool CanParallelize() const { return true; }
 #endif
 };
diff --git a/Engine/Source/Editor/UnrealEd/Classes/Commandlets/GenerateGatherManifestCommandlet.h b/Engine/Source/Editor/UnrealEd/Classes/Commandlets/GenerateGatherManifestCommandlet.h
--- a/Engine/Source/Editor/UnrealEd/Classes/Commandlets/GenerateGatherManifestCommandlet.h
+++ b/Engine/Source/Editor/UnrealEd/Classes/Commandlets/GenerateGatherManifestCommandlet.h
@@ -20,6 +20,7 @@ public:
 	virtual int32 Main( const FString& Params ) override;
 	//~ End UCommandlet Interface
 		//~ Begin UGatherTextCommandletBase  Interface
+	virtual EGatherTextCommandletPhase GetPhase() const override { return EGatherTextCommandletPhase::UpdateManifests; }
 	virtual bool ShouldRunInPreview(const TArray<FString>& Switches, const TMap<FString, FString>& ParamVals) const override;
 	//~ End UGatherTextCommandletBase  Interface
 #endif
diff --git a/Engine/Source/Editor/UnrealEd/Classes/Commandlets/GenerateTextLocalizationReportCommandlet.h b/Engine/Source/Editor/UnrealEd/Classes/Commandlets/GenerateTextLocalizationReportCommandlet.h
--- a/Engine/Source/Editor/UnrealEd/Classes/Commandlets/GenerateTextLocalizationReportCommandlet.h
+++ b/Engine/Source/Editor/UnrealEd/Classes/Commandlets/GenerateTextLocalizationReportCommandlet.h
@@ -17,6 +17,9 @@ public:
 	virtual int32 Main(const FString& Params) override;
 	//~ End UCommandlet Interface
 
+	virtual EGatherTextCommandletPhase GetPhase() const override { return EGatherTextCommandletPhase::Finalize; }
+	virtual bool CanParallelize() const { return true; }
+
 	virtual bool ShouldRunInPreview(const TArray<FString>& Switches, const TMap<FString, FString>& ParamVals) const
 	{
 		return true;
diff --git a/Engine/Source/Editor/UnrealEd/Classes/Commandlets/GenerateTextLocalizationResourceCommandlet.h b/Engine/Source/Editor/UnrealEd/Classes/Commandlets/GenerateTextLocalizationResourceCommandlet.h
--- a/Engine/Source/Editor/UnrealEd/Classes/Commandlets/GenerateTextLocalizationResourceCommandlet.h
+++ b/Engine/Source/Editor/UnrealEd/Classes/Commandlets/GenerateTextLocalizationResourceCommandlet.h
@@ -20,4 +20,7 @@ public:
 	//~ Begin UCommandlet Interface
 	virtual int32 Main(const FString& Params) override;
 	//~ End UCommandlet Interface
+
+	virtual EGatherTextCommandletPhase GetPhase() const override { return EGatherTextCommandletPhase::Compile; }
+	virtual bool CanParallelize() const { return true; }
 };
diff --git a/Engine/Source/Editor/UnrealEd/Classes/Commandlets/ImportDialogueScriptCommandlet.h b/Engine/Source/Editor/UnrealEd/Classes/Commandlets/ImportDialogueScriptCommandlet.h
--- a/Engine/Source/Editor/UnrealEd/Classes/Commandlets/ImportDialogueScriptCommandlet.h
+++ b/Engine/Source/Editor/UnrealEd/Classes/Commandlets/ImportDialogueScriptCommandlet.h
@@ -20,6 +20,8 @@ public:
 	virtual int32 Main(const FString& Params) override;
 	//~ End UCommandlet Interface
 
+	virtual EGatherTextCommandletPhase GetPhase() const override { return EGatherTextCommandletPhase::Import; }
+
 private:
 	bool ImportDialogueScriptForCulture(FLocTextHelper& InLocTextHelper, const FString& InDialogueScriptFileName, const FString& InCultureName, const bool bIsNativeCulture);
 };
diff --git a/Engine/Source/Editor/UnrealEd/Classes/Commandlets/ImportLocalizedDialogueCommandlet.h b/Engine/Source/Editor/UnrealEd/Classes/Commandlets/ImportLocalizedDialogueCommandlet.h
--- a/Engine/Source/Editor/UnrealEd/Classes/Commandlets/ImportLocalizedDialogueCommandlet.h
+++ b/Engine/Source/Editor/UnrealEd/Classes/Commandlets/ImportLocalizedDialogueCommandlet.h
@@ -24,6 +24,8 @@ public:
 	virtual int32 Main(const FString& Params) override;
 	//~ End UCommandlet Interface
 
+	virtual EGatherTextCommandletPhase GetPhase() const override { return EGatherTextCommandletPhase::Import; }
+
 private:
 	struct FCultureImportInfo
 	{
diff --git a/Engine/Source/Editor/UnrealEd/Classes/Commandlets/InternationalizationExportCommandlet.h b/Engine/Source/Editor/UnrealEd/Classes/Commandlets/InternationalizationExportCommandlet.h
--- a/Engine/Source/Editor/UnrealEd/Classes/Commandlets/InternationalizationExportCommandlet.h
+++ b/Engine/Source/Editor/UnrealEd/Classes/Commandlets/InternationalizationExportCommandlet.h
@@ -23,4 +23,12 @@ public:
 	//~ Begin UCommandlet Interface
 	virtual int32 Main(const FString& Params) override;
 	//~ End UCommandlet Interface
+
+	virtual bool ConfigurePhase(const FString& Params) override;
+	virtual EGatherTextCommandletPhase GetPhase() const override;
+	virtual bool CanParallelize() const { return true; }
+
+private:
+	bool bDoImport = false;
+	bool bDoExport = false;
 };
diff --git a/Engine/Source/Editor/UnrealEd/Private/Commandlets/GatherTextCommandlet.cpp b/Engine/Source/Editor/UnrealEd/Private/Commandlets/GatherTextCommandlet.cpp
--- a/Engine/Source/Editor/UnrealEd/Private/Commandlets/GatherTextCommandlet.cpp
+++ b/Engine/Source/Editor/UnrealEd/Private/Commandlets/GatherTextCommandlet.cpp
@@ -2,6 +2,7 @@
 
 #include "Commandlets/GatherTextCommandlet.h"
 
+#include "Async/ParallelFor.h"
 #include "Commandlets/GatherTextFromSourceCommandlet.h"
 #include "GeneralProjectSettings.h"
 #include "HAL/FileManager.h"
@@ -22,6 +23,12 @@ DEFINE_LOG_CATEGORY_STATIC(LogGatherTextCommandlet, Log, All);
 namespace GatherTextCommandlet
 {
 	static constexpr int32 LocalizationLogIdentifier = 304;
+
+	static bool bForceLegacyScheduler = false;
+	static FAutoConsoleVariableRef CVarForceLegacyScheduler(TEXT("Localization.ForceGatherTextLegacyScheduler"), bForceLegacyScheduler, TEXT("True to force the GatherText commandlet to use the legacy scheduler for its tasks, or false to use the new scheduler"));
+
+	static bool bAllowParallelTasks = true;
+	static FAutoConsoleVariableRef CVarAllowParallelTasks(TEXT("Localization.AllowParallelGatherTextTasks"), bAllowParallelTasks, TEXT("True to allow the GatherText commandlet to run tasks in parallel, or false to force everything to run in sequence"));
 }
 
 /**
@@ -85,7 +92,7 @@ int32 UGatherTextCommandlet::Execute(const FString& Params, const TSharedPtr<con
 			else
 			{
 				UE_LOGFMT(LogGatherTextCommandlet, Warning, "There are no config file paths in specified config ,list '{configList}'. Please check to see the is correctly populated.",
-					("configList", **ConfigListFileParamPtr),
+					("configList", *ConfigListFileParamPtr),
 					("id", GatherTextCommandlet::LocalizationLogIdentifier)
 				);
 			}
@@ -93,7 +100,7 @@ int32 UGatherTextCommandlet::Execute(const FString& Params, const TSharedPtr<con
 		else
 		{
 			UE_LOGFMT(LogGatherTextCommandlet, Warning, "Specified config list file '{configList}' does not exist. No additional config files from -ConfigList can be added.",
-				("configList", **ConfigListFileParamPtr),
+				("configList", *ConfigListFileParamPtr),
 				("id", GatherTextCommandlet::LocalizationLogIdentifier)
 			);
 		}
@@ -110,18 +117,22 @@ int32 UGatherTextCommandlet::Execute(const FString& Params, const TSharedPtr<con
 		{
 			GatherTextConfigPath = FPaths::Combine(*ProjectBasePath, *GatherTextConfigPath);
 		}
+		GatherTextConfigPath = FConfigCacheIni::NormalizeConfigIniPath(GatherTextConfigPath);
 	}
 
 	if (GatherTextConfigPaths.Num() == 0)
 	{
 		UE_LOGFMT(LogGatherTextCommandlet, Error, "-config or -ConfigList not specified. If -ConfigList was specified, please check that the file path is valid.\n{usageText}",
-			("usageText", *UsageText),
+			("usageText", UsageText),
 			("id", GatherTextCommandlet::LocalizationLogIdentifier)
 		);
 		return -1;
 	}
 
-	const bool bRunningInPreview = Switches.Contains(UGatherTextCommandletBase::PreviewSwitch);
+	bLegacyScheduler = GatherTextCommandlet::bForceLegacyScheduler || Switches.Contains(TEXT("LegacyScheduler"));
+	UE_LOG(LogGatherTextCommandlet, Display, TEXT("LegacyScheduler: %s"), bLegacyScheduler ? TEXT("Yes") : TEXT("No"));
+
+	bRunningInPreview = Switches.Contains(UGatherTextCommandletBase::PreviewSwitch);
 	if (bRunningInPreview)
 	{
 		UE_LOG(LogGatherTextCommandlet, Display, TEXT("Running commandlet in preview mode. Some child commandlets and steps will be skipped."));
@@ -150,41 +161,65 @@ int32 UGatherTextCommandlet::Execute(const FString& Params, const TSharedPtr<con
 		if (!CommandletSourceControlInfo->IsReady(SCCErrorStr))
 		{
 			UE_LOGFMT(LogGatherTextCommandlet, Error, "Revision Control error: {error}",
-				("error", *SCCErrorStr.ToString()),
+				("error", SCCErrorStr.ToString()),
 				("id", GatherTextCommandlet::LocalizationLogIdentifier)
 			);
 			return -1;
 		}
 	}
 
-	FScopedSlowTask SlowTask(GatherTextConfigPaths.Num(), (EmbeddedContext && EmbeddedContext->SlowTaskMessageOverride) ? *EmbeddedContext->SlowTaskMessageOverride : NSLOCTEXT("GatherTextCommandlet", "GatherTextTask", "Running Localization Tasks for Targets..."));
-
 	double AllCommandletExecutionStartTime = FPlatformTime::Seconds();
-	for (const FString& GatherTextConfigPath : GatherTextConfigPaths)
+
+	ON_SCOPE_EXIT
 	{
-		SlowTask.EnterProgressFrame();
+		for (const FString& GatherTextConfigPath : GatherTextConfigPaths)
+		{
+			if (GConfig->FindConfigFile(GatherTextConfigPath))
+			{
+				GConfig->UnloadFile(GatherTextConfigPath);
+			}
+		}
+	};
 
-		if (SlowTask.ShouldCancel() || (EmbeddedContext && EmbeddedContext->ShouldAbort()))
+	// Schedule the work
+	{
+		if (!ScheduleGatherConfigs(GatherTextConfigPaths, CommandletSourceControlInfo, Tokens, Switches, ParamVals))
 		{
-			UE_LOGFMT(LogGatherTextCommandlet, Error, "GatherText aborted!",
-				("id", GatherTextCommandlet::LocalizationLogIdentifier)
-			);
 			return -1;
 		}
 
-		if (EmbeddedContext)
+		if (Switches.Contains(TEXT("ScheduleOnly")))
 		{
-			EmbeddedContext->RunTick();
+			// Only wanted to see how the tasks would be scheduled, but not actually run the work
+			return 0;
 		}
+	}
 
-		const int32 Result = ProcessGatherConfig(GatherTextConfigPath, CommandletSourceControlInfo, Tokens, Switches, ParamVals);
-		if (Result != 0)
+	// Execute the work
+	{
+		ON_SCOPE_EXIT
 		{
 			UGatherTextFromSourceCommandlet::LogStats();
-			return Result;
+		};
+
+		if (!ExecuteGatherTextPhases(CommandletSourceControlInfo))
+		{
+			if (CommandletSourceControlInfo.IsValid() && !bDisableSubmit)
+			{
+				FText SCCErrorStr;
+				if (!CommandletSourceControlInfo->CleanUp(SCCErrorStr))
+				{
+					UE_LOGFMT(LogGatherTextCommandlet, Error, "Revision Control error: {error}",
+						("error", SCCErrorStr.ToString()),
+						("id", GatherTextCommandlet::LocalizationLogIdentifier)
+					);
+				}
+			}
+			return -1;
 		}
 	}
-	UGatherTextFromSourceCommandlet::LogStats();
+
+	CleanupStalePlatformData();
 
 	if (CommandletSourceControlInfo.IsValid() && !bDisableSubmit)
 	{
@@ -196,13 +231,13 @@ int32 UGatherTextCommandlet::Execute(const FString& Params, const TSharedPtr<con
 		else
 		{
 			UE_LOGFMT(LogGatherTextCommandlet, Error, "Revision control error: {error}",
-				("error", *SCCErrorStr.ToString()),
+				("error", SCCErrorStr.ToString()),
 				("id", GatherTextCommandlet::LocalizationLogIdentifier)
 			);
 			if (!CommandletSourceControlInfo->CleanUp(SCCErrorStr))
 			{
 				UE_LOGFMT(LogGatherTextCommandlet, Error, "Revision control error post-cleanup: {error}",
-					("error", *SCCErrorStr.ToString()),
+					("error", SCCErrorStr.ToString()),
 					("id", GatherTextCommandlet::LocalizationLogIdentifier)
 				);
 			}
@@ -216,20 +251,98 @@ int32 UGatherTextCommandlet::Execute(const FString& Params, const TSharedPtr<con
 	return 0;
 }
 
-int32 UGatherTextCommandlet::ProcessGatherConfig(const FString& GatherTextConfigPath, const TSharedPtr<FLocalizationSCC>& CommandletSourceControlInfo, const TArray<FString>& Tokens, const TArray<FString>& Switches, const TMap<FString, FString>& ParamVals)
+bool UGatherTextCommandlet::ScheduleGatherConfigs(const TArray<FString>& GatherTextConfigPaths, const TSharedPtr<FLocalizationSCC>& CommandletSourceControlInfo, const TArray<FString>& Tokens, const TArray<FString>& Switches, const TMap<FString, FString>& ParamVals)
 {
-	GConfig->LoadFile(*GatherTextConfigPath);
+	FScopedSlowTask SlowTask(GatherTextConfigPaths.Num(), (EmbeddedContext && EmbeddedContext->SlowTaskMessageOverride) ? *EmbeddedContext->SlowTaskMessageOverride : NSLOCTEXT("GatherTextCommandlet", "GatherTextTask.SchedulingTasks", "Scheduling Localization Tasks for Targets..."));
 
-	if (!GConfig->FindConfigFile(*GatherTextConfigPath))
+	// Schedule the set of localization tasks we will run, as that needs to happen on the game-thread
+	for (const FString& GatherTextConfigPath : GatherTextConfigPaths)
 	{
-		UE_LOGFMT(LogGatherTextCommandlet, Error, "Loading Config File '{configFile}' failed.",
-			("configFile", *GatherTextConfigPath),
-			("id", GatherTextCommandlet::LocalizationLogIdentifier)
-		);
-		return -1; 
+		SlowTask.EnterProgressFrame();
+
+		if (SlowTask.ShouldCancel() || (EmbeddedContext && EmbeddedContext->ShouldAbort()))
+		{
+			UE_LOGFMT(LogGatherTextCommandlet, Error, "GatherText aborted!",
+				("id", GatherTextCommandlet::LocalizationLogIdentifier)
+			);
+			return false;
+		}
+
+		if (!GConfig->Find(GatherTextConfigPath))
+		{
+			UE_LOGFMT(LogGatherTextCommandlet, Error, "Loading Config File '{configFile}' failed.",
+				("configFile", GatherTextConfigPath),
+				("id", GatherTextCommandlet::LocalizationLogIdentifier)
+			);
+			return false;
+		}
+
+		if (!ScheduleGatherConfig(GatherTextConfigPath, CommandletSourceControlInfo, Tokens, Switches, ParamVals))
+		{
+			UE_LOGFMT(LogGatherTextCommandlet, Error, "Scheduling tasks from '{configFile}' failed.",
+				("configFile", GatherTextConfigPath),
+				("id", GatherTextCommandlet::LocalizationLogIdentifier)
+			);
+			return false;
+		}
+	}
+
+	// If any phases only have a single parallel task, then just add that to the end of the sequential tasks as the parallelization will just add overhead
+	for (int32 PhaseIndex = (int32)EGatherTextCommandletPhase::FirstPhase; PhaseIndex < (int32)EGatherTextCommandletPhase::NumPhases; ++PhaseIndex)
+	{
+		if (FGatherTextCommandletPhase& Phase = GatherTextPhases[PhaseIndex];
+			Phase.ParallelTasks.Num() == 1)
+		{
+			UE_LOG(LogGatherTextCommandlet, Display, TEXT("Phase '%s' has a single parallel task. Adding it to the end of its sequential tasks."), LexToString((EGatherTextCommandletPhase)PhaseIndex));
+			Phase.SequentialTasks.Add(MoveTemp(Phase.ParallelTasks[0]));
+			Phase.ParallelTasks.Reset();
+		}
+	}
+
+	// Log the full set of work that we scheduled for each phase
+	for (int32 PhaseIndex = (int32)EGatherTextCommandletPhase::FirstPhase; PhaseIndex < (int32)EGatherTextCommandletPhase::NumPhases; ++PhaseIndex)
+	{
+		auto BuildTaskStatsReport =
+			[](const TArray<FGatherTextCommandletTask>& TasksArray)
+			{
+				TMap<FName, int32> NumTasksByClass;
+				for (const FGatherTextCommandletTask& Task : TasksArray)
+				{
+					++NumTasksByClass.FindOrAdd(Task.Commandlet->GetClass()->GetFName());
+				}
+
+				FString TaskStatsReport;
+				{
+					bool bAddComma = false;
+					for (const TTuple<FName, int32>& NumTasksForClass : NumTasksByClass)
+					{
+						if (bAddComma)
+						{
+							TaskStatsReport += TEXT(", ");
+						}
+						TaskStatsReport.AppendInt(NumTasksForClass.Value);
+						TaskStatsReport += TEXT(' ');
+						NumTasksForClass.Key.AppendString(TaskStatsReport);
+						bAddComma = true;
+					}
+				}
+				return TaskStatsReport;
+			};
+
+		const FGatherTextCommandletPhase& Phase = GatherTextPhases[PhaseIndex];
+		UE_LOG(LogGatherTextCommandlet, Display, TEXT("Phase '%s':\n\t%d sequential task(s): %s\n\t%d parallel tasks(s): %s"),
+			LexToString((EGatherTextCommandletPhase)PhaseIndex),
+			Phase.SequentialTasks.Num(), *BuildTaskStatsReport(Phase.SequentialTasks),
+			Phase.ParallelTasks.Num(), *BuildTaskStatsReport(Phase.ParallelTasks)
+			);
 	}
 
-	UE_LOG(LogGatherTextCommandlet, Display, TEXT("Beginning GatherText Commandlet for '%s'"), *GatherTextConfigPath);
+	return true;
+}
+
+bool UGatherTextCommandlet::ScheduleGatherConfig(const FString& GatherTextConfigPath, const TSharedPtr<FLocalizationSCC>& CommandletSourceControlInfo, const TArray<FString>& Tokens, const TArray<FString>& Switches, const TMap<FString, FString>& ParamVals)
+{
+	UE_LOG(LogGatherTextCommandlet, Display, TEXT("Scheduling GatherText tasks for '%s'"), *GatherTextConfigPath);
 
 	// Read in the platform split mode to use
 	ELocTextPlatformSplitMode PlatformSplitMode = ELocTextPlatformSplitMode::None;
@@ -305,35 +418,19 @@ int32 UGatherTextCommandlet::ProcessGatherConfig(const FString& GatherTextConfig
 		GeneratedParamsAndSwitches += FString::Printf(TEXT(" -%s"), *Switch);
 	}
 
-	FScopedSlowTask SlowTask(StepNames.Num(), NSLOCTEXT("GatherTextCommandlet", "GatherTextSubTask", "Running Localization Tasks..."));
-
-	const bool bRunningInPreview = Switches.Contains(TEXT("Preview"));
-	// Execute each step defined in the config file.
+	// Schedule each step defined in the config file.
+	EGatherTextCommandletPhase PreviousTaskPhase = EGatherTextCommandletPhase::FirstPhase;
+	bool bPreviousTaskWasParallel = false;
 	for (const FString& StepName : StepNames)
 	{
-		SlowTask.EnterProgressFrame();
-
-		if (SlowTask.ShouldCancel() || (EmbeddedContext && EmbeddedContext->ShouldAbort()))
-		{
-			UE_LOGFMT(LogGatherTextCommandlet, Error, "GatherText aborted!",
-				("id", GatherTextCommandlet::LocalizationLogIdentifier)
-			);
-			return -1;
-		}
-
-		if (EmbeddedContext)
-		{
-			EmbeddedContext->RunTick();
-		}
-
 		FString CommandletClassName = GConfig->GetStr( *StepName, TEXT("CommandletClass"), GatherTextConfigPath ) + TEXT("Commandlet");
 
 		UClass* CommandletClass = FindFirstObject<UClass>(*CommandletClassName, EFindFirstObjectOptions::None, ELogVerbosity::Warning, TEXT("UGatherTextCommandlet::ProcessGatherConfig"));
 		if (!CommandletClass)
 		{
 			UE_LOGFMT(LogGatherTextCommandlet, Error, "The commandlet name {commandletName} in section {section} is invalid.",
-				("commandletName", *CommandletClassName),
-				("section", *StepName),
+				("commandletName", CommandletClassName),
+				("section", StepName),
 				("id", GatherTextCommandlet::LocalizationLogIdentifier)
 			);
 			continue;
@@ -352,93 +449,296 @@ int32 UGatherTextCommandlet::ProcessGatherConfig(const FString& GatherTextConfig
 			continue;
 		}
 
-		// Execute the commandlet.
-		double CommandletExecutionStartTime = FPlatformTime::Seconds();
-
-		UE_LOG(LogGatherTextCommandlet, Display, TEXT("Executing %s: %s"), *StepName, *CommandletClassName);
-		
-		FString GeneratedCmdLine = FString::Printf(TEXT("-Config=\"%s\" -Section=%s"), *GatherTextConfigPath , *StepName);
+		FString GeneratedCmdLine = FString::Printf(TEXT("-Config=\"%s\" -Section=%s"), *GatherTextConfigPath, *StepName);
 		GeneratedCmdLine += GeneratedParamsAndSwitches;
 
-		if( 0 != Commandlet->Main( GeneratedCmdLine ) )
+		bool bParallelTask = false;
+		EGatherTextCommandletPhase TaskPhase = PreviousTaskPhase;
+		if (!bLegacyScheduler)
 		{
-			UE_LOGFMT(LogGatherTextCommandlet, Error, "{stepName}-{commandletName} reported an error.",
-				("stepName", *StepName),
-				("commandletName", *CommandletClassName),
-				("id", GatherTextCommandlet::LocalizationLogIdentifier)
-			);
-			if( CommandletSourceControlInfo.IsValid() )
+			bParallelTask = Commandlet->CanParallelize();
+
+			// Which phase should this commandlet run it?
+			if (!Commandlet->ConfigurePhase(GeneratedCmdLine))
 			{
-				FText SCCErrorStr;
-				if( !CommandletSourceControlInfo->CleanUp( SCCErrorStr ) )
+				UE_LOGFMT(LogGatherTextCommandlet, Error, "The commandlet name {commandletName} in section {section} failed to configure a valid phase.",
+					("commandletName", CommandletClassName),
+					("section", StepName),
+					("id", GatherTextCommandlet::LocalizationLogIdentifier)
+				);
+				return false;
+			}
+
+			TaskPhase = Commandlet->GetPhase();
+			if (TaskPhase == EGatherTextCommandletPhase::NumPhases)
+			{
+				TaskPhase = EGatherTextCommandletPhase::Undefined;
+				UE_LOGFMT(LogGatherTextCommandlet, Error, "The commandlet name {commandletName} in section {section} returned an invalid phase. It will be set to Undefined.",
+					("commandletName", CommandletClassName),
+					("section", StepName),
+					("id", GatherTextCommandlet::LocalizationLogIdentifier)
+				);
+			}
+			if (TaskPhase == EGatherTextCommandletPhase::Undefined)
+			{
+				TaskPhase = PreviousTaskPhase;
+				if (bPreviousTaskWasParallel && !bParallelTask)
 				{
-					UE_LOGFMT(LogGatherTextCommandlet, Error, "Revision Control error: {error}",
-						("error", *SCCErrorStr.ToString()),
-						("id", GatherTextCommandlet::LocalizationLogIdentifier)
-					);
+					// If the previous task was parallel and this is sequential, then we actually want to place this in the next phase to 
+					// ensure it runs after the previous task (as sequential tasks run before parallel tasks within the same phase)
+					if (int32 PhaseIndex = (int32)TaskPhase;
+						PhaseIndex < (int32)EGatherTextCommandletPhase::NumPhases - 1)
+					{
+						TaskPhase = (EGatherTextCommandletPhase)(PhaseIndex + 1);
+					}
 				}
+				UE_LOGFMT(LogGatherTextCommandlet, Warning, "The commandlet name {commandletName} in section {section} is using an Undefined phase and has been placed in the {phase} phase. You should upgrade your commandlet to use a specific phase to ensure it gets sequenced correctly!",
+					("commandletName", CommandletClassName),
+					("section", StepName),
+					("phase", LexToString(TaskPhase)),
+					("id", GatherTextCommandlet::LocalizationLogIdentifier)
+				);
 			}
-			return -1;
+
+			bPreviousTaskWasParallel = bParallelTask;
+			PreviousTaskPhase = TaskPhase;
 		}
+		check(TaskPhase < EGatherTextCommandletPhase::NumPhases);
+
+		FGatherTextCommandletPhase& Phase = GatherTextPhases[(int32)TaskPhase];
+		TArray<FGatherTextCommandletTask>& TasksArray = bParallelTask ? Phase.ParallelTasks : Phase.SequentialTasks;
 
-		UE_LOG(LogGatherTextCommandlet, Display, TEXT("Completed %s: %s in %.2f seconds"), *StepName, *CommandletClassName, FPlatformTime::Seconds() - CommandletExecutionStartTime);
+		FGatherTextCommandletTask& Task = TasksArray.AddDefaulted_GetRef();
+		Task.CommandletParams = MoveTemp(GeneratedCmdLine);
+		Task.Commandlet = Commandlet.Get();
 	}
 
-	// Clean-up any stale per-platform data
+	// Schedule the clean-up of any stale per-platform data
 	{
 		FString DestinationPath;
 		if (GetPathFromConfig(TEXT("CommonSettings"), TEXT("DestinationPath"), DestinationPath, GatherTextConfigPath))
 		{
-			IFileManager& FileManager = IFileManager::Get();
+			FString PlatformLocalizationPath = DestinationPath / FPaths::GetPlatformLocalizationFolderName();
+			FSplitPlatformConfig& SplitPlatformConfig = SplitPlatformConfigs.FindOrAdd(MoveTemp(PlatformLocalizationPath));
 
-			auto RemoveDirectory = [&FileManager](const TCHAR* InDirectory)
+			if (CommandletGatherManifestHelper->ShouldSplitPlatformData())
 			{
-				FileManager.IterateDirectoryRecursively(InDirectory, [&FileManager](const TCHAR* FilenameOrDirectory, bool bIsDirectory) -> bool
+				SplitPlatformConfig.bShouldSplitPlatformData = true;
+				SplitPlatformConfig.PlatformsToSplit.Append(CommandletGatherManifestHelper->GetPlatformsToSplit());
+			}
+		}
+		else
+		{
+			UE_LOGFMT(LogGatherTextCommandlet, Warning, "No destination path specified in the 'CommonSettings' section. Cannot check for stale per-platform data!",
+				("id", GatherTextCommandlet::LocalizationLogIdentifier)
+			);
+		}
+	}
+
+	return true;
+}
+
+bool UGatherTextCommandlet::ExecuteGatherTextPhases(const TSharedPtr<FLocalizationSCC>& CommandletSourceControlInfo)
+{
+	FScopedSlowTask SlowTask((int32)EGatherTextCommandletPhase::NumPhases, (EmbeddedContext && EmbeddedContext->SlowTaskMessageOverride) ? *EmbeddedContext->SlowTaskMessageOverride : NSLOCTEXT("GatherTextCommandlet", "GatherTextTask.RunningTasks", "Running Localization Tasks for Targets..."));
+
+	auto RunTaskCommandlet =
+		[](const FGatherTextCommandletTask& Task)
+		{
+			const double CommandletExecutionStartTime = FPlatformTime::Seconds();
+
+			const FString CommandletClassName = Task.Commandlet->GetClass()->GetName();
+			UE_LOG(LogGatherTextCommandlet, Display, TEXT("Executing %s (params: %s)"), *CommandletClassName, *Task.CommandletParams);
+
+			if (Task.Commandlet->Main(Task.CommandletParams) != 0)
+			{
+				UE_LOGFMT(LogGatherTextCommandlet, Error, "{commandletName} reported an error (params: {params}).",
+					("commandletName", CommandletClassName),
+					("params", Task.CommandletParams),
+					("id", GatherTextCommandlet::LocalizationLogIdentifier)
+				);
+				return false;
+			}
+
+			UE_LOG(LogGatherTextCommandlet, Display, TEXT("Executing %s took %.2f seconds (params: %s)"), *CommandletClassName, FPlatformTime::Seconds() - CommandletExecutionStartTime, *Task.CommandletParams);
+			return true;
+		};
+
+	const bool bCanRunParallelTasks = GatherTextCommandlet::bAllowParallelTasks && (!EmbeddedContext || EmbeddedContext->IsThreadSafe());
+	UE_LOG(LogGatherTextCommandlet, Display, TEXT("CanRunParallelTasks: %s"), bCanRunParallelTasks ? TEXT("Yes") : TEXT("No"));
+
+	for (int32 PhaseIndex = (int32)EGatherTextCommandletPhase::FirstPhase; PhaseIndex < (int32)EGatherTextCommandletPhase::NumPhases; ++PhaseIndex)
+	{
+		SlowTask.EnterProgressFrame();
+
+		if (SlowTask.ShouldCancel() || (EmbeddedContext && EmbeddedContext->ShouldAbort()))
+		{
+			UE_LOGFMT(LogGatherTextCommandlet, Error, "GatherText aborted!",
+				("id", GatherTextCommandlet::LocalizationLogIdentifier)
+			);
+			return false;
+		}
+
+		const FGatherTextCommandletPhase& Phase = GatherTextPhases[PhaseIndex];
+		if (Phase.SequentialTasks.IsEmpty() && Phase.ParallelTasks.IsEmpty())
+		{
+			continue;
+		}
+
+		const double PhaseExecutionStartTime = FPlatformTime::Seconds();
+		UE_LOG(LogGatherTextCommandlet, Display, TEXT("Running phase '%s' (%d sequential tasks and %d parallel tasks)."), LexToString((EGatherTextCommandletPhase)PhaseIndex), Phase.SequentialTasks.Num(), Phase.ParallelTasks.Num());
+
+		// Run sequential tasks
+		if (Phase.SequentialTasks.Num() > 0)
+		{
+			const double SequentialTasksExecutionStartTime = FPlatformTime::Seconds();
+			UE_LOG(LogGatherTextCommandlet, Display, TEXT("Running %d sequentials task(s) in phase '%s'..."), Phase.SequentialTasks.Num(), LexToString((EGatherTextCommandletPhase)PhaseIndex));
+
+			int32 TaskProgress = 0;
+			for (const FGatherTextCommandletTask& Task : Phase.SequentialTasks)
+			{
+				UE_LOG(LogGatherTextCommandlet, Display, TEXT("Starting sequential task %d of %d in phase '%s'..."), ++TaskProgress, Phase.SequentialTasks.Num(), LexToString((EGatherTextCommandletPhase)PhaseIndex));
+
+				if (SlowTask.ShouldCancel() || (EmbeddedContext && EmbeddedContext->ShouldAbort()))
 				{
-					if (!bIsDirectory)
-					{
-						if (!USourceControlHelpers::IsAvailable() || !USourceControlHelpers::MarkFileForDelete(FilenameOrDirectory))
-						{
-							FileManager.Delete(FilenameOrDirectory, false, true);
-						}
-					}
-					return true;
-				});
-				FileManager.DeleteDirectory(InDirectory, false, true);
-			};
+					UE_LOGFMT(LogGatherTextCommandlet, Error, "GatherText aborted!",
+						("id", GatherTextCommandlet::LocalizationLogIdentifier)
+					);
+					return false;
+				}
 
-			const FString PlatformLocalizationPath = DestinationPath / FPaths::GetPlatformLocalizationFolderName();
-			if (CommandletGatherManifestHelper->ShouldSplitPlatformData())
+				if (EmbeddedContext)
+				{
+					EmbeddedContext->RunTick();
+				}
+
+				if (!RunTaskCommandlet(Task))
+				{
+					return false;
+				}
+			}
+
+			UE_LOG(LogGatherTextCommandlet, Display, TEXT("%d sequentials task(s) in phase '%s' took %.2f seconds."), Phase.SequentialTasks.Num(), LexToString((EGatherTextCommandletPhase)PhaseIndex), FPlatformTime::Seconds() - SequentialTasksExecutionStartTime);
+		}
+			
+		if (SlowTask.ShouldCancel() || (EmbeddedContext && EmbeddedContext->ShouldAbort()))
+		{
+			UE_LOGFMT(LogGatherTextCommandlet, Error, "GatherText aborted!",
+				("id", GatherTextCommandlet::LocalizationLogIdentifier)
+			);
+			return false;
+		}
+
+		// Run parallel tasks
+		if (Phase.ParallelTasks.Num() > 0)
+		{
+			const double ParallelTasksExecutionStartTime = FPlatformTime::Seconds();
+			UE_LOG(LogGatherTextCommandlet, Display, TEXT("Starting %d parallel task(s) in phase '%s'..."), Phase.ParallelTasks.Num(), LexToString((EGatherTextCommandletPhase)PhaseIndex));
+
+			if (CommandletSourceControlInfo)
+			{
+				CommandletSourceControlInfo->BeginParallelTasks();
+			}
+
+			std::atomic<bool> bFailedParallelTasks = false;
+			ParallelFor(TEXT("GatherTextCommandlet"), Phase.ParallelTasks.Num(), 1, [this, &Phase, &RunTaskCommandlet, &bFailedParallelTasks](int32 TaskIndex)
 			{
-				// Remove any stale platform sub-folders
-				FileManager.IterateDirectory(*PlatformLocalizationPath, [&](const TCHAR* FilenameOrDirectory, bool bIsDirectory) -> bool
+				if (EmbeddedContext)
 				{
-					if (bIsDirectory)
-					{
-						const FString SplitPlatformName = FPaths::GetCleanFilename(FilenameOrDirectory);
-						if (!CommandletGatherManifestHelper->GetPlatformsToSplit().Contains(SplitPlatformName))
-						{
-							RemoveDirectory(FilenameOrDirectory);
-						}
-					}
-					return true;
-				});
+					EmbeddedContext->RunTick();
+				}
+
+				const FGatherTextCommandletTask& Task = Phase.ParallelTasks[TaskIndex];
+				if (!RunTaskCommandlet(Task))
+				{
+					bFailedParallelTasks = true;
+				}
+			},
+			bCanRunParallelTasks ? EParallelForFlags::None : EParallelForFlags::ForceSingleThread);
+
+			if (CommandletSourceControlInfo)
+			{
+				FText SCCErrorStr;
+				if (!CommandletSourceControlInfo->EndParallelTasks(SCCErrorStr))
+				{
+					UE_LOGFMT(LogGatherTextCommandlet, Error, "Revision Control error: {error}",
+						("error", SCCErrorStr.ToString()),
+						("id", GatherTextCommandlet::LocalizationLogIdentifier)
+					);
+				}
 			}
-			else
+
+			if (bFailedParallelTasks)
 			{
-				// Remove the entire Platforms folder
-				RemoveDirectory(*PlatformLocalizationPath);
+				return false;
 			}
+
+			UE_LOG(LogGatherTextCommandlet, Display, TEXT("%d parallel task(s) in phase '%s' took %.2f seconds."), Phase.ParallelTasks.Num(), LexToString((EGatherTextCommandletPhase)PhaseIndex), FPlatformTime::Seconds() - ParallelTasksExecutionStartTime);
 		}
-		else
+
+		if (SlowTask.ShouldCancel() || (EmbeddedContext && EmbeddedContext->ShouldAbort()))
 		{
-			UE_LOGFMT(LogGatherTextCommandlet, Warning, "No destination path specified in the 'CommonSettings' section. Cannot check for stale per-platform data!",
+			UE_LOGFMT(LogGatherTextCommandlet, Error, "GatherText aborted!",
 				("id", GatherTextCommandlet::LocalizationLogIdentifier)
 			);
+			return false;
 		}
+
+		UE_LOG(LogGatherTextCommandlet, Display, TEXT("Phase '%s' %.2f seconds (%d sequential tasks and %d parallel tasks)."), LexToString((EGatherTextCommandletPhase)PhaseIndex), FPlatformTime::Seconds() - PhaseExecutionStartTime, Phase.SequentialTasks.Num(), Phase.ParallelTasks.Num());
 	}
 
-	return 0;
+	return true;
+}
+
+void UGatherTextCommandlet::CleanupStalePlatformData()
+{
+	IFileManager& FileManager = IFileManager::Get();
+
+	auto RemoveDirectory = [&FileManager](const TCHAR* InDirectory)
+	{
+		FileManager.IterateDirectoryRecursively(InDirectory, [&FileManager](const TCHAR* FilenameOrDirectory, bool bIsDirectory) -> bool
+		{
+			if (!bIsDirectory)
+			{
+				if (!USourceControlHelpers::IsAvailable() || !USourceControlHelpers::MarkFileForDelete(FilenameOrDirectory))
+				{
+					FileManager.Delete(FilenameOrDirectory, false, true);
+				}
+			}
+			return true;
+		});
+		FileManager.DeleteDirectory(InDirectory, false, true);
+	};
+
+	for (const TTuple<FString, FSplitPlatformConfig>& SplitPlatformConfig : SplitPlatformConfigs)
+	{
+		if (!FileManager.DirectoryExists(*SplitPlatformConfig.Key))
+		{
+			continue;
+		}
+
+		if (SplitPlatformConfig.Value.bShouldSplitPlatformData)
+		{
+			// Remove any stale platform sub-folders
+			FileManager.IterateDirectory(*SplitPlatformConfig.Key, [&SplitPlatformConfig, &RemoveDirectory](const TCHAR* FilenameOrDirectory, bool bIsDirectory) -> bool
+			{
+				if (bIsDirectory)
+				{
+					const FString SplitPlatformName = FPaths::GetCleanFilename(FilenameOrDirectory);
+					if (!SplitPlatformConfig.Value.PlatformsToSplit.Contains(SplitPlatformName))
+					{
+						RemoveDirectory(FilenameOrDirectory);
+					}
+				}
+				return true;
+			});
+		}
+		else
+		{
+			// Remove the entire Platforms folder
+			RemoveDirectory(*SplitPlatformConfig.Key);
+		}
+	}
 }
 
 FText UGatherTextCommandlet::GetChangelistDescription(const TArray<FString>& GatherTextConfigPaths) const
diff --git a/Engine/Source/Editor/UnrealEd/Private/Commandlets/GatherTextCommandletBase.cpp b/Engine/Source/Editor/UnrealEd/Private/Commandlets/GatherTextCommandletBase.cpp
--- a/Engine/Source/Editor/UnrealEd/Private/Commandlets/GatherTextCommandletBase.cpp
+++ b/Engine/Source/Editor/UnrealEd/Private/Commandlets/GatherTextCommandletBase.cpp
@@ -23,6 +23,42 @@ DEFINE_LOG_CATEGORY_STATIC(LogGatherTextCommandletBase, Log, All);
 FGatherTextDelegates::FGetAdditionalGatherPaths FGatherTextDelegates::GetAdditionalGatherPaths;
 FGatherTextDelegates::FGetAdditionalGatherPathsForContext FGatherTextDelegates::GetAdditionalGatherPathsForContext;
 
+const TCHAR* LexToString(EGatherTextCommandletPhase Phase)
+{
+#define GATHERTEXTPHASETOSTRING(PhaseName) case EGatherTextCommandletPhase::PhaseName: return TEXT(#PhaseName)
+	switch (Phase)
+	{
+		GATHERTEXTPHASETOSTRING(PreInitialize);
+		GATHERTEXTPHASETOSTRING(Initialize);
+		GATHERTEXTPHASETOSTRING(PostInitialize);
+		GATHERTEXTPHASETOSTRING(PreUpdateManifests);
+		GATHERTEXTPHASETOSTRING(UpdateManifests);
+		GATHERTEXTPHASETOSTRING(PostUpdateManifests);
+		GATHERTEXTPHASETOSTRING(PreUpdateArchives);
+		GATHERTEXTPHASETOSTRING(UpdateArchives);
+		GATHERTEXTPHASETOSTRING(PostUpdateArchives);
+		GATHERTEXTPHASETOSTRING(PreImport);
+		GATHERTEXTPHASETOSTRING(Import);
+		GATHERTEXTPHASETOSTRING(PostImport);
+		GATHERTEXTPHASETOSTRING(PreExport);
+		GATHERTEXTPHASETOSTRING(Export);
+		GATHERTEXTPHASETOSTRING(PostExport);
+		GATHERTEXTPHASETOSTRING(PreCompile);
+		GATHERTEXTPHASETOSTRING(Compile);
+		GATHERTEXTPHASETOSTRING(PostCompile);
+		GATHERTEXTPHASETOSTRING(PreFinalize);
+		GATHERTEXTPHASETOSTRING(Finalize);
+		GATHERTEXTPHASETOSTRING(PostFinalize);
+		GATHERTEXTPHASETOSTRING(Undefined);
+		default:
+			break;
+	}
+#undef GATHERTEXTPHASETOSTRING
+
+	checkf(false, TEXT("Unknown EGatherTextCommandletPhase!"));
+	return TEXT("Unknown");
+}
+
 //////////////////////////////////////////////////////////////////////////
 //UGatherTextCommandletBase
 
diff --git a/Engine/Source/Editor/UnrealEd/Private/Commandlets/InternationalizationExportCommandlet.cpp b/Engine/Source/Editor/UnrealEd/Private/Commandlets/InternationalizationExportCommandlet.cpp
--- a/Engine/Source/Editor/UnrealEd/Private/Commandlets/InternationalizationExportCommandlet.cpp
+++ b/Engine/Source/Editor/UnrealEd/Private/Commandlets/InternationalizationExportCommandlet.cpp
@@ -171,10 +171,7 @@ int32 UInternationalizationExportCommandlet::Main(const FString& Params)
 		}
 	}
 
-	bool bDoImport = false;
 	GetBoolFromConfig(*SectionName, TEXT("bImportLoc"), bDoImport, ConfigPath);
-
-	bool bDoExport = false;
 	GetBoolFromConfig(*SectionName, TEXT("bExportLoc"), bDoExport, ConfigPath);
 
 	if (!bDoImport && !bDoExport)
@@ -246,3 +243,57 @@ int32 UInternationalizationExportCommandlet::Main(const FString& Params)
 
 	return 0;
 }
+
+bool UInternationalizationExportCommandlet::ConfigurePhase(const FString& Params)
+{
+	TArray<FString> Tokens;
+	TArray<FString> Switches;
+	TMap<FString, FString> ParamVals;
+	UCommandlet::ParseCommandLine(*Params, Tokens, Switches, ParamVals);
+
+	// Set config
+	FString ConfigPath;
+	if (const FString* ConfigParamVal = ParamVals.Find(FString(TEXT("Config"))))
+	{
+		ConfigPath = *ConfigParamVal;
+	}
+	else
+	{
+		UE_LOGFMT(LogInternationalizationExportCommandlet, Error, "No config specified.",
+			("id", InternationalizationExportCommandlet::LocalizationLogIdentifier)
+		);
+		return false;
+	}
+
+	// Set config section
+	FString SectionName;
+	if (const FString* ConfigSectionParamVal = ParamVals.Find(FString(TEXT("Section"))))
+	{
+		SectionName = *ConfigSectionParamVal;
+	}
+	else
+	{
+		UE_LOGFMT(LogInternationalizationExportCommandlet, Error, "No config section specified.",
+			("id", InternationalizationExportCommandlet::LocalizationLogIdentifier)
+		);
+		return false;
+	}
+
+	GetBoolFromConfig(*SectionName, TEXT("bImportLoc"), bDoImport, ConfigPath);
+	GetBoolFromConfig(*SectionName, TEXT("bExportLoc"), bDoExport, ConfigPath);
+
+	return true;
+}
+
+EGatherTextCommandletPhase UInternationalizationExportCommandlet::GetPhase() const
+{
+	if (bDoImport)
+	{
+		return EGatherTextCommandletPhase::Import;
+	}
+	if (bDoExport)
+	{
+		return EGatherTextCommandletPhase::Export;
+	}
+	return EGatherTextCommandletPhase::Undefined;
+}
