IncrementalCook: First draft of ICookArtfact API to provide hooks for game-specific systems that create artifacts in the cook output folder to manage the invalidation of their artifacts during incremental cook.
#rnx

diff --git a/Engine/Source/Editor/UnrealEd/Classes/CookOnTheSide/CookOnTheFlyServer.h b/Engine/Source/Editor/UnrealEd/Classes/CookOnTheSide/CookOnTheFlyServer.h
--- a/Engine/Source/Editor/UnrealEd/Classes/CookOnTheSide/CookOnTheFlyServer.h
+++ b/Engine/Source/Editor/UnrealEd/Classes/CookOnTheSide/CookOnTheFlyServer.h
@@ -4,6 +4,7 @@
 
 #include "Containers/ArrayView.h"
 #include "Containers/RingBuffer.h"
+#include "Cooker/CookArtifact.h"
 #include "CookOnTheSide/CookLog.h"
 #include "CookPackageSplitter.h"
 #include "HAL/PlatformMemory.h"
@@ -185,6 +186,7 @@ namespace UE::Cook
 	class FCookWorkerServer;
 	class FDiagnostics;
 	class FIncrementallyModifiedDiagnostics;
+	class FGlobalCookArtifact;
 	class FLogHandler;
 	class FODSCClientData;
 	class FPackagePreloader;
@@ -192,6 +194,7 @@ namespace UE::Cook
 	class FRequestCluster;
 	class FRequestQueue;
 	class FSaveCookedPackageContext;
+	class FShaderLibraryCookArtifact;
 	class FSoftGCHistory;
 	class FStallDetector;
 	class FWorkerRequestsLocal;
@@ -699,6 +702,8 @@ public:
 	UNREALED_API virtual UE::Cook::ECookType GetCookType() override;
 	UNREALED_API virtual UE::Cook::ECookingDLC GetCookingDLC() override;
 	UNREALED_API virtual FString GetDLCName() override;
+	UNREALED_API virtual FString GetBasedOnReleaseVersion() override;
+	UNREALED_API virtual FString GetCreateReleaseVersion() override;
 	UNREALED_API virtual UE::Cook::EProcessType GetProcessType() override;
 	UNREALED_API virtual UE::Cook::ECookValidationOptions GetCookValidationOptions() override;
 	UNREALED_API virtual bool IsIncremental() override;
@@ -709,6 +714,9 @@ public:
 	UNREALED_API virtual void RegisterCollector(UE::Cook::IMPCollector* Collector,
 		UE::Cook::EProcessType ProcessType = UE::Cook::EProcessType::AllMPCook) override;
 	UNREALED_API virtual void UnregisterCollector(UE::Cook::IMPCollector* Collector) override;
+	UNREALED_API virtual void RegisterArtifact(UE::Cook::ICookArtifact* Artifact) override;
+	UNREALED_API virtual void UnregisterArtifact(UE::Cook::ICookArtifact* Artifact) override;
+
 	UNREALED_API virtual void GetCulturesToCook(TArray<FString>& OutCulturesToCook) const override;
 	// End ICookInfo interface
 
@@ -1068,6 +1076,8 @@ private:
 	/** Return the name to use for the project's global shader library */
 	FString GetProjectShaderLibraryName() const;
 
+	void RegisterShaderLibraryIncrementalCookArtifact(FBeginCookContext& BeginContext);
+
 	/** Invokes the necessary FShaderCodeLibrary functions to start cooking the shader code library. */
 	void BeginCookStartShaderCodeLibrary(FBeginCookContext& BeginContext);
 	/**
@@ -1100,9 +1110,6 @@ private:
 	 */
 	void CreatePipelineCache(const ITargetPlatform* TargetPlatform, const FString& LibraryName);
 
-	/** Invokes the necessary FShaderCodeLibrary functions to clean out all the temporary files. */
-	void CleanShaderCodeLibraries();
-
 	void RegisterShaderChunkDataGenerator();
 
 	/** Called at the end of CookByTheBook to write aggregated data such as AssetRegistry and shaders. */
@@ -1318,23 +1325,14 @@ private:
 	bool GetCookedIniVersionStrings( const ITargetPlatform* TargetPlatform,
 		UE::Cook::FIniSettingContainer& IniVersionStrings, TMap<FString, FString>& AdditionalStrings ) const;
 
-	/**
-	 * Test the CurrentCookSettings against the previous cooksettings to decide if we have to wipe the
-	 * previous cook even when running incrementally.
-	 */
-	bool ArePreviousCookSettingsCompatible(const TMap<FName, FString>& CurrentCookSettings,
-		const ITargetPlatform* TargetPlatform);
-	/** Save the CurrentCookSettings into the output directory. */
-	void SaveCookSettings(const TMap<FName, FString>& CurrentCookSettings, const ITargetPlatform* TargetPlatform);
-	/** Load the CookSettings written at beginning of cook and rewrite them with CookInProgress flag removed. */
-	void ClearCookInProgressFlagFromCookSettings(const ITargetPlatform* TargetPlatform) const;
-	/**
-	 * Populate a map suitable for saving as an ini with the current value of the Cook settings that need to be
-	 * tested for compatibility.
-	 */
-	TMap<FName, FString> CalculateCookSettingStrings(const FCookByTheBookStartupOptions* StartupOptions) const;
-	/** Get the path to the CookSettings metadata file for the given platform and current output directory. */
-	FString GetCookSettingsFileName(const ITargetPlatform* TargetPlatform) const;
+	/** Get the path to the CookSettings file for the given artifact in the output directory for the given platform. */
+	FString GetCookSettingsFileName(const ITargetPlatform* TargetPlatform, const FString& ArtifactName) const;
+	/** Load the CookSettings configfile for the given ArtifactName from output directory for the given platform. */
+	FConfigFile LoadCookSettings(const ITargetPlatform* TargetPlatform, const FString& ArtifactName);
+	/** Save the CookSettings configfile for the given ArtifactName to the output directory for the given platform. */
+	void SaveCookSettings(const FConfigFile& ConfigFile, const ITargetPlatform* TargetPlatform, const FString& ArtifactName) const;
+	/** Load the global CookSettings written at beginning of cook and rewrite them with CookInProgress flag removed. */
+	void ClearCookInProgressFlagFromGlobalCookSettings(const ITargetPlatform* TargetPlatform) const;
 
 	/**
 	 * Convert a path to a full sandbox path. Is effected by the cooking dlc settings.
@@ -1723,6 +1721,9 @@ private:
 	/** Override for Cook.CVarControl, parsed from the commandline (-CookCVarControl=) */
 	int32 OverrideCookCVarControl = -1;
 
+	TArray<TRefCountPtr<UE::Cook::ICookArtifact>> CookArtifacts;
+	UE::Cook::FGlobalCookArtifact* GlobalArtifact = nullptr;
+
 	/**
 	 * Heap of Pollables to tick, ordered by NextTimeSeconds.
 	 * Only accessible by PumpPollables when bPollablesInTick==true, can be accessed outside of PollablesLock.
@@ -1763,6 +1764,7 @@ private:
 	friend UE::Cook::FCookWorkerServer;
 	friend UE::Cook::FDiagnostics;
 	friend UE::Cook::FGenerationHelper;
+	friend UE::Cook::FGlobalCookArtifact;
 	friend UE::Cook::FInitializeConfigSettings;
 	friend UE::Cook::FLogHandler;
 	friend UE::Cook::FPackageData;
@@ -1775,6 +1777,7 @@ private:
 	friend UE::Cook::FRequestCluster;
 	friend UE::Cook::FSaveCookedPackageContext;
 	friend UE::Cook::FScopeFindCookReferences;
+	friend UE::Cook::FShaderLibraryCookArtifact;
 	friend UE::Cook::FSoftGCHistory;
 	friend UE::Cook::FWorkerRequestsLocal;
 	friend UE::Cook::FWorkerRequestsRemote;
diff --git a/Engine/Source/Editor/UnrealEd/Private/CookOnTheFlyServer.cpp b/Engine/Source/Editor/UnrealEd/Private/CookOnTheFlyServer.cpp
--- a/Engine/Source/Editor/UnrealEd/Private/CookOnTheFlyServer.cpp
+++ b/Engine/Source/Editor/UnrealEd/Private/CookOnTheFlyServer.cpp
@@ -29,6 +29,7 @@
 #include "Cooker/CookGenerationHelper.h"
 #include "Cooker/CookGlobalDependencies.h"
 #include "Cooker/CookImportsChecker.h"
+#include "Cooker/CookArtifact.h"
 #include "Cooker/CookLogPrivate.h"
 #include "Cooker/CookOnTheFlyServerInterface.h"
 #include "Cooker/CookPackageData.h"
@@ -42,6 +43,7 @@
 #include "Cooker/CookTypes.h"
 #include "Cooker/CookWorkerClient.h"
 #include "Cooker/DiffPackageWriter.h"
+#include "Cooker/GlobalCookArtifact.h"
 #include "Cooker/IncrementalValidatePackageWriter.h"
 #include "Cooker/IoStoreCookOnTheFlyRequestManager.h"
 #include "Cooker/LooseCookedPackageWriter.h"
@@ -5984,6 +5986,14 @@ void UCookOnTheFlyServer::Initialize( ECookMode::Type DesiredCookMode, ECookInit
 	GCDiagnosticContext = MakeUnique<UE::Cook::FCookGCDiagnosticContext>();
 	StallDetector = MakeUnique<UE::Cook::FStallDetector>();
 	LogHandler.Reset(UE::Cook::CreateLogHandler(*this));
+	if (!IsCookWorkerMode())
+	{
+		GlobalArtifact = new UE::Cook::FGlobalCookArtifact(*this);
+		TRefCountPtr<UE::Cook::ICookArtifact> RefCounter(GlobalArtifact);
+		RegisterArtifact(GlobalArtifact);
+		// We keep a raw pointer and rely on RegisterArtifact holding the reference.
+		check(RefCounter.GetRefCount() > 1);
+	}
 
 	if (!IsCookWorkerMode())
 	{
@@ -6834,6 +6844,20 @@ FString UCookOnTheFlyServer::GetDLCName()
 	return CookByTheBookOptions->DlcName;
 }
 
+FString UCookOnTheFlyServer::GetBasedOnReleaseVersion()
+{
+	if (!IsCookingDLC() && CookByTheBookOptions->CreateReleaseVersion.IsEmpty())
+	{
+		return FString();
+	}
+	return CookByTheBookOptions->BasedOnReleaseVersion;
+}
+
+FString UCookOnTheFlyServer::GetCreateReleaseVersion()
+{
+	return CookByTheBookOptions->CreateReleaseVersion;
+}
+
 UE::Cook::EProcessType UCookOnTheFlyServer::GetProcessType()
 {
 	using namespace UE::Cook;
@@ -6938,6 +6962,55 @@ void UCookOnTheFlyServer::UnregisterCollector(UE::Cook::IMPCollector* Collector)
 	}
 }
 
+void UCookOnTheFlyServer::RegisterArtifact(UE::Cook::ICookArtifact* Artifact)
+{
+	using namespace UE::Cook;
+	if (!Artifact)
+	{
+		return;
+	}
+
+	TRefCountPtr<UE::Cook::ICookArtifact> DeleteIfCallerHasNoReference(Artifact);
+	if (IsCookWorkerMode())
+	{
+		// Artifacts are not used on CookWorkers
+		return;
+	}
+
+	FString ArtifactName = Artifact->GetArtifactName();
+	if (ArtifactName.IsEmpty())
+	{
+		UE_LOG(LogCook, Error, TEXT("RegisterArtifact called with empty GetArtifactName(), which is invalid. Ignoring the Artifact."));
+		FDebug::DumpStackTraceToLog(ELogVerbosity::Display);
+		return;
+	}
+
+	for (TRefCountPtr<ICookArtifact>& Existing : CookArtifacts)
+	{
+		if (Existing->GetArtifactName() == ArtifactName)
+		{
+			if (Existing != Artifact)
+			{
+				UE_LOG(LogCook, Error,
+					TEXT("RegisterArtifact Called with already existing ArtifactName \"%s\". ArtifactNames must be unique. Ignoring the artifact."),
+					*ArtifactName);
+				FDebug::DumpStackTraceToLog(ELogVerbosity::Display);
+			}
+			return;
+		}
+	}
+	CookArtifacts.Add(Artifact);
+}
+
+void UCookOnTheFlyServer::UnregisterArtifact(UE::Cook::ICookArtifact* Artifact)
+{
+	if (!Artifact)
+	{
+		return;
+	}
+	TRefCountPtr<UE::Cook::ICookArtifact> RefCountPtr(Artifact);
+	CookArtifacts.Remove(RefCountPtr);
+}
 
 void UCookOnTheFlyServer::GetCulturesToCook(TArray<FString>& OutCulturesToCook) const
 {
@@ -7538,212 +7611,53 @@ bool UCookOnTheFlyServer::GetCookedIniVersionStrings(const ITargetPlatform* Targ
 	return true;
 }
 
-static const TCHAR* TEXT_CookSettings(TEXT("CookSettings"));
-static const TCHAR* TEXT_CookInProgress(TEXT("CookInProgress"));
-static FName ExecutableHashName(TEXT("ExecutableHash"));
-static FName ExecutableHashInvalidModuleName(TEXT("ExecutableHashInvalidModule"));
-
-
-TMap<FName, FString> UCookOnTheFlyServer::CalculateCookSettingStrings(const FCookByTheBookStartupOptions* StartupOptions) const
+FString UCookOnTheFlyServer::GetCookSettingsFileName(const ITargetPlatform* TargetPlatform,
+	const FString& ArtifactName) const
 {
-	TMap<FName, FString> CookSettingStrings;
-	const FName NAME_CookMode(TEXT("CookMode"));
-
-	TArray<FModuleStatus> Modules;
-	FModuleManager::Get().QueryModules(Modules);
-	Modules.Sort([](const FModuleStatus& A, const FModuleStatus& B)
-		{
-			return A.FilePath.Compare(B.FilePath, ESearchCase::IgnoreCase) < 0;
-		});
-	IPlatformFile& PlatformFile = FPlatformFileManager::Get().GetPlatformFile();
-
-	CookSettingStrings.Add(FName(TEXT("Version")), TEXT("21F52B9EDD4D456AB1AF381CA172BD28"));
-	CookSettingStrings.Add(FName(TEXT("LegacyBuildDependencies")), *LexToString(bLegacyBuildDependencies));
-
-	if (!bLegacyBuildDependencies)
-	{
-		// Store the CookIncrementalVersion in the global settings, so that it will cause deletion of all artifacts
-		// when it changes, in addition to invalidating all cooked packages.
-		CookSettingStrings.Add(FName(TEXT("CookIncrementalVersion")), UE::Cook::CookIncrementalVersion.ToString());
-	}
-	else
-	{
-		// Calculate the executable hash by combining the module file hash of every loaded module
-		// TODO: Write the module file hash from UnrealBuildTool into the .modules file and read it
-		// here from the .modules file instead of calculating it on every cook.
-		if (bLegacyIterativeCalculateExe)
-		{
-			FString InvalidModule;
-			bool bValid = true;
-			FXxHash64Builder Hasher;
-			TArray<uint8> Buffer;
-			for (FModuleStatus& ModuleStatus : Modules)
-			{
-				TUniquePtr<IFileHandle> FileHandle(PlatformFile.OpenRead(*ModuleStatus.FilePath));
-				if (!FileHandle)
-				{
-					InvalidModule = ModuleStatus.FilePath;
-					break;
-				}
-				int64 FileSize = FileHandle->Size();
-				Buffer.SetNumUninitialized(FileSize, EAllowShrinking::No);
-				if (!FileHandle->Read(Buffer.GetData(), FileSize))
-				{
-					InvalidModule = ModuleStatus.FilePath;
-					break;
-				}
-				Hasher.Update(Buffer.GetData(), FileSize);
-			}
-			if (InvalidModule.IsEmpty())
-			{
-				CookSettingStrings.Add(ExecutableHashName, FString(*WriteToString<64>(Hasher.Finalize())));
-			}
-			else
-			{
-				CookSettingStrings.Add(ExecutableHashInvalidModuleName, InvalidModule);
-			}
-		}
-	}
-
-	if (IsDirectorCookByTheBook())
-	{
-		CookSettingStrings.Add(NAME_CookMode, TEXT("CookByTheBook"));
-		CookSettingStrings.Add(FName(TEXT("DLCName")), CookByTheBookOptions->DlcName);
-		if (!CookByTheBookOptions->DlcName.IsEmpty())
-		{
-			CookSettingStrings.Add(FName(TEXT("BasedOnReleaseVersion")),
-				StartupOptions ? StartupOptions->BasedOnReleaseVersion : FString());
-		}
-	}
-	else
-	{
-		check(IsDirectorCookOnTheFly());
-		CookSettingStrings.Add(NAME_CookMode, TEXT("CookOnTheFly"));
-	}
-	return CookSettingStrings;
+	check(!ArtifactName.IsEmpty());
+	FString LeafName = ArtifactName != TEXT("global")
+		? FString::Printf(TEXT("CookedSettings_%s.txt"), *ArtifactName)
+		: TEXT("CookedSettings.txt");
+	FString UncookedPath = GetMetadataDirectory() / LeafName;
+	return ConvertToFullSandboxPath(*UncookedPath, true /* bForWrite */, TargetPlatform->PlatformName());
 }
 
-FString UCookOnTheFlyServer::GetCookSettingsFileName(const ITargetPlatform* TargetPlatform) const
+FConfigFile UCookOnTheFlyServer::LoadCookSettings(const ITargetPlatform* TargetPlatform, const FString& ArtifactName)
 {
-	FString CookedSettingsIni = GetMetadataDirectory() / TEXT("CookedSettings.txt");
-	return ConvertToFullSandboxPath(*CookedSettingsIni, true, TargetPlatform->PlatformName());
-}
-
-bool UCookOnTheFlyServer::ArePreviousCookSettingsCompatible(const TMap<FName, FString>& CurrentCookSettings, const ITargetPlatform* TargetPlatform)
-{
-	FConfigFile ConfigFile;
-	FString CookSettingsFileName = GetCookSettingsFileName(TargetPlatform);
-	if (TUniquePtr<FArchive> Reader(FindOrCreateCookArtifactReader(TargetPlatform).CreateFileReader(*CookSettingsFileName)); Reader)
+	UE::ConfigAccessTracking::FIgnoreScope IgnoreScope;
+	FConfigFile Result;
+	FString CookSettingsFileName = GetCookSettingsFileName(TargetPlatform, ArtifactName);
+	TUniquePtr<FArchive> Reader;
+	Reader.Reset(FindOrCreateCookArtifactReader(TargetPlatform).CreateFileReader(*CookSettingsFileName));
+	if (Reader)
 	{
 		FString CookSettingsFileContents;
 		if (FFileHelper::LoadFileToString(CookSettingsFileContents, *Reader.Get()))
 		{
-			ConfigFile.ProcessInputFileContents(CookSettingsFileContents, CookSettingsFileName);
-		}
-	}
-
-	const FConfigSection* CookSettings = ConfigFile.FindSection(TEXT_CookSettings);
-	if (CookSettings == nullptr)
-	{
-		UE_LOG(LogCook, Display, TEXT("Cook invalidated for platform %s because CookSettings file %s is invalid. Clearing previously cooked packages."),
-			*TargetPlatform->PlatformName(), *CookSettingsFileName);
-		return false;
-	}
-
-	TSet<FName> IgnoreKeys;
-	IgnoreKeys.Add(ExecutableHashName);
-	IgnoreKeys.Add(ExecutableHashInvalidModuleName);
-
-	for (const TPair<FName, FString>& CurrentSetting : CurrentCookSettings)
-	{
-		if (IgnoreKeys.Contains(CurrentSetting.Key))
-		{
-			continue;
-		}
-		const FConfigValue* PreviousSetting = CookSettings->Find(CurrentSetting.Key);
-		if (!PreviousSetting || PreviousSetting->GetValue() != CurrentSetting.Value)
-		{
-			UE_LOG(LogCook, Display, TEXT("Cook invalidated for platform %s because %s has changed. Old: %s, New: %s. Clearing previously cooked packages."),
-				*TargetPlatform->PlatformName(), *CurrentSetting.Key.ToString(),
-				PreviousSetting ? *PreviousSetting->GetValue() : TEXT(""),
-				*CurrentSetting.Value);
-			return false;
-		}
-	}
-
-	if (GIsBuildMachine)
-	{
-		bool bCookInProgress;
-		if (ConfigFile.GetBool(TEXT_CookSettings, TEXT_CookInProgress, bCookInProgress) && bCookInProgress)
-		{
-			UE_LOG(LogCook, Display, TEXT("Cook invalidated for platform %s because the previous cook crashed (or otherwise did not report completion).")
-				TEXT(" CookSettings file %s still has [%s]:%s=true. Clearing previously cooked packages."),
-				*TargetPlatform->PlatformName(), *CookSettingsFileName, TEXT_CookSettings, TEXT_CookInProgress);
-			return false;
-		}
-	}
-
-	if (!bLegacyIterativeIgnoreIni && bLegacyBuildDependencies && IniSettingsOutOfDate(TargetPlatform))
-	{
-		UE_LOG(LogCook, Display, TEXT("Cook invalidated for platform %s because ini settings have changed. Clearing previously cooked packages."),
-			*TargetPlatform->PlatformName());
-		return false;
-	}
-
-	if (!bLegacyIterativeIgnoreExe && bLegacyBuildDependencies)
-	{
-		const FString* CurrentHash = CurrentCookSettings.Find(ExecutableHashName);
-		if (!CurrentHash)
-		{
-			UE_LOG(LogCook, Display, TEXT("Cook invalidated for platform %s because current executable hash is invalid. Invalid module=%s. Clearing previously cooked packages."),
-				*TargetPlatform->PlatformName(), *CurrentCookSettings.FindRef(ExecutableHashInvalidModuleName));
-			return false;
-		}
-		const FConfigValue* PreviousHash = CookSettings->Find(ExecutableHashName);
-		if (!PreviousHash)
-		{
-			const FConfigValue* InvalidModuleName = CookSettings->Find(ExecutableHashInvalidModuleName);
-			UE_LOG(LogCook, Display, TEXT("Cook invalidated for platform %s because old executable hash is invalid. Invalid module=%s. Clearing previously cooked packages."),
-				*TargetPlatform->PlatformName(), InvalidModuleName ? *InvalidModuleName->GetValue() : TEXT(""));
-			return false;
-		}
-		if (!CurrentHash->Equals(*PreviousHash->GetValue(), ESearchCase::CaseSensitive))
-		{
-			UE_LOG(LogCook, Display, TEXT("Cook invalidated for platform %s because executable hash has changed. Old: %s, New: %s. Clearing previously cooked packages."),
-				*TargetPlatform->PlatformName(), *PreviousHash->GetValue(), **CurrentHash);
-			return false;
+			Result.ProcessInputFileContents(CookSettingsFileContents, CookSettingsFileName);
 		}
 	}
-
-	return true;
+	return Result;
 }
 
-void UCookOnTheFlyServer::SaveCookSettings(const TMap<FName, FString>& CurrentCookSettings, const ITargetPlatform* TargetPlatform)
+void UCookOnTheFlyServer::SaveCookSettings(const FConfigFile& ConfigFile, const ITargetPlatform* TargetPlatform,
+	const FString& ArtifactName) const
 {
 	UE::ConfigAccessTracking::FIgnoreScope IgnoreScope;
+	FString SettingsFileName = GetCookSettingsFileName(TargetPlatform, ArtifactName);
 
-	FConfigFile ConfigFile;
-	for (const TPair<FName, FString>& CurrentSetting : CurrentCookSettings)
+	if (ConfigFile.IsEmpty())
 	{
-		ConfigFile.AddToSection(TEXT_CookSettings, CurrentSetting.Key, CurrentSetting.Value);
+		IFileManager::Get().Delete(*SettingsFileName, false /* bRequireExists */);
+	}
+	else
+	{
+		// Use ConfigFile.WriteToString instead of Write because Write applies suppressions specific to GConfig and Dirty.
+		FString Text;
+		// TODO: Change WriteToString to be a const function.
+		const_cast<FConfigFile&>(ConfigFile).WriteToString(Text);
+		FFileHelper::SaveStringToFile(Text, *SettingsFileName);
 	}
-	ConfigFile.AddToSection(TEXT_CookSettings, TEXT_CookInProgress, TEXT("true"));
-	// Use ConfigFile.WriteToString instead of Write because Write applies suppressions specific to GConfig and Dirty.
-	FString Text;
-	ConfigFile.WriteToString(Text);
-	FFileHelper::SaveStringToFile(Text, *GetCookSettingsFileName(TargetPlatform));
-}
-
-void UCookOnTheFlyServer::ClearCookInProgressFlagFromCookSettings(const ITargetPlatform* TargetPlatform) const
-{
-	FConfigFile ConfigFile;
-	FString Filename = GetCookSettingsFileName(TargetPlatform);
-	ConfigFile.Read(Filename);
-	ConfigFile.RemoveKeyFromSection(TEXT_CookSettings, TEXT_CookInProgress);
-	// Use ConfigFile.WriteToString instead of Write because Write applies suppressions specific to GConfig and Dirty.
-	FString Text;
-	ConfigFile.WriteToString(Text);
-	FFileHelper::SaveStringToFile(Text, *Filename);
 }
 
 bool UCookOnTheFlyServer::IniSettingsOutOfDate(const ITargetPlatform* TargetPlatform) const
@@ -9105,6 +9019,42 @@ public:
 	static FGuid MessageType;
 	static constexpr ANSICHAR RootObjectID[2]{ "S" };
 };
+
+const TCHAR* TEXT_GlobalMaterialHash(TEXT("GlobalMaterialHash"));
+const TCHAR* TEXT_ShaderLibrary(TEXT("ShaderLibrary"));
+/** CookArtifact to manage the shaderlibraries in the cook's OutputFolder. */
+class FShaderLibraryCookArtifact : public ICookArtifact
+{
+public:
+	FShaderLibraryCookArtifact(UCookOnTheFlyServer& InCOTFS)
+		: COTFS(InCOTFS)
+	{
+	}
+
+	virtual FString GetArtifactName() const
+	{
+		return TEXT_ShaderLibrary;
+	}
+
+	virtual void OnFullRecook(const ITargetPlatform* TargetPlatform)
+	{
+		CleanIntermediateFiles(TargetPlatform);
+	}
+
+	void CleanIntermediateFiles(const ITargetPlatform* TargetPlatform)
+	{
+		TArray<FName> ShaderFormats;
+		TargetPlatform->GetAllTargetedShaderFormats(ShaderFormats);
+		if (ShaderFormats.Num() > 0)
+		{
+			FShaderLibraryCooker::CleanDirectories(ShaderFormats);
+		}
+	}
+
+private:
+	UCookOnTheFlyServer& COTFS;
+};
+
 FGuid FShaderLibraryCollector::MessageType(TEXT("4DF3B36BBA2F4E04A846E894E24EB2C4"));
 
 void FShaderLibraryCollector::ClientTick(FMPCollectorClientTickContext& Context)
@@ -9148,11 +9098,19 @@ void FShaderLibraryCollector::ServerReceiveMessage(FMPCollectorServerMessageCont
 
 }
 
+void UCookOnTheFlyServer::RegisterShaderLibraryIncrementalCookArtifact(FBeginCookContext& BeginContext)
+{
+	if (IsCookWorkerMode() || !IsUsingShaderCodeLibrary())
+	{
+		return;
+	}
+	RegisterArtifact(new UE::Cook::FShaderLibraryCookArtifact(*this));
+}
+
 void UCookOnTheFlyServer::BeginCookStartShaderCodeLibrary(FBeginCookContext& BeginContext)
 {
 	const UProjectPackagingSettings* const PackagingSettings = GetDefault<UProjectPackagingSettings>();
-	bool const bCacheShaderLibraries = IsUsingShaderCodeLibrary();
-	if (bCacheShaderLibraries)
+	if (IsUsingShaderCodeLibrary())
 	{
 		FShaderLibraryCooker::InitForCooking(PackagingSettings->bSharedMaterialNativeLibraries, AllContextArtifactReader.Get());
 
@@ -9216,11 +9174,6 @@ void UCookOnTheFlyServer::BeginCookStartShaderCodeLibrary(FBeginCookContext& Beg
 	{
 		CookWorkerClient->Register(new FShaderStatsAggregator(FShaderStatsAggregator::EMode::Worker));
 	}
-
-	if (!IsCookWorkerMode())
-	{
-		CleanShaderCodeLibraries();
-	}
 }
 
 void UCookOnTheFlyServer::BeginCookFinishShaderCodeLibrary(FBeginCookContext& BeginContext)
@@ -9527,26 +9480,6 @@ void UCookOnTheFlyServer::SaveShaderLibrary(const ITargetPlatform* TargetPlatfor
 	}
 }
 
-void UCookOnTheFlyServer::CleanShaderCodeLibraries()
-{
-	const bool bCacheShaderLibraries = IsUsingShaderCodeLibrary();
-
-	for (const ITargetPlatform* TargetPlatform : PlatformManager->GetSessionPlatforms())
-	{
-		UE::Cook::FPlatformData* PlatformData = PlatformManager->GetPlatformData(TargetPlatform);
-		// If this is a full (aka forcerecook) build then clean up our temporary files
-		if (bCacheShaderLibraries && PlatformData->bFullBuild)
-		{
-			TArray<FName> ShaderFormats;
-			TargetPlatform->GetAllTargetedShaderFormats(ShaderFormats);
-			if (ShaderFormats.Num() > 0)
-			{
-				FShaderLibraryCooker::CleanDirectories(ShaderFormats);
-			}
-		}
-	}
-}
-
 void UCookOnTheFlyServer::WriteCookMetadata(const ITargetPlatform* InTargetPlatform, uint64 InDevelopmentAssetRegistryHash)
 {
 	FString PlatformNameString = InTargetPlatform->PlatformName();
@@ -9875,7 +9808,7 @@ void UCookOnTheFlyServer::CookByTheBookFinished()
 
 	for (const ITargetPlatform* TargetPlatform : PlatformManager->GetSessionPlatforms())
 	{
-		ClearCookInProgressFlagFromCookSettings(TargetPlatform);
+		ClearCookInProgressFlagFromGlobalCookSettings(TargetPlatform);
 	}
 	ShutdownCookSession();
 	UE_LOG(LogCook, Display, TEXT("Done!"));
@@ -10912,7 +10845,12 @@ void UCookOnTheFlyServer::LoadBeginCookIncrementalFlagsLocal(FBeginCookContext&
 				UE_LOG(LogCook, Display, TEXT("Local cook is newer then shared cooked build, legacyiteratively cooking from local build."));
 			}
 		}
-		PlatformContext.CurrentCookSettings = CalculateCookSettingStrings(BeginContext.StartupOptions);
+		for (const TRefCountPtr<UE::Cook::ICookArtifact>& Artifact : CookArtifacts)
+		{
+			UE::Cook::FCookArtifactsCurrentSettings& CurrentSettings
+				= PlatformContext.CookArtifactsCurrentSettings.FindOrAdd(Artifact.GetReference());
+			CurrentSettings.Settings = Artifact->CalculateCurrentSettings(*this, TargetPlatform);
+		}
 		PlatformContext.bHasMemoryResults = PlatformData->bIsSandboxInitialized;
 
 		if (bIsDiffOnly)
@@ -10962,19 +10900,44 @@ void UCookOnTheFlyServer::LoadBeginCookIncrementalFlagsLocal(FBeginCookContext&
 					);
 				}
 			}
-			else if (!ArePreviousCookSettingsCompatible(PlatformContext.CurrentCookSettings, TargetPlatform))
+			else
 			{
-				UE_LOG(LogCook, Display,
-					TEXT("FULL COOK: %s was specified, but global settings have changed and all previously cooked packages are invalidated. Deleting previously cooked packages for platform %s and recooking all packages discovered in the current cook."),
-					bLegacyBuildDependencies ? TEXT("-legacyiterative") : TEXT("-cookincremental"),
-					*TargetPlatform->PlatformName());
-				if (bRunningAsShaderServer)
+				UE::ConfigAccessTracking::FIgnoreScope IgnoreScope;
+				for (const TRefCountPtr<UE::Cook::ICookArtifact>& Artifact: CookArtifacts)
 				{
+					if (bLegacyBuildDependencies && Artifact != GlobalArtifact)
+					{
+						continue;
+					}
+
+					FString ArtifactName = Artifact->GetArtifactName();
+					FConfigFile Previous = LoadCookSettings(TargetPlatform, ArtifactName);
+					UE::Cook::FCookArtifactsCurrentSettings& Current
+						= PlatformContext.CookArtifactsCurrentSettings.FindChecked(Artifact.GetReference());
+					FString PreviousFileName = GetCookSettingsFileName(TargetPlatform, ArtifactName);
+
+					UE::Cook::Artifact::FCompareSettingsContext Context(*this, Previous, Current.Settings, PreviousFileName);
+					Context.TargetPlatform = TargetPlatform;
+					Artifact->CompareSettings(Context);
+
+					Current.bRequestInvalidate = Context.IsRequestInvalidate();
+					if (!Context.IsRequestFullRecook())
+					{
+						continue;
+					}
+
 					UE_LOG(LogCook, Display,
-						TEXT("'-odsc -legacyiterative' was passed on commandline, but due to unrelated changes in global settings the cooker has to clear cook results. The build will need to be recooked before it can be staged.")
-					);
+						TEXT("FULL COOK: %s was specified, but %s settings have changed and all previously cooked packages are invalidated. Deleting previously cooked packages for platform %s and recooking all packages discovered in the current cook."),
+						bLegacyBuildDependencies ? TEXT("-legacyiterative") : TEXT("-cookincremental"), *ArtifactName, *TargetPlatform->PlatformName());
+					if (bRunningAsShaderServer)
+					{
+						UE_LOG(LogCook, Display,
+							TEXT("'-odsc -legacyiterative' was passed on commandline, but due to unrelated changes in global settings the cooker has to clear cook results. The build will need to be recooked before it can be staged.")
+						);
+					}
+					bIncrementalOrLegacyIterativeAllowed = false;
+					break;
 				}
-				bIncrementalOrLegacyIterativeAllowed = false;
 			}
 
 			if (bIncrementalOrLegacyIterativeAllowed)
@@ -11058,7 +11021,32 @@ void UCookOnTheFlyServer::BeginCookSandbox(FBeginCookContext& BeginContext)
 			ResetPlatforms.Emplace(TargetPlatform, bResetResults);
 			if (!PlatformContext.bWorkerOnSharedSandbox)
 			{
-				SaveCookSettings(PlatformContext.CurrentCookSettings, TargetPlatform);
+				check(!IsCookWorkerMode());
+				for (const TRefCountPtr<UE::Cook::ICookArtifact>& Artifact: CookArtifacts)
+				{
+					UE::Cook::FCookArtifactsCurrentSettings& Current
+						= PlatformContext.CookArtifactsCurrentSettings.FindChecked(Artifact.GetReference());
+					FString ArtifactName = Artifact->GetArtifactName();
+
+					if (!bLegacyBuildDependencies || PlatformContext.bFullBuild)
+					{
+						if (PlatformContext.bFullBuild)
+						{
+							Artifact->OnFullRecook(TargetPlatform);
+						}
+						else if (Current.bRequestInvalidate)
+						{
+							UE_LOG(LogCook, Display,
+								TEXT("COOKINCREMENTAL: CLEARING ARTIFACT: -cookincremental was specified, but %s settings have changed and the artifact is invalidated. Deleting the artifact; it will be reconstructed from the packags that cook during this session."),
+								*ArtifactName);
+							Artifact->OnInvalidate(TargetPlatform);
+						}
+					}
+					if (!bLegacyBuildDependencies || Artifact == GlobalArtifact)
+					{
+						SaveCookSettings(Current.Settings, TargetPlatform, ArtifactName);
+					}
+				}
 			}
 
 			PlatformData->bIsSandboxInitialized = true;
@@ -11376,6 +11364,7 @@ void UCookOnTheFlyServer::StartCookByTheBook( const FCookByTheBookStartupOptions
 	LoadBeginCookConfigSettings(BeginContext);
 	SelectSessionPlatforms(BeginContext);
 	CalculateGlobalDependenciesHashes();
+	RegisterShaderLibraryIncrementalCookArtifact(BeginContext);
 	LoadBeginCookIncrementalFlags(BeginContext);
 
 	// Initialize systems referenced by later stages or that need to start early for async performance
@@ -11569,6 +11558,7 @@ FBeginCookContext UCookOnTheFlyServer::CreateBeginCookByTheBookContext(const FCo
 	CookByTheBookOptions->CookStartTime = FPlatformTime::Seconds();
 	CookByTheBookOptions->bGenerateStreamingInstallManifests = StartupOptions.bGenerateStreamingInstallManifests;
 	CookByTheBookOptions->bGenerateDependenciesForMaps = StartupOptions.bGenerateDependenciesForMaps;
+	CookByTheBookOptions->BasedOnReleaseVersion = StartupOptions.BasedOnReleaseVersion;
 	CookByTheBookOptions->CreateReleaseVersion = StartupOptions.CreateReleaseVersion;
 	CookByTheBookOptions->bSkipHardReferences = !!(CookOptions & ECookByTheBookOptions::SkipHardReferences);
 	CookByTheBookOptions->bSkipSoftReferences = !!(CookOptions & ECookByTheBookOptions::SkipSoftReferences);
@@ -12108,8 +12098,8 @@ void UCookOnTheFlyServer::GenerateInitialRequests(FBeginCookContext& BeginContex
 	}
 	InitialRequestCount = WorkerRequests->GetNumExternalRequests();
 
-	const FString& CreateReleaseVersion = BeginContext.StartupOptions->CreateReleaseVersion;
-	const FString& BasedOnReleaseVersion = BeginContext.StartupOptions->BasedOnReleaseVersion;
+	const FString& BasedOnReleaseVersion = CookByTheBookOptions->BasedOnReleaseVersion;
+	const FString& CreateReleaseVersion = CookByTheBookOptions->CreateReleaseVersion;
 	if (!IsCookingDLC() && !BasedOnReleaseVersion.IsEmpty())
 	{
 		// if we are based on a release and we are not cooking dlc then we should always be creating a new one (note that we could be creating the same one we are based on).
@@ -12155,7 +12145,7 @@ void UCookOnTheFlyServer::RecordDLCPackagesFromBaseGame(FBeginCookContext& Begin
 	}
 
 	const ECookByTheBookOptions& CookOptions = CookByTheBookOptions->StartupOptions;
-	const FString& BasedOnReleaseVersion = BeginContext.StartupOptions->BasedOnReleaseVersion;
+	const FString& BasedOnReleaseVersion = CookByTheBookOptions->BasedOnReleaseVersion;
 
 	// If we're cooking against a fixed base, we don't need to verify the packages exist on disk, we simply want to use the Release Data 
 	const bool bVerifyPackagesExist = !IsCookingAgainstFixedBase();
diff --git a/Engine/Source/Editor/UnrealEd/Private/Cooker/CookTypes.cpp b/Engine/Source/Editor/UnrealEd/Private/Cooker/CookTypes.cpp
--- a/Engine/Source/Editor/UnrealEd/Private/Cooker/CookTypes.cpp
+++ b/Engine/Source/Editor/UnrealEd/Private/Cooker/CookTypes.cpp
@@ -955,6 +955,7 @@ FCbWriter& operator<<(FCbWriter& Writer, const UE::Cook::FCookByTheBookOptions&
 	// StartupPackages and SessionStartupObjects are process-specific
 
 	Writer << "DlcName" << Value.DlcName;
+	Writer << "BasedOnReleaseVersion" << Value.BasedOnReleaseVersion;
 	Writer << "CreateReleaseVersion" << Value.CreateReleaseVersion;
 	Writer << "SourceToLocalizedPackageVariants" << Value.SourceToLocalizedPackageVariants;
 	Writer << "AllCulturesToCook" << Value.AllCulturesToCook;
@@ -985,6 +986,7 @@ bool LoadFromCompactBinary(FCbFieldView Field, UE::Cook::FCookByTheBookOptions&
 	OutValue.SessionStartupObjects.Empty();
 
 	bOk = LoadFromCompactBinary(Field["DlcName"], OutValue.DlcName) & bOk;
+	bOk = LoadFromCompactBinary(Field["BasedOnReleaseVersion"], OutValue.BasedOnReleaseVersion) & bOk;
 	bOk = LoadFromCompactBinary(Field["CreateReleaseVersion"], OutValue.CreateReleaseVersion) & bOk;
 	bOk = LoadFromCompactBinary(Field["SourceToLocalizedPackageVariants"], OutValue.SourceToLocalizedPackageVariants) & bOk;
 	bOk = LoadFromCompactBinary(Field["AllCulturesToCook"], OutValue.AllCulturesToCook) & bOk;
diff --git a/Engine/Source/Editor/UnrealEd/Private/Cooker/CookTypes.h b/Engine/Source/Editor/UnrealEd/Private/Cooker/CookTypes.h
--- a/Engine/Source/Editor/UnrealEd/Private/Cooker/CookTypes.h
+++ b/Engine/Source/Editor/UnrealEd/Private/Cooker/CookTypes.h
@@ -541,7 +541,10 @@ public:
 	/** DlcName setup if we are cooking dlc will be used as the directory to save cooked files to */
 	FString							DlcName;
 
-	/** Create a release from this manifest and store it in the releases directory for this cgame */
+	/** If cooking DLC or a patch, the release version of the basegame it is based on. */
+	FString							BasedOnReleaseVersion;
+
+	/** Create a release from this manifest and store it in the releases directory for this game. */
 	FString							CreateReleaseVersion;
 
 	/**
@@ -660,7 +663,14 @@ private:
 	EDiscoveredPlatformSet Source;
 };
 
-}
+/** Platform-specific data about each ICookArtifact during cook start. */
+struct FCookArtifactsCurrentSettings
+{
+	FConfigFile Settings;
+	bool bRequestInvalidate = false;
+};
+
+} // namespace UE::Cook
 
 
 /** Helper struct for FBeginCookContext; holds the context data for each platform being cooked */
@@ -668,7 +678,7 @@ struct FBeginCookContextPlatform
 {
 	ITargetPlatform* TargetPlatform = nullptr;
 	UE::Cook::FPlatformData* PlatformData = nullptr;
-	TMap<FName, FString> CurrentCookSettings;
+	TMap<UE::Cook::ICookArtifact*, UE::Cook::FCookArtifactsCurrentSettings> CookArtifactsCurrentSettings;
 
 	/** If true, we are deleting all old results from the previous cook. If false, we are keeping the old results. */
 	bool bFullBuild = false;
diff --git a/Engine/Source/Editor/UnrealEd/Private/Cooker/GlobalCookArtifact.cpp b/Engine/Source/Editor/UnrealEd/Private/Cooker/GlobalCookArtifact.cpp
new file mode 100644
--- /dev/null
+++ b/Engine/Source/Editor/UnrealEd/Private/Cooker/GlobalCookArtifact.cpp
@@ -0,0 +1,231 @@
+// Copyright Epic Games, Inc. All Rights Reserved.
+
+#include "GlobalCookArtifact.h"
+
+#include "Containers/Array.h"
+#include "Containers/Map.h"
+#include "Containers/Set.h"
+#include "Cooker/CookTypes.h"
+#include "CookOnTheSide/CookLog.h"
+#include "CookOnTheSide/CookOnTheFlyServer.h"
+#include "HAL/PlatformFileManager.h"
+#include "Hash/xxhash.h"
+#include "Interfaces/ITargetPlatform.h"
+#include "Logging/LogMacros.h"
+#include "Misc/ConfigAccessTracking.h"
+#include "Misc/StringBuilder.h"
+#include "Modules/ModuleManager.h"
+#include "Templates/UniquePtr.h"
+#include "UObject/NameTypes.h"
+
+namespace UE::Cook
+{
+
+FGlobalCookArtifact::FGlobalCookArtifact(UCookOnTheFlyServer& InCOTFS)
+	: COTFS(InCOTFS)
+{
+}
+
+FString FGlobalCookArtifact::GetArtifactName() const
+{
+	return FString(TEXT("global"));
+}
+
+static const TCHAR* TEXT_CookSettings(TEXT("CookSettings"));
+static const TCHAR* TEXT_CookInProgress(TEXT("CookInProgress"));
+static FName ExecutableHashName(TEXT("ExecutableHash"));
+static FName ExecutableHashInvalidModuleName(TEXT("ExecutableHashInvalidModule"));
+
+// TODO: Move into GlobalCookArtifact.cpp after review
+FConfigFile FGlobalCookArtifact::CalculateCurrentSettings(ICookInfo& CookInfo, const ITargetPlatform* TargetPlatform)
+{
+	TMap<FName, FString> CookSettingStrings;
+	const FName NAME_CookMode(TEXT("CookMode"));
+
+	TArray<FModuleStatus> Modules;
+	FModuleManager::Get().QueryModules(Modules);
+	Modules.Sort([](const FModuleStatus& A, const FModuleStatus& B)
+		{
+			return A.FilePath.Compare(B.FilePath, ESearchCase::IgnoreCase) < 0;
+		});
+	IPlatformFile& PlatformFile = FPlatformFileManager::Get().GetPlatformFile();
+
+	CookSettingStrings.Add(FName(TEXT("Version")), TEXT("21F52B9EDD4D456AB1AF381CA172BD28"));
+	CookSettingStrings.Add(FName(TEXT("LegacyBuildDependencies")), ::LexToString(COTFS.bLegacyBuildDependencies));
+
+	if (!COTFS.bLegacyBuildDependencies)
+	{
+		// Store the CookIncrementalVersion in the global settings, so that it will cause deletion of all artifacts
+		// when it changes, in addition to invalidating all cooked packages.
+		CookSettingStrings.Add(FName(TEXT("CookIncrementalVersion")), CookIncrementalVersion.ToString());
+	}
+	else
+	{
+		// Calculate the executable hash by combining the module file hash of every loaded module
+		// TODO: Write the module file hash from UnrealBuildTool into the .modules file and read it
+		// here from the .modules file instead of calculating it on every cook.
+		if (COTFS.bLegacyIterativeCalculateExe)
+		{
+			FString InvalidModule;
+			bool bValid = true;
+			FXxHash64Builder Hasher;
+			TArray<uint8> Buffer;
+			for (FModuleStatus& ModuleStatus : Modules)
+			{
+				TUniquePtr<IFileHandle> FileHandle(PlatformFile.OpenRead(*ModuleStatus.FilePath));
+				if (!FileHandle)
+				{
+					InvalidModule = ModuleStatus.FilePath;
+					break;
+				}
+				int64 FileSize = FileHandle->Size();
+				Buffer.SetNumUninitialized(FileSize, EAllowShrinking::No);
+				if (!FileHandle->Read(Buffer.GetData(), FileSize))
+				{
+					InvalidModule = ModuleStatus.FilePath;
+					break;
+				}
+				Hasher.Update(Buffer.GetData(), FileSize);
+			}
+			if (InvalidModule.IsEmpty())
+			{
+				CookSettingStrings.Add(ExecutableHashName, FString(*WriteToString<64>(Hasher.Finalize())));
+			}
+			else
+			{
+				CookSettingStrings.Add(ExecutableHashInvalidModuleName, InvalidModule);
+			}
+		}
+	}
+
+	if (CookInfo.GetCookType() == ECookType::ByTheBook)
+	{
+		CookSettingStrings.Add(NAME_CookMode, TEXT("CookByTheBook"));
+		CookSettingStrings.Add(FName(TEXT("DLCName")), CookInfo.GetDLCName());
+		if (!CookInfo.GetDLCName().IsEmpty() || !CookInfo.GetCreateReleaseVersion().IsEmpty())
+		{
+			CookSettingStrings.Add(FName(TEXT("BasedOnReleaseVersion")), CookInfo.GetBasedOnReleaseVersion());
+		}
+	}
+	else
+	{
+		check(CookInfo.GetCookType() == ECookType::OnTheFly);
+		CookSettingStrings.Add(NAME_CookMode, TEXT("CookOnTheFly"));
+	}
+
+	CookSettingStrings.Add(FName(TEXT_CookInProgress), TEXT("true"));
+
+	FConfigSection ConfigSection;
+	FConfigFile ConfigFile;
+	for (TPair<FName, FString>& CurrentSetting : CookSettingStrings)
+	{
+		ConfigSection.Add(CurrentSetting.Key, FConfigValue(MoveTemp(CurrentSetting.Value)));
+	}
+	ConfigFile.Add(TEXT_CookSettings, MoveTemp(ConfigSection));
+	return ConfigFile;
+}
+
+void FGlobalCookArtifact::CompareSettings(UE::Cook::Artifact::FCompareSettingsContext& Context)
+{
+	// For the global settings, we only use RequestFullRecook; we have no files to invalidate.
+	const ITargetPlatform* TargetPlatform = Context.GetTargetPlatform();
+
+	const FConfigSection* PreviousSettings = Context.GetPrevious().FindSection(TEXT_CookSettings);
+	if (PreviousSettings == nullptr)
+	{
+		UE_LOG(LogCook, Display, TEXT("Cook invalidated for platform %s because CookSettings file %s is invalid. Clearing previously cooked packages."),
+			*TargetPlatform->PlatformName(), *Context.GetPreviousFileName());
+		Context.RequestFullRecook(true);
+		return;
+	}
+
+	TSet<FName> IgnoreKeys;
+	IgnoreKeys.Add(FName(TEXT_CookInProgress));
+	IgnoreKeys.Add(ExecutableHashName);
+	IgnoreKeys.Add(ExecutableHashInvalidModuleName);
+
+	const FConfigSection* CurrentSettings = Context.GetCurrent().FindSection(TEXT_CookSettings);
+	check(CurrentSettings);
+	for (const TPair<FName, FConfigValue>& CurrentSetting : *CurrentSettings)
+	{
+		if (IgnoreKeys.Contains(CurrentSetting.Key))
+		{
+			continue;
+		}
+		const FConfigValue* PreviousSetting = PreviousSettings->Find(CurrentSetting.Key);
+		if (!PreviousSetting || PreviousSetting->GetValue() != CurrentSetting.Value.GetValue())
+		{
+			UE_LOG(LogCook, Display, TEXT("Cook invalidated for platform %s because %s has changed. Old: %s, New: %s. Clearing previously cooked packages."),
+				*TargetPlatform->PlatformName(), *CurrentSetting.Key.ToString(),
+				PreviousSetting ? *PreviousSetting->GetValue() : TEXT(""),
+				*CurrentSetting.Value.GetValue());
+			Context.RequestFullRecook(true);
+			return;
+		}
+	}
+
+	if (GIsBuildMachine)
+	{
+		bool bCookInProgress;
+		if (Context.GetPrevious().GetBool(TEXT_CookSettings, TEXT_CookInProgress, bCookInProgress) && bCookInProgress)
+		{
+			UE_LOG(LogCook, Display, TEXT("Cook invalidated for platform %s because the previous cook crashed (or otherwise did not report completion).")
+				TEXT(" CookSettings file %s still has [%s]:%s=true. Clearing previously cooked packages."),
+				*TargetPlatform->PlatformName(), *Context.GetPreviousFileName(), TEXT_CookSettings, TEXT_CookInProgress);
+			Context.RequestFullRecook(true);
+			return;
+		}
+	}
+
+	if (!COTFS.bLegacyIterativeIgnoreIni && COTFS.bLegacyBuildDependencies && COTFS.IniSettingsOutOfDate(TargetPlatform))
+	{
+		UE_LOG(LogCook, Display, TEXT("Cook invalidated for platform %s because ini settings have changed. Clearing previously cooked packages."),
+			*TargetPlatform->PlatformName());
+		Context.RequestFullRecook(true);
+		return;
+	}
+
+	if (!COTFS.bLegacyIterativeIgnoreExe && COTFS.bLegacyBuildDependencies)
+	{
+		const FConfigValue* CurrentHash = CurrentSettings->Find(ExecutableHashName);
+		if (!CurrentHash)
+		{
+			UE_LOG(LogCook, Display, TEXT("Cook invalidated for platform %s because current executable hash is invalid. Invalid module=%s. Clearing previously cooked packages."),
+				*TargetPlatform->PlatformName(), *CurrentSettings->FindRef(ExecutableHashInvalidModuleName).GetValue());
+			Context.RequestFullRecook(true);
+			return;
+		}
+		const FConfigValue* PreviousHash = PreviousSettings->Find(ExecutableHashName);
+		if (!PreviousHash)
+		{
+			const FConfigValue* InvalidModuleName = PreviousSettings->Find(ExecutableHashInvalidModuleName);
+			UE_LOG(LogCook, Display, TEXT("Cook invalidated for platform %s because old executable hash is invalid. Invalid module=%s. Clearing previously cooked packages."),
+				*TargetPlatform->PlatformName(), InvalidModuleName ? *InvalidModuleName->GetValue() : TEXT(""));
+			Context.RequestFullRecook(true);
+			return;
+		}
+		if (!CurrentHash->GetValue().Equals(*PreviousHash->GetValue(), ESearchCase::CaseSensitive))
+		{
+			UE_LOG(LogCook, Display, TEXT("Cook invalidated for platform %s because executable hash has changed. Old: %s, New: %s. Clearing previously cooked packages."),
+				*TargetPlatform->PlatformName(), *PreviousHash->GetValue(), *CurrentHash->GetValue());
+			Context.RequestFullRecook(true);
+			return;
+		}
+	}
+}
+
+} // namespace UE::Cook
+
+void UCookOnTheFlyServer::ClearCookInProgressFlagFromGlobalCookSettings(const ITargetPlatform* TargetPlatform) const
+{
+	using namespace UE::Cook;
+
+	UE::ConfigAccessTracking::FIgnoreScope IgnoreScope;
+	FConfigFile ConfigFile;
+	check(GlobalArtifact);
+	FString ArtifactName = GlobalArtifact->GetArtifactName();
+	FString Filename = GetCookSettingsFileName(TargetPlatform, ArtifactName);
+	ConfigFile.Read(Filename);
+	ConfigFile.RemoveKeyFromSection(TEXT_CookSettings, TEXT_CookInProgress);
+	SaveCookSettings(ConfigFile, TargetPlatform, ArtifactName);
+}
diff --git a/Engine/Source/Editor/UnrealEd/Private/Cooker/GlobalCookArtifact.h b/Engine/Source/Editor/UnrealEd/Private/Cooker/GlobalCookArtifact.h
new file mode 100644
--- /dev/null
+++ b/Engine/Source/Editor/UnrealEd/Private/Cooker/GlobalCookArtifact.h
@@ -0,0 +1,30 @@
+// Copyright Epic Games, Inc. All Rights Reserved.
+
+#pragma once
+
+#include "Cooker/CookArtifact.h"
+
+class UCookOnTheFlyServer;
+
+namespace UE::Cook
+{
+
+/**
+ * CookArtifact that triggers a recook of all packages if any of the global settings that affect all
+ * packages change. This artifact does not store files on disk other than its record of the global CookSettings.
+ */
+class FGlobalCookArtifact : public ICookArtifact
+{
+public:
+	FGlobalCookArtifact(UCookOnTheFlyServer& InCOTFS);
+
+	FString GetArtifactName() const override;
+	FConfigFile CalculateCurrentSettings(ICookInfo& CookInfo, const ITargetPlatform* TargetPlatform) override;
+	void CompareSettings(UE::Cook::Artifact::FCompareSettingsContext& Context) override;
+
+private:
+	UCookOnTheFlyServer& COTFS;
+};
+
+
+} // namespace UE::Cook
\ No newline at end of file
diff --git a/Engine/Source/Editor/UnrealEd/Private/Cooker/WorkerRequestsRemote.cpp b/Engine/Source/Editor/UnrealEd/Private/Cooker/WorkerRequestsRemote.cpp
--- a/Engine/Source/Editor/UnrealEd/Private/Cooker/WorkerRequestsRemote.cpp
+++ b/Engine/Source/Editor/UnrealEd/Private/Cooker/WorkerRequestsRemote.cpp
@@ -152,7 +152,7 @@ void FWorkerRequestsRemote::GetBeginCookIncrementalFlags(UCookOnTheFlyServer& CO
 			*TargetPlatform->PlatformName());
 
 		UE::Cook::FPlatformData* PlatformData = PlatformContext.PlatformData;
-		PlatformContext.CurrentCookSettings.Empty(); // Not needed on CookWorkerClients
+		PlatformContext.CookArtifactsCurrentSettings.Empty(); // Not needed on CookWorkerClients
 		PlatformContext.bHasMemoryResults = false; // Not needed on CookWorkerClients
 		PlatformContext.bFullBuild = DirectorPlatformContext->bFullBuild;
 		PlatformContext.bClearMemoryResults = false; // Not needed on CookWorkerClients
diff --git a/Engine/Source/Runtime/CoreUObject/Public/Cooker/CookArtifact.h b/Engine/Source/Runtime/CoreUObject/Public/Cooker/CookArtifact.h
new file mode 100644
--- /dev/null
+++ b/Engine/Source/Runtime/CoreUObject/Public/Cooker/CookArtifact.h
@@ -0,0 +1,174 @@
+// Copyright Epic Games, Inc. All Rights Reserved.
+
+#pragma once
+
+#if WITH_EDITOR
+#include "Containers/UnrealString.h"
+#include "Misc/ConfigCacheIni.h"
+#include "Misc/Guid.h"
+#include "Templates/RefCounting.h"
+
+class ITargetPlatform;
+class UCookOnTheFlyServer;
+namespace UE::Cook { class ICookInfo; }
+
+namespace UE::Cook::Artifact
+{
+
+/** Argument and return value passing structure for ICookArtifact::CompareSettings. */
+struct FCompareSettingsContext
+{
+public:
+	ICookInfo& GetCookInfo() const;
+	/** The TargetPlatform being tested; invalidation requests are allowed to differ between platforms. Non-null. */
+	const ITargetPlatform* GetTargetPlatform() const;
+	const FConfigFile& GetPrevious() const;
+	const FConfigFile& GetCurrent() const;
+	const FString& GetPreviousFileName() const;
+	bool IsRequestInvalidate() const;
+	bool IsRequestFullRecook() const;
+
+	/**
+	 * Request an invalidation of just this artifact. Initial value is false. OnInvalidate will be called, unless
+	 * the cooker or some artifact requests OnFullRecook, in which case OnFullRecook will be called and OnInvalidate
+	 * will not be called.
+	 */
+	void RequestInvalidate(bool bValue);
+	/**
+	 * Request a full recook; all artifacts invalidated and all packages recooked. Initial value is false.
+	 * OnFullRecook will be called. OnInvalidate will not be called.
+	 */
+	void RequestFullRecook(bool bValue);
+
+private:
+	FCompareSettingsContext(ICookInfo& InCookInfo, const FConfigFile& InPrevious, const FConfigFile& InCurrent,
+		const FString& InPreviousFileName);
+
+	ICookInfo& CookInfo;
+	const ITargetPlatform* TargetPlatform = nullptr;
+	const FConfigFile& Previous;
+	const FConfigFile& Current;
+	const FString& PreviousFileName;
+	bool bRequestInvalidate = false;
+	bool bRequestFullRecook = false;
+
+	friend UCookOnTheFlyServer;
+};
+
+} // namespace UE::Cook::Artifact
+
+namespace UE::Cook
+{
+
+/**
+ * Interface used during cooking for systems that create an artifact collected from cooked packages.
+ * Provides hooks to save global settings that invalidate the artifact when changed.
+ * Provides hooks to clean the artifact when invalidated.
+ */
+class ICookArtifact : public FRefCountBase
+{
+public:
+	virtual ~ICookArtifact() {}
+
+	/** The name of the artifact in log messages and its Metadata/CookSettings_<ArtifactName>.txt file. */
+	virtual FString GetArtifactName() const = 0;
+
+	/**
+	 * Construct a container of config settings from the current executable, config, and global assets.
+	 * If non-empty, it will be passed to CompareSettings each time the cook can cook incrementally,
+	 * and will be saved to disk in the Metadata folder. If the FConfigFile is empty, it will not be saved,
+	 * CompareSettings will be called with empty Previous and Current values.
+	 */
+	virtual FConfigFile CalculateCurrentSettings(ICookInfo& CookInfo, const ITargetPlatform* TargetPlatform)
+	{
+		return FConfigFile();
+	}
+
+	/**
+	 * Compare the results of the previous cook's CalculateCurrentSettings with the current results, and
+	 * report to the cooker if needs to invalidate the artifact constructed by the previous cook.
+	 */
+	virtual void CompareSettings(UE::Cook::Artifact::FCompareSettingsContext& Context)
+	{
+	}
+
+	/**
+	 * Called from the cooker when a full recook was not required and CompareSettings requested invalidate.
+	 * This function must clean the artifact's files under Saved\Cooked\<Platform> and any other invalidated files
+	 * stored elsewhere.
+	 */
+	virtual void OnInvalidate(const ITargetPlatform* TargetPlatform)
+	{
+	}
+
+	/**
+	 * Called from the cooker when a full recook was required by the cooker or by CompareSettings from any
+	 * artifact. Any files under Saved\Cooked\<Platform> will be deleted by the cooker; this function is responsible
+	 * for cleaning invalidated files stored elsewhere.
+	 */
+	virtual void OnFullRecook(const ITargetPlatform* TargetPlatform)
+	{
+	}
+};
+
+
+} // namespace UE::Cook
+
+
+///////////////////////////////////////////////////////
+// Inline implementations
+///////////////////////////////////////////////////////
+
+
+namespace UE::Cook::Artifact
+{
+
+inline FCompareSettingsContext::FCompareSettingsContext(ICookInfo& InCookInfo,
+	const FConfigFile& InPrevious, const FConfigFile& InCurrent, const FString& InPreviousFileName)
+	: CookInfo(InCookInfo)
+	, Previous(InPrevious)
+	, Current(InCurrent)
+	, PreviousFileName(InPreviousFileName)
+{
+}
+
+inline ICookInfo& FCompareSettingsContext::GetCookInfo() const
+{
+	return CookInfo;
+}
+inline const ITargetPlatform* FCompareSettingsContext::GetTargetPlatform() const
+{
+	return TargetPlatform;
+}
+inline const FConfigFile& FCompareSettingsContext::GetPrevious() const
+{
+	return Previous;
+}
+inline const FConfigFile& FCompareSettingsContext::GetCurrent() const
+{
+	return Current;
+}
+inline const FString& FCompareSettingsContext::GetPreviousFileName() const
+{
+	return PreviousFileName;
+}
+inline bool FCompareSettingsContext::IsRequestInvalidate() const
+{
+	return bRequestInvalidate;
+}
+inline bool FCompareSettingsContext::IsRequestFullRecook() const
+{
+	return bRequestFullRecook;
+}
+inline void FCompareSettingsContext::RequestInvalidate(bool bValue)
+{
+	bRequestInvalidate = bValue;
+}
+inline void FCompareSettingsContext::RequestFullRecook(bool bValue)
+{
+	bRequestFullRecook = bValue;
+}
+
+} // namespace UE::Cook::Artifact
+
+#endif // WITH_EDITOR
\ No newline at end of file
diff --git a/Engine/Source/Runtime/CoreUObject/Public/UObject/ICookInfo.h b/Engine/Source/Runtime/CoreUObject/Public/UObject/ICookInfo.h
--- a/Engine/Source/Runtime/CoreUObject/Public/UObject/ICookInfo.h
+++ b/Engine/Source/Runtime/CoreUObject/Public/UObject/ICookInfo.h
@@ -22,6 +22,7 @@ class ITargetPlatform;
 class IPlugin;
 class UPackage;
 
+namespace UE::Cook { class ICookArtifact; }
 namespace UE::Cook { class IMPCollector; }
 #endif
 
@@ -215,6 +216,14 @@ public:
 	virtual UE::Cook::ECookingDLC GetCookingDLC() = 0;
 	/** When DLC is being cooked (@see GetCookingDLC), this will return the name of the DLC plugin. */
 	virtual FString GetDLCName() = 0;
+	/**
+	 * The release version of the basegame the current DLC is based on, when cooking DLC or when making a patch.
+	 * Empty string when not cooking DLC or making a patch.
+	 */
+	virtual FString GetBasedOnReleaseVersion() = 0;
+	/** The release version being created when making a patch. */
+	virtual FString GetCreateReleaseVersion() = 0;
+
 	/** The role the current process plays in its MPCook session, or EProcessType::SingleProcess if it is running standalone. */
 	virtual UE::Cook::EProcessType GetProcessType() = 0;
 	/** Get the validation options used by the running cook, if any. */
@@ -264,7 +273,7 @@ public:
 	 * information between CookWorkers and the CookDirector. Registration will be skipped if the current cook is
 	 * singleprocess, or if the provided ProcessType does not match the current processtype. If registration is
 	 * skipped, the Collector will be referenced but then immediately released, which will delete the Collector if
-	 * the caller does not have their own TRefCountPtr to it.
+	 * not referenced by the caller.
 	 * 
 	 * @param ProcessType: The given collector will only be registered on the given process types, allowing you
 	 *                     to register a different class on Workers than on the Director if desired. 
@@ -274,10 +283,24 @@ public:
 	/**
 	 * MPCook: Unregister in the current process a collector that was registered via RegisterCollector. Silently
 	 * returns if the collector is not registered. References to the Collector will be released, which will
-	 * delete the Collector if he caller does not have their own TRefCountPtr to it.
+	 * delete the Collector if not referenced by the caller.
 	 */
 	virtual void UnregisterCollector(IMPCollector* Collector) = 0;
 
+	/**
+	 * Call in SPCook or on the cookdirector to register an Artifact that saves non-package files in
+	 * the cook's output directory for each TargetPlatform. On CookWorkers this is a noop and the Artifact will be
+	 * deleted if not referenced by the caller. Registered artifacts receives calls that allow them to handle
+	 * invalidation of the artifact during incremental cooks.
+	 */
+	virtual void RegisterArtifact(ICookArtifact* Artifact) = 0;
+	/**
+	 * Unregister Artifact that was registered via RegisterArtifact. Silently returns if the artifact is not
+	 * registered. References to the Artifact will be released, which will delete the Artifact if not referenced by the
+	 * caller.
+	 */
+	virtual void UnregisterArtifact(ICookArtifact* Artifact) = 0;
+
 	/**
 	 * Gets the current cultures that are being cooked.
 	 */
