Cooker: Fix private header name include to match the style for Cooker headers.
#rnx

diff --git a/Engine/Source/Editor/UnrealEd/Private/Cooker/GlobalCookArtifact.cpp b/Engine/Source/Editor/UnrealEd/Private/Cooker/GlobalCookArtifact.cpp
--- a/Engine/Source/Editor/UnrealEd/Private/Cooker/GlobalCookArtifact.cpp
+++ b/Engine/Source/Editor/UnrealEd/Private/Cooker/GlobalCookArtifact.cpp
@@ -1,6 +1,6 @@
 // Copyright Epic Games, Inc. All Rights Reserved.
 
-#include "GlobalCookArtifact.h"
+#include "Cooker/GlobalCookArtifact.h"
 
 #include "Containers/Array.h"
 #include "Containers/Map.h"
