Added experimental support for a multi-process asset localization gather

This can be enabled by passing `-NumGatherTextWorkers=N` to your GatherText commandlet, where N is the number of additional worker processes you want to use (default 0).
You can optionally pass `-StartGatherTextWorkersImmediately` to start any workers at the start of GatherTextFromAssets, rather than wait until there are uncached packages to gather.

Note: IPC to the workers relies on MessageBus, so you also need to pass `-Messaging` to enable that.

diff --git a/Engine/Config/BaseEditor.ini b/Engine/Config/BaseEditor.ini
--- a/Engine/Config/BaseEditor.ini
+++ b/Engine/Config/BaseEditor.ini
@@ -29,6 +29,7 @@ SdfFontAtlasSize=2048
 +DisablePythonForCommandlet=CookShaders
 +DisablePythonForCommandlet=Cook
 +DisablePythonForCommandlet=GatherText
++DisablePythonForCommandlet=GatherTextFromAssetsWorker
 
 [AssetTools]
 ; Full Package Path being hidden from the Asset view
diff --git a/Engine/Source/Editor/UnrealEd/Classes/Commandlets/GatherTextFromAssetsCommandlet.h b/Engine/Source/Editor/UnrealEd/Classes/Commandlets/GatherTextFromAssetsCommandlet.h
--- a/Engine/Source/Editor/UnrealEd/Classes/Commandlets/GatherTextFromAssetsCommandlet.h
+++ b/Engine/Source/Editor/UnrealEd/Classes/Commandlets/GatherTextFromAssetsCommandlet.h
@@ -2,16 +2,76 @@
 
 #pragma once
 
-#include "CoreMinimal.h"
-#include "UObject/ObjectMacros.h"
-#include "Internationalization/GatherableTextData.h"
-#include "Commandlets/GatherTextCommandletBase.h"
 #include "AssetRegistry/AssetData.h"
+#include "Commandlets/GatherTextCommandletBase.h"
+#include "Containers/MpscQueue.h"
+#include "IMessageContext.h"
+#include "Internationalization/GatherableTextData.h"
+
 #include "GatherTextFromAssetsCommandlet.generated.h"
 
 struct FARFilter;
 struct FPackageFileSummary;
 
+class FMessageEndpoint;
+
+USTRUCT()
+struct FGatherTextFromAssetsWorkerMessage_Ping
+{
+	GENERATED_BODY()
+
+	UPROPERTY()
+	int32 ProtocolVersion = 0;
+};
+
+USTRUCT()
+struct FGatherTextFromAssetsWorkerMessage_Pong
+{
+	GENERATED_BODY()
+
+	UPROPERTY()
+	FGuid WorkerId;
+
+	UPROPERTY()
+	TOptional<FDateTime> IdleStartTimeUtc;
+};
+
+USTRUCT()
+struct FGatherTextFromAssetsWorkerMessage_PackageRequest
+{
+	GENERATED_BODY()
+
+	UPROPERTY()
+	FName PackageName;
+
+	UPROPERTY()
+	TSet<FName> Dependencies;
+
+	UPROPERTY()
+	TSet<FGuid> ExternalActors;
+};
+
+USTRUCT()
+struct FGatherTextFromAssetsWorkerMessage_PackageResult
+{
+	GENERATED_BODY()
+
+	UPROPERTY()
+	FGuid WorkerId;
+
+	UPROPERTY()
+	FName PackageName;
+
+	UPROPERTY()
+	TArray<uint8> GatherableTextData;
+
+	UPROPERTY()
+	FString LoadLogCapture;
+
+	UPROPERTY()
+	bool bLoadError = false;
+};
+
 /**
  *	UGatherTextFromAssetsCommandlet: Localization commandlet that collects all text to be localized from the game assets.
  */
@@ -20,11 +80,9 @@ class UGatherTextFromAssetsCommandlet : public UGatherTextCommandletBase
 {
 	GENERATED_UCLASS_BODY()
 
+	void StartWorkers(const int32 MinPackagesToUseWorkers);
 	void ProcessGatherableTextDataArray(const TArray<FGatherableTextData>& GatherableTextDataArray);
 	void CalculateDependenciesForPackagesPendingGather();
-	bool HasExceededMemoryLimit(const bool bLog);
-	void PurgeGarbage(const bool bPurgeReferencedPackages);
-	static void AddReferencedObjects(UObject* InThis, FReferenceCollector& Collector);
 
 public:
 	//~ Begin UCommandlet Interface
@@ -122,6 +180,15 @@ private:
 
 	TSet<FName> PackagesWithDuplicateLocalizationIds;
 
+	/** Record of which packages were sent to workers, so that we can re-add them to the main PackagesPendingGather queue if the worker crashes */
+	TMap<FName, FPackagePendingGather> PackagesDistributedToWorkers;
+
+	/** Number of worker processes to use */
+	int32 NumWorkers = 0;
+
+	/** True if we should start the worker processes immediately, rather than wait until we know we have packages to load */
+	bool bStartWorkersImmediately = false;
+
 	/** Run a GC if the free system memory is less than this value (or zero to disable) */
 	uint64 MinFreeMemoryBytes;
 
@@ -130,8 +197,6 @@ private:
 
 	uint64 NumPackagesDupLocId;
 
-	/** Array of objects that should be kept alive during the next call to CollectGarbage (used by PurgeGarbage and AddReferencedObjects) */
-	TSet<TObjectPtr<UObject>> ObjectsToKeepAlive;
 	/** Path to the directory where output reports etc will be saved.*/
 	FString DestinationPath;
 	bool bSkipGatherCache;
@@ -144,3 +209,228 @@ private:
 	bool bShouldExcludeDerivedClasses;
 	bool bFixPackageLocalizationIdConflict;
 };
+
+/**
+ *	UGatherTextFromAssetsWorkerCommandlet: Localization commandlet worker process that collects all text to be localized from the requested assets.
+ */
+UCLASS()
+class UGatherTextFromAssetsWorkerCommandlet : public UCommandlet
+{
+	GENERATED_BODY()
+
+public:
+	//~ Begin UCommandlet Interface
+	virtual int32 Main(const FString& Params) override;
+	//~ End UCommandlet Interface
+
+private:
+	void HandlePingMessage(const FGatherTextFromAssetsWorkerMessage_Ping& Message, const TSharedRef<IMessageContext>& Context);
+	void HandlePackageRequestMessage(const FGatherTextFromAssetsWorkerMessage_PackageRequest& Message, const TSharedRef<IMessageContext>& Context);
+
+	struct FPackagePendingGather
+	{
+		/** The name of the package */
+		FName PackageName;
+
+		/** The complete set of dependencies for the package */
+		TSet<FName> Dependencies;
+
+		/** The set of external actors to process for a world partition map package */
+		TSet<FGuid> ExternalActors;
+
+		/** Address of the endpoint that requested this package to be gathered */
+		FMessageAddress RequestorAddress;
+	};
+	TArray<FPackagePendingGather> PackagesPendingGather;
+	mutable UE::FMutex PackagesPendingGatherMutex;
+
+	/** Run a GC if the free system memory is less than this value (or zero to disable) */
+	uint64 MinFreeMemoryBytes = 0;
+
+	/** Run a GC if the used process memory is greater than this value (or zero to disable) */
+	uint64 MaxUsedMemoryBytes = 0;
+
+	/** ID of this worker process */
+	FGuid WorkerId;
+
+	/** Holds the messaging endpoint we are sending from */
+	TSharedPtr<FMessageEndpoint> MessageEndpoint;
+
+	/** If this worker is currently idle, when did it last enter an idle state? */
+	TOptional<FDateTime> IdleStartTimeUtc;
+	mutable UE::FMutex IdleStartTimeUtcMutex;
+};
+
+/**
+ * Director for managing worker processes from the main gather commandlet.
+ * @note This is separate from UGatherTextFromAssetsCommandlet as UGatherTextFromAssetsCommandlet may be invoked multiple times within a single gather.
+ * 
+ * Example usage:
+ *	StartWorkers(...);
+ *	WaitForWorkersToStart();
+ *	AssignPackageToWorker(...);
+ *	while (!IsIdle())
+ *	{
+ *		TickWorkers(); // Process GT messages and detect crashed/hung workers
+ *		
+ *		while (TOptional<...> Result = IngestPackageResult())
+ *		{
+ *			// Merge worker result into main result set
+ *		}
+ * 
+ *		... = IngestPackagesFromCrashedWorkers(); // Merge back into main job queue
+ *	}
+ *	StopWorkers();
+ * 
+ * Alternatively you may skip the blocking wait, and use GetAvailableWorkerIds to lazily assign jobs to the workers once available:
+ *	StartWorkers(...);
+ *	for (;;)
+ *	{
+ *		TickWorkers(); // Process GT messages and detect crashed/hung workers
+ *
+ *		if (... = GetAvailableWorkerIds(0))
+ *		{
+ *			AssignPackageToWorker(...);
+ *		}
+ * 
+ *		while (TOptional<...> Result = IngestPackageResult())
+ *		{
+ *			// Merge worker result into main result set
+ *		}
+ *
+ *		... = IngestPackagesFromCrashedWorkers(); // Merge back into main job queue
+ * 
+ *		if (IsIdle())
+ *		{
+ *			break;
+ *		}
+ *	}
+ *	StopWorkers();
+ */
+class FGatherTextFromAssetsWorkerDirector
+{
+private:
+	FGatherTextFromAssetsWorkerDirector() = default;
+
+public:
+	~FGatherTextFromAssetsWorkerDirector();
+
+	/**
+	 * Get the singleton director instance.
+	 */
+	static FGatherTextFromAssetsWorkerDirector& Get();
+
+	/**
+	 * Start the given number of worker processes.
+	 * @note This function will return immediately; call WaitForWorkersToStart if you need to wait for new workers to enter a state where they can accept work.
+	 * 
+	 * @param NumWorkers					The number of worker processes to use.
+	 * @param bStopAdditionalWorkers		True to stop any additional worker processes (>NumWorkers) that are already running.
+	 * @param NumRestartAttemptsIfCrashed	How many times should we attempt to restart the workers if they crash?
+	 * 
+	 * @return True if the workers could be started, false otherwise.
+	 */
+	bool StartWorkers(const int32 NumWorkers, const bool bStopAdditionalWorkers = true, const int32 NumRestartAttemptsIfCrashed = 0);
+
+	/**
+	 * Block while any workers created by StartWorkers are in state where they can accept work.
+	 * 
+	 * @param Timeout Optional timeout to wait for, or unset to wait for an unlimited amount of time.
+	 * 
+	 * @return True if all workers started in the given timeout, false otherwise.
+	 */
+	bool WaitForWorkersToStart(const TOptional<FTimespan> Timeout = {});
+
+	/**
+	 * Stop the current workers (if any), discarding any work that may currently be assigned to them.
+	 */
+	bool StopWorkers();
+
+	/**
+	 * Tick workers, to detect which have crashed or hung.
+	 */
+	void TickWorkers();
+
+	/**
+	 * Get the IDs of the available workers (that are ready to accept work), optionally only returning workers currently meet the given idle threshold (have <= IdleThreshold packages assigned).
+	 * 
+	 * @param IdleThreshold If set, only workers that have <= IdleThreshold packages assigned will be returned.
+	 * 
+	 * @return The available worker IDs, to be used with AssignPackageToWorker.
+	 */
+	TArray<FGuid> GetAvailableWorkerIds(const TOptional<int32> IdleThreshold = {}) const;
+
+	/**
+	 * Is this director idle?
+	 * @note Idle means that all workers have no pending work and there is there no work left to ingest into the main process
+	 * 
+	 * @return True if idle, false otherwise.
+	 */
+	bool IsIdle(int32* OutNumPendingWorkerPackages = nullptr) const;
+
+	/**
+	 * Assign the given package to the given worker.
+	 * 
+	 * @param WorkerId			ID of the worker to assign the package to (see GetAvailableWorkerIds).
+	 * @param PackageRequest	Information about the package to be assigned.
+	 * 
+	 * @return True if the assignment was possible (the worker exists), false otherwise.
+	 */
+	bool AssignPackageToWorker(const FGuid& WorkerId, const FGatherTextFromAssetsWorkerMessage_PackageRequest& PackageRequest);
+
+	/**
+	 * Ingest the next package result received from any worker, if any.
+	 * 
+	 * @return The package result, or an unset value.
+	 */
+	TOptional<FGatherTextFromAssetsWorkerMessage_PackageResult> IngestPackageResult();
+
+	/**
+	 * Ingest the current queue of packages that failed to be processed due to their worker crashing.
+	 * 
+	 * @return The array of packages to handle within the main process.
+	 */
+	TArray<FName> IngestPackagesFromCrashedWorkers();
+
+private:
+	/** Utilities for FWorkerInfo; these should be called while CurrentWorkersMutex is held */
+	struct FWorkerInfo;
+	FString GenerateWorkerCommandLine(const FGuid& WorkerId);
+	FProcHandle StartWorker(const FString& WorkerCommandLine);
+	bool HandleWorkerCrashed(const FGuid& WorkerId, FWorkerInfo& WorkerInfo);
+	void ResetWorkerDiscoveryAndTimeout(FWorkerInfo& WorkerInfo);
+	void ResendWorkerPendingPackageRequests(FWorkerInfo& WorkerInfo);
+
+	void BroadcastPingMessage(const bool bIgnoreDelay = false);
+
+	void HandlePongMessage(const FGatherTextFromAssetsWorkerMessage_Pong& Message, const TSharedRef<IMessageContext>& Context);
+	void HandlePackageResultMessage(const FGatherTextFromAssetsWorkerMessage_PackageResult& Message, const TSharedRef<IMessageContext>& Context);
+
+	/** Information about the current worker processes, including the work that has been assigned to them */
+	struct FWorkerInfo
+	{
+		FProcHandle WorkerProc;
+		FString WorkerCommandLine;
+		FMessageAddress EndpointAddress;
+		TOptional<FDateTime> LastMessageReceivedUtc;
+		TOptional<FDateTime> LastPackageRequestUtc;
+		TOptional<FDateTime> IdleStartTimeUtc;
+		TMap<FName, FGatherTextFromAssetsWorkerMessage_PackageRequest> PendingPackageRequests;
+		int32 NumRestartAttemptsIfCrashed = 0;
+		bool bResendPendingPackageRequests = false;
+	};
+	TMap<FGuid, TSharedPtr<FWorkerInfo>> CurrentWorkers;
+	mutable UE::FMutex CurrentWorkersMutex;
+
+	/** The last time that we broadcast a ping message */
+	FDateTime LastPingBroadcastUtc = FDateTime::MinValue();
+
+	/** The queue of results received from any worker */
+	TMpscQueue<FGatherTextFromAssetsWorkerMessage_PackageResult> PackageResults;
+
+	/** The queue of packages that failed to be processed due to their worker crashing */
+	TArray<FName> PackagesFromCrashedWorkers;
+
+	/** Holds the messaging endpoint we are sending from */
+	TSharedPtr<FMessageEndpoint> MessageEndpoint;
+};
diff --git a/Engine/Source/Editor/UnrealEd/Private/Commandlets/GatherTextCommandlet.cpp b/Engine/Source/Editor/UnrealEd/Private/Commandlets/GatherTextCommandlet.cpp
--- a/Engine/Source/Editor/UnrealEd/Private/Commandlets/GatherTextCommandlet.cpp
+++ b/Engine/Source/Editor/UnrealEd/Private/Commandlets/GatherTextCommandlet.cpp
@@ -3,6 +3,7 @@
 #include "Commandlets/GatherTextCommandlet.h"
 
 #include "Async/ParallelFor.h"
+#include "Commandlets/GatherTextFromAssetsCommandlet.h"
 #include "Commandlets/GatherTextFromSourceCommandlet.h"
 #include "GeneralProjectSettings.h"
 #include "HAL/FileManager.h"
@@ -199,6 +200,8 @@ int32 UGatherTextCommandlet::Execute(const FString& Params, const TSharedPtr<con
 	{
 		ON_SCOPE_EXIT
 		{
+			FGatherTextFromAssetsWorkerDirector::Get().StopWorkers();
+
 			UGatherTextFromSourceCommandlet::LogStats();
 		};
 
@@ -685,6 +688,12 @@ bool UGatherTextCommandlet::ExecuteGatherTextPhases(const TSharedPtr<FLocalizati
 		}
 
 		UE_LOG(LogGatherTextCommandlet, Display, TEXT("Phase '%s' %.2f seconds (%d sequential tasks and %d parallel tasks)."), LexToString((EGatherTextCommandletPhase)PhaseIndex), FPlatformTime::Seconds() - PhaseExecutionStartTime, Phase.SequentialTasks.Num(), Phase.ParallelTasks.Num());
+
+		if ((EGatherTextCommandletPhase)PhaseIndex == EGatherTextCommandletPhase::PostUpdateManifests)
+		{
+			// No longer need the asset gather workers after the manifests have been updated
+			FGatherTextFromAssetsWorkerDirector::Get().StopWorkers();
+		}
 	}
 
 	return true;
diff --git a/Engine/Source/Editor/UnrealEd/Private/Commandlets/GatherTextFromAssetsCommandlet.cpp b/Engine/Source/Editor/UnrealEd/Private/Commandlets/GatherTextFromAssetsCommandlet.cpp
--- a/Engine/Source/Editor/UnrealEd/Private/Commandlets/GatherTextFromAssetsCommandlet.cpp
+++ b/Engine/Source/Editor/UnrealEd/Private/Commandlets/GatherTextFromAssetsCommandlet.cpp
@@ -7,6 +7,7 @@
 #include "AssetRegistry/AssetRegistryModule.h"
 #include "Async/ParallelFor.h"
 #include "CollectionManagerModule.h"
+#include "Commandlets/CommandletHelpers.h"
 #include "Commandlets/GatherTextFromSourceCommandlet.h"
 #include "DistanceFieldAtlas.h"
 #include "Editor.h"
@@ -21,6 +22,9 @@
 #include "Logging/StructuredLog.h"
 #include "MeshCardBuild.h"
 #include "MeshCardRepresentation.h"
+#include "MessageEndpoint.h"
+#include "MessageEndpointBuilder.h"
+#include "Misc/AsciiSet.h"
 #include "Misc/CommandLine.h"
 #include "Misc/ConfigCacheIni.h"
 #include "Misc/FeedbackContext.h"
@@ -30,6 +34,7 @@
 #include "Misc/PackageName.h"
 #include "Misc/Paths.h"
 #include "Misc/ScopedSlowTask.h"
+#include "Misc/StringOutputDevice.h"
 #include "Modules/ModuleManager.h"
 #include "PackageHelperFunctions.h"
 #include "ProfilingDebugging/ScopedTimers.h"
@@ -37,6 +42,7 @@
 #include "ShaderCompiler.h"
 #include "Sound/DialogueWave.h"
 #include "Templates/UniquePtr.h"
+#include "UnrealEdMisc.h"
 #include "UObject/Class.h"
 #include "UObject/EditorObjectVersion.h"
 #include "UObject/FortniteMainBranchObjectVersion.h"
@@ -76,7 +82,11 @@ public:
 		FLoadPackageLogOutputRedirector* LogOutputRedirector;
 	};
 
-	FLoadPackageLogOutputRedirector() = default;
+	explicit FLoadPackageLogOutputRedirector(FOutputDevice* InTargetOutputDevice = GWarn)
+		: TargetOutputDevice(InTargetOutputDevice)
+	{
+	}
+
 	virtual ~FLoadPackageLogOutputRedirector() = default;
 
 	void BeginCapturingLogData(FStringView InPackageContext)
@@ -97,14 +107,14 @@ public:
 		OriginalWarningContext = nullptr;
 
 		// Report any messages, and also report a warning if we silenced some warnings or errors when loading
-		if (ErrorCount > 0 || WarningCount > 0)
+		if (TargetOutputDevice && (ErrorCount > 0 || WarningCount > 0))
 		{
 			static const FString LogIndentation = TEXT("    ");
 
-			UE_LOG(LogGatherTextFromAssetsCommandlet, Display, TEXT("Package '%s' produced %d error(s) and %d warning(s) while loading (see below). Please verify that your text has gathered correctly."), *PackageContext, ErrorCount, WarningCount);
+			TargetOutputDevice->CategorizedLogf(LogGatherTextFromAssetsCommandlet.GetCategoryName(), ELogVerbosity::Display, TEXT("Package '%s' produced %d error(s) and %d warning(s) while loading (see below). Please verify that your text has gathered correctly."), *PackageContext, ErrorCount, WarningCount);
 			for (const FString& FormattedOutput : FormattedErrorsAndWarningsList)
 			{
-				GWarn->Log(NAME_None, ELogVerbosity::Display, LogIndentation + FormattedOutput);
+				TargetOutputDevice->Log(NAME_None, ELogVerbosity::Display, LogIndentation + FormattedOutput);
 			}
 		}
 
@@ -191,6 +201,8 @@ private:
 
 	FString PackageContext;
 	FFeedbackContext* OriginalWarningContext = nullptr;
+
+	FOutputDevice* TargetOutputDevice = nullptr;
 };
 
 class FAssetGatherCacheMetrics
@@ -269,6 +281,224 @@ namespace UE::Private::GatherTextFromAssetsCommandlet
 
 	static int32 ProcessAndRemoveCachedPackagesMaxThreads = -1;
 	static FAutoConsoleVariableRef CVarProcessAndRemoveCachedPackagesMaxThreads(TEXT("Localization.GatherTextFromAssetsCommandlet.ProcessAndRemoveCachedPackagesMaxThreads"), ProcessAndRemoveCachedPackagesMaxThreads, TEXT("Max number of threads to use if parallelizing the 'process and remove cached packages' step, or <= 0 to use as many threads as possible."));
+
+	// Bump whenever changing the data in any of the FGatherTextFromAssetsWorker*Message types
+	static constexpr int32 WorkerProtocolVersion = 1;
+
+	// Values controlling how packages are distributed to workers
+	static constexpr int32 MinPackagesToUseWorkers = 1000;	// How many packages should we have to process (in total) before we start the workers?
+	static constexpr int32 MinPackagesToKeepLocal = 50;		// How many packages should we keep in the main process, rather than distribute to workers? (to avoid waiting on workers at the end of the gather)
+	static constexpr int32 MaxPackagesToDistribute = 100;	// How many packages should we distribute to each worker in a single batch?
+	static constexpr int32 WorkerIdleThreshold = 10;		// How many packages can a worker have pending in its queue before we consider it "idle" and try and assign it more packages? (to avoid workers becoming idle)
+
+	UPackage* LoadPackageToGather(const FString& PackageName, const bool bHasExternalActors, FLoadPackageLogOutputRedirector& LogOutputRedirector)
+	{
+		UPackage* Package = nullptr;
+		{
+			FLoadPackageLogOutputRedirector::FScopedCapture ScopedCapture(&LogOutputRedirector, PackageName);
+			if (bHasExternalActors)
+			{
+				Package = LoadWorldPackageForEditor(PackageName, EWorldType::Editor, LOAD_NoWarn | LOAD_Quiet);
+			}
+			else
+			{
+				Package = LoadPackage(nullptr, *PackageName, LOAD_NoWarn | LOAD_Quiet);
+			}
+		}
+		return Package;
+	}
+
+	bool GatherTextFromPackage(UPackage* Package, const TSet<FGuid>& ExternalActors, const bool bGatherFromPrimaryPackage, const bool bGatherFromExternalActorPackages, FLoadPackageLogOutputRedirector& LogOutputRedirector, TArray<FGatherableTextData>& OutGatherableTextDataArray)
+	{
+		// Because packages may not have been resaved after this flagging was implemented, we may have added packages to load that weren't flagged - potential false positives.
+		// The loading process should have reflagged said packages so that only true positives will have this flag.
+		if (Package->RequiresLocalizationGather() || ExternalActors.Num() > 0)
+		{
+			if (bGatherFromPrimaryPackage)
+			{
+				// Gathers from the given package
+				EPropertyLocalizationGathererResultFlags GatherableTextResultFlags = EPropertyLocalizationGathererResultFlags::Empty;
+				FPropertyLocalizationDataGatherer(OutGatherableTextDataArray, Package, GatherableTextResultFlags);
+			}
+
+			// If this is a WP world then query the localization for any external actors actors that were determined to be stale
+			if (bGatherFromExternalActorPackages && ExternalActors.Num() > 0)
+			{
+				if (UWorld* World = UWorld::FindWorldInPackage(Package))
+				{
+					UWorld::InitializationValues IVS;
+					IVS.InitializeScenes(false);
+					IVS.AllowAudioPlayback(false);
+					IVS.RequiresHitProxies(false);
+					IVS.CreatePhysicsScene(false);
+					IVS.CreateNavigation(false);
+					IVS.CreateAISystem(false);
+					IVS.ShouldSimulatePhysics(false);
+					IVS.EnableTraceCollision(false);
+					IVS.SetTransactional(false);
+					IVS.CreateFXSystem(false);
+					IVS.CreateWorldPartition(true);
+
+					TOptional<FScopedEditorWorld> ScopeEditorWorld;
+					if (!World->IsInitialized())
+					{
+						FLoadPackageLogOutputRedirector::FScopedCapture ScopedCapture(&LogOutputRedirector, Package->GetName());
+						ScopeEditorWorld.Emplace(World, IVS); // Initializing FScopedEditorWorld can log warnings, so capture those like we do with loading errors
+					}
+
+					if (UWorldPartition* WorldPartition = World->GetWorldPartition())
+					{
+						// ForEachActorWithLoading may GC while running, so keep the world partition (and indirectly the world and its package) alive
+						TGCObjectScopeGuard<UWorldPartition> WorldPartitionGCGuard(WorldPartition);
+
+						FWorldPartitionHelpers::FForEachActorWithLoadingParams ForEachActorParams;
+						ForEachActorParams.ActorGuids = ExternalActors.Array();
+
+						FWorldPartitionHelpers::ForEachActorWithLoading(WorldPartition,
+							[&OutGatherableTextDataArray](const FWorldPartitionActorDescInstance* ActorDescInstance)
+							{
+								if (const AActor* Actor = ActorDescInstance->GetActor())
+								{
+									EPropertyLocalizationGathererResultFlags ActorGatherableTextResultFlags = EPropertyLocalizationGathererResultFlags::Empty;
+									FPropertyLocalizationDataGatherer(OutGatherableTextDataArray, Actor->GetExternalPackage(), ActorGatherableTextResultFlags);
+								}
+								return true;
+							}, ForEachActorParams);
+					}
+				}
+			}
+
+			return true;
+		}
+
+		return false;
+	}
+
+	bool HasExceededMemoryLimit(const uint64 MinFreeMemoryBytes, const uint64 MaxUsedMemoryBytes, const bool bLog)
+	{
+		const FPlatformMemoryStats MemStats = FPlatformMemory::GetStats();
+
+		const uint64 FreeMemoryBytes = MemStats.AvailablePhysical;
+		if (MinFreeMemoryBytes > 0u && FreeMemoryBytes < MinFreeMemoryBytes)
+		{
+			UE_CLOG(bLog, LogGatherTextFromAssetsCommandlet, Display, TEXT("Free system memory is currently %s of %s (%s used by process), which is less than the requested limit of %s; a flush will be performed."), 
+				*FText::AsMemory(FreeMemoryBytes).ToString(), *FText::AsMemory(MemStats.TotalPhysical).ToString(), *FText::AsMemory(MemStats.UsedPhysical).ToString(), *FText::AsMemory(MinFreeMemoryBytes).ToString());
+			return true;
+		}
+
+		const uint64 UsedMemoryBytes = MemStats.UsedPhysical;
+		if (MaxUsedMemoryBytes > 0u && UsedMemoryBytes >= MaxUsedMemoryBytes)
+		{
+			UE_CLOG(bLog, LogGatherTextFromAssetsCommandlet, Display, TEXT("Used process memory is currently %s, which is greater than the requested limit of %s; a flush will be performed."), *FText::AsMemory(UsedMemoryBytes).ToString(), *FText::AsMemory(MaxUsedMemoryBytes).ToString());
+			return true;
+		}
+
+		return false;
+	}
+
+	template <typename TPackagePendingGather>
+	void PurgeGarbage(TArray<TPackagePendingGather>* PackagesPendingGather)
+	{
+		FlushAsyncLoading();
+
+		TSet<FName> LoadedPackageNames;
+		TSet<FName> PackageNamesToKeepAlive;
+		TArray<UObject*> ObjectsToKeepAlive;
+
+		const bool bHasPackagesPendingGather = PackagesPendingGather && PackagesPendingGather->Num() > 0;
+		if (bHasPackagesPendingGather)
+		{
+			// Build a complete list of packages that we still need to keep alive, either because we still 
+			// have to process them, or because they're a dependency for something we still have to process
+			for (const TPackagePendingGather& PackagePendingGather : *PackagesPendingGather)
+			{
+				PackageNamesToKeepAlive.Add(PackagePendingGather.PackageName);
+				PackageNamesToKeepAlive.Append(PackagePendingGather.Dependencies);
+			}
+
+			for (TObjectIterator<UPackage> PackageIt; PackageIt; ++PackageIt)
+			{
+				UPackage* Package = *PackageIt;
+				if (PackageNamesToKeepAlive.Contains(Package->GetFName()))
+				{
+					LoadedPackageNames.Add(Package->GetFName());
+
+					// Keep any requested packages (and their RF_Standalone inners) alive during a call to PurgeGarbage
+					ObjectsToKeepAlive.Add(Package);
+					ForEachObjectWithPackage(Package, [&ObjectsToKeepAlive](UObject* InPackageInner)
+					{
+						if (InPackageInner->HasAnyFlags(RF_Standalone | RF_HasExternalPackage))
+						{
+							ObjectsToKeepAlive.Add(InPackageInner);
+						}
+						return true;
+					}, true, RF_NoFlags, EInternalObjectFlags::Garbage);
+				}
+			}
+		}
+
+		if (ObjectsToKeepAlive.Num() > 0)
+		{
+			TGCObjectsScopeGuard<UObject> ObjectsToKeepAliveScopeGuard(ObjectsToKeepAlive);
+			CollectGarbage(IsRunningCommandlet() ? RF_NoFlags : GARBAGE_COLLECTION_KEEPFLAGS);
+		}
+		else
+		{
+			CollectGarbage(IsRunningCommandlet() ? RF_NoFlags : GARBAGE_COLLECTION_KEEPFLAGS);
+		}
+
+		// Fully process the shader compilation results when performing a full purge, as it's the only way to reclaim that memory
+		if (!PackagesPendingGather && GShaderCompilingManager)
+		{
+			GShaderCompilingManager->ProcessAsyncResults(false, false);
+		}
+
+		if (bHasPackagesPendingGather)
+		{
+			// Sort the remaining packages to gather so that currently loaded packages are processed first, followed by those with the most dependencies
+			// This aims to allow packages to be GC'd as soon as possible once nothing is no longer referencing them as a dependency
+			// Note: This array is processed backwards, so "first" is actually the end of the array
+			PackagesPendingGather->Sort([&LoadedPackageNames](const TPackagePendingGather& PackagePendingGatherOne, const TPackagePendingGather& PackagePendingGatherTwo)
+			{
+				const bool bIsPackageOneLoaded = LoadedPackageNames.Contains(PackagePendingGatherOne.PackageName);
+				const bool bIsPackageTwoLoaded = LoadedPackageNames.Contains(PackagePendingGatherTwo.PackageName);
+				return (bIsPackageOneLoaded == bIsPackageTwoLoaded)
+					? PackagePendingGatherOne.Dependencies.Num() < PackagePendingGatherTwo.Dependencies.Num()
+					: bIsPackageTwoLoaded;
+			});
+		}
+	}
+
+	template <typename TPackagePendingGather>
+	void ConditionalPurgeGarbage(const uint64 MinFreeMemoryBytes, const uint64 MaxUsedMemoryBytes, TArray<TPackagePendingGather>& PackagesPendingGather)
+	{
+		if (HasExceededMemoryLimit(MinFreeMemoryBytes, MaxUsedMemoryBytes, /*bLog*/true))
+		{
+			// First try a minimal purge to only remove things that are no longer referenced or needed by other packages pending gather
+			PurgeGarbage<TPackagePendingGather>(&PackagesPendingGather);
+
+			if (HasExceededMemoryLimit(MinFreeMemoryBytes, MaxUsedMemoryBytes, /*bLog*/false))
+			{
+				// If we're still over the memory limit after a minimal purge, then attempt a full purge
+				PurgeGarbage<TPackagePendingGather>(nullptr);
+
+				// If we're still over the memory limit after both purges, then log a warning as we may be about to OOM
+				UE_CLOG(HasExceededMemoryLimit(MinFreeMemoryBytes, MaxUsedMemoryBytes, /*bLog*/false), LogGatherTextFromAssetsCommandlet, Warning, TEXT("Flushing failed to reduce process memory to within the requested limits; this process may OOM!"));
+			}
+		}
+	}
+
+	void TickMessageBusGT()
+	{
+		// Flush the task graph to grab any pending messages
+		// We put a dummy fence task into the queue to avoid potentially waiting indefinitely if other threads keep adding game thread events
+		if (!FTaskGraphInterface::Get().IsThreadProcessingTasks(ENamedThreads::GameThread))
+		{
+			DECLARE_CYCLE_STAT(TEXT("GatherTextFromAssets.TickMessageBusGT"), STAT_GatherTextFromAssets_TickMessageBusGT, STATGROUP_TaskGraphTasks);
+			FGraphEventRef FenceHandle = FSimpleDelegateGraphTask::CreateAndDispatchWhenReady(FSimpleDelegateGraphTask::FDelegate(), GET_STATID(STAT_GatherTextFromAssets_TickMessageBusGT), nullptr, ENamedThreads::GameThread);
+			FTaskGraphInterface::Get().WaitUntilTaskCompletes(FenceHandle, ENamedThreads::GameThread);
+		}
+	}
 }
 
 #define LOC_DEFINE_REGION
@@ -378,101 +608,6 @@ void UGatherTextFromAssetsCommandlet::CalculateDependenciesForPackagesPendingGat
 	}
 }
 
-bool UGatherTextFromAssetsCommandlet::HasExceededMemoryLimit(const bool bLog)
-{
-	const FPlatformMemoryStats MemStats = FPlatformMemory::GetStats();
-
-	const uint64 FreeMemoryBytes = MemStats.AvailablePhysical;
-	if (MinFreeMemoryBytes > 0u && FreeMemoryBytes < MinFreeMemoryBytes)
-	{
-		UE_CLOG(bLog, LogGatherTextFromAssetsCommandlet, Display, TEXT("Free system memory is currently %s, which is less than the requested limit of %s; a flush will be performed."), *FText::AsMemory(FreeMemoryBytes).ToString(), *FText::AsMemory(MinFreeMemoryBytes).ToString());
-		return true;
-	}
-
-	const uint64 UsedMemoryBytes = MemStats.UsedPhysical;
-	if (MaxUsedMemoryBytes > 0u && UsedMemoryBytes >= MaxUsedMemoryBytes)
-	{
-		UE_CLOG(bLog, LogGatherTextFromAssetsCommandlet, Display, TEXT("Used process memory is currently %s, which is greater than the requested limit of %s; a flush will be performed."), *FText::AsMemory(UsedMemoryBytes).ToString(), *FText::AsMemory(MaxUsedMemoryBytes).ToString());
-		return true;
-	}
-
-	return false;
-}
-
-void UGatherTextFromAssetsCommandlet::PurgeGarbage(const bool bPurgeReferencedPackages)
-{
-	check(ObjectsToKeepAlive.Num() == 0);
-
-	FlushAsyncLoading();
-
-	TSet<FName> LoadedPackageNames;
-	TSet<FName> PackageNamesToKeepAlive;
-
-	if (!bPurgeReferencedPackages)
-	{
-		// Build a complete list of packages that we still need to keep alive, either because we still 
-		// have to process them, or because they're a dependency for something we still have to process
-		for (const FPackagePendingGather& PackagePendingGather : PackagesPendingGather)
-		{
-			PackageNamesToKeepAlive.Add(PackagePendingGather.PackageName);
-			PackageNamesToKeepAlive.Append(PackagePendingGather.Dependencies);
-		}
-
-		for (TObjectIterator<UPackage> PackageIt; PackageIt; ++PackageIt)
-		{
-			UPackage* Package = *PackageIt;
-			if (PackageNamesToKeepAlive.Contains(Package->GetFName()))
-			{
-				LoadedPackageNames.Add(Package->GetFName());
-
-				// Keep any requested packages (and their RF_Standalone inners) alive during a call to PurgeGarbage
-				ObjectsToKeepAlive.Add(Package);
-				ForEachObjectWithPackage(Package, [this](UObject* InPackageInner)
-				{
-					if (InPackageInner->HasAnyFlags(RF_Standalone | RF_HasExternalPackage))
-					{
-						ObjectsToKeepAlive.Add(InPackageInner);
-					}
-					return true;
-				}, true, RF_NoFlags, EInternalObjectFlags::Garbage);
-			}
-		}
-	}
-
-	CollectGarbage(IsRunningCommandlet() ? RF_NoFlags : GARBAGE_COLLECTION_KEEPFLAGS);
-	ObjectsToKeepAlive.Reset();
-
-	// Fully process the shader compilation results when performing a full purge, as it's the only way to reclaim that memory
-	if (bPurgeReferencedPackages && GShaderCompilingManager)
-	{
-		GShaderCompilingManager->ProcessAsyncResults(false, false);
-	}
-
-	if (!bPurgeReferencedPackages)
-	{
-		// Sort the remaining packages to gather so that currently loaded packages are processed first, followed by those with the most dependencies
-		// This aims to allow packages to be GC'd as soon as possible once nothing is no longer referencing them as a dependency
-		// Note: This array is processed backwards, so "first" is actually the end of the array
-		PackagesPendingGather.Sort([&LoadedPackageNames](const FPackagePendingGather& PackagePendingGatherOne, const FPackagePendingGather& PackagePendingGatherTwo)
-		{
-			const bool bIsPackageOneLoaded = LoadedPackageNames.Contains(PackagePendingGatherOne.PackageName);
-			const bool bIsPackageTwoLoaded = LoadedPackageNames.Contains(PackagePendingGatherTwo.PackageName);
-			return (bIsPackageOneLoaded == bIsPackageTwoLoaded) 
-				? PackagePendingGatherOne.Dependencies.Num() < PackagePendingGatherTwo.Dependencies.Num() 
-				: bIsPackageTwoLoaded;
-		});
-	}
-}
-
-void UGatherTextFromAssetsCommandlet::AddReferencedObjects(UObject* InThis, FReferenceCollector& Collector)
-{
-	Super::AddReferencedObjects(InThis, Collector);
-
-	// Keep any requested objects alive during a call to PurgeGarbage
-	UGatherTextFromAssetsCommandlet* This = CastChecked<UGatherTextFromAssetsCommandlet>(InThis);
-	Collector.AddReferencedObjects(This->ObjectsToKeepAlive);
-}
-
 bool IsGatherableTextDataIdentical(const TArray<FGatherableTextData>& GatherableTextDataArrayOne, const TArray<FGatherableTextData>& GatherableTextDataArrayTwo)
 {
 	struct FSignificantGatherableTextData
@@ -1332,6 +1467,8 @@ bool UGatherTextFromAssetsCommandlet::LoadAndProcessUncachedPackages(TArray<FNam
 	int32 PackageCount = PackagesPendingGather.Num();
 	int32 NumPackagesFailedLoading = 0;
 	
+	FGatherTextFromAssetsWorkerDirector& WorkerDirector = FGatherTextFromAssetsWorkerDirector::Get();
+
 	FScopedSlowTask SlowTask(PackageCount, NSLOCTEXT("GatherTextCommandlet", "LoadAndProcessUncachedPackages", "Loading and Gathering Packages..."));
 	UE_LOG(LogGatherTextFromAssetsCommandlet, Display, TEXT("Number of packages with duplicate loc ids: %d"), NumPackagesDupLocId);
 	UE_LOG(LogGatherTextFromAssetsCommandlet, Display, TEXT("Number of packages to load: %d"), PackagesPendingGather.Num());
@@ -1347,8 +1484,70 @@ bool UGatherTextFromAssetsCommandlet::LoadAndProcessUncachedPackages(TArray<FNam
 		return FLocalizedAssetSCCUtil::SavePackageWithSCC(SourceControlInfo, Package, PackageFileName);
 	};
 
+	auto AssignPackagesToWorker = [this, &WorkerDirector, &SlowTask, &NumPackagesProcessed](const FGuid& WorkerId, const int32 NumPackages)
+	{
+		check(PackagesPendingGather.Num() >= NumPackages);
+
+		// We assign worker packages from the tail of the queue (the front of the array) 
+		// as these packages have the fewest number of dependencies so are the simplest 
+		// to distribute and load, which helps to balance the memory usage between 
+		// the main process (which takes the heavy tasks) and the worker processes
+		TArray<FPackagePendingGather> PackagesToReinsert;
+		for (int32 Index = 0; Index < NumPackages; ++Index)
+		{
+			FPackagePendingGather PackagePendingGather = MoveTemp(PackagesPendingGather[Index]);
+
+			FGatherTextFromAssetsWorkerMessage_PackageRequest PackageRequest;
+			PackageRequest.PackageName = PackagePendingGather.PackageName;
+			PackageRequest.Dependencies = PackagePendingGather.Dependencies;
+			PackageRequest.ExternalActors = PackagePendingGather.ExternalActors;
+
+			if (WorkerDirector.AssignPackageToWorker(WorkerId, PackageRequest))
+			{
+				SlowTask.EnterProgressFrame();
+				++NumPackagesProcessed;
+
+				const FName PackageName = PackagePendingGather.PackageName;
+				PackagesDistributedToWorkers.Add(PackageName, MoveTemp(PackagePendingGather));
+			}
+			else
+			{
+				PackagesToReinsert.Add(MoveTemp(PackagePendingGather));
+			}
+		}
+		PackagesPendingGather.RemoveAt(0, NumPackages, EAllowShrinking::No);
+		if (PackagesToReinsert.Num() > 0)
+		{
+			PackagesPendingGather.Insert(MoveTemp(PackagesToReinsert), 0);
+		}
+	};
+
 	while (PackagesPendingGather.Num() > 0)
 	{
+		// Should we try and distribute some work?
+		{
+			if (PackagesPendingGather.Num() > UE::Private::GatherTextFromAssetsCommandlet::MinPackagesToKeepLocal)
+			{
+				// If we have (soon to be) idle workers, distribute some work to them now
+				if (const TArray<FGuid> IdleWorkerIds = WorkerDirector.GetAvailableWorkerIds(UE::Private::GatherTextFromAssetsCommandlet::WorkerIdleThreshold);
+					IdleWorkerIds.Num() > 0)
+				{
+					if (const int32 NumPackagesToDistributeToEachWorker = FMath::Clamp((PackagesPendingGather.Num() - UE::Private::GatherTextFromAssetsCommandlet::MinPackagesToKeepLocal) / IdleWorkerIds.Num(), 0, UE::Private::GatherTextFromAssetsCommandlet::MaxPackagesToDistribute);
+						NumPackagesToDistributeToEachWorker > 0)
+					{
+						for (const FGuid& WorkerId : IdleWorkerIds)
+						{
+							AssignPackagesToWorker(WorkerId, NumPackagesToDistributeToEachWorker);
+						}
+					}
+				}
+				if (PackagesPendingGather.Num() == 0)
+				{
+					break;
+				}
+			}
+		}
+
 		const FPackagePendingGather PackagePendingGather = PackagesPendingGather.Pop(EAllowShrinking::No);
 		FNameBuilder PackageNameStr(PackagePendingGather.PackageName);
 
@@ -1366,19 +1565,7 @@ bool UGatherTextFromAssetsCommandlet::LoadAndProcessUncachedPackages(TArray<FNam
 			return false;
 		}
 
-		UPackage* Package = nullptr;
-		{
-			FLoadPackageLogOutputRedirector::FScopedCapture ScopedCapture(&LogOutputRedirector, *PackageNameStr);
-			if (PackagePendingGather.ExternalActors.Num() > 0)
-			{
-				Package = LoadWorldPackageForEditor(*PackageNameStr, EWorldType::Editor, LOAD_NoWarn | LOAD_Quiet);
-			}
-			else
-			{
-				Package = LoadPackage(nullptr, *PackageNameStr, LOAD_NoWarn | LOAD_Quiet);
-			}
-		}
-
+		UPackage* Package = UE::Private::GatherTextFromAssetsCommandlet::LoadPackageToGather(*PackageNameStr, PackagePendingGather.ExternalActors.Num() > 0, LogOutputRedirector);
 		if (!Package)
 		{
 			UE_LOGFMT(LogGatherTextFromAssetsCommandlet, Warning, "Failed to load package: '{package}'.",
@@ -1423,10 +1610,9 @@ bool UGatherTextFromAssetsCommandlet::LoadAndProcessUncachedPackages(TArray<FNam
 		{
 			EmbeddedContext->RunTick();
 		}
+		WorkerDirector.TickWorkers();
 
-		// Because packages may not have been resaved after this flagging was implemented, we may have added packages to load that weren't flagged - potential false positives.
-		// The loading process should have reflagged said packages so that only true positives will have this flag.
-		if (Package->RequiresLocalizationGather() || PackagePendingGather.ExternalActors.Num() > 0)
+		if (UE::Private::GatherTextFromAssetsCommandlet::GatherTextFromPackage(Package, PackagePendingGather.ExternalActors, /*bGatherFromPrimaryPackage*/true, /*bGatherFromExternalActorPackages*/false, LogOutputRedirector, GatherableTextDataArray))
 		{
 			UE_LOG(LogGatherTextFromAssetsCommandlet, Display, TEXT("[%6.2f%%] Gathering package: '%s'..."), PercentageComplete, *PackageNameStr);
 
@@ -1481,47 +1667,7 @@ bool UGatherTextFromAssetsCommandlet::LoadAndProcessUncachedPackages(TArray<FNam
 			// If this is a WP world then query the localization for any external actors actors that were determined to be stale
 			if (PackagePendingGather.ExternalActors.Num() > 0)
 			{
-				if (UWorld* World = UWorld::FindWorldInPackage(Package))
-				{
-					UWorld::InitializationValues IVS;
-					IVS.InitializeScenes(false);
-					IVS.AllowAudioPlayback(false);
-					IVS.RequiresHitProxies(false);
-					IVS.CreatePhysicsScene(false);
-					IVS.CreateNavigation(false);
-					IVS.CreateAISystem(false);
-					IVS.ShouldSimulatePhysics(false);
-					IVS.EnableTraceCollision(false);
-					IVS.SetTransactional(false);
-					IVS.CreateFXSystem(false);
-					IVS.CreateWorldPartition(true);
-
-					TOptional<FScopedEditorWorld> ScopeEditorWorld;
-					if (!World->IsInitialized())
-					{
-						FLoadPackageLogOutputRedirector::FScopedCapture ScopedCapture(&LogOutputRedirector, *PackageNameStr);
-						ScopeEditorWorld.Emplace(World, IVS); // Initializing FScopedEditorWorld can log warnings, so capture those like we do with loading errors
-					}
-
-					if (UWorldPartition* WorldPartition = World->GetWorldPartition())
-					{
-						// ForEachActorWithLoading may GC while running, so keep the world partition (and indirectly the world and its package) alive
-						TGCObjectScopeGuard<UWorldPartition> WorldPartitionGCGuard(WorldPartition);
-
-						FWorldPartitionHelpers::FForEachActorWithLoadingParams ForEachActorParams;
-						ForEachActorParams.ActorGuids = PackagePendingGather.ExternalActors.Array();
-
-						FWorldPartitionHelpers::ForEachActorWithLoading(WorldPartition, [&GatherableTextDataArray](const FWorldPartitionActorDescInstance* ActorDescInstance)
-							{
-								if (const AActor* Actor = ActorDescInstance->GetActor())
-								{
-									EPropertyLocalizationGathererResultFlags ActorGatherableTextResultFlags = EPropertyLocalizationGathererResultFlags::Empty;
-									FPropertyLocalizationDataGatherer(GatherableTextDataArray, Actor->GetExternalPackage(), ActorGatherableTextResultFlags);
-								}
-								return true;
-							}, ForEachActorParams);
-					}
-				}
+				UE::Private::GatherTextFromAssetsCommandlet::GatherTextFromPackage(Package, PackagePendingGather.ExternalActors, /*bGatherFromPrimaryPackage*/false, /*bGatherFromExternalActorPackages*/true, LogOutputRedirector, GatherableTextDataArray);
 			}
 
 			// This package may have already been cached in cases where we're reporting or fixing assets with a stale gather cache
@@ -1547,25 +1693,9 @@ bool UGatherTextFromAssetsCommandlet::LoadAndProcessUncachedPackages(TArray<FNam
 			}
 		}
 
-		if (HasExceededMemoryLimit(/*bLog*/true))
-		{
-			// First try a minimal purge to only remove things that are no longer referenced or needed by other packages pending gather
-			PurgeGarbage(/*bPurgeReferencedPackages*/false);
-
-			if (HasExceededMemoryLimit(/*bLog*/false))
-			{
-				// If we're still over the memory limit after a minimal purge, then attempt a full purge
-				PurgeGarbage(/*bPurgeReferencedPackages*/true);
-
-				// If we're still over the memory limit after both purges, then log a warning as we may be about to OOM
-				UE_CLOG(HasExceededMemoryLimit(/*bLog*/false), LogGatherTextFromAssetsCommandlet, Warning, TEXT("Flushing failed to reduce process memory to within the requested limits; this process may OOM!"));
-			}
-		}
+		UE::Private::GatherTextFromAssetsCommandlet::ConditionalPurgeGarbage<FPackagePendingGather>(MinFreeMemoryBytes, MaxUsedMemoryBytes, PackagesPendingGather);
 	}
 
-	// Clear list of packages with duplicate localization IDs
-	PackagesWithDuplicateLocalizationIds.Empty();
-
 	UE_LOG(LogGatherTextFromAssetsCommandlet, Display, TEXT("Loaded %d packages. %d failed."), NumPackagesProcessed, NumPackagesFailedLoading);
 	return true;
 }
@@ -1663,6 +1793,11 @@ int32 UGatherTextFromAssetsCommandlet::Main(const FString& Params)
 		return -1;
 	}
 
+	if (bStartWorkersImmediately)
+	{
+		StartWorkers(/*MinPackagesToUseWorkers*/0);
+	}
+
 	// If the editor has loaded a persistent world then create an empty world prior to starting the asset gather
 	// This avoids any issues when loading and initializing worlds during the gather, as WP needs to re-initialize the world
 	// Note: We can skip this when running embedded within a normal editor (ie, not a commandlet) as editor worlds are already fully initialized
@@ -1759,11 +1894,13 @@ int32 UGatherTextFromAssetsCommandlet::Main(const FString& Params)
 	const double PackageLoadingStartTime = FPlatformTime::Seconds();
 	UE_LOG(LogGatherTextFromAssetsCommandlet, Display, TEXT("Preparing to load %d packages..."), PackagesPendingGather.Num());
 
+	StartWorkers(UE::Private::GatherTextFromAssetsCommandlet::MinPackagesToUseWorkers);
+
 	CalculateDependenciesForPackagesPendingGather();
 
 	// Collect garbage before beginning to load packages
 	// This also sorts the list of packages into the best processing order
-	PurgeGarbage(/*bPurgeReferencedPackages*/false);
+	UE::Private::GatherTextFromAssetsCommandlet::PurgeGarbage<FPackagePendingGather>(&PackagesPendingGather);
 
 	// We don't need to have compiled shaders to gather text
 	bool bWasShaderCompilationEnabled = false;
@@ -1773,17 +1910,80 @@ int32 UGatherTextFromAssetsCommandlet::Main(const FString& Params)
 		GShaderCompilingManager->SkipShaderCompilation(true);
 	}
 
+	FGatherTextFromAssetsWorkerDirector& WorkerDirector = FGatherTextFromAssetsWorkerDirector::Get();
+
 	TArray<FName> PackagesWithStaleGatherCache;
-	if (!LoadAndProcessUncachedPackages(PackagesWithStaleGatherCache))
+	for (;;)
 	{
-		return -1;
+		// Process the main task work
+		if (PackagesPendingGather.Num() > 0)
+		{
+			if (!LoadAndProcessUncachedPackages(PackagesWithStaleGatherCache))
+			{
+				return -1;
+			}
+		}
+
+		WorkerDirector.TickWorkers();
+
+		// Ingest any results from workers
+		while (TOptional<FGatherTextFromAssetsWorkerMessage_PackageResult> PackageResult = WorkerDirector.IngestPackageResult())
+		{
+			if (!PackageResult->LoadLogCapture.IsEmpty())
+			{
+				GWarn->Log(NAME_None, ELogVerbosity::Display, PackageResult->LoadLogCapture);
+			}
+
+			if (PackageResult->bLoadError)
+			{
+				UE_LOGFMT(LogGatherTextFromAssetsCommandlet, Warning, "Failed to load package: '{package}'.",
+					("package", PackageResult->PackageName),
+					("id", GatherTextFromAssetsCommandlet::LocalizationLogIdentifier)
+				);
+			}
+			else
+			{
+				FMemoryReader MemoryReader(PackageResult->GatherableTextData);
+
+				TArray<FGatherableTextData> GatherableTextData;
+				MemoryReader << GatherableTextData;
+
+				UE_LOG(LogGatherTextFromAssetsCommandlet, Display, TEXT("Ingesting worker result for '%s'..."), *PackageResult->PackageName.ToString());
+				ProcessGatherableTextDataArray(GatherableTextData);
+			}
+		}
+
+		// Re-add any jobs sent to crashed workers to the main task work queue
+		{
+			const TArray<FName> PackagesFromCrashedWorkers = WorkerDirector.IngestPackagesFromCrashedWorkers();
+			for (const FName PackageName : PackagesFromCrashedWorkers)
+			{
+				if (FPackagePendingGather* PackagePendingGather = PackagesDistributedToWorkers.Find(PackageName))
+				{
+					PackagesPendingGather.Add(MoveTemp(*PackagePendingGather));
+					PackagesDistributedToWorkers.Remove(PackageName);
+				}
+			}
+		}
+
+		int32 NumPendingWorkerPackages = 0;
+		if (PackagesPendingGather.Num() == 0 && WorkerDirector.IsIdle(&NumPendingWorkerPackages))
+		{
+			break;
+		}
+
+		UE_LOG(LogGatherTextFromAssetsCommandlet, Display, TEXT("Waiting for workers to finish (%d pending package(s))..."), NumPendingWorkerPackages);
+		FPlatformProcess::SleepNoStats(1.0f);
 	}
 
+	// Clear list of packages with duplicate localization IDs
+	PackagesWithDuplicateLocalizationIds.Empty();
+
 	UE_LOG(LogGatherTextFromAssetsCommandlet, Display, TEXT("Loading packages took %.2f seconds."), FPlatformTime::Seconds() - PackageLoadingStartTime);
 
 	// Collect garbage after loading all packages
 	// This reclaims as much memory as possible for the rest of the gather pipeline
-	PurgeGarbage(/*bPurgeReferencedPackages*/true);
+	UE::Private::GatherTextFromAssetsCommandlet::PurgeGarbage<FPackagePendingGather>(nullptr);
 	
 	if (GShaderCompilingManager)
 	{
@@ -1982,6 +2182,39 @@ bool UGatherTextFromAssetsCommandlet::ConfigureFromScript(const FString& GatherT
 	ReadBoolFlagWithFallback(TEXT("FixMissingGatherCache"), bFixMissingGatherCache);
 	ReadBoolFlagWithFallback(TEXT("FixPackageLocalizationIdConflict"), bFixPackageLocalizationIdConflict);
 
+	{
+		bStartWorkersImmediately = FParse::Param(FCommandLine::Get(), TEXT("StartGatherTextWorkersImmediately"));
+
+		NumWorkers = 0;
+		FParse::Value(FCommandLine::Get(), TEXT("NumGatherTextWorkers="), NumWorkers);
+		NumWorkers = FMath::Max(0, NumWorkers);
+
+		if (NumWorkers > 0 && (bReportStaleGatherCache || bFixStaleGatherCache || bFixMissingGatherCache || bFixPackageLocalizationIdConflict))
+		{
+			UE_LOGFMT(LogGatherTextFromAssetsCommandlet, Warning, "Workers cannot be used when running in report/fix stale/missing gather cache mode. The request to use workers will be ignored.",
+				("id", GatherTextFromAssetsCommandlet::LocalizationLogIdentifier)
+			);
+			NumWorkers = 0;
+		}
+
+		// TODO: We should support this case by sending the ModulesToPreload to each worker when we start them
+		if (NumWorkers > 0 && ModulesToPreload.Num() > 0)
+		{
+			UE_LOGFMT(LogGatherTextFromAssetsCommandlet, Warning, "Workers cannot be used when module preloading is needed. The request to use workers will be ignored.",
+				("id", GatherTextFromAssetsCommandlet::LocalizationLogIdentifier)
+			);
+			NumWorkers = 0;
+		}
+
+		if (NumWorkers > 0 && !FParse::Param(FCommandLine::Get(), TEXT("Messaging")))
+		{
+			UE_LOGFMT(LogGatherTextFromAssetsCommandlet, Warning, "Workers cannot be used without MessageBus (did you forget to pass '-Messaging'?). The request to use workers will be ignored.",
+				("id", GatherTextFromAssetsCommandlet::LocalizationLogIdentifier)
+			);
+			NumWorkers = 0;
+		}
+	}
+
 	// Read some settings from the editor config
 	{
 		int32 MinFreeMemoryMB = 0;
@@ -2007,6 +2240,805 @@ bool UGatherTextFromAssetsCommandlet::ConfigureFromScript(const FString& GatherT
 	return !bHasFatalError;
 }
 
+void UGatherTextFromAssetsCommandlet::StartWorkers(const int32 MinPackagesToUseWorkers)
+{
+	FGatherTextFromAssetsWorkerDirector& WorkerDirector = FGatherTextFromAssetsWorkerDirector::Get();
+
+	int32 NumWorkersToUse = 0;
+	if (PackagesPendingGather.Num() >= MinPackagesToUseWorkers)
+	{
+		NumWorkersToUse = NumWorkers;
+	}
+	WorkerDirector.StartWorkers(NumWorkersToUse, /*bStopAdditionalWorkers*/false, /*NumRestartAttemptsIfCrashed*/4);
+}
+
+int32 UGatherTextFromAssetsWorkerCommandlet::Main(const FString& Params)
+{
+	FProcHandle DirectorProc;
+
+	// Parse config data
+	{
+		TArray<FString> Tokens;
+		TArray<FString> Switches;
+		TMap<FString, FString> ParamVals;
+		UCommandlet::ParseCommandLine(*Params, Tokens, Switches, ParamVals);
+
+		{
+			uint32 DirectorPid = 0;
+			if (const FString* DirectorPidStr = ParamVals.Find(TEXT("GatherDirectorPid")))
+			{
+				LexFromString(DirectorPid, **DirectorPidStr);
+			}
+			if (DirectorPid == 0)
+			{
+				UE_LOGFMT(LogGatherTextFromAssetsCommandlet, Warning, "No -GatherDirectorPid argument was provided. This process will continue to run even if the director process stops.",
+					("id", GatherTextFromAssetsCommandlet::LocalizationLogIdentifier)
+				);
+			}
+			else
+			{
+				DirectorProc = FPlatformProcess::OpenProcess(DirectorPid);
+			}
+			if (!DirectorProc.IsValid())
+			{
+				UE_LOGFMT(LogGatherTextFromAssetsCommandlet, Warning, "Failed to open a handle for the director process. This process will continue to run even if the director process stops.",
+					("id", GatherTextFromAssetsCommandlet::LocalizationLogIdentifier)
+				);
+			}
+		}
+
+		if (const FString* WorkerIdStr = ParamVals.Find(TEXT("GatherWorkerId")))
+		{
+			FGuid::Parse(*WorkerIdStr, WorkerId);
+		}
+		if (!WorkerId.IsValid())
+		{
+			UE_LOGFMT(LogGatherTextFromAssetsCommandlet, Error, "An invalid -GatherWorkerId argument was provided. Aborting.",
+				("id", GatherTextFromAssetsCommandlet::LocalizationLogIdentifier)
+			);
+			return -1;
+		}
+
+		{
+			int32 MinFreeMemoryMB = 0;
+			if (!GConfig->GetInt(TEXT("GatherTextFromAssetsWorker"), TEXT("MinFreeMemory"), MinFreeMemoryMB, GEditorIni))
+			{
+				GConfig->GetInt(TEXT("GatherTextFromAssets"), TEXT("MinFreeMemory"), MinFreeMemoryMB, GEditorIni);
+			}
+			MinFreeMemoryMB = FMath::Max(MinFreeMemoryMB, 0);
+			MinFreeMemoryBytes = MinFreeMemoryMB * 1024LL * 1024LL;
+		}
+
+		{
+			int32 MaxUsedMemoryMB = 0;
+			if (!GConfig->GetInt(TEXT("GatherTextFromAssetsWorker"), TEXT("MaxUsedMemory"), MaxUsedMemoryMB, GEditorIni))
+			{
+				GConfig->GetInt(TEXT("GatherTextFromAssets"), TEXT("MaxUsedMemory"), MaxUsedMemoryMB, GEditorIni);
+			}
+			MaxUsedMemoryMB = FMath::Max(MaxUsedMemoryMB, 0);
+			MaxUsedMemoryBytes = MaxUsedMemoryMB * 1024LL * 1024LL;
+		}
+	}
+
+	// Set-up MessageBus
+	{
+		if (!FParse::Param(FCommandLine::Get(), TEXT("Messaging")))
+		{
+			UE_LOGFMT(LogGatherTextFromAssetsCommandlet, Error, "MessageBus is not enabled for the worker process (did you forget to pass '-Messaging'?). Aborting.",
+				("id", GatherTextFromAssetsCommandlet::LocalizationLogIdentifier)
+			);
+			return -1;
+		}
+
+		MessageEndpoint = FMessageEndpoint::Builder(*(TEXT("GatherTextFromAssetsWorker") + WorkerId.ToString()))
+			.ReceivingOnAnyThread()
+			.Handling<FGatherTextFromAssetsWorkerMessage_Ping>(this, &UGatherTextFromAssetsWorkerCommandlet::HandlePingMessage)
+			.Handling<FGatherTextFromAssetsWorkerMessage_PackageRequest>(this, &UGatherTextFromAssetsWorkerCommandlet::HandlePackageRequestMessage);
+		if (!MessageEndpoint)
+		{
+			UE_LOGFMT(LogGatherTextFromAssetsCommandlet, Error, "Failed to create MessageEndpoint for the worker process. Aborting.",
+				("id", GatherTextFromAssetsCommandlet::LocalizationLogIdentifier)
+			);
+			return -1;
+		}
+		MessageEndpoint->Subscribe<FGatherTextFromAssetsWorkerMessage_Ping>();
+	}
+
+	// Run a GC now that we're ready to start accepting work, to clean-up anything loaded during start-up that is no longer needed
+	UE::Private::GatherTextFromAssetsCommandlet::PurgeGarbage<FPackagePendingGather>(nullptr);
+
+	// Wait for work, and process each package requested
+	for (;;)
+	{
+		TOptional<FPackagePendingGather> PackagePendingGather;
+		{
+			UE::TScopeLock _(PackagesPendingGatherMutex);
+			if (PackagesPendingGather.Num() > 0)
+			{
+				PackagePendingGather = PackagesPendingGather.Pop(EAllowShrinking::No);
+			}
+		}
+
+		if (PackagePendingGather)
+		{
+			{
+				UE::TScopeLock _(IdleStartTimeUtcMutex);
+				IdleStartTimeUtc.Reset();
+			}
+
+			const FString PackageNameStr = PackagePendingGather->PackageName.ToString();
+			UE_LOG(LogGatherTextFromAssetsCommandlet, Display, TEXT("Gathering package: '%s'."), *PackageNameStr);
+
+			FStringOutputDevice LogOutputTargetDevice;
+			FLoadPackageLogOutputRedirector LogOutputRedirector(&LogOutputTargetDevice);
+
+			FGatherTextFromAssetsWorkerMessage_PackageResult* PackageResultMessage = FMessageEndpoint::MakeMessage<FGatherTextFromAssetsWorkerMessage_PackageResult>();
+			PackageResultMessage->WorkerId = WorkerId;
+			PackageResultMessage->PackageName = PackagePendingGather->PackageName;
+
+			if (UPackage* Package = UE::Private::GatherTextFromAssetsCommandlet::LoadPackageToGather(PackageNameStr, PackagePendingGather->ExternalActors.Num() > 0, LogOutputRedirector))
+			{
+				TArray<FGatherableTextData> GatherableTextDataArray;
+				UE::Private::GatherTextFromAssetsCommandlet::GatherTextFromPackage(Package, PackagePendingGather->ExternalActors, /*bGatherFromPrimaryPackage*/true, /*bGatherFromExternalActorPackages*/true, LogOutputRedirector, GatherableTextDataArray);
+
+				FMemoryWriter MemoryWriter(PackageResultMessage->GatherableTextData);
+				MemoryWriter << GatherableTextDataArray;
+			}
+			else
+			{
+				PackageResultMessage->bLoadError = true;
+			}
+			PackageResultMessage->LoadLogCapture = LogOutputTargetDevice;
+
+			MessageEndpoint->Send(
+				PackageResultMessage, // Should be deleted by MessageBus
+				EMessageFlags::Reliable,
+				nullptr, // No Attachment
+				{ PackagePendingGather->RequestorAddress },
+				FTimespan::Zero(), // No Delay
+				FDateTime::MaxValue() // No Expiration
+				);
+
+			UE_LOG(LogGatherTextFromAssetsCommandlet, Display, TEXT("Sent package result for '%s'."), *PackageNameStr);
+
+			UE::Private::GatherTextFromAssetsCommandlet::TickBackgroundTasks();
+
+			UE::Private::GatherTextFromAssetsCommandlet::ConditionalPurgeGarbage<FPackagePendingGather>(MinFreeMemoryBytes, MaxUsedMemoryBytes, PackagesPendingGather);
+		}
+		else
+		{
+			{
+				UE::TScopeLock _(IdleStartTimeUtcMutex);
+				if (!IdleStartTimeUtc)
+				{
+					IdleStartTimeUtc = FDateTime::UtcNow();
+				}
+			}
+			FPlatformProcess::SleepNoStats(1.0f);
+		}
+
+		UE::Private::GatherTextFromAssetsCommandlet::TickMessageBusGT();
+
+		// If our director process has stopped, then we can also exit
+		if (DirectorProc.IsValid() && !FPlatformProcess::IsProcRunning(DirectorProc))
+		{
+			UE_LOG(LogGatherTextFromAssetsCommandlet, Display, TEXT("Director process has stopped. Exiting worker..."));
+			break;
+		}
+	}
+
+	if (DirectorProc.IsValid())
+	{
+		FPlatformProcess::CloseProc(DirectorProc);
+	}
+
+	// Disable the Endpoint message handling since the message could keep it alive a bit.
+	if (MessageEndpoint)
+	{
+		MessageEndpoint->Disable();
+		MessageEndpoint.Reset();
+	}
+
+	return 0;
+}
+
+void UGatherTextFromAssetsWorkerCommandlet::HandlePingMessage(const FGatherTextFromAssetsWorkerMessage_Ping& Message, const TSharedRef<IMessageContext>& Context)
+{
+	if (Message.ProtocolVersion == UE::Private::GatherTextFromAssetsCommandlet::WorkerProtocolVersion)
+	{
+		UE_LOG(LogGatherTextFromAssetsCommandlet, Display, TEXT("Received ping and responded with a pong."));
+
+		FGatherTextFromAssetsWorkerMessage_Pong* PongMessage = FMessageEndpoint::MakeMessage<FGatherTextFromAssetsWorkerMessage_Pong>();
+		PongMessage->WorkerId = WorkerId;
+		{
+			UE::TScopeLock _(IdleStartTimeUtcMutex);
+			PongMessage->IdleStartTimeUtc = IdleStartTimeUtc;
+		}
+
+		MessageEndpoint->Send(
+			PongMessage, // Should be deleted by MessageBus
+			EMessageFlags::None,
+			nullptr, // No Attachment
+			{ Context->GetSender() },
+			FTimespan::Zero(), // No Delay
+			FDateTime::MaxValue() // No Expiration
+			);
+	}
+	else
+	{
+		UE_LOG(LogGatherTextFromAssetsCommandlet, Display, TEXT("Received ping with the wrong protocol version (received %d, expected %d)."), Message.ProtocolVersion, UE::Private::GatherTextFromAssetsCommandlet::WorkerProtocolVersion);
+	}
+}
+
+void UGatherTextFromAssetsWorkerCommandlet::HandlePackageRequestMessage(const FGatherTextFromAssetsWorkerMessage_PackageRequest& Message, const TSharedRef<IMessageContext>& Context)
+{
+	UE_LOG(LogGatherTextFromAssetsCommandlet, Display, TEXT("Received a package request for '%s'."), *Message.PackageName.ToString());
+	
+	// Note: This isn't an optimal way to insert into the array, but the array is ordered so it can pop from the end, so inserting at element 0 keeps it FIFO
+	{
+		UE::TScopeLock _(PackagesPendingGatherMutex);
+		PackagesPendingGather.Insert(FPackagePendingGather{ Message.PackageName, Message.Dependencies, Message.ExternalActors, Context->GetSender() }, 0);
+	}
+}
+
+FGatherTextFromAssetsWorkerDirector::~FGatherTextFromAssetsWorkerDirector()
+{
+	checkf(CurrentWorkers.IsEmpty(), TEXT("FGatherTextFromAssetsWorkerDirector still had workers active during destruction. Did you forget to call StopWorkers?"));
+}
+
+FGatherTextFromAssetsWorkerDirector& FGatherTextFromAssetsWorkerDirector::Get()
+{
+	static FGatherTextFromAssetsWorkerDirector Instance;
+	return Instance;
+}
+
+bool FGatherTextFromAssetsWorkerDirector::StartWorkers(const int32 NumWorkers, const bool bStopAdditionalWorkers, const int32 NumRestartAttemptsIfCrashed)
+{
+	// If no workers were requested, then this is the same as calling StopWorkers
+	if (NumWorkers <= 0 && bStopAdditionalWorkers)
+	{
+		StopWorkers();
+		return false;
+	}
+
+	// Start any extra worker processes that are needed
+	// Stop any worker processes that are no longer needed
+	{
+		UE::TScopeLock _(CurrentWorkersMutex);
+		if (NumWorkers > CurrentWorkers.Num())
+		{
+			const int32 NumWorkersToCreate = NumWorkers - CurrentWorkers.Num();
+			for (int32 WorkerIndex = 0; WorkerIndex < NumWorkersToCreate; ++WorkerIndex)
+			{
+				const FGuid WorkerId = FGuid::NewGuid();
+				FString WorkerCommandLine = GenerateWorkerCommandLine(WorkerId);
+
+				UE_LOG(LogGatherTextFromAssetsCommandlet, Display, TEXT("Starting worker process '%s'..."), *WorkerId.ToString());
+
+				FProcHandle WorkerProc = StartWorker(WorkerCommandLine);
+				if (WorkerProc.IsValid())
+				{
+					TSharedPtr<FWorkerInfo> WorkerInfo = MakeShared<FWorkerInfo>();
+					WorkerInfo->WorkerProc = WorkerProc;
+					WorkerInfo->WorkerCommandLine = MoveTemp(WorkerCommandLine);
+					CurrentWorkers.Add(WorkerId, MoveTemp(WorkerInfo));
+				}
+				else
+				{
+					UE_LOG(LogGatherTextFromAssetsCommandlet, Warning, TEXT("Worker process '%s' failed to start. This worker request will be ignored."), *WorkerId.ToString());
+				}
+			}
+		}
+		else if (bStopAdditionalWorkers && NumWorkers < CurrentWorkers.Num())
+		{
+			int32 WorkerCount = 0;
+			for (auto It = CurrentWorkers.CreateIterator(); It; ++It)
+			{
+				const TTuple<FGuid, TSharedPtr<FWorkerInfo>>& CurrentWorkerPair = *It;
+				if (++WorkerCount > NumWorkers)
+				{
+					UE_LOG(LogGatherTextFromAssetsCommandlet, Warning, TEXT("Stopping worker process '%s' as it is no longer needed."), *CurrentWorkerPair.Key.ToString());
+					FPlatformProcess::TerminateProc(CurrentWorkerPair.Value->WorkerProc);
+					FPlatformProcess::CloseProc(CurrentWorkerPair.Value->WorkerProc);
+					It.RemoveCurrent();
+					continue;
+				}
+			}
+		}
+	}
+
+	bool bHasWorkers = false;
+
+	// Any existing workers may have been idle for a while, so reset their discovery and timeout, so that we have to rediscover them rather than have them immediately timeout
+	{
+		UE::TScopeLock _(CurrentWorkersMutex);
+		for (TTuple<FGuid, TSharedPtr<FWorkerInfo>> CurrentWorkerPair : CurrentWorkers)
+		{
+			CurrentWorkerPair.Value->NumRestartAttemptsIfCrashed = NumRestartAttemptsIfCrashed;
+			ResetWorkerDiscoveryAndTimeout(*CurrentWorkerPair.Value);
+			bHasWorkers = true;
+		}
+	}
+
+	if (bHasWorkers)
+	{
+		// Ensure we have a MessageBus endpoint to communicate via
+		if (!MessageEndpoint)
+		{
+			MessageEndpoint = FMessageEndpoint::Builder("GatherTextFromAssetsWorkerDirector")
+				.ReceivingOnAnyThread()
+				.Handling<FGatherTextFromAssetsWorkerMessage_Pong>(this, &FGatherTextFromAssetsWorkerDirector::HandlePongMessage)
+				.Handling<FGatherTextFromAssetsWorkerMessage_PackageResult>(this, &FGatherTextFromAssetsWorkerDirector::HandlePackageResultMessage);
+		}
+
+		// Force broadcast a ping to start discovery
+		// Newly started workers likely won't respond to this, so they will need to wait for the next TickWorkers (unless using WaitForWorkersToStart)
+		BroadcastPingMessage(/*bIgnoreDelay*/true);
+		UE::Private::GatherTextFromAssetsCommandlet::TickMessageBusGT();
+	}
+
+	return bHasWorkers;
+}
+
+bool FGatherTextFromAssetsWorkerDirector::WaitForWorkersToStart(const TOptional<FTimespan> Timeout)
+{
+	// Wait for the worker processes to reach a state where they are ready to accept work
+	const FDateTime WaitStartUtc = FDateTime::UtcNow();
+	for (;;)
+	{
+		const FDateTime UtcNow = FDateTime::UtcNow();
+		const FTimespan WaitDuration = UtcNow - WaitStartUtc;
+
+		int32 NumPendingWorkers = 0;
+		{
+			int32 NumReadyWorkers = 0;
+
+			UE::TScopeLock _(CurrentWorkersMutex);
+			for (auto It = CurrentWorkers.CreateIterator(); It; ++It)
+			{
+				const TTuple<FGuid, TSharedPtr<FWorkerInfo>>& CurrentWorkerPair = *It;
+				
+				// Has this process crashed (or quit) during start-up?
+				if (!FPlatformProcess::IsProcRunning(CurrentWorkerPair.Value->WorkerProc))
+				{
+					FPlatformProcess::CloseProc(CurrentWorkerPair.Value->WorkerProc);
+					UE_LOG(LogGatherTextFromAssetsCommandlet, Warning, TEXT("Worker process '%s' crashed during start-up. This worker request will be ignored."), *CurrentWorkerPair.Key.ToString());
+					It.RemoveCurrent();
+					continue;
+				}
+
+				if (CurrentWorkerPair.Value->EndpointAddress.IsValid())
+				{
+					++NumReadyWorkers;
+				}
+				else
+				{
+					++NumPendingWorkers;
+				}
+			}
+
+			if (NumReadyWorkers == CurrentWorkers.Num())
+			{
+				return true;
+			}
+
+			if (Timeout && WaitDuration >= *Timeout)
+			{
+				UE_LOG(LogGatherTextFromAssetsCommandlet, Warning, TEXT("Timed-out waiting for %d worker(s) to be ready (duration %s)."), NumPendingWorkers, *WaitDuration.ToString());
+				return false;
+			}
+		}
+
+		UE_LOG(LogGatherTextFromAssetsCommandlet, Display, TEXT("Waiting for %d worker(s) to be ready (duration %s)..."), NumPendingWorkers, *WaitDuration.ToString());
+
+		// Force broadcast a ping request and sleep
+		BroadcastPingMessage(/*bIgnoreDelay*/true);
+		UE::Private::GatherTextFromAssetsCommandlet::TickMessageBusGT();
+		FPlatformProcess::SleepNoStats(1.0f);
+	}
+}
+
+bool FGatherTextFromAssetsWorkerDirector::StopWorkers()
+{
+	{
+		UE::TScopeLock _(CurrentWorkersMutex);
+		for (const TTuple<FGuid, TSharedPtr<FWorkerInfo>>& CurrentWorkerPair : CurrentWorkers)
+		{
+			FPlatformProcess::TerminateProc(CurrentWorkerPair.Value->WorkerProc);
+			FPlatformProcess::CloseProc(CurrentWorkerPair.Value->WorkerProc);
+		}
+		CurrentWorkers.Reset();
+	}
+
+	// Disable the Endpoint message handling since the message could keep it alive a bit.
+	if (MessageEndpoint)
+	{
+		MessageEndpoint->Disable();
+		MessageEndpoint.Reset();
+	}
+
+	while (PackageResults.Dequeue()) {}
+	PackagesFromCrashedWorkers.Reset();
+
+	return true;
+}
+
+void FGatherTextFromAssetsWorkerDirector::TickWorkers()
+{
+	if (MessageEndpoint)
+	{
+		BroadcastPingMessage();
+		UE::Private::GatherTextFromAssetsCommandlet::TickMessageBusGT();
+	}
+
+	{
+		const FDateTime UtcNow = FDateTime::UtcNow();
+
+		UE::TScopeLock _(CurrentWorkersMutex);
+		for (auto It = CurrentWorkers.CreateIterator(); It; ++It)
+		{
+			const TTuple<FGuid, TSharedPtr<FWorkerInfo>>& CurrentWorkerPair = *It;
+
+			// Has this process crashed?
+			if (!FPlatformProcess::IsProcRunning(CurrentWorkerPair.Value->WorkerProc))
+			{
+				FPlatformProcess::CloseProc(CurrentWorkerPair.Value->WorkerProc);
+				if (!HandleWorkerCrashed(CurrentWorkerPair.Key, *CurrentWorkerPair.Value))
+				{
+					It.RemoveCurrent();
+				}
+				continue;
+			}
+
+			// Is this worker still starting?
+			// If so, we can skip the timeout detection as we won't have that data yet
+			if (!CurrentWorkerPair.Value->EndpointAddress.IsValid())
+			{
+				continue;
+			}
+
+			// Has this process likely hung? (not received any MessageBus messages for over 5 minutes)
+			if (CurrentWorkerPair.Value->LastMessageReceivedUtc && *CurrentWorkerPair.Value->LastMessageReceivedUtc + FTimespan(0, 5, 0) <= UtcNow)
+			{
+				FPlatformProcess::TerminateProc(CurrentWorkerPair.Value->WorkerProc);
+				FPlatformProcess::CloseProc(CurrentWorkerPair.Value->WorkerProc);
+				if (!HandleWorkerCrashed(CurrentWorkerPair.Key, *CurrentWorkerPair.Value))
+				{
+					It.RemoveCurrent();
+				}
+				continue;
+			}
+
+			// Were we asked to re-send any pending package requests? (eg, after rediscovering a previously crashed worker)
+			if (CurrentWorkerPair.Value->bResendPendingPackageRequests)
+			{
+				if (CurrentWorkerPair.Value->PendingPackageRequests.Num() > 0)
+				{
+					UE_LOG(LogGatherTextFromAssetsCommandlet, Display, TEXT("Worker process '%s' has packages assigned to it, and has been rediscovered. Re-sending %d pending package request(s)..."), *CurrentWorkerPair.Key.ToString(), CurrentWorkerPair.Value->PendingPackageRequests.Num());
+					ResendWorkerPendingPackageRequests(*CurrentWorkerPair.Value);
+				}
+				CurrentWorkerPair.Value->bResendPendingPackageRequests = false;
+			}
+
+			// Has this worker has been idle for over 2 minutes, and last had work assigned to it over 2 minutes ago?
+			// If so, it's likely that either the request or response was lost, so try sending any pending requests again...
+			if (CurrentWorkerPair.Value->PendingPackageRequests.Num() > 0 &&
+				CurrentWorkerPair.Value->LastPackageRequestUtc && CurrentWorkerPair.Value->IdleStartTimeUtc &&
+				*CurrentWorkerPair.Value->LastPackageRequestUtc + FTimespan(0, 2, 0) <= UtcNow && *CurrentWorkerPair.Value->IdleStartTimeUtc + FTimespan(0, 2, 0) <= UtcNow)
+			{
+				UE_LOG(LogGatherTextFromAssetsCommandlet, Display, TEXT("Worker process '%s' has packages assigned to it, but has been idle for a while. Re-sending %d pending package request(s)..."), *CurrentWorkerPair.Key.ToString(), CurrentWorkerPair.Value->PendingPackageRequests.Num());
+				ResendWorkerPendingPackageRequests(*CurrentWorkerPair.Value);
+			}
+		}
+	}
+}
+
+TArray<FGuid> FGatherTextFromAssetsWorkerDirector::GetAvailableWorkerIds(const TOptional<int32> IdleThreshold) const
+{
+	TArray<FGuid> WorkerIds;
+	{
+		UE::TScopeLock _(CurrentWorkersMutex);
+		for (const TTuple<FGuid, TSharedPtr<FWorkerInfo>>& CurrentWorkerPair : CurrentWorkers)
+		{
+			if (CurrentWorkerPair.Value->EndpointAddress.IsValid() && (!IdleThreshold.IsSet() || CurrentWorkerPair.Value->PendingPackageRequests.Num() <= IdleThreshold.GetValue()))
+			{
+				WorkerIds.Add(CurrentWorkerPair.Key);
+			}
+		}
+	}
+	return WorkerIds;
+}
+
+bool FGatherTextFromAssetsWorkerDirector::IsIdle(int32* OutNumPendingWorkerPackages) const
+{
+	int32 NumPendingWorkerPackages = 0;
+	{
+		UE::TScopeLock _(CurrentWorkersMutex);
+		for (const TTuple<FGuid, TSharedPtr<FWorkerInfo>>& CurrentWorkerPair : CurrentWorkers)
+		{
+			NumPendingWorkerPackages += CurrentWorkerPair.Value->PendingPackageRequests.Num();
+		}
+	}
+
+	if (OutNumPendingWorkerPackages)
+	{
+		*OutNumPendingWorkerPackages = NumPendingWorkerPackages;
+	}
+
+	return NumPendingWorkerPackages == 0
+		&& PackageResults.IsEmpty()
+		&& PackagesFromCrashedWorkers.IsEmpty();
+}
+
+bool FGatherTextFromAssetsWorkerDirector::AssignPackageToWorker(const FGuid& WorkerId, const FGatherTextFromAssetsWorkerMessage_PackageRequest& PackageRequest)
+{
+	FMessageAddress WorkerAddress;
+	{
+		UE::TScopeLock _(CurrentWorkersMutex);
+		if (TSharedPtr<FWorkerInfo> CurrentWorker = CurrentWorkers.FindRef(WorkerId))
+		{
+			WorkerAddress = CurrentWorker->EndpointAddress;
+			CurrentWorker->LastPackageRequestUtc = FDateTime::UtcNow();
+			CurrentWorker->PendingPackageRequests.Add(PackageRequest.PackageName, PackageRequest);
+		}
+		else
+		{
+			return false;
+		}
+	}
+
+	if (WorkerAddress.IsValid())
+	{
+		FGatherTextFromAssetsWorkerMessage_PackageRequest* PackageRequestMessage = FMessageEndpoint::MakeMessage<FGatherTextFromAssetsWorkerMessage_PackageRequest>(PackageRequest);
+
+		MessageEndpoint->Send(
+			PackageRequestMessage, // Should be deleted by MessageBus
+			EMessageFlags::Reliable,
+			nullptr, // No Attachment
+			{ WorkerAddress },
+			FTimespan::Zero(), // No Delay
+			FDateTime::MaxValue() // No Expiration
+			);
+	}
+
+	UE_LOG(LogGatherTextFromAssetsCommandlet, Display, TEXT("Sent package request to worker '%s' for '%s'."), *WorkerId.ToString(), *PackageRequest.PackageName.ToString());
+
+	return true;
+}
+
+TOptional<FGatherTextFromAssetsWorkerMessage_PackageResult> FGatherTextFromAssetsWorkerDirector::IngestPackageResult()
+{
+	return PackageResults.Dequeue();
+}
+
+TArray<FName> FGatherTextFromAssetsWorkerDirector::IngestPackagesFromCrashedWorkers()
+{
+	TArray<FName> TmpPackagesFromCrashedWorkers = MoveTemp(PackagesFromCrashedWorkers);
+	PackagesFromCrashedWorkers.Reset();
+
+	// Merge in any packages from workers that haven't started yet, as the main process may want to redistribute those or handle them itself
+	{
+		UE::TScopeLock _(CurrentWorkersMutex);
+		for (const TTuple<FGuid, TSharedPtr<FWorkerInfo>>& CurrentWorkerPair : CurrentWorkers)
+		{
+			if (!CurrentWorkerPair.Value->EndpointAddress.IsValid())
+			{
+				for (const TTuple<FName, FGatherTextFromAssetsWorkerMessage_PackageRequest>& PendingPackageRequestPair : CurrentWorkerPair.Value->PendingPackageRequests)
+				{
+					TmpPackagesFromCrashedWorkers.Add(PendingPackageRequestPair.Key);
+				}
+				CurrentWorkerPair.Value->PendingPackageRequests.Reset();
+			}
+		}
+	}
+
+	return TmpPackagesFromCrashedWorkers;
+}
+
+FString FGatherTextFromAssetsWorkerDirector::GenerateWorkerCommandLine(const FGuid& WorkerId)
+{
+	// Build the set of tokens to keep from the current command line args
+	TArray<FString> Tokens;
+	{
+		const TCHAR* ProjectName = FApp::GetProjectName();
+		const TCHAR* CommandLine = FCommandLine::Get();
+
+		FString Token;
+		while (FParse::Token(CommandLine, Token, /*bUseEscape*/false))
+		{
+			if (Token.IsEmpty())
+			{
+				continue;
+			}
+			if (Token == ProjectName || Token.EndsWith(TEXT(".uproject")) || 
+				Token.StartsWith(TEXT("-run=")) ||
+				Token.StartsWith(TEXT("-config=")) ||
+				Token.StartsWith(TEXT("-configlist=")) ||
+				Token.StartsWith(TEXT("-abslog=")) ||
+				Token == TEXT("-unattended") ||
+				Token == TEXT("-messaging") ||
+				Token.StartsWith(TEXT("-NumGatherTextWorkers=")) ||
+				Token == TEXT("-StartGatherTextWorkersImmediately")
+				)
+			{
+				continue;
+			}
+			Tokens.Add(MoveTemp(Token));
+		}
+	}
+	Tokens.Add(TEXT("-NoAssetRegistryCache")); // Disable the AR cache to avoid IO conflicts with the main process; the worker doesn't directly use the AR
+	Tokens.Add(TEXT("-unattended"));
+	Tokens.Add(TEXT("-messaging"));
+	Tokens.Add(TEXT("-multiprocess"));
+	Tokens.Add(TEXT("-GatherWorkerId=") + WorkerId.ToString());
+	Tokens.Add(FString::Printf(TEXT("-GatherDirectorPid=%u"), FPlatformProcess::GetCurrentProcessId()));
+
+	// We are joining the tokens back into a commandline string; wrap tokens with whitespace in quotes
+	for (FString& Token : Tokens)
+	{
+		static constexpr FAsciiSet WhitespaceSet(" \r\n");
+		if (FAsciiSet::HasAny(Token, WhitespaceSet))
+		{
+			if (int32 IndexOfQuote = INDEX_NONE;
+				!Token.FindChar(TEXT('\"'), IndexOfQuote))
+			{
+				Token.InsertAt(0, TEXT('\"'));
+				Token += TEXT('\"');
+			}
+		}
+	}
+
+	const FString ProjectFilePath = FApp::HasProjectName() ? FString::Printf(TEXT("\"%s\""), *FPaths::ConvertRelativePathToFull(FPaths::GetProjectFilePath())) : FString();
+	return CommandletHelpers::BuildCommandletProcessArguments(TEXT("GatherTextFromAssetsWorker"), *ProjectFilePath, *FString::Join(Tokens, TEXT(" ")));
+}
+
+FProcHandle FGatherTextFromAssetsWorkerDirector::StartWorker(const FString& WorkerCommandLine)
+{
+	constexpr bool bLaunchDetached = true;
+	constexpr bool bLaunchHidden = true;
+	constexpr bool bLaunchReallyHidden = true;
+	return FPlatformProcess::CreateProc(*FUnrealEdMisc::Get().GetExecutableForCommandlets(), *WorkerCommandLine, bLaunchDetached, bLaunchHidden, bLaunchReallyHidden, nullptr, 0, nullptr, nullptr);
+}
+
+bool FGatherTextFromAssetsWorkerDirector::HandleWorkerCrashed(const FGuid& WorkerId, FWorkerInfo& WorkerInfo)
+{
+	if (WorkerInfo.NumRestartAttemptsIfCrashed > 0)
+	{
+		UE_LOG(LogGatherTextFromAssetsCommandlet, Display, TEXT("Restarting worker process '%s'..."), *WorkerId.ToString());
+		--WorkerInfo.NumRestartAttemptsIfCrashed;
+
+		// Restart the process and wait for it to start again
+		// Any pending work will be resent once the worker is available
+		WorkerInfo.WorkerProc = StartWorker(WorkerInfo.WorkerCommandLine);
+		if (WorkerInfo.WorkerProc.IsValid())
+		{
+			ResetWorkerDiscoveryAndTimeout(WorkerInfo);
+			return true;
+		}
+		else
+		{
+			UE_LOG(LogGatherTextFromAssetsCommandlet, Warning, TEXT("Worker process '%s' failed to start. Its work will be returned to the main pool."), *WorkerId.ToString());
+		}
+	}
+	else
+	{
+		UE_LOG(LogGatherTextFromAssetsCommandlet, Warning, TEXT("Worker process '%s' crashed or hung. Its work will be returned to the main pool."), *WorkerId.ToString());
+	}
+
+	// Failed to restart, or restart is not allowed
+	// Move the pending work to the crashed queue
+	for (const TTuple<FName, FGatherTextFromAssetsWorkerMessage_PackageRequest>& PendingPackageRequestPair : WorkerInfo.PendingPackageRequests)
+	{
+		PackagesFromCrashedWorkers.Add(PendingPackageRequestPair.Key);
+	}
+	WorkerInfo.PendingPackageRequests.Reset();
+
+	return false;
+}
+
+void FGatherTextFromAssetsWorkerDirector::ResetWorkerDiscoveryAndTimeout(FWorkerInfo& WorkerInfo)
+{
+	WorkerInfo.EndpointAddress = FMessageAddress();
+	WorkerInfo.LastMessageReceivedUtc.Reset();
+	WorkerInfo.LastPackageRequestUtc.Reset();
+	WorkerInfo.IdleStartTimeUtc.Reset();
+}
+
+void FGatherTextFromAssetsWorkerDirector::ResendWorkerPendingPackageRequests(FWorkerInfo& WorkerInfo)
+{
+	WorkerInfo.LastPackageRequestUtc = FDateTime::UtcNow();
+	WorkerInfo.IdleStartTimeUtc.Reset(); // Reset our cached idle state so that we don't re-trigger the re-send until the worker confirms it's re-entered an idle state
+
+	for (const TTuple<FName, FGatherTextFromAssetsWorkerMessage_PackageRequest>& PendingPackageRequestPair : WorkerInfo.PendingPackageRequests)
+	{
+		FGatherTextFromAssetsWorkerMessage_PackageRequest* PackageRequestMessage = FMessageEndpoint::MakeMessage<FGatherTextFromAssetsWorkerMessage_PackageRequest>(PendingPackageRequestPair.Value);
+
+		MessageEndpoint->Send(
+			PackageRequestMessage, // Should be deleted by MessageBus
+			EMessageFlags::Reliable,
+			nullptr, // No Attachment
+			{ WorkerInfo.EndpointAddress },
+			FTimespan::Zero(), // No Delay
+			FDateTime::MaxValue() // No Expiration
+			);
+	}
+}
+
+void FGatherTextFromAssetsWorkerDirector::BroadcastPingMessage(const bool bIgnoreDelay)
+{
+	const FDateTime UtcNow = FDateTime::UtcNow();
+	if (bIgnoreDelay || (LastPingBroadcastUtc + FTimespan(0, 0, 30) <= UtcNow))
+	{
+		LastPingBroadcastUtc = UtcNow;
+
+		FGatherTextFromAssetsWorkerMessage_Ping* PingMessage = FMessageEndpoint::MakeMessage<FGatherTextFromAssetsWorkerMessage_Ping>();
+		PingMessage->ProtocolVersion = UE::Private::GatherTextFromAssetsCommandlet::WorkerProtocolVersion;
+
+		MessageEndpoint->Publish(
+			PingMessage,
+			EMessageScope::Network,
+			FTimespan::Zero(), // No Delay
+			FDateTime::MaxValue() // No Expiration
+			);
+	}
+}
+
+void FGatherTextFromAssetsWorkerDirector::HandlePongMessage(const FGatherTextFromAssetsWorkerMessage_Pong& Message, const TSharedRef<IMessageContext>& Context)
+{
+	UE::TScopeLock _(CurrentWorkersMutex);
+	if (TSharedPtr<FWorkerInfo> CurrentWorker = CurrentWorkers.FindRef(Message.WorkerId))
+	{
+		const FDateTime UtcNow = FDateTime::UtcNow();
+		if (Message.IdleStartTimeUtc)
+		{
+			const FTimespan IdleDuration = UtcNow - *Message.IdleStartTimeUtc;
+			UE_LOG(LogGatherTextFromAssetsCommandlet, Display, TEXT("Received pong from worker '%s' with endpoint '%s' (Idle for %s)."), *Message.WorkerId.ToString(), *Context->GetSender().ToString(), *IdleDuration.ToString());
+		}
+		else
+		{
+			UE_LOG(LogGatherTextFromAssetsCommandlet, Display, TEXT("Received pong from worker '%s' with endpoint '%s' (Working)."), *Message.WorkerId.ToString(), *Context->GetSender().ToString());
+		}
+
+		if (!CurrentWorker->EndpointAddress.IsValid() && CurrentWorker->PendingPackageRequests.Num() > 0)
+		{
+			// If we (re)discovered this worker, resend any pending requests as they may have been lost (eg, if restarting the process after a crash)
+			CurrentWorker->bResendPendingPackageRequests = true;
+		}
+		CurrentWorker->EndpointAddress = Context->GetSender();
+		CurrentWorker->LastMessageReceivedUtc = UtcNow;
+		CurrentWorker->IdleStartTimeUtc = Message.IdleStartTimeUtc;
+	}
+}
+
+void FGatherTextFromAssetsWorkerDirector::HandlePackageResultMessage(const FGatherTextFromAssetsWorkerMessage_PackageResult& Message, const TSharedRef<IMessageContext>& Context)
+{
+	bool bWasSolicitedResult = false;
+	{
+		UE::TScopeLock _(CurrentWorkersMutex);
+		if (TSharedPtr<FWorkerInfo> CurrentWorker = CurrentWorkers.FindRef(Message.WorkerId))
+		{
+			CurrentWorker->LastMessageReceivedUtc = FDateTime::UtcNow();
+			CurrentWorker->IdleStartTimeUtc.Reset(); // Don't consider it idle if we're still receiving results from it
+
+			if (CurrentWorker->PendingPackageRequests.Contains(Message.PackageName))
+			{
+				bWasSolicitedResult = true;
+				CurrentWorker->PendingPackageRequests.Remove(Message.PackageName);
+			}
+		}
+	}
+
+	if (bWasSolicitedResult)
+	{
+		UE_LOG(LogGatherTextFromAssetsCommandlet, Display, TEXT("Received package result from worker '%s' for '%s' and queued it for ingestion."), *Message.WorkerId.ToString(), *Message.PackageName.ToString());
+		PackageResults.Enqueue(Message);
+	}
+	else
+	{
+		UE_LOG(LogGatherTextFromAssetsCommandlet, Display, TEXT("Received unsolicited package result from worker '%s' for '%s' and dropped it."), *Message.WorkerId.ToString(), *Message.PackageName.ToString());
+	}
+}
+
 #undef LOC_DEFINE_REGION
 
 //////////////////////////////////////////////////////////////////////////
