Remove ingested results from PackagesDistributedToWorkers

#rnx

diff --git a/Engine/Source/Editor/UnrealEd/Private/Commandlets/GatherTextFromAssetsCommandlet.cpp b/Engine/Source/Editor/UnrealEd/Private/Commandlets/GatherTextFromAssetsCommandlet.cpp
--- a/Engine/Source/Editor/UnrealEd/Private/Commandlets/GatherTextFromAssetsCommandlet.cpp
+++ b/Engine/Source/Editor/UnrealEd/Private/Commandlets/GatherTextFromAssetsCommandlet.cpp
@@ -1951,6 +1951,8 @@ int32 UGatherTextFromAssetsCommandlet::Main(const FString& Params)
 				UE_LOG(LogGatherTextFromAssetsCommandlet, Display, TEXT("Ingesting worker result for '%s'..."), *PackageResult->PackageName.ToString());
 				ProcessGatherableTextDataArray(GatherableTextData);
 			}
+
+			PackagesDistributedToWorkers.Remove(PackageResult->PackageName);
 		}
 
 		// Re-add any jobs sent to crashed workers to the main task work queue
