Cooker: Move UCookOnTheFlyServer's shader-related member functions, structs, helper functions, and global variables into new files ShaderLibraryCooking.h/.cpp. Intent of this change is to improve single-file compile time and improve discoverability of the general cooking code and of the shader cooking code by reducing the size of the CookOnTheFlyServer.cpp file.
#rnx

#FailedValidations Preflight

diff --git a/Engine/Source/Editor/UnrealEd/Classes/CookOnTheSide/CookOnTheFlyServer.h b/Engine/Source/Editor/UnrealEd/Classes/CookOnTheSide/CookOnTheFlyServer.h
--- a/Engine/Source/Editor/UnrealEd/Classes/CookOnTheSide/CookOnTheFlyServer.h
+++ b/Engine/Source/Editor/UnrealEd/Classes/CookOnTheSide/CookOnTheFlyServer.h
@@ -1104,6 +1104,12 @@ private:
 	/** Invokes the necessary FShaderCodeLibrary functions to save and close a named code library. */
 	void SaveShaderLibrary(const ITargetPlatform* TargetPlatform, FString const& Name);
 
+	/** Called during CookByTheBookFinished. */
+	void ShutdownShaderLibraryCookerAndCompilers(const FString& ProjectShaderLibraryName);
+
+	/** Log shader stats during CookByTheBookFinished. */
+	void DumpShaderTypeStats(const FString& PlatformNameString);
+
 	/**
 	 * Calls the ShaderPipelineCacheToolsCommandlet to build a upipelinecache file from the stable pipeline cache
 	 * (.spc, in the past textual .stablepc.csv) file, if any
diff --git a/Engine/Source/Editor/UnrealEd/Private/CookOnTheFlyServer.cpp b/Engine/Source/Editor/UnrealEd/Private/CookOnTheFlyServer.cpp
--- a/Engine/Source/Editor/UnrealEd/Private/CookOnTheFlyServer.cpp
+++ b/Engine/Source/Editor/UnrealEd/Private/CookOnTheFlyServer.cpp
@@ -17,7 +17,6 @@
 #include "Async/Async.h"
 #include "Async/ParallelFor.h"
 #include "Commandlets/AssetRegistryGenerator.h"
-#include "Commandlets/ShaderPipelineCacheToolsCommandlet.h"
 #include "Containers/DirectoryTree.h"
 #include "Containers/RingBuffer.h"
 #include "Cooker/AsyncIODelete.h"
@@ -51,6 +50,7 @@
 #include "Cooker/NetworkFileCookOnTheFlyRequestManager.h"
 #include "Cooker/OnDemandShaderCompilation.h"
 #include "Cooker/PackageTracker.h"
+#include "Cooker/ShaderLibraryCooking.h"
 #include "Cooker/StallDetector.h"
 #include "Cooker/WorkerRequestsLocal.h"
 #include "Cooker/WorkerRequestsRemote.h"
@@ -76,7 +76,6 @@
 #include "FileServerMessages.h"
 #include "GameDelegates.h"
 #include "GenericPlatform/GenericPlatformCrashContext.h"
-#include "GlobalShader.h"
 #include "HAL/FileManager.h"
 #include "HAL/IConsoleManager.h"
 #include "HAL/MemoryMisc.h"
@@ -129,7 +128,6 @@
 #include "Modules/ModuleManager.h"
 #include "ObjectTools.h"
 #include "PackageHelperFunctions.h"
-#include "PipelineCacheChunkDataGenerator.h"
 #include "PlatformInfo.h"
 #include "ProfilingDebugging/CookStats.h"
 #include "ProfilingDebugging/PlatformFileTrace.h"
@@ -144,11 +142,8 @@
 #include "Serialization/CustomVersion.h"
 #include "Settings/LevelEditorPlaySettings.h"
 #include "Settings/ProjectPackagingSettings.h"
-#include "ShaderCodeLibrary.h"
 #include "ShaderCompiler.h"
-#include "ShaderLibraryChunkDataGenerator.h"
 #include "ShaderStats.h"
-#include "ShaderStatsCollector.h"
 #include "Stats/StatsSystem.h"
 #include "String/Find.h"
 #include "String/ParseLines.h"
@@ -241,19 +236,6 @@ static FAutoConsoleVariableRef CVarCookTimeCVarControl(
 		"NOTE: CURRENTLY any non-zero mode will also enable redirecting a ShaderPlatform to the cooking platform"),
 	ECVF_Default);
 
-bool bDisableShaderCompilationDuringCookOnTheFly = false;
-static FAutoConsoleVariableRef CVarDisableShaderCompilationDuringCookOnTheFly(
-	TEXT("Cook.OnTheFly.ShaderCompilationDisabled"),
-	bDisableShaderCompilationDuringCookOnTheFly,
-	TEXT("Controls whether shader compilation is disabled during cook on the fly\n"),
-	ECVF_Default);
-
-bool bAllowIncompleteShaderMapsDuringCookOnTheFly = false;
-static FAutoConsoleVariableRef CVarAllowEmptyShaderMapsDuringCookOnTheFly(
-	TEXT("Cook.OnTheFly.AllowIncompleteShadermaps"),
-	bAllowIncompleteShaderMapsDuringCookOnTheFly,
-	TEXT("Controls whether incomplete shader maps are processed during cook on the fly\n"),
-	ECVF_Default);
 
 ////////////////////////////////////////////////////////////////
 /// Cook on the fly server
@@ -298,20 +280,6 @@ static FString GetReleaseVersionAssetRegistryPath(const FString& ReleaseVersion,
 	return (*ReleasesRoot) / ReleaseVersion / PlatformName;
 }
 
-static void GenerateAutogenShaderHeadersForAllShaderFormats(const TArray<ITargetPlatform*>& Platforms)
-{
-	for (const ITargetPlatform* Platform : Platforms)
-	{
-		TArray<FName> ShaderPlatformNameArray;
-		Platform->GetAllTargetedShaderFormats(ShaderPlatformNameArray);
-
-		for (const FName& ShaderPlatformName : ShaderPlatformNameArray)
-		{
-			FShaderCompileUtilities::GenerateBrdfHeaders(ShaderPlatformName);
-		}
-	}
-}
-
 template<typename T>
 struct FOneTimeCommandlineReader
 {
@@ -1246,12 +1214,6 @@ bool UCookOnTheFlyServer::IsDirectorCookByTheBook() const
 	return ::IsCookByTheBookMode(DirectorCookMode);
 }
 
-bool UCookOnTheFlyServer::IsUsingShaderCodeLibrary() const
-{
-	const UProjectPackagingSettings* const PackagingSettings = GetDefault<UProjectPackagingSettings>();
-	return IsDirectorCookByTheBook() && AllowShaderCompiling() && PackagingSettings->bShareMaterialShaderCode;
-}
-
 bool UCookOnTheFlyServer::IsUsingZenStore() const
 {
 	return bZenStore;
@@ -2353,19 +2315,13 @@ static void CBTBTickCommandletStats()
 	UE::Stats::FStats::TickCommandletStats();
 }
 
-static void TickShaderCompilingManager(UE::Cook::FTickStackData& StackData)
-{
-	UE_SCOPED_COOKTIMER_AND_DURATION(CookByTheBook_ShaderProcessAsync, DetailedCookStats::TickLoopShaderProcessAsyncResultsTimeSec);
-	GShaderCompilingManager->ProcessAsyncResults(StackData.Timer.GetActionTimeSlice(), false);
-}
-
 static void TickAssetRegistry()
 {
 	UE_SCOPED_COOKTIMER(CookByTheBook_TickAssetRegistry);
 	FAssetRegistryModule::TickAssetRegistry(-1.0f);
 }
 
-}
+} // namespace UE::Cook
 
 void UCookOnTheFlyServer::InitializePollables()
 {
@@ -5739,33 +5695,6 @@ void UCookOnTheFlyServer::TickRequestManager()
 	}
 }
 
-void UCookOnTheFlyServer::TickRecompileShaderRequestsPrivate(UE::Cook::FTickStackData& StackData)
-{
-	// try to pull off a request
-	UE::Cook::FRecompileShaderRequest RecompileShaderRequest;
-	bool bProcessedRequests = false;
-	if (PackageTracker->RecompileRequests.Dequeue(&RecompileShaderRequest))
-	{
-		if (RecompileShaderRequest.RecompileArguments.CommandType != ODSCRecompileCommand::ResetMaterialCache)
-		{
-			RecompileShadersForRemote(RecompileShaderRequest.RecompileArguments, GetSandboxDirectory(RecompileShaderRequest.RecompileArguments.PlatformName));
-		}
-
-		RecompileShaderRequest.CompletionCallback();
-		bProcessedRequests = true;
-	}
-	if (PackageTracker->RecompileRequests.HasItems())
-	{
-		RecompileRequestsPollable->Trigger(*this);
-	}
-
-	if (bProcessedRequests)
-	{
-		// Ask for GC to run again when we processed some shaders requests to ensure material get evicted and we don't keep their package open
-		StackData.ResultFlags |= COSR_RequiresGC | COSR_RequiresGC_Periodic | COSR_YieldTick;
-	}
-}
-
 class FDiffModeCookServerUtils
 {
 public:
@@ -6271,8 +6200,6 @@ static EMPCookGeneratorSplit ParseMPCookGeneratorSplitFromString(const FString&
 
 }
 
-#include "RHIShaderFormatDefinitions.inl"
-
 namespace UE::Cook::CVarControl
 {
 	void UpdateCVars(FBeginCookContext& BeginContext, FName OverrideDeviceProfileName, int32 OverrideCookCVarControl)
@@ -6315,23 +6242,8 @@ namespace UE::Cook::CVarControl
 				IConsoleManager::Get().StompPlatformCVars(PlatformName, DPName.ToString(), CookTimeCVarTag, ECVF_SetByCode, ECVF_Default, ECVF_Default);
 			}
 
-			// for now, only do this if we are performing any cvar control modes
-			if (CookTimeCVarControl != 0)
-			{
-				if (Platform != TPM->GetRunningTargetPlatform())
-				{
-					// register that when we cook for this platform's shaderplatforms, we will want to use the given platform when looking up cvars
-					TArray<FName> ShaderFormats;
-					Platform->GetAllTargetedShaderFormats(ShaderFormats);
-					for (FName SF : ShaderFormats)
-					{
-						EShaderPlatform SP = ShaderFormatNameToShaderPlatform(SF);
-						ConsoleVariablePlatformMapping::RegisterShaderPlatformToPlatformMapping((int)SP, PlatformName);
-					}
-				}
-			}
+			UpdateShaderCookingCVars(TPM, CookTimeCVarControl, Platform, PlatformName);
 		}
-
 	}
 }
 
@@ -7357,7 +7269,7 @@ void GetAdditionalCurrentIniVersionStrings( const UCookOnTheFlyServer* CookOnThe
 	IniVersionMap.Add(TEXT("GlobalDDCVersion"), *GetGlobalShaderMapDDCGuid().ToString());
 
 	UProjectPackagingSettings* PackagingSettings = Cast<UProjectPackagingSettings>(UProjectPackagingSettings::StaticClass()->GetDefaultObject());
-	IniVersionMap.Add(TEXT("IsUsingShaderCodeLibrary"), FString::Printf(TEXT("%d"), PackagingSettings->bShareMaterialShaderCode && CookOnTheFlyServer->IsUsingShaderCodeLibrary()));
+	IniVersionMap.Add(TEXT("IsUsingShaderCodeLibrary"), FString::Printf(TEXT("%d"), CookOnTheFlyServer->IsUsingShaderCodeLibrary()));
 }
 
 bool UCookOnTheFlyServer::GetCurrentIniVersionStrings( const ITargetPlatform* TargetPlatform, UE::Cook::FIniSettingContainer& IniVersionStrings ) const
@@ -8928,26 +8840,6 @@ bool UCookOnTheFlyServer::IsCookByTheBookRunning() const
 	return IsCookByTheBookMode() && IsInSession();
 }
 
-
-void UCookOnTheFlyServer::SaveGlobalShaderMapFiles(const TArrayView<const ITargetPlatform* const>& Platforms, ODSCRecompileCommand RecompileCommand)
-{
-	LLM_SCOPE(ELLMTag::Shaders);
-	check(!IsCookingDLC()); // GlobalShaderMapFiles are not supported when cooking DLC
-	check(IsInGameThread());
-	for (const ITargetPlatform* TargetPlatform : Platforms)
-	{
-		const FString& PlatformName = TargetPlatform->PlatformName();
-		UE_LOG(LogCook, Display, TEXT("Compiling global%s shaders for platform '%s'"),
-			RecompileCommand == ODSCRecompileCommand::Changed ? TEXT(" changed") : TEXT(""), *PlatformName);
-
-		TArray<uint8> GlobalShaderMap;
-		FShaderRecompileData RecompileData(PlatformName, SP_NumPlatforms, RecompileCommand, nullptr, nullptr, &GlobalShaderMap);
-		RecompileData.ODSCCustomLoadMaterial = &UE::Cook::FODSCClientData::FindMaterial;
-		
-		RecompileShadersForRemote(RecompileData, GetSandboxDirectory(PlatformName));
-	}
-}
-
 FString UCookOnTheFlyServer::GetSandboxDirectory( const FString& PlatformName ) const
 {
 	return SandboxFile->GetSandboxDirectory(PlatformName);
@@ -9003,483 +8895,6 @@ FString UCookOnTheFlyServer::GetSandboxCachedEditorThumbnailsFilename()
 	return ConvertToFullSandboxPath(*CachedEditorThumbnailsFilename, true);
 }
 
-namespace UE::Cook
-{
-
-/** CookMultiprocess collector for ShaderLibrary data. */
-class FShaderLibraryCollector : public IMPCollector
-{
-public:
-	virtual FGuid GetMessageType() const override { return MessageType; }
-	virtual const TCHAR* GetDebugName() const override { return TEXT("FShaderLibraryCollector"); }
-
-	virtual void ClientTick(FMPCollectorClientTickContext& Context) override;
-	virtual void ServerReceiveMessage(FMPCollectorServerMessageContext& Context, FCbObjectView Message) override;
-
-	static FGuid MessageType;
-	static constexpr ANSICHAR RootObjectID[2]{ "S" };
-};
-
-const TCHAR* TEXT_GlobalMaterialHash(TEXT("GlobalMaterialHash"));
-const TCHAR* TEXT_ShaderLibrary(TEXT("ShaderLibrary"));
-/** CookArtifact to manage the shaderlibraries in the cook's OutputFolder. */
-class FShaderLibraryCookArtifact : public ICookArtifact
-{
-public:
-	FShaderLibraryCookArtifact(UCookOnTheFlyServer& InCOTFS)
-		: COTFS(InCOTFS)
-	{
-	}
-
-	virtual FString GetArtifactName() const
-	{
-		return TEXT_ShaderLibrary;
-	}
-
-	virtual void OnFullRecook(const ITargetPlatform* TargetPlatform)
-	{
-		CleanIntermediateFiles(TargetPlatform);
-	}
-
-	void CleanIntermediateFiles(const ITargetPlatform* TargetPlatform)
-	{
-		TArray<FName> ShaderFormats;
-		TargetPlatform->GetAllTargetedShaderFormats(ShaderFormats);
-		if (ShaderFormats.Num() > 0)
-		{
-			FShaderLibraryCooker::CleanDirectories(ShaderFormats);
-		}
-	}
-
-private:
-	UCookOnTheFlyServer& COTFS;
-};
-
-FGuid FShaderLibraryCollector::MessageType(TEXT("4DF3B36BBA2F4E04A846E894E24EB2C4"));
-
-void FShaderLibraryCollector::ClientTick(FMPCollectorClientTickContext& Context)
-{
-	constexpr int32 NumLoopsPerWarning = 100;
-	for (int32 NumMessages = 0; ; ++NumMessages)
-	{
-		// Maximum size for a message is 1GB. Caller will crash if we go over that.
-		// Provide a maximum shader size of half that because CopyToCompactBinaryAndClear adds on additional
-		// small amounts of data beyond the shader size limit.
-		const int64 MaximumSize = UE::CompactBinaryTCP::MaxOSPacketSize / 2;
-
-		bool bHasData;
-		bool bRanOutOfRoom;
-		FCbWriter Writer;
-		Writer.BeginObject();
-		Writer.SetName(RootObjectID);
-		FShaderLibraryCooker::CopyToCompactBinaryAndClear(Writer, bHasData, bRanOutOfRoom, MaximumSize);
-		if (bHasData)
-		{
-			Writer.EndObject();
-			Context.AddMessage(Writer.Save().AsObject());
-		}
-		if (!bRanOutOfRoom)
-		{
-			break;
-		}
-		if (NumMessages > 0 && NumMessages % NumLoopsPerWarning == 0)
-		{
-			UE_LOG(LogCook, Warning, TEXT("FShaderLibraryCollector::ClientTick has an unexpectedly high number of loops. Infinite loop?"))
-		}
-	}
-}
-
-void FShaderLibraryCollector::ServerReceiveMessage(FMPCollectorServerMessageContext& Context, FCbObjectView Message)
-{
-	bool bSuccessful = FShaderLibraryCooker::AppendFromCompactBinary(Message[RootObjectID]);
-	UE_CLOG(!bSuccessful, LogCook, Error,
-		TEXT("Corrupt message received from CookWorker when replicating ShaderLibrary. Shaders will be missing from the cook."));
-}
-
-}
-
-void UCookOnTheFlyServer::RegisterShaderLibraryIncrementalCookArtifact(FBeginCookContext& BeginContext)
-{
-	if (IsCookWorkerMode() || !IsUsingShaderCodeLibrary())
-	{
-		return;
-	}
-	RegisterArtifact(new UE::Cook::FShaderLibraryCookArtifact(*this));
-}
-
-void UCookOnTheFlyServer::BeginCookStartShaderCodeLibrary(FBeginCookContext& BeginContext)
-{
-	const UProjectPackagingSettings* const PackagingSettings = GetDefault<UProjectPackagingSettings>();
-	if (IsUsingShaderCodeLibrary())
-	{
-		FShaderLibraryCooker::InitForCooking(PackagingSettings->bSharedMaterialNativeLibraries, AllContextArtifactReader.Get());
-
-		bool bAllPlatformsNeedStableKeys = false;
-		// support setting without Hungarian prefix for the compatibility, but allow newer one to override
-		GConfig->GetBool(TEXT("DevOptions.Shaders"), TEXT("NeedsShaderStableKeys"), bAllPlatformsNeedStableKeys, GEngineIni);
-		GConfig->GetBool(TEXT("DevOptions.Shaders"), TEXT("bNeedsShaderStableKeys"), bAllPlatformsNeedStableKeys, GEngineIni);
-
-		// PSO manual cache in DLC is not currently supported. Although stable keys can have other uses, disable this for DLC as it will
-		// also make it cook faster.
-		bAllPlatformsNeedStableKeys &= !IsCookingDLC();
-
- 		for (const ITargetPlatform* TargetPlatform : BeginContext.TargetPlatforms)
- 		{
-			// Find out if this platform requires stable shader keys, by reading the platform setting file.
-			// Stable shader keys are needed if we are going to create a PSO cache.
-			bool bNeedShaderStableKeys = bAllPlatformsNeedStableKeys;
-			FConfigFile PlatformIniFile;
-			FConfigCacheIni::LoadLocalIniFile(PlatformIniFile, TEXT("Engine"), true, *TargetPlatform->IniPlatformName());
-			PlatformIniFile.GetBool(TEXT("DevOptions.Shaders"), TEXT("NeedsShaderStableKeys"), bNeedShaderStableKeys);
-			PlatformIniFile.GetBool(TEXT("DevOptions.Shaders"), TEXT("bNeedsShaderStableKeys"), bNeedShaderStableKeys);
-
-			bool bNeedsDeterministicOrder = PackagingSettings->bDeterministicShaderCodeOrder;
-			FConfigFile PlatformGameIniFile;
-			FConfigCacheIni::LoadLocalIniFile(PlatformGameIniFile, TEXT("Game"), true, *TargetPlatform->IniPlatformName());
-			PlatformGameIniFile.GetBool(TEXT("/Script/UnrealEd.ProjectPackagingSettings"), TEXT("bDeterministicShaderCodeOrder"), bNeedsDeterministicOrder);
-
- 			TArray<FName> ShaderFormats;
- 			TargetPlatform->GetAllTargetedShaderFormats(ShaderFormats);
-			TArray<FShaderLibraryCooker::FShaderFormatDescriptor> ShaderFormatsWithStableKeys;
-			for (FName& Format : ShaderFormats)
-			{
-				FShaderLibraryCooker::FShaderFormatDescriptor NewDesc;
-				NewDesc.ShaderFormat = Format;
-				NewDesc.bNeedsStableKeys = bNeedShaderStableKeys;
-				NewDesc.bNeedsDeterministicOrder = bNeedsDeterministicOrder;
-				ShaderFormatsWithStableKeys.Push(NewDesc);
-			}
-
-			if (ShaderFormats.Num() > 0)
-			{
-				FShaderLibraryCooker::CookShaderFormats(ShaderFormatsWithStableKeys);
-			}
-		}
-
-		if (CookDirector)
-		{
-			CookDirector->Register(new UE::Cook::FShaderLibraryCollector());
-		}
-		else if (CookWorkerClient)
-		{
-			CookWorkerClient->Register(new UE::Cook::FShaderLibraryCollector());
-		}
-	}
-
-	if (CookDirector)
-	{
-		CookDirector->Register(new FShaderStatsAggregator(FShaderStatsAggregator::EMode::Director));
-	}
-	else if (CookWorkerClient)
-	{
-		CookWorkerClient->Register(new FShaderStatsAggregator(FShaderStatsAggregator::EMode::Worker));
-	}
-}
-
-void UCookOnTheFlyServer::BeginCookFinishShaderCodeLibrary(FBeginCookContext& BeginContext)
-{
-	check(IsDirectorCookByTheBook()); // CookByTheBook only for now
-
-	// Generate the AutogenShaderHeaders.usd files for every shader platforms we are about to cook. Those files can be hashed and used
-	// in dependencies so they need to be present. Only do this for the director.
-	bool bIsCookWorker = IsCookWorkerMode();
-	if (!bIsCookWorker)
-	{
-		GenerateAutogenShaderHeadersForAllShaderFormats(BeginContext.TargetPlatforms);
-	}
-
-	// don't resave the global shader map files in dlc
-	if (!bIsCookWorker && !IsCookingDLC() && !EnumHasAnyFlags(CookByTheBookOptions->StartupOptions, ECookByTheBookOptions::ForceDisableSaveGlobalShaders))
-	{
-		OpenGlobalShaderLibrary();
-
-		// make sure global shaders are up to date!
-		SaveGlobalShaderMapFiles(BeginContext.TargetPlatforms, ODSCRecompileCommand::Changed);
-
-		SaveAndCloseGlobalShaderLibrary();
-	}
-
-	// Open the shader code library for the current project or the current DLC pack, depending on which we are cooking
-	FString LibraryName = GetProjectShaderLibraryName();
-	check(!LibraryName.IsEmpty());
-	OpenShaderLibrary(LibraryName);
-}
-
-
-void UCookOnTheFlyServer::RegisterShaderChunkDataGenerator()
-{
-	check(!IsCookWorkerMode());
-	// add shader library and PSO cache chunkers
-	FString LibraryName = GetProjectShaderLibraryName();
-	for (const ITargetPlatform* TargetPlatform : PlatformManager->GetSessionPlatforms())
-	{
-		FAssetRegistryGenerator& RegistryGenerator = *(PlatformManager->GetPlatformData(TargetPlatform)->RegistryGenerator);
-		RegistryGenerator.RegisterChunkDataGenerator(MakeShared<FShaderLibraryChunkDataGenerator>(*this, TargetPlatform));
-		RegistryGenerator.RegisterChunkDataGenerator(MakeShared<FPipelineCacheChunkDataGenerator>(TargetPlatform, LibraryName));
-	}
-}
-
-FString UCookOnTheFlyServer::GetProjectShaderLibraryName() const
-{
-	static FString ShaderLibraryName = [this]()
-	{
-		FString OverrideShaderLibraryName;
-		if (FParse::Value(FCommandLine::Get(), TEXT("OverrideShaderLibraryName="), OverrideShaderLibraryName))
-		{
-			return OverrideShaderLibraryName;
-		}
-
-		if (!IsCookingDLC())
-		{
-			FString Result = FApp::GetProjectName();
-			if (Result.IsEmpty())
-			{
-				Result = TEXT("UnrealGame");
-			}
-			return Result;
-		}
-		else
-		{
-			return CookByTheBookOptions->DlcName;
-		}
-	}();
-	return ShaderLibraryName;
-}
-
-static FString GenerateShaderCodeLibraryName(FString const& Name, bool bIsLegacyIterativeSharedBuild)
-{
-	FString ActualName = (!bIsLegacyIterativeSharedBuild) ? Name : Name + TEXT("_SC");
-	return ActualName;
-}
-
-void UCookOnTheFlyServer::OpenGlobalShaderLibrary()
-{
-	const bool bCacheShaderLibraries = IsUsingShaderCodeLibrary();
-	if (bCacheShaderLibraries)
-	{
-		const TCHAR* GlobalShaderLibName = TEXT("Global");
-		FString ActualName = GenerateShaderCodeLibraryName(GlobalShaderLibName, IsCookFlagSet(ECookInitializationFlags::LegacyIterativeSharedBuild));
-
-		// The shader code library directory doesn't matter while cooking
-		FShaderLibraryCooker::BeginCookingLibrary(ActualName);
-	}
-}
-
-void UCookOnTheFlyServer::OpenShaderLibrary(FString const& Name)
-{
-	const bool bCacheShaderLibraries = IsUsingShaderCodeLibrary();
-	if (bCacheShaderLibraries)
-	{
-		FString ActualName = GenerateShaderCodeLibraryName(Name, IsCookFlagSet(ECookInitializationFlags::LegacyIterativeSharedBuild));
-
-		// The shader code library directory doesn't matter while cooking
-		FShaderLibraryCooker::BeginCookingLibrary(ActualName);
-	}
-}
-
-void UCookOnTheFlyServer::CreatePipelineCache(const ITargetPlatform* TargetPlatform, const FString& LibraryName)
-{
-	// make sure we have a registry generated for all the platforms 
-	const FString TargetPlatformName = TargetPlatform->PlatformName();
-	TArray<FString>* SCLCSVPaths = OutSCLCSVPaths.Find(FName(TargetPlatformName));
-	if (SCLCSVPaths && SCLCSVPaths->Num())
-	{
-		TArray<FName> ShaderFormats;
-		TargetPlatform->GetAllTargetedShaderFormats(ShaderFormats);
-		for (FName ShaderFormat : ShaderFormats)
-		{
-			const FString StablePCDir = FPaths::ProjectDir() / TEXT("Build") / TargetPlatform->IniPlatformName() / TEXT("PipelineCaches");
-			// look for the new binary format for stable pipeline cache - spc
-			const FString StablePCBinary = StablePCDir / FString::Printf(TEXT("*%s_%s.spc"), *LibraryName, *ShaderFormat.ToString());
-
-			bool bBinaryStablePipelineCacheFilesFound = [&StablePCBinary]()
-			{
-				TArray<FString> ExpandedFiles;
-				IFileManager::Get().FindFilesRecursive(ExpandedFiles, *FPaths::GetPath(StablePCBinary), *FPaths::GetCleanFilename(StablePCBinary), true, false, false);
-				return ExpandedFiles.Num() > 0;
-			}();
-
-			// for now, also look for the older *stablepc.csv or *stablepc.csv.compressed
-			const FString StablePCTextual = StablePCDir / FString::Printf(TEXT("*%s_%s.stablepc.csv"), *LibraryName, *ShaderFormat.ToString());
-			const FString StablePCTextualCompressed = StablePCTextual + TEXT(".compressed");
-
-			bool bTextualStablePipelineCacheFilesFound = [&StablePCTextual, &StablePCTextualCompressed]()
-			{
-				TArray<FString> ExpandedFiles;
-				IFileManager::Get().FindFilesRecursive(ExpandedFiles, *FPaths::GetPath(StablePCTextual), *FPaths::GetCleanFilename(StablePCTextual), true, false, false);
-				IFileManager::Get().FindFilesRecursive(ExpandedFiles, *FPaths::GetPath(StablePCTextualCompressed), *FPaths::GetCleanFilename(StablePCTextualCompressed), true, false, false);
-				return ExpandedFiles.Num() > 0;
-			}();
-
-			// because of the compute shaders that are cached directly from stable shader keys files, we need to run this also if we have stable keys (which is pretty much always)
-			static const IConsoleVariable* CVarIncludeComputePSOsDuringCook = IConsoleManager::Get().FindConsoleVariable(TEXT("r.ShaderPipelineCacheTools.IncludeComputePSODuringCook"));
-			const bool bIncludeComputePSOsDuringCook = CVarIncludeComputePSOsDuringCook && CVarIncludeComputePSOsDuringCook->GetInt() >= 1;
-			if (!bBinaryStablePipelineCacheFilesFound && !bTextualStablePipelineCacheFilesFound && !bIncludeComputePSOsDuringCook)
-			{
-				UE_LOG(LogCook, Display, TEXT("---- NOT Running UShaderPipelineCacheToolsCommandlet for platform %s  shader format %s, no files found at %s, and either no stable keys or not including compute PSOs during the cook"), *TargetPlatformName, *ShaderFormat.ToString(), *StablePCDir);
-			}
-			else
-			{
-				UE_LOG(LogCook, Display, TEXT("---- Running UShaderPipelineCacheToolsCommandlet for platform %s  shader format %s"), *TargetPlatformName, *ShaderFormat.ToString());
-
-				const FString OutFilename = FString::Printf(TEXT("%s_%s.stable.upipelinecache"), *LibraryName, *ShaderFormat.ToString());
-				const FString PCUncookedPath = FPaths::ProjectDir() / TEXT("Content") / TEXT("PipelineCaches") / TargetPlatform->IniPlatformName() / OutFilename;
-
-				if (IFileManager::Get().FileExists(*PCUncookedPath))
-				{
-					UE_LOG(LogCook, Warning, TEXT("Deleting %s, cooked data doesn't belong here."), *PCUncookedPath);
-					IFileManager::Get().Delete(*PCUncookedPath, false, true);
-				}
-
-				const FString PCCookedPath = ConvertToFullSandboxPath(*PCUncookedPath, true);
-				const FString PCPath = PCCookedPath.Replace(TEXT("[Platform]"), *TargetPlatformName);
-
-
-				FString Args(TEXT("build "));
-				if (bBinaryStablePipelineCacheFilesFound)
-				{
-					Args += TEXT("\"");
-					Args += StablePCBinary;
-					Args += TEXT("\" ");
-				}
-				if (bTextualStablePipelineCacheFilesFound)
-				{
-					Args += TEXT("\"");
-					Args += StablePCTextual;
-					Args += TEXT("\" ");
-				}
-
-				int32 NumMatched = 0;
-				for (int32 Index = 0; Index < SCLCSVPaths->Num(); Index++)
-				{
-					if (!(*SCLCSVPaths)[Index].Contains(ShaderFormat.ToString()))
-					{
-						continue;
-					}
-					NumMatched++;
-					Args += TEXT(" ");
-					Args += TEXT("\"");
-					Args += (*SCLCSVPaths)[Index];
-					Args += TEXT("\"");
-				}
-				if (!NumMatched)
-				{
-					UE_LOG(LogCook, Warning, TEXT("Shader format %s for platform %s had stable pipeline cache files, but no stable keys files."), *ShaderFormat.ToString(), *TargetPlatformName);
-					for (int32 Index = 0; Index < SCLCSVPaths->Num(); Index++)
-					{
-						UE_LOG(LogCook, Warning, TEXT("    stable keys file: %s"), *((*SCLCSVPaths)[Index]));
-					}							
-					continue;
-				}
-
-				Args += TEXT(" -chunkinfodir=\"");
-				Args += ConvertToFullSandboxPath(FPaths::ProjectDir() / TEXT("Content"), true).Replace(TEXT("[Platform]"), *TargetPlatformName);
-				Args += TEXT("\" ");
-				Args += TEXT(" -library=");
-				Args += LibraryName;
-				Args += TEXT(" ");
-				Args += TEXT(" -platform=");
-				Args += TargetPlatformName;
-				Args += TEXT(" ");
-				Args += TEXT("\"");
-				Args += PCPath;
-				Args += TEXT("\"");
-				UE_LOG(LogCook, Display, TEXT("  With Args: %s"), *Args);
-
-				int32 Result = UShaderPipelineCacheToolsCommandlet::StaticMain(Args);
-
-				if (Result)
-				{
-					LogCookerMessage(FString::Printf(TEXT("UShaderPipelineCacheToolsCommandlet failed %d"), Result), EMessageSeverity::Error);
-				}
-				else
-				{
-					UE_LOG(LogCook, Display, TEXT("---- Done running UShaderPipelineCacheToolsCommandlet for platform %s"), *TargetPlatformName);
-
-					// copy the resulting file to metadata for easier examination later
-					if (IFileManager::Get().FileExists(*PCPath))
-					{
-						const FString RootPipelineCacheMetadataPath = GetMetadataDirectory() / TEXT("PipelineCaches");
-						const FString PipelineCacheMetadataPathSB = ConvertToFullSandboxPath(*RootPipelineCacheMetadataPath, true);
-						const FString PipelineCacheMetadataPath = PipelineCacheMetadataPathSB.Replace(TEXT("[Platform]"), *TargetPlatform->PlatformName());
-						const FString PipelineCacheMetadataFileName = PipelineCacheMetadataPath / OutFilename;
-
-						UE_LOG(LogCook, Display, TEXT("Copying the binary PSO cache file %s to %s."), *PCPath, *PipelineCacheMetadataFileName);
-						if (IFileManager::Get().Copy(*PipelineCacheMetadataFileName, *PCPath) != COPY_OK)
-						{
-							UE_LOG(LogCook, Warning, TEXT("Failed to copy the binary PSO cache file %s to %s."), *PCPath, *PipelineCacheMetadataFileName);
-						}
-					}
-				}
-			}
-		}
-	}
-}
-
-void UCookOnTheFlyServer::SaveAndCloseGlobalShaderLibrary()
-{
-	const bool bCacheShaderLibraries = IsUsingShaderCodeLibrary();
-	if (bCacheShaderLibraries)
-	{
-		const TCHAR* GlobalShaderLibName = TEXT("Global");
-		FString ActualName = GenerateShaderCodeLibraryName(GlobalShaderLibName, IsCookFlagSet(ECookInitializationFlags::LegacyIterativeSharedBuild));
-
-		// Save shader code map - cleaning directories is deliberately a separate loop here as we open the cache once per shader platform and we don't assume that they can't be shared across target platforms.
-		for (const ITargetPlatform* TargetPlatform : PlatformManager->GetSessionPlatforms())
-		{
-			FinishPopulateShaderLibrary(TargetPlatform, GlobalShaderLibName);
-			SaveShaderLibrary(TargetPlatform, GlobalShaderLibName);
-		}
-
-		FShaderLibraryCooker::EndCookingLibrary(ActualName);
-	}
-}
-
-void UCookOnTheFlyServer::GetShaderLibraryPaths(const ITargetPlatform* TargetPlatform,
-	FString& OutShaderCodeDir, FString& OutMetaDataPath, bool bUseProjectDirForDLC)
-{
-	// TODO: Saving ShaderChunks into the DLC directory currently does not work, so we have the bUseProjectDirForDLC arg to save to Project
-	const FString BasePath = (!IsCookingDLC() || bUseProjectDirForDLC) ? FPaths::ProjectContentDir() : GetContentDirectoryForDLC();
-	OutShaderCodeDir = ConvertToFullSandboxPath(*BasePath, true, TargetPlatform->PlatformName());
-
-	const FString RootMetaDataPath = GetMetadataDirectory() / TEXT("PipelineCaches");
-	OutMetaDataPath = ConvertToFullSandboxPath(*RootMetaDataPath, true, *TargetPlatform->PlatformName()); 
-}
-
-void UCookOnTheFlyServer::FinishPopulateShaderLibrary(const ITargetPlatform* TargetPlatform, const FString& Name)
-{
-	FString ShaderCodeDir;
-	FString MetaDataPath;
-	GetShaderLibraryPaths(TargetPlatform, ShaderCodeDir, MetaDataPath);
-
-	FShaderLibraryCooker::FinishPopulateShaderLibrary(TargetPlatform, Name, ShaderCodeDir, MetaDataPath);
-}
-
-void UCookOnTheFlyServer::SaveShaderLibrary(const ITargetPlatform* TargetPlatform, const FString& Name)
-{
-	FString ShaderCodeDir;
-	FString MetaDataPath;
-	GetShaderLibraryPaths(TargetPlatform, ShaderCodeDir, MetaDataPath);
-
-	TArray<FString>& PlatformSCLCSVPaths = OutSCLCSVPaths.FindOrAdd(FName(TargetPlatform->PlatformName()));
-	FString ErrorString;
-	bool bHasData;
-	if (!FShaderLibraryCooker::SaveShaderLibraryWithoutChunking(TargetPlatform, Name, ShaderCodeDir, MetaDataPath,
-		PlatformSCLCSVPaths, ErrorString, bHasData))
-	{
-		// This is fatal - In this case we should cancel any launch on device operation or package write but we don't want to assert and crash the editor
-		LogCookerMessage(FString::Printf(TEXT("%s"), *ErrorString), EMessageSeverity::Error);
-	}
-	else if (bHasData)
-	{
-		for (const FString& Item : PlatformSCLCSVPaths)
-		{
-			UE_LOG(LogCook, Display, TEXT("Saved scl.csv %s for platform %s, %d bytes"), *Item, *TargetPlatform->PlatformName(),
-				IFileManager::Get().FileSize(*Item));
-		}
-	}
-}
-
 void UCookOnTheFlyServer::WriteCookMetadata(const ITargetPlatform* InTargetPlatform, uint64 InDevelopmentAssetRegistryHash)
 {
 	FString PlatformNameString = InTargetPlatform->PlatformName();
@@ -10082,7 +9497,7 @@ void UCookOnTheFlyServer::CookByTheBookFinishedInternal()
 			// Save shader code map
 			SaveShaderLibrary(TargetPlatform, LibraryName);
 			CreatePipelineCache(TargetPlatform, LibraryName);
-			FShaderLibraryCooker::DumpShaderTypeStats(GShaderCompilingManager->GetAbsoluteShaderDebugInfoDirectory(), ConvertToFullSandboxPath(GetMetadataDirectory(), true, PlatformNameString));
+			DumpShaderTypeStats(PlatformNameString);
 		}
 
 		if (FParse::Param(FCommandLine::Get(), TEXT("fastcook")))
@@ -10139,11 +9554,7 @@ void UCookOnTheFlyServer::CookByTheBookFinishedInternal()
 	}
 
 	UE_LOG(LogCookStatus, Display, TEXT("Shutting down shader library cooker & shader compilers"));
-
-	FString ActualLibraryName = GenerateShaderCodeLibraryName(LibraryName, IsCookFlagSet(ECookInitializationFlags::LegacyIterativeSharedBuild));
-	FShaderLibraryCooker::EndCookingLibrary(ActualLibraryName);
-	FShaderLibraryCooker::Shutdown();
-	ShutdownShaderCompilers(PlatformManager->GetSessionPlatforms());
+	ShutdownShaderLibraryCookerAndCompilers(LibraryName);
 
 	if (CookByTheBookOptions->bGenerateDependenciesForMaps)
 	{
@@ -11695,11 +11106,7 @@ void UCookOnTheFlyServer::CookAsCookWorkerFinished()
 		LastCookableObjectTickTime = CurrentTime;
 	}
 
-	FString LibraryName = GetProjectShaderLibraryName();
-	FString ActualLibraryName = GenerateShaderCodeLibraryName(LibraryName, IsCookFlagSet(ECookInitializationFlags::LegacyIterativeSharedBuild));
-	FShaderLibraryCooker::EndCookingLibrary(ActualLibraryName);
-	FShaderLibraryCooker::Shutdown();
-	ShutdownShaderCompilers(PlatformManager->GetSessionPlatforms());
+	ShutdownShaderLibraryCookerAndCompilers(GetProjectShaderLibraryName());
 	FinalizePackageStore();
 
 	if (IsDirectorCookOnTheFly())
@@ -12582,16 +11989,6 @@ TArray<FName> UCookOnTheFlyServer::GetNeverCookPackageNames(TArrayView<const FSt
 	return NeverCookPackages;
 }
 
-bool UCookOnTheFlyServer::RecompileChangedShaders(const TArray<const ITargetPlatform*>& TargetPlatforms)
-{
-	bool bShadersRecompiled = false;
-	for (const ITargetPlatform* TargetPlatform : TargetPlatforms)
-	{
-		bShadersRecompiled |= RecompileChangedShadersForPlatform(TargetPlatform->PlatformName());
-	}
-	return bShadersRecompiled;
-}
-
 /* UCookOnTheFlyServer callbacks
  *****************************************************************************/
 
diff --git a/Engine/Source/Editor/UnrealEd/Private/Cooker/CookTypes.h b/Engine/Source/Editor/UnrealEd/Private/Cooker/CookTypes.h
--- a/Engine/Source/Editor/UnrealEd/Private/Cooker/CookTypes.h
+++ b/Engine/Source/Editor/UnrealEd/Private/Cooker/CookTypes.h
@@ -820,3 +820,11 @@ IPackageWriter::ECommitStatus PackageResultToCommitStatus(ESavePackageResult Res
 extern int32 GCookProgressDisplay;
 
 const TCHAR* LexToString(UE::Cook::ECookResult CookResult);
+
+/**
+ * Uses the FMessageLog to log a message
+ *
+ * @param Message to log
+ * @param Severity of the message
+ */
+void LogCookerMessage(const FString& MessageText, EMessageSeverity::Type Severity);
diff --git a/Engine/Source/Editor/UnrealEd/Private/Cooker/ShaderLibraryCooking.cpp b/Engine/Source/Editor/UnrealEd/Private/Cooker/ShaderLibraryCooking.cpp
new file mode 100644
--- /dev/null
+++ b/Engine/Source/Editor/UnrealEd/Private/Cooker/ShaderLibraryCooking.cpp
@@ -0,0 +1,641 @@
+// Copyright Epic Games, Inc. All Rights Reserved.
+
+#include "Cooker/ShaderLibraryCooking.h"
+
+#include "Commandlets/AssetRegistryGenerator.h"
+#include "Commandlets/ShaderPipelineCacheToolsCommandlet.h"
+#include "Cooker/CookDirector.h"
+#include "Cooker/CookOnTheFlyServerInterface.h"
+#include "Cooker/CookPlatformManager.h"
+#include "Cooker/CookTypes.h"
+#include "Cooker/CookWorkerClient.h"
+#include "Cooker/OnDemandShaderCompilation.h"
+#include "Cooker/PackageTracker.h"
+#include "CookOnTheSide/CookLog.h"
+#include "CookOnTheSide/CookOnTheFlyServer.h"
+#include "CoreGlobals.h"
+#include "GlobalShader.h"
+#include "HAL/ConsoleManager.h"
+#include "HAL/FileManager.h"
+#include "HAL/LowLevelMemTracker.h"
+#include "Interfaces/IShaderFormat.h"
+#include "Interfaces/ITargetPlatform.h"
+#include "Interfaces/ITargetPlatformManagerModule.h"
+#include "LayeredCookArtifactReader.h"
+#include "Misc/App.h"
+#include "Misc/CommandLine.h"
+#include "Misc/Parse.h"
+#include "Misc/Paths.h"
+#include "PipelineCacheChunkDataGenerator.h"
+#include "Serialization/CompactBinary.h"
+#include "Serialization/CompactBinaryWriter.h"
+#include "Settings/ProjectPackagingSettings.h"
+#include "ShaderCodeLibrary.h"
+#include "ShaderCompiler.h"
+#include "ShaderLibraryChunkDataGenerator.h"
+#include "ShaderStats.h"
+#include "ShaderStatsCollector.h"
+
+// Include RHIShaderFormatDefinitions.inl only after all header includes; it defines some static variables,
+// used by UpdateShaderCookingCVars.
+#include "RHIShaderFormatDefinitions.inl"
+
+
+bool bDisableShaderCompilationDuringCookOnTheFly = false;
+static FAutoConsoleVariableRef CVarDisableShaderCompilationDuringCookOnTheFly(
+	TEXT("Cook.OnTheFly.ShaderCompilationDisabled"),
+	bDisableShaderCompilationDuringCookOnTheFly,
+	TEXT("Controls whether shader compilation is disabled during cook on the fly\n"),
+	ECVF_Default);
+
+bool bAllowIncompleteShaderMapsDuringCookOnTheFly = false;
+static FAutoConsoleVariableRef CVarAllowEmptyShaderMapsDuringCookOnTheFly(
+	TEXT("Cook.OnTheFly.AllowIncompleteShadermaps"),
+	bAllowIncompleteShaderMapsDuringCookOnTheFly,
+	TEXT("Controls whether incomplete shader maps are processed during cook on the fly\n"),
+	ECVF_Default);
+
+static void GenerateAutogenShaderHeadersForAllShaderFormats(const TArray<ITargetPlatform*>& Platforms)
+{
+	for (const ITargetPlatform* Platform : Platforms)
+	{
+		TArray<FName> ShaderPlatformNameArray;
+		Platform->GetAllTargetedShaderFormats(ShaderPlatformNameArray);
+
+		for (const FName& ShaderPlatformName : ShaderPlatformNameArray)
+		{
+			FShaderCompileUtilities::GenerateBrdfHeaders(ShaderPlatformName);
+		}
+	}
+}
+
+bool UCookOnTheFlyServer::IsUsingShaderCodeLibrary() const
+{
+	const UProjectPackagingSettings* const PackagingSettings = GetDefault<UProjectPackagingSettings>();
+	return IsDirectorCookByTheBook() && AllowShaderCompiling() && PackagingSettings->bShareMaterialShaderCode;
+}
+
+namespace UE::Cook
+{
+
+void TickShaderCompilingManager(UE::Cook::FTickStackData& StackData)
+{
+	UE_SCOPED_COOKTIMER_AND_DURATION(CookByTheBook_ShaderProcessAsync, DetailedCookStats::TickLoopShaderProcessAsyncResultsTimeSec);
+	GShaderCompilingManager->ProcessAsyncResults(StackData.Timer.GetActionTimeSlice(), false);
+}
+
+} // namespace UE::Cook
+
+void UCookOnTheFlyServer::TickRecompileShaderRequestsPrivate(UE::Cook::FTickStackData& StackData)
+{
+	using namespace UE::Cook;
+
+	// try to pull off a request
+	FRecompileShaderRequest RecompileShaderRequest;
+	bool bProcessedRequests = false;
+	if (PackageTracker->RecompileRequests.Dequeue(&RecompileShaderRequest))
+	{
+		if (RecompileShaderRequest.RecompileArguments.CommandType != ODSCRecompileCommand::ResetMaterialCache)
+		{
+			RecompileShadersForRemote(RecompileShaderRequest.RecompileArguments, GetSandboxDirectory(RecompileShaderRequest.RecompileArguments.PlatformName));
+		}
+
+		RecompileShaderRequest.CompletionCallback();
+		bProcessedRequests = true;
+	}
+	if (PackageTracker->RecompileRequests.HasItems())
+	{
+		RecompileRequestsPollable->Trigger(*this);
+	}
+
+	if (bProcessedRequests)
+	{
+		// Ask for GC to run again when we processed some shaders requests to ensure material get evicted and we don't keep their package open
+		StackData.ResultFlags |= COSR_RequiresGC | COSR_RequiresGC_Periodic | COSR_YieldTick;
+	}
+}
+
+namespace UE::Cook::CVarControl
+{
+
+void UpdateShaderCookingCVars(ITargetPlatformManagerModule* TPM, int32 CookTimeCVarControl,
+	const ITargetPlatform* Platform, FName PlatformName)
+{
+	// for now, only do this if we are performing any cvar control modes
+	if (CookTimeCVarControl != 0)
+	{
+		if (Platform != TPM->GetRunningTargetPlatform())
+		{
+			// register that when we cook for this platform's shaderplatforms, we will want to use the given platform when looking up cvars
+			TArray<FName> ShaderFormats;
+			Platform->GetAllTargetedShaderFormats(ShaderFormats);
+			for (FName SF : ShaderFormats)
+			{
+				EShaderPlatform SP = ShaderFormatNameToShaderPlatform(SF);
+				ConsoleVariablePlatformMapping::RegisterShaderPlatformToPlatformMapping((int)SP, PlatformName);
+			}
+		}
+	}
+}
+
+} // namespace UE::Cook::CVarControl
+
+
+void UCookOnTheFlyServer::SaveGlobalShaderMapFiles(const TArrayView<const ITargetPlatform* const>& Platforms, ODSCRecompileCommand RecompileCommand)
+{
+	LLM_SCOPE(ELLMTag::Shaders);
+	check(!IsCookingDLC()); // GlobalShaderMapFiles are not supported when cooking DLC
+	check(IsInGameThread());
+	for (const ITargetPlatform* TargetPlatform : Platforms)
+	{
+		const FString& PlatformName = TargetPlatform->PlatformName();
+		UE_LOG(LogCook, Display, TEXT("Compiling global%s shaders for platform '%s'"),
+			RecompileCommand == ODSCRecompileCommand::Changed ? TEXT(" changed") : TEXT(""), *PlatformName);
+
+		TArray<uint8> GlobalShaderMap;
+		FShaderRecompileData RecompileData(PlatformName, SP_NumPlatforms, RecompileCommand, nullptr, nullptr, &GlobalShaderMap);
+		RecompileData.ODSCCustomLoadMaterial = &UE::Cook::FODSCClientData::FindMaterial;
+
+		RecompileShadersForRemote(RecompileData, GetSandboxDirectory(PlatformName));
+	}
+}
+
+namespace UE::Cook
+{
+
+FGuid FShaderLibraryCollector::MessageType(TEXT("4DF3B36BBA2F4E04A846E894E24EB2C4"));
+
+void FShaderLibraryCollector::ClientTick(FMPCollectorClientTickContext& Context)
+{
+	constexpr int32 NumLoopsPerWarning = 100;
+	for (int32 NumMessages = 0; ; ++NumMessages)
+	{
+		// Maximum size for a message is 1GB. Caller will crash if we go over that.
+		// Provide a maximum shader size of half that because CopyToCompactBinaryAndClear adds on additional
+		// small amounts of data beyond the shader size limit.
+		const int64 MaximumSize = UE::CompactBinaryTCP::MaxOSPacketSize / 2;
+
+		bool bHasData;
+		bool bRanOutOfRoom;
+		FCbWriter Writer;
+		Writer.BeginObject();
+		Writer.SetName(RootObjectID);
+		FShaderLibraryCooker::CopyToCompactBinaryAndClear(Writer, bHasData, bRanOutOfRoom, MaximumSize);
+		if (bHasData)
+		{
+			Writer.EndObject();
+			Context.AddMessage(Writer.Save().AsObject());
+		}
+		if (!bRanOutOfRoom)
+		{
+			break;
+		}
+		if (NumMessages > 0 && NumMessages % NumLoopsPerWarning == 0)
+		{
+			UE_LOG(LogCook, Warning, TEXT("FShaderLibraryCollector::ClientTick has an unexpectedly high number of loops. Infinite loop?"))
+		}
+	}
+}
+
+void FShaderLibraryCollector::ServerReceiveMessage(FMPCollectorServerMessageContext& Context, FCbObjectView Message)
+{
+	bool bSuccessful = FShaderLibraryCooker::AppendFromCompactBinary(Message[RootObjectID]);
+	UE_CLOG(!bSuccessful, LogCook, Error,
+		TEXT("Corrupt message received from CookWorker when replicating ShaderLibrary. Shaders will be missing from the cook."));
+}
+
+const TCHAR* TEXT_ShaderLibrary(TEXT("ShaderLibrary"));
+
+FShaderLibraryCookArtifact::FShaderLibraryCookArtifact(UCookOnTheFlyServer& InCOTFS)
+	: COTFS(InCOTFS)
+{
+}
+
+FString FShaderLibraryCookArtifact::GetArtifactName() const
+{
+	return TEXT_ShaderLibrary;
+}
+
+void FShaderLibraryCookArtifact::OnFullRecook(const ITargetPlatform* TargetPlatform)
+{
+	CleanIntermediateFiles(TargetPlatform);
+}
+
+void FShaderLibraryCookArtifact::CleanIntermediateFiles(const ITargetPlatform* TargetPlatform)
+{
+	TArray<FName> ShaderFormats;
+	TargetPlatform->GetAllTargetedShaderFormats(ShaderFormats);
+	if (ShaderFormats.Num() > 0)
+	{
+		FShaderLibraryCooker::CleanDirectories(ShaderFormats);
+	}
+}
+
+} // namespace UE::Cook
+
+void UCookOnTheFlyServer::RegisterShaderLibraryIncrementalCookArtifact(FBeginCookContext& BeginContext)
+{
+	if (IsCookWorkerMode() || !IsUsingShaderCodeLibrary())
+	{
+		return;
+	}
+	RegisterArtifact(new UE::Cook::FShaderLibraryCookArtifact(*this));
+}
+
+void UCookOnTheFlyServer::BeginCookStartShaderCodeLibrary(FBeginCookContext& BeginContext)
+{
+	const UProjectPackagingSettings* const PackagingSettings = GetDefault<UProjectPackagingSettings>();
+	if (IsUsingShaderCodeLibrary())
+	{
+		FShaderLibraryCooker::InitForCooking(PackagingSettings->bSharedMaterialNativeLibraries, AllContextArtifactReader.Get());
+
+		bool bAllPlatformsNeedStableKeys = false;
+		// support setting without Hungarian prefix for the compatibility, but allow newer one to override
+		GConfig->GetBool(TEXT("DevOptions.Shaders"), TEXT("NeedsShaderStableKeys"), bAllPlatformsNeedStableKeys, GEngineIni);
+		GConfig->GetBool(TEXT("DevOptions.Shaders"), TEXT("bNeedsShaderStableKeys"), bAllPlatformsNeedStableKeys, GEngineIni);
+
+		// PSO manual cache in DLC is not currently supported. Although stable keys can have other uses, disable this for DLC as it will
+		// also make it cook faster.
+		bAllPlatformsNeedStableKeys &= !IsCookingDLC();
+
+		for (const ITargetPlatform* TargetPlatform : BeginContext.TargetPlatforms)
+		{
+			// Find out if this platform requires stable shader keys, by reading the platform setting file.
+			// Stable shader keys are needed if we are going to create a PSO cache.
+			bool bNeedShaderStableKeys = bAllPlatformsNeedStableKeys;
+			FConfigFile PlatformIniFile;
+			FConfigCacheIni::LoadLocalIniFile(PlatformIniFile, TEXT("Engine"), true, *TargetPlatform->IniPlatformName());
+			PlatformIniFile.GetBool(TEXT("DevOptions.Shaders"), TEXT("NeedsShaderStableKeys"), bNeedShaderStableKeys);
+			PlatformIniFile.GetBool(TEXT("DevOptions.Shaders"), TEXT("bNeedsShaderStableKeys"), bNeedShaderStableKeys);
+
+			bool bNeedsDeterministicOrder = PackagingSettings->bDeterministicShaderCodeOrder;
+			FConfigFile PlatformGameIniFile;
+			FConfigCacheIni::LoadLocalIniFile(PlatformGameIniFile, TEXT("Game"), true, *TargetPlatform->IniPlatformName());
+			PlatformGameIniFile.GetBool(TEXT("/Script/UnrealEd.ProjectPackagingSettings"), TEXT("bDeterministicShaderCodeOrder"), bNeedsDeterministicOrder);
+
+			TArray<FName> ShaderFormats;
+			TargetPlatform->GetAllTargetedShaderFormats(ShaderFormats);
+			TArray<FShaderLibraryCooker::FShaderFormatDescriptor> ShaderFormatsWithStableKeys;
+			for (FName& Format : ShaderFormats)
+			{
+				FShaderLibraryCooker::FShaderFormatDescriptor NewDesc;
+				NewDesc.ShaderFormat = Format;
+				NewDesc.bNeedsStableKeys = bNeedShaderStableKeys;
+				NewDesc.bNeedsDeterministicOrder = bNeedsDeterministicOrder;
+				ShaderFormatsWithStableKeys.Push(NewDesc);
+			}
+
+			if (ShaderFormats.Num() > 0)
+			{
+				FShaderLibraryCooker::CookShaderFormats(ShaderFormatsWithStableKeys);
+			}
+		}
+
+		if (CookDirector)
+		{
+			CookDirector->Register(new UE::Cook::FShaderLibraryCollector());
+		}
+		else if (CookWorkerClient)
+		{
+			CookWorkerClient->Register(new UE::Cook::FShaderLibraryCollector());
+		}
+	}
+
+	if (CookDirector)
+	{
+		CookDirector->Register(new FShaderStatsAggregator(FShaderStatsAggregator::EMode::Director));
+	}
+	else if (CookWorkerClient)
+	{
+		CookWorkerClient->Register(new FShaderStatsAggregator(FShaderStatsAggregator::EMode::Worker));
+	}
+}
+
+void UCookOnTheFlyServer::BeginCookFinishShaderCodeLibrary(FBeginCookContext& BeginContext)
+{
+	check(IsDirectorCookByTheBook()); // CookByTheBook only for now
+
+	// Generate the AutogenShaderHeaders.usd files for every shader platforms we are about to cook. Those files can be hashed and used
+	// in dependencies so they need to be present. Only do this for the director.
+	bool bIsCookWorker = IsCookWorkerMode();
+	if (!bIsCookWorker)
+	{
+		GenerateAutogenShaderHeadersForAllShaderFormats(BeginContext.TargetPlatforms);
+	}
+
+	// don't resave the global shader map files in dlc
+	if (!bIsCookWorker && !IsCookingDLC() && !EnumHasAnyFlags(CookByTheBookOptions->StartupOptions, ECookByTheBookOptions::ForceDisableSaveGlobalShaders))
+	{
+		OpenGlobalShaderLibrary();
+
+		// make sure global shaders are up to date!
+		SaveGlobalShaderMapFiles(BeginContext.TargetPlatforms, ODSCRecompileCommand::Changed);
+
+		SaveAndCloseGlobalShaderLibrary();
+	}
+
+	// Open the shader code library for the current project or the current DLC pack, depending on which we are cooking
+	FString LibraryName = GetProjectShaderLibraryName();
+	check(!LibraryName.IsEmpty());
+	OpenShaderLibrary(LibraryName);
+}
+
+void UCookOnTheFlyServer::RegisterShaderChunkDataGenerator()
+{
+	check(!IsCookWorkerMode());
+	// add shader library and PSO cache chunkers
+	FString LibraryName = GetProjectShaderLibraryName();
+	for (const ITargetPlatform* TargetPlatform : PlatformManager->GetSessionPlatforms())
+	{
+		FAssetRegistryGenerator& RegistryGenerator = *(PlatformManager->GetPlatformData(TargetPlatform)->RegistryGenerator);
+		RegistryGenerator.RegisterChunkDataGenerator(MakeShared<FShaderLibraryChunkDataGenerator>(*this, TargetPlatform));
+		RegistryGenerator.RegisterChunkDataGenerator(MakeShared<FPipelineCacheChunkDataGenerator>(TargetPlatform, LibraryName));
+	}
+}
+
+FString UCookOnTheFlyServer::GetProjectShaderLibraryName() const
+{
+	static FString ShaderLibraryName = [this]()
+		{
+			FString OverrideShaderLibraryName;
+			if (FParse::Value(FCommandLine::Get(), TEXT("OverrideShaderLibraryName="), OverrideShaderLibraryName))
+			{
+				return OverrideShaderLibraryName;
+			}
+
+			if (!IsCookingDLC())
+			{
+				FString Result = FApp::GetProjectName();
+				if (Result.IsEmpty())
+				{
+					Result = TEXT("UnrealGame");
+				}
+				return Result;
+			}
+			else
+			{
+				return CookByTheBookOptions->DlcName;
+			}
+		}();
+	return ShaderLibraryName;
+}
+
+static FString GenerateShaderCodeLibraryName(const FString& Name, bool bIsLegacyIterativeSharedBuild)
+{
+	FString ActualName = (!bIsLegacyIterativeSharedBuild) ? Name : Name + TEXT("_SC");
+	return ActualName;
+}
+
+void UCookOnTheFlyServer::OpenGlobalShaderLibrary()
+{
+	const bool bCacheShaderLibraries = IsUsingShaderCodeLibrary();
+	if (bCacheShaderLibraries)
+	{
+		const TCHAR* GlobalShaderLibName = TEXT("Global");
+		FString ActualName = GenerateShaderCodeLibraryName(GlobalShaderLibName, IsCookFlagSet(ECookInitializationFlags::LegacyIterativeSharedBuild));
+
+		// The shader code library directory doesn't matter while cooking
+		FShaderLibraryCooker::BeginCookingLibrary(ActualName);
+	}
+}
+
+void UCookOnTheFlyServer::OpenShaderLibrary(FString const& Name)
+{
+	const bool bCacheShaderLibraries = IsUsingShaderCodeLibrary();
+	if (bCacheShaderLibraries)
+	{
+		FString ActualName = GenerateShaderCodeLibraryName(Name, IsCookFlagSet(ECookInitializationFlags::LegacyIterativeSharedBuild));
+
+		// The shader code library directory doesn't matter while cooking
+		FShaderLibraryCooker::BeginCookingLibrary(ActualName);
+	}
+}
+
+void UCookOnTheFlyServer::CreatePipelineCache(const ITargetPlatform* TargetPlatform, const FString& LibraryName)
+{
+	// make sure we have a registry generated for all the platforms 
+	const FString TargetPlatformName = TargetPlatform->PlatformName();
+	TArray<FString>* SCLCSVPaths = OutSCLCSVPaths.Find(FName(TargetPlatformName));
+	if (SCLCSVPaths && SCLCSVPaths->Num())
+	{
+		TArray<FName> ShaderFormats;
+		TargetPlatform->GetAllTargetedShaderFormats(ShaderFormats);
+		for (FName ShaderFormat : ShaderFormats)
+		{
+			const FString StablePCDir = FPaths::ProjectDir() / TEXT("Build") / TargetPlatform->IniPlatformName() / TEXT("PipelineCaches");
+			// look for the new binary format for stable pipeline cache - spc
+			const FString StablePCBinary = StablePCDir / FString::Printf(TEXT("*%s_%s.spc"), *LibraryName, *ShaderFormat.ToString());
+
+			bool bBinaryStablePipelineCacheFilesFound = [&StablePCBinary]()
+				{
+					TArray<FString> ExpandedFiles;
+					IFileManager::Get().FindFilesRecursive(ExpandedFiles, *FPaths::GetPath(StablePCBinary), *FPaths::GetCleanFilename(StablePCBinary), true, false, false);
+					return ExpandedFiles.Num() > 0;
+				}();
+
+			// for now, also look for the older *stablepc.csv or *stablepc.csv.compressed
+			const FString StablePCTextual = StablePCDir / FString::Printf(TEXT("*%s_%s.stablepc.csv"), *LibraryName, *ShaderFormat.ToString());
+			const FString StablePCTextualCompressed = StablePCTextual + TEXT(".compressed");
+
+			bool bTextualStablePipelineCacheFilesFound = [&StablePCTextual, &StablePCTextualCompressed]()
+				{
+					TArray<FString> ExpandedFiles;
+					IFileManager::Get().FindFilesRecursive(ExpandedFiles, *FPaths::GetPath(StablePCTextual), *FPaths::GetCleanFilename(StablePCTextual), true, false, false);
+					IFileManager::Get().FindFilesRecursive(ExpandedFiles, *FPaths::GetPath(StablePCTextualCompressed), *FPaths::GetCleanFilename(StablePCTextualCompressed), true, false, false);
+					return ExpandedFiles.Num() > 0;
+				}();
+
+			// because of the compute shaders that are cached directly from stable shader keys files, we need to run this also if we have stable keys (which is pretty much always)
+			static const IConsoleVariable* CVarIncludeComputePSOsDuringCook = IConsoleManager::Get().FindConsoleVariable(TEXT("r.ShaderPipelineCacheTools.IncludeComputePSODuringCook"));
+			const bool bIncludeComputePSOsDuringCook = CVarIncludeComputePSOsDuringCook && CVarIncludeComputePSOsDuringCook->GetInt() >= 1;
+			if (!bBinaryStablePipelineCacheFilesFound && !bTextualStablePipelineCacheFilesFound && !bIncludeComputePSOsDuringCook)
+			{
+				UE_LOG(LogCook, Display, TEXT("---- NOT Running UShaderPipelineCacheToolsCommandlet for platform %s  shader format %s, no files found at %s, and either no stable keys or not including compute PSOs during the cook"), *TargetPlatformName, *ShaderFormat.ToString(), *StablePCDir);
+			}
+			else
+			{
+				UE_LOG(LogCook, Display, TEXT("---- Running UShaderPipelineCacheToolsCommandlet for platform %s  shader format %s"), *TargetPlatformName, *ShaderFormat.ToString());
+
+				const FString OutFilename = FString::Printf(TEXT("%s_%s.stable.upipelinecache"), *LibraryName, *ShaderFormat.ToString());
+				const FString PCUncookedPath = FPaths::ProjectDir() / TEXT("Content") / TEXT("PipelineCaches") / TargetPlatform->IniPlatformName() / OutFilename;
+
+				if (IFileManager::Get().FileExists(*PCUncookedPath))
+				{
+					UE_LOG(LogCook, Warning, TEXT("Deleting %s, cooked data doesn't belong here."), *PCUncookedPath);
+					IFileManager::Get().Delete(*PCUncookedPath, false, true);
+				}
+
+				const FString PCCookedPath = ConvertToFullSandboxPath(*PCUncookedPath, true);
+				const FString PCPath = PCCookedPath.Replace(TEXT("[Platform]"), *TargetPlatformName);
+
+
+				FString Args(TEXT("build "));
+				if (bBinaryStablePipelineCacheFilesFound)
+				{
+					Args += TEXT("\"");
+					Args += StablePCBinary;
+					Args += TEXT("\" ");
+				}
+				if (bTextualStablePipelineCacheFilesFound)
+				{
+					Args += TEXT("\"");
+					Args += StablePCTextual;
+					Args += TEXT("\" ");
+				}
+
+				int32 NumMatched = 0;
+				for (int32 Index = 0; Index < SCLCSVPaths->Num(); Index++)
+				{
+					if (!(*SCLCSVPaths)[Index].Contains(ShaderFormat.ToString()))
+					{
+						continue;
+					}
+					NumMatched++;
+					Args += TEXT(" ");
+					Args += TEXT("\"");
+					Args += (*SCLCSVPaths)[Index];
+					Args += TEXT("\"");
+				}
+				if (!NumMatched)
+				{
+					UE_LOG(LogCook, Warning, TEXT("Shader format %s for platform %s had stable pipeline cache files, but no stable keys files."), *ShaderFormat.ToString(), *TargetPlatformName);
+					for (int32 Index = 0; Index < SCLCSVPaths->Num(); Index++)
+					{
+						UE_LOG(LogCook, Warning, TEXT("    stable keys file: %s"), *((*SCLCSVPaths)[Index]));
+					}
+					continue;
+				}
+
+				Args += TEXT(" -chunkinfodir=\"");
+				Args += ConvertToFullSandboxPath(FPaths::ProjectDir() / TEXT("Content"), true).Replace(TEXT("[Platform]"), *TargetPlatformName);
+				Args += TEXT("\" ");
+				Args += TEXT(" -library=");
+				Args += LibraryName;
+				Args += TEXT(" ");
+				Args += TEXT(" -platform=");
+				Args += TargetPlatformName;
+				Args += TEXT(" ");
+				Args += TEXT("\"");
+				Args += PCPath;
+				Args += TEXT("\"");
+				UE_LOG(LogCook, Display, TEXT("  With Args: %s"), *Args);
+
+				int32 Result = UShaderPipelineCacheToolsCommandlet::StaticMain(Args);
+
+				if (Result)
+				{
+					LogCookerMessage(FString::Printf(TEXT("UShaderPipelineCacheToolsCommandlet failed %d"), Result), EMessageSeverity::Error);
+				}
+				else
+				{
+					UE_LOG(LogCook, Display, TEXT("---- Done running UShaderPipelineCacheToolsCommandlet for platform %s"), *TargetPlatformName);
+
+					// copy the resulting file to metadata for easier examination later
+					if (IFileManager::Get().FileExists(*PCPath))
+					{
+						const FString RootPipelineCacheMetadataPath = GetMetadataDirectory() / TEXT("PipelineCaches");
+						const FString PipelineCacheMetadataPathSB = ConvertToFullSandboxPath(*RootPipelineCacheMetadataPath, true);
+						const FString PipelineCacheMetadataPath = PipelineCacheMetadataPathSB.Replace(TEXT("[Platform]"), *TargetPlatform->PlatformName());
+						const FString PipelineCacheMetadataFileName = PipelineCacheMetadataPath / OutFilename;
+
+						UE_LOG(LogCook, Display, TEXT("Copying the binary PSO cache file %s to %s."), *PCPath, *PipelineCacheMetadataFileName);
+						if (IFileManager::Get().Copy(*PipelineCacheMetadataFileName, *PCPath) != COPY_OK)
+						{
+							UE_LOG(LogCook, Warning, TEXT("Failed to copy the binary PSO cache file %s to %s."), *PCPath, *PipelineCacheMetadataFileName);
+						}
+					}
+				}
+			}
+		}
+	}
+}
+
+void UCookOnTheFlyServer::SaveAndCloseGlobalShaderLibrary()
+{
+	const bool bCacheShaderLibraries = IsUsingShaderCodeLibrary();
+	if (bCacheShaderLibraries)
+	{
+		const TCHAR* GlobalShaderLibName = TEXT("Global");
+		FString ActualName = GenerateShaderCodeLibraryName(GlobalShaderLibName, IsCookFlagSet(ECookInitializationFlags::LegacyIterativeSharedBuild));
+
+		// Save shader code map - cleaning directories is deliberately a separate loop here as we open the cache once per shader platform and we don't assume that they can't be shared across target platforms.
+		for (const ITargetPlatform* TargetPlatform : PlatformManager->GetSessionPlatforms())
+		{
+			FinishPopulateShaderLibrary(TargetPlatform, GlobalShaderLibName);
+			SaveShaderLibrary(TargetPlatform, GlobalShaderLibName);
+		}
+
+		FShaderLibraryCooker::EndCookingLibrary(ActualName);
+	}
+}
+
+void UCookOnTheFlyServer::GetShaderLibraryPaths(const ITargetPlatform* TargetPlatform,
+	FString& OutShaderCodeDir, FString& OutMetaDataPath, bool bUseProjectDirForDLC)
+{
+	// TODO: Saving ShaderChunks into the DLC directory currently does not work, so we have the bUseProjectDirForDLC arg to save to Project
+	const FString BasePath = (!IsCookingDLC() || bUseProjectDirForDLC) ? FPaths::ProjectContentDir() : GetContentDirectoryForDLC();
+	OutShaderCodeDir = ConvertToFullSandboxPath(*BasePath, true, TargetPlatform->PlatformName());
+
+	const FString RootMetaDataPath = GetMetadataDirectory() / TEXT("PipelineCaches");
+	OutMetaDataPath = ConvertToFullSandboxPath(*RootMetaDataPath, true, *TargetPlatform->PlatformName());
+}
+
+void UCookOnTheFlyServer::FinishPopulateShaderLibrary(const ITargetPlatform* TargetPlatform, const FString& Name)
+{
+	FString ShaderCodeDir;
+	FString MetaDataPath;
+	GetShaderLibraryPaths(TargetPlatform, ShaderCodeDir, MetaDataPath);
+
+	FShaderLibraryCooker::FinishPopulateShaderLibrary(TargetPlatform, Name, ShaderCodeDir, MetaDataPath);
+}
+
+void UCookOnTheFlyServer::SaveShaderLibrary(const ITargetPlatform* TargetPlatform, const FString& Name)
+{
+	FString ShaderCodeDir;
+	FString MetaDataPath;
+	GetShaderLibraryPaths(TargetPlatform, ShaderCodeDir, MetaDataPath);
+
+	TArray<FString>& PlatformSCLCSVPaths = OutSCLCSVPaths.FindOrAdd(FName(TargetPlatform->PlatformName()));
+	FString ErrorString;
+	bool bHasData;
+	if (!FShaderLibraryCooker::SaveShaderLibraryWithoutChunking(TargetPlatform, Name, ShaderCodeDir, MetaDataPath,
+		PlatformSCLCSVPaths, ErrorString, bHasData))
+	{
+		// This is fatal - In this case we should cancel any launch on device operation or package write but we don't want to assert and crash the editor
+		LogCookerMessage(FString::Printf(TEXT("%s"), *ErrorString), EMessageSeverity::Error);
+	}
+	else if (bHasData)
+	{
+		for (const FString& Item : PlatformSCLCSVPaths)
+		{
+			UE_LOG(LogCook, Display, TEXT("Saved scl.csv %s for platform %s, %d bytes"), *Item, *TargetPlatform->PlatformName(),
+				IFileManager::Get().FileSize(*Item));
+		}
+	}
+}
+
+void UCookOnTheFlyServer::ShutdownShaderLibraryCookerAndCompilers(const FString& ProjectShaderLibraryName)
+{
+	FString ActualLibraryName = GenerateShaderCodeLibraryName(ProjectShaderLibraryName,
+		IsCookFlagSet(ECookInitializationFlags::LegacyIterativeSharedBuild));
+	FShaderLibraryCooker::EndCookingLibrary(ActualLibraryName);
+	FShaderLibraryCooker::Shutdown();
+	ShutdownShaderCompilers(PlatformManager->GetSessionPlatforms());
+}
+
+void UCookOnTheFlyServer::DumpShaderTypeStats(const FString& PlatformNameString)
+{
+	FShaderLibraryCooker::DumpShaderTypeStats(
+		GShaderCompilingManager->GetAbsoluteShaderDebugInfoDirectory(),
+		ConvertToFullSandboxPath(GetMetadataDirectory(), true, PlatformNameString));
+}
+
+bool UCookOnTheFlyServer::RecompileChangedShaders(const TArray<const ITargetPlatform*>& TargetPlatforms)
+{
+	bool bShadersRecompiled = false;
+	for (const ITargetPlatform* TargetPlatform : TargetPlatforms)
+	{
+		bShadersRecompiled |= RecompileChangedShadersForPlatform(TargetPlatform->PlatformName());
+	}
+	return bShadersRecompiled;
+}
\ No newline at end of file
diff --git a/Engine/Source/Editor/UnrealEd/Private/Cooker/ShaderLibraryCooking.h b/Engine/Source/Editor/UnrealEd/Private/Cooker/ShaderLibraryCooking.h
new file mode 100644
--- /dev/null
+++ b/Engine/Source/Editor/UnrealEd/Private/Cooker/ShaderLibraryCooking.h
@@ -0,0 +1,79 @@
+// Copyright Epic Games, Inc. All Rights Reserved.
+
+#pragma once
+
+#include "Cooker/CookArtifact.h"
+#include "Cooker/MPCollector.h"
+#include "CoreTypes.h"
+#include "UObject/NameTypes.h"
+
+class ITargetPlatform;
+class ITargetPlatformManagerModule;
+namespace UE::Cook { struct FTickStackData; }
+
+
+extern bool bDisableShaderCompilationDuringCookOnTheFly;
+extern bool bAllowIncompleteShaderMapsDuringCookOnTheFly;
+
+namespace UE::Cook
+{
+
+void TickShaderCompilingManager(UE::Cook::FTickStackData& StackData);
+
+/** CookMultiprocess collector for ShaderLibrary data. */
+class FShaderLibraryCollector : public IMPCollector
+{
+public:
+	virtual FGuid GetMessageType() const override;
+	virtual const TCHAR* GetDebugName() const override;
+
+	virtual void ClientTick(FMPCollectorClientTickContext& Context) override;
+	virtual void ServerReceiveMessage(FMPCollectorServerMessageContext& Context, FCbObjectView Message) override;
+
+	static FGuid MessageType;
+	static constexpr ANSICHAR RootObjectID[2]{ "S" };
+};
+
+/** CookArtifact to manage the shaderlibraries in the cook's OutputFolder. */
+class FShaderLibraryCookArtifact : public ICookArtifact
+{
+public:
+	FShaderLibraryCookArtifact(UCookOnTheFlyServer& InCOTFS);
+	virtual FString GetArtifactName() const;
+	virtual void OnFullRecook(const ITargetPlatform* TargetPlatform);
+
+	void CleanIntermediateFiles(const ITargetPlatform* TargetPlatform);
+
+private:
+	UCookOnTheFlyServer& COTFS;
+};
+
+} // namespace UE::Cook
+
+namespace UE::Cook::CVarControl
+{
+
+void UpdateShaderCookingCVars(ITargetPlatformManagerModule* TPM, int32 CookTimeCVarControl,
+	const ITargetPlatform* Platform, FName PlatformName);
+
+} // namespace UE::Cook::CVarControl
+
+
+///////////////////////////////////////////////////////
+// Inline implementations
+///////////////////////////////////////////////////////
+
+namespace UE::Cook
+{
+
+inline FGuid FShaderLibraryCollector::GetMessageType() const
+{
+	return MessageType;
+}
+
+inline const TCHAR* FShaderLibraryCollector::GetDebugName() const
+{
+	return TEXT("FShaderLibraryCollector");
+}
+
+} // namespace UE::Cook
\ No newline at end of file
