Fixed log typo

#rnx

diff --git a/Engine/Source/Editor/UnrealEd/Private/Commandlets/GatherTextCommandlet.cpp b/Engine/Source/Editor/UnrealEd/Private/Commandlets/GatherTextCommandlet.cpp
--- a/Engine/Source/Editor/UnrealEd/Private/Commandlets/GatherTextCommandlet.cpp
+++ b/Engine/Source/Editor/UnrealEd/Private/Commandlets/GatherTextCommandlet.cpp
@@ -687,7 +687,7 @@ bool UGatherTextCommandlet::ExecuteGatherTextPhases(const TSharedPtr<FLocalizati
 			return false;
 		}
 
-		UE_LOG(LogGatherTextCommandlet, Display, TEXT("Phase '%s' %.2f seconds (%d sequential tasks and %d parallel tasks)."), LexToString((EGatherTextCommandletPhase)PhaseIndex), FPlatformTime::Seconds() - PhaseExecutionStartTime, Phase.SequentialTasks.Num(), Phase.ParallelTasks.Num());
+		UE_LOG(LogGatherTextCommandlet, Display, TEXT("Phase '%s' took %.2f seconds (%d sequential tasks and %d parallel tasks)."), LexToString((EGatherTextCommandletPhase)PhaseIndex), FPlatformTime::Seconds() - PhaseExecutionStartTime, Phase.SequentialTasks.Num(), Phase.ParallelTasks.Num());
 
 		if ((EGatherTextCommandletPhase)PhaseIndex == EGatherTextCommandletPhase::PostUpdateManifests)
 		{
