Cooker: Logging: Add missing IsSuppressed call to log with dynamically chosen verbosity. The UE_LOG macro calls IsSuppressed; when we circumvent UE_LOG because we have a dynamic verbosity we need to handle calling it.
#rnx

diff --git a/Engine/Source/Editor/UnrealEd/Private/CookOnTheFlyServer.cpp b/Engine/Source/Editor/UnrealEd/Private/CookOnTheFlyServer.cpp
--- a/Engine/Source/Editor/UnrealEd/Private/CookOnTheFlyServer.cpp
+++ b/Engine/Source/Editor/UnrealEd/Private/CookOnTheFlyServer.cpp
@@ -4807,10 +4807,13 @@ void UCookOnTheFlyServer::ProcessUnsolicitedPackages(TArray<FName>* OutDiscovere
 				Instigator.Category != EInstigator::StartupSoftObjectPath)
 			{
 #if !NO_LOGGING
-				FMsg::Logf(__FILE__, __LINE__, LogCook.GetCategoryName(), UnexpectedLoadWarningSeverity,
-					TEXT("UnexpectedLoad: Package %s was loaded outside of any cook operation on a referencer package; we don't know why it was loaded. ")
-					TEXT("Adding it to the current cook, but it will possibly not be found in future incremental cooks."),
-					*PackageNameStr);
+				if (!LogCook.IsSuppressed(UnexpectedLoadWarningSeverity))
+				{
+					FMsg::Logf(__FILE__, __LINE__, LogCook.GetCategoryName(), UnexpectedLoadWarningSeverity,
+						TEXT("UnexpectedLoad: Package %s was loaded outside of any cook operation on a referencer package; we don't know why it was loaded. ")
+						TEXT("Adding it to the current cook, but it will possibly not be found in future incremental cooks."),
+						*PackageNameStr);
+				}
 #endif
 			}
 			QueueDiscoveredPackage(*PackageData, FInstigator(Instigator), EDiscoveredPlatformSet::CopyFromInstigator);
@@ -4871,11 +4874,14 @@ void UCookOnTheFlyServer::ProcessUnsolicitedPackages(TArray<FName>* OutDiscovere
 					*InstancedLeafLoadedName.ToString(), *Instigator.Referencer.ToString());
 			}
 #if !NO_LOGGING
-			FMsg::Logf(__FILE__, __LINE__, LogCook.GetCategoryName(), UnexpectedLoadWarningSeverity,
-				TEXT("UnexpectedLoad: TargetPackage %s was unexpectededly loaded during the cook of SourcePackage %s; it was not declared as a hard dependency (aka Import) of the SourcePackage. ")
-				TEXT("We have to conservatively add this package to the cook, but it might not be needed at runtime. ")
-				TEXT("Declare the TargetPackage as an import during editor save of the SourcePackage, or mark up its load with FCookLoadScope to specify whether it is runtime or editoronly."),
-				*PackageNameStr, *DescriptionOfReferencer);
+			if (!LogCook.IsSuppressed(UnexpectedLoadWarningSeverity))
+			{
+				FMsg::Logf(__FILE__, __LINE__, LogCook.GetCategoryName(), UnexpectedLoadWarningSeverity,
+					TEXT("UnexpectedLoad: TargetPackage %s was unexpectededly loaded during the cook of SourcePackage %s; it was not declared as a hard dependency (aka Import) of the SourcePackage. ")
+					TEXT("We have to conservatively add this package to the cook, but it might not be needed at runtime. ")
+					TEXT("Declare the TargetPackage as an import during editor save of the SourcePackage, or mark up its load with FCookLoadScope to specify whether it is runtime or editoronly."),
+					*PackageNameStr, *DescriptionOfReferencer);
+			}
 #endif
 		}
 
