[UE] Ensure that skeleton class references are cleared on forced deletion of a Blueprint class asset.

#FailedValidations Preflight

diff --git a/Engine/Source/Editor/UnrealEd/Private/ObjectTools.cpp b/Engine/Source/Editor/UnrealEd/Private/ObjectTools.cpp
--- a/Engine/Source/Editor/UnrealEd/Private/ObjectTools.cpp
+++ b/Engine/Source/Editor/UnrealEd/Private/ObjectTools.cpp
@@ -3814,6 +3814,12 @@ namespace ObjectTools
 								ObjectsToReplace.AddUnique(BlueprintObject->GeneratedClass);
 							}
 
+							// Also include the skeleton class, in case it is referenced
+							if (BlueprintObject->SkeletonGeneratedClass)
+							{
+								ObjectsToReplace.AddUnique(BlueprintObject->SkeletonGeneratedClass);
+							}
+
 							// Reparent any direct children to the parent class of the blueprint that's about to be deleted
 							if (BlueprintObject->ParentClass != nullptr)
 							{
