Disabling parallelization for the report writing job while failures are investigated

#rnx

diff --git a/Engine/Source/Editor/UnrealEd/Classes/Commandlets/GenerateTextLocalizationReportCommandlet.h b/Engine/Source/Editor/UnrealEd/Classes/Commandlets/GenerateTextLocalizationReportCommandlet.h
--- a/Engine/Source/Editor/UnrealEd/Classes/Commandlets/GenerateTextLocalizationReportCommandlet.h
+++ b/Engine/Source/Editor/UnrealEd/Classes/Commandlets/GenerateTextLocalizationReportCommandlet.h
@@ -18,7 +18,7 @@ public:
 	//~ End UCommandlet Interface
 
 	virtual EGatherTextCommandletPhase GetPhase() const override { return EGatherTextCommandletPhase::Finalize; }
-	virtual bool CanParallelize() const { return true; }
+	//virtual bool CanParallelize() const { return true; }
 
 	virtual bool ShouldRunInPreview(const TArray<FString>& Switches, const TMap<FString, FString>& ParamVals) const
 	{
