Editor: Call PostLoad after PostDuplicate when duplicating objects in the editor

diff --git a/Engine/Source/Editor/UnrealEd/Private/ObjectTools.cpp b/Engine/Source/Editor/UnrealEd/Private/ObjectTools.cpp
--- a/Engine/Source/Editor/UnrealEd/Private/ObjectTools.cpp
+++ b/Engine/Source/Editor/UnrealEd/Private/ObjectTools.cpp
@@ -1019,6 +1019,8 @@ namespace ObjectTools
 				Params.CreatedObjects = &CreatedObjects;
 			}
 
+			// Skip postload so we can call it later - this ensures PostDuplicate is called on all subobjects before any PostLoad is performed
+			Params.bSkipPostLoad = true;
 			DupObject = StaticDuplicateObjectEx(Params);
 
 			if (DuplicatedObjects)
@@ -1034,6 +1036,9 @@ namespace ObjectTools
 
 		if( DupObject )
 		{
+			// Now call PostLoad (since we disabled it with bSkipPostLoad above)
+			DupObject->ConditionalPostLoad();
+
 			ReplacementMap.Add( Object, DupObject );
 			DupObject->MarkPackageDirty();
 
@@ -3103,13 +3108,12 @@ namespace ObjectTools
 			return 0;
 		}
 
-		FResultMessage Result;
-		Result.bSuccess = true;
-		FEditorDelegates::OnPreDestructiveAssetAction.Broadcast(ObjectsToDelete, EDestructiveAssetActions::AssetDelete, Result);
+		FOnPreDestructiveAssetActionArgs PreDestructiveArgs(ObjectsToDelete, EDestructiveAssetActions::AssetDelete);
+		FEditorDelegates::OnPreDestructiveAssetActionWithArgs.Broadcast(PreDestructiveArgs);
 
-		if (!Result.bSuccess)
+		if (!PreDestructiveArgs.Result.bSuccess)
 		{
-			UE_LOG(LogObjectTools, Warning, TEXT("%s"), *Result.ErrorMessage);
+			UE_LOG(LogObjectTools, Warning, TEXT("%s"), *PreDestructiveArgs.Result.ErrorMessage);
 			return 0;
 		}
 
@@ -3210,13 +3214,12 @@ namespace ObjectTools
 			return 0;
 		}
 
-		FResultMessage Result;
-		Result.bSuccess = true;
-		FEditorDelegates::OnPreDestructiveAssetAction.Broadcast(ObjectsToPrivatize, EDestructiveAssetActions::AssetPrivatize, Result);
+		FOnPreDestructiveAssetActionArgs PreDestructiveArgs(ObjectsToPrivatize, EDestructiveAssetActions::AssetPrivatize);
+		FEditorDelegates::OnPreDestructiveAssetActionWithArgs.Broadcast(PreDestructiveArgs);
 
-		if (!Result.bSuccess)
+		if (!PreDestructiveArgs.Result.bSuccess)
 		{
-			UE_LOG(LogObjectTools, Warning, TEXT("%s"), *Result.ErrorMessage);
+			UE_LOG(LogObjectTools, Warning, TEXT("%s"), *PreDestructiveArgs.Result.ErrorMessage);
 			return 0;
 		}
 
