[Settings] Expand Configuration do not ask for engine restart
ExpandConfiguration does not ask for an engine restart even though its required for some entry to change

#FailedValidations Preflight

diff --git a/Engine/Source/Editor/UnrealEd/Classes/Settings/EditorStyleSettings.h b/Engine/Source/Editor/UnrealEd/Classes/Settings/EditorStyleSettings.h
--- a/Engine/Source/Editor/UnrealEd/Classes/Settings/EditorStyleSettings.h
+++ b/Engine/Source/Editor/UnrealEd/Classes/Settings/EditorStyleSettings.h
@@ -152,7 +152,7 @@ public:
 	uint32 bShowNativeComponentNames:1;
 
 	/** When enabled, the Editor Preferences and Project Settings menu items in the main menu will be expanded with sub-menus for each settings section. */
-	UPROPERTY(EditAnywhere, config, Category=UserInterface, AdvancedDisplay)
+	UPROPERTY(EditAnywhere, config, Category=UserInterface, AdvancedDisplay, meta = (ConfigRestartRequired = true))
 	uint32 bExpandConfigurationMenus:1;
 
 	/** When enabled, the project subsection of the File menu will be shown. */
