Fix compile error from robomerge issue

#rnx

diff --git a/Engine/Source/Editor/UnrealEd/Private/ObjectTools.cpp b/Engine/Source/Editor/UnrealEd/Private/ObjectTools.cpp
--- a/Engine/Source/Editor/UnrealEd/Private/ObjectTools.cpp
+++ b/Engine/Source/Editor/UnrealEd/Private/ObjectTools.cpp
@@ -3108,12 +3108,13 @@ namespace ObjectTools
 			return 0;
 		}
 
-		FOnPreDestructiveAssetActionArgs PreDestructiveArgs(ObjectsToDelete, EDestructiveAssetActions::AssetDelete);
-		FEditorDelegates::OnPreDestructiveAssetActionWithArgs.Broadcast(PreDestructiveArgs);
+		FResultMessage Result;
+		Result.bSuccess = true;
+		FEditorDelegates::OnPreDestructiveAssetAction.Broadcast(ObjectsToDelete, EDestructiveAssetActions::AssetDelete, Result);
 
-		if (!PreDestructiveArgs.Result.bSuccess)
+		if (!Result.bSuccess)
 		{
-			UE_LOG(LogObjectTools, Warning, TEXT("%s"), *PreDestructiveArgs.Result.ErrorMessage);
+			UE_LOG(LogObjectTools, Warning, TEXT("%s"), *Result.ErrorMessage);
 			return 0;
 		}
 
@@ -3214,12 +3215,13 @@ namespace ObjectTools
 			return 0;
 		}
 
-		FOnPreDestructiveAssetActionArgs PreDestructiveArgs(ObjectsToPrivatize, EDestructiveAssetActions::AssetPrivatize);
-		FEditorDelegates::OnPreDestructiveAssetActionWithArgs.Broadcast(PreDestructiveArgs);
+		FResultMessage Result;
+		Result.bSuccess = true;
+		FEditorDelegates::OnPreDestructiveAssetAction.Broadcast(ObjectsToPrivatize, EDestructiveAssetActions::AssetPrivatize, Result);
 
-		if (!PreDestructiveArgs.Result.bSuccess)
+		if (!Result.bSuccess)
 		{
-			UE_LOG(LogObjectTools, Warning, TEXT("%s"), *PreDestructiveArgs.Result.ErrorMessage);
+			UE_LOG(LogObjectTools, Warning, TEXT("%s"), *Result.ErrorMessage);
 			return 0;
 		}
 
