Cooker: Suppress spurious warning in FinishGeneratedPlatformSave when a dependency reported by the splitter has a dependency of its own that does not exist. Report such missing dependencies instead at display verbosity, during the CookRequestCluster's exploration of the dependency graph.
#rnx

#FailedValidations Preflight

diff --git a/Engine/Source/Editor/UnrealEd/Private/Cooker/CookGenerationHelper.cpp b/Engine/Source/Editor/UnrealEd/Private/Cooker/CookGenerationHelper.cpp
--- a/Engine/Source/Editor/UnrealEd/Private/Cooker/CookGenerationHelper.cpp
+++ b/Engine/Source/Editor/UnrealEd/Private/Cooker/CookGenerationHelper.cpp
@@ -1673,6 +1673,7 @@ void FGenerationHelper::FinishGeneratedPlatformSave(FPackageData& PackageData, c
 
 	// External packages are detected by their presence in the NeverCookPackageList 
 	FThreadSafeSet<FName>& NeverCookPackageList = GetOwner().GetPackageDatas().GetCookOnTheFlyServer().PackageTracker->NeverCookPackageList; 
+	IAssetRegistry& AssetRegistry = IAssetRegistry::GetChecked();
 
 	for (const FAssetDependency& Dependency : Info->PackageDependencies)
 	{
@@ -1681,12 +1682,12 @@ void FGenerationHelper::FinishGeneratedPlatformSave(FPackageData& PackageData, c
 		if (bIsBuildDependency)
 		{
 			TArray<FAssetDependency> PackageDependencies;
-			IAssetRegistry::GetChecked().GetDependencies(FAssetIdentifier(Dependency.AssetId.PackageName), PackageDependencies,
+			AssetRegistry.GetDependencies(FAssetIdentifier(Dependency.AssetId.PackageName), PackageDependencies,
 				UE::AssetRegistry::EDependencyCategory::Package, UE::AssetRegistry::EDependencyQuery::Game);
 
 			for (const FAssetDependency& PackageDependency : PackageDependencies)
 			{
-				TStringBuilder<256> PackakgeDependencyName(InPlace, PackageDependency.AssetId.PackageName);
+				TStringBuilder<256> PackageDependencyName(InPlace, PackageDependency.AssetId.PackageName);
 
 				// TODO : This is a special case for world splitter. If the package is a level, ignore it. A level can contain another level. 
 				// But they are already included in the generated package so there is no need to add them as runtime dependencies. It would trigger 
@@ -1694,14 +1695,19 @@ void FGenerationHelper::FinishGeneratedPlatformSave(FPackageData& PackageData, c
 				// find the correct dependencies directly so I don't need any special case here.
 				
 				// If it's a script package, add it directly. Calling GetLocalFullPath to check if it's a map will trigger a warning.
-				if (FPackageName::IsScriptPackage(PackakgeDependencyName))
+				if (FPackageName::IsScriptPackage(PackageDependencyName))
 				{
 					OutDependencies.Add(PackageDependency);
 					continue;
 				}
 				
-				FPackagePath DependencyPackagePath = FPackagePath::FromPackageNameChecked(PackakgeDependencyName);
-				EPackageExtension Extension = FPackagePath::ParseExtension(DependencyPackagePath.GetLocalFullPath());
+				FString ExtensionStr;
+				if (!AssetRegistry.DoesPackageExistOnDisk(PackageDependency.AssetId.PackageName, nullptr, &ExtensionStr))
+				{
+					// No need to add dependencies to packages that do not exist on disk (not even generated packages).
+					continue;
+				}
+				EPackageExtension Extension = FPackagePath::ParseExtension(ExtensionStr);
 				if (Extension == EPackageExtension::Map)
 				{
 					continue;
diff --git a/Engine/Source/Editor/UnrealEd/Private/Cooker/CookRequestCluster.cpp b/Engine/Source/Editor/UnrealEd/Private/Cooker/CookRequestCluster.cpp
--- a/Engine/Source/Editor/UnrealEd/Private/Cooker/CookRequestCluster.cpp
+++ b/Engine/Source/Editor/UnrealEd/Private/Cooker/CookRequestCluster.cpp
@@ -29,6 +29,7 @@
 #include "Engine/Level.h"
 #include "HAL/Event.h"
 #include "HAL/PlatformProcess.h"
+#include "Logging/StructuredLog.h"
 #include "Interfaces/ITargetPlatform.h"
 #include "Misc/RedirectCollector.h"
 #include "Misc/ReverseIterate.h"
@@ -2011,6 +2012,12 @@ void FRequestCluster::FGraphSearch::FExploreEdgesContext::QueueVisitsOfDependenc
 			GenerationHelper.GetReference());
 		if (!DependencyVertex.GetPackageData())
 		{
+			TStringBuilder<256> DependencyNameStr(InPlace, DependencyName);
+			if (!FPackageName::IsScriptPackage(DependencyNameStr))
+			{
+				UE_LOGFMT(LogCook, Display, "Package {PackageName} has a dependency on package {MissingDependency} which does not exist.",
+					PackageName.ToString(), *DependencyNameStr);
+			}
 			continue;
 		}
 		FPackageData& DependencyPackageData(*DependencyVertex.GetPackageData());
