BuildLayers and MultiProductBuild: Fix bug with Archive.IsCookingDLC: It is supposed to return ECookingDLC::No when cooking the new BuildLayer prototype (commandline flag ), but it was using the old IsCookingDLC function that does not consider that override.
This was causing Materials to incorrectly set bIncludedInBaseGame=false in MultiProductBuild.
#rnx

diff --git a/Engine/Source/Editor/UnrealEd/Private/Cooker/CookSavePackage.cpp b/Engine/Source/Editor/UnrealEd/Private/Cooker/CookSavePackage.cpp
--- a/Engine/Source/Editor/UnrealEd/Private/Cooker/CookSavePackage.cpp
+++ b/Engine/Source/Editor/UnrealEd/Private/Cooker/CookSavePackage.cpp
@@ -265,7 +265,7 @@ void FSaveCookedPackageContext::SetupPlatform(const ITargetPlatform* InTargetPla
 
 	ArchiveCookContext.Emplace(Package,
 		COTFS.IsDirectorCookByTheBook() ? UE::Cook::ECookType::ByTheBook : UE::Cook::ECookType::OnTheFly,
-		COTFS.IsCookingDLC() ? UE::Cook::ECookingDLC::Yes : UE::Cook::ECookingDLC::No,
+		COTFS.GetCookingDLC(),
 		TargetPlatform, &COTFS);
 
 	if (CommitType == EReachability::Runtime)
