[AutoViz] GeoInsp Modes not available in the Blueprint Editor
- Added AutoViz geometry inspection modes to the Kismet editor
- Fixed UnrealEd V iewprt Toolbar to allow wireframe support without unlit wireframe

diff --git a/Engine/Source/Editor/Kismet/Private/SSCSEditorViewport.cpp b/Engine/Source/Editor/Kismet/Private/SSCSEditorViewport.cpp
--- a/Engine/Source/Editor/Kismet/Private/SSCSEditorViewport.cpp
+++ b/Engine/Source/Editor/Kismet/Private/SSCSEditorViewport.cpp
@@ -67,6 +67,10 @@ bool IsViewModeSupported(EViewModeIndex InViewModeIndex)
 	case VMI_Lit:
 	case VMI_BrushWireframe:
 	case VMI_CollisionVisibility:
+	case VMI_Clay:
+	case VMI_Zebra:
+	case VMI_FrontBackFace:
+	case VMI_RandomColor:
 		return true;
 	default:
 		return false;
diff --git a/Engine/Source/Editor/UnrealEd/Private/ViewportToolbar/UnrealEdViewportToolbar.cpp b/Engine/Source/Editor/UnrealEd/Private/ViewportToolbar/UnrealEdViewportToolbar.cpp
--- a/Engine/Source/Editor/UnrealEd/Private/ViewportToolbar/UnrealEdViewportToolbar.cpp
+++ b/Engine/Source/Editor/UnrealEd/Private/ViewportToolbar/UnrealEdViewportToolbar.cpp
@@ -2184,16 +2184,13 @@ void PopulateViewModesMenu(UToolMenu* InMenu)
 							"GeometryInspectionWireframe", LOCTEXT("GeometryInspectionWireframeHeader", "Wireframe")
 						);
 
-						if (!IsViewModeSupported.IsBound() || IsViewModeSupported.Execute(VMI_Lit_Wireframe))
-						{
-							Section.AddMenuEntry(FEditorViewportCommands::Get().GeometryInspectionShowWireframe);
+						Section.AddMenuEntry(FEditorViewportCommands::Get().GeometryInspectionShowWireframe);
 
-							if (TSharedPtr<SEditorViewport> EditorViewport = EditorViewportWeakPtr.Pin())
-							{
-								Section.AddEntry(FToolMenuEntry::InitWidget(
-							"WireframeOpacity", EditorViewport->BuildWireframeMenu(), LOCTEXT("WireframeOpacity", "Opacity")
-								));
-							}
+						if (TSharedPtr<SEditorViewport> EditorViewport = EditorViewportWeakPtr.Pin())
+						{
+							Section.AddEntry(FToolMenuEntry::InitWidget(
+						"WireframeOpacity", EditorViewport->BuildWireframeMenu(), LOCTEXT("WireframeOpacity", "Opacity")
+							));
 						}
 					}
 				}
