Fixed locked camera reset issue introduced by CL 39915852.

This incorrectly caused our piloted camera to reset to the origin.

diff --git a/Engine/Source/Editor/UnrealEd/Private/LevelEditorViewport.cpp b/Engine/Source/Editor/UnrealEd/Private/LevelEditorViewport.cpp
--- a/Engine/Source/Editor/UnrealEd/Private/LevelEditorViewport.cpp
+++ b/Engine/Source/Editor/UnrealEd/Private/LevelEditorViewport.cpp
@@ -4338,10 +4338,10 @@ void FLevelEditorViewportClient::MoveCameraToLockedActor()
 				FMinimalViewInfo ViewInfo;
 				if (ensure(ViewComponent->GetEditorPreviewInfo(0.f, ViewInfo)))
 				{
-					ViewFOV = ControllingActorViewInfo.FOV;
-					AspectRatio = ControllingActorViewInfo.AspectRatio;
-					SetViewLocation(ControllingActorViewInfo.Location);
-					SetViewRotation(ControllingActorViewInfo.Rotation);
+					ViewFOV = ViewInfo.FOV;
+					AspectRatio = ViewInfo.AspectRatio;
+					SetViewLocation(ViewInfo.Location);
+					SetViewRotation(ViewInfo.Rotation);
 					bGotCameraView = true;
 				}
 			}
