MPCook: Fix error "GenerationHelper for package %s is still allocated and initialized at end of cooksession.", in the case that a generator package is retracted from a cookworker in between the completion of QueueGeneratedPackages on the cookworker and the reporting of that call to the director.
#rnx
#FailedValidations Preflight

diff --git a/Engine/Source/Editor/UnrealEd/Private/Cooker/CookWorkerClient.cpp b/Engine/Source/Editor/UnrealEd/Private/Cooker/CookWorkerClient.cpp
--- a/Engine/Source/Editor/UnrealEd/Private/Cooker/CookWorkerClient.cpp
+++ b/Engine/Source/Editor/UnrealEd/Private/Cooker/CookWorkerClient.cpp
@@ -103,8 +103,8 @@ void FCookWorkerClient::TickFromSchedulerThread(FTickStackData& StackData)
 		PumpReceiveMessages();
 		if (ConnectStatus == EConnectStatus::Connected)
 		{
-			SendPendingResults();
-			PumpSendMessages();
+			SendPendingMessages();
+			PumpSendBuffer();
 			TickCollectors(StackData, false /* bFlush */);
 		}
 	}
@@ -638,7 +638,7 @@ void FCookWorkerClient::LogConnected()
 	UE_LOG(LogCook, Display, TEXT("Connection to CookDirector successful."));
 }
 
-void FCookWorkerClient::PumpSendMessages()
+void FCookWorkerClient::PumpSendBuffer()
 {
 	UE::CompactBinaryTCP::EConnectionStatus Status = UE::CompactBinaryTCP::TryFlushBuffer(ServerSocket, SendBuffer);
 	UpdateSocketSendDiagnostics(Status);
@@ -650,8 +650,9 @@ void FCookWorkerClient::PumpSendMessages()
 	}
 }
 
-void FCookWorkerClient::SendPendingResults()
+void FCookWorkerClient::SendPendingMessages()
 {
+	TArray<UE::CompactBinaryTCP::FMarshalledMessage> LocalCopyQueuedMessages;
 	FPackageResultsMessage Message;
 	{
 		FScopeLock PendingResultsScopeLock(&PendingResultsLock);
@@ -664,6 +665,8 @@ void FCookWorkerClient::SendPendingResults()
 			}
 			PendingResults.Reset();
 		}
+		LocalCopyQueuedMessages = MoveTemp(QueuedMessages);
+		QueuedMessages.Empty();
 	}
 	if (!Message.Results.IsEmpty())
 	{
@@ -697,6 +700,11 @@ void FCookWorkerClient::SendPendingResults()
 	{
 		SendMessage(FileTransferMessage);
 	}
+
+	for (UE::CompactBinaryTCP::FMarshalledMessage& QueuedMessage: LocalCopyQueuedMessages)
+	{
+		SendMessage(MoveTemp(QueuedMessage));
+	}
 }
 
 void FCookWorkerClient::PumpReceiveMessages()
@@ -862,7 +870,7 @@ void FCookWorkerClient::PumpDisconnect(FTickStackData& StackData)
 			PumpReceiveMessages();
 			if (ConnectStatus == EConnectStatus::WaitForAbortAcknowledge)
 			{
-				PumpSendMessages();
+				PumpSendBuffer();
 
 				constexpr float WaitForDisconnectTimeout = 60.f;
 				if (FPlatformTime::Seconds() - ConnectStartTimeSeconds > WaitForDisconnectTimeout
@@ -902,11 +910,21 @@ void FCookWorkerClient::PumpDisconnect(FTickStackData& StackData)
 
 void FCookWorkerClient::SendMessage(const IMPCollectorMessage& Message)
 {
-	using namespace UE::CompactBinaryTCP;
-	EConnectionStatus Status = TryWritePacket(ServerSocket, SendBuffer, MarshalToCompactBinaryTCP(Message));
+	SendMessage(MarshalToCompactBinaryTCP(Message));
+}
+
+void FCookWorkerClient::SendMessage(UE::CompactBinaryTCP::FMarshalledMessage&& Message)
+{
+	UE::CompactBinaryTCP::EConnectionStatus Status = TryWritePacket(ServerSocket, SendBuffer, MoveTemp(Message));
 	UpdateSocketSendDiagnostics(Status);
 }
 
+void FCookWorkerClient::QueueMessage(UE::CompactBinaryTCP::FMarshalledMessage&& Message)
+{
+	FScopeLock PendingResultsScopeLock(&PendingResultsLock);
+	QueuedMessages.Add(MoveTemp(Message));
+}
+
 void FCookWorkerClient::SendToState(EConnectStatus TargetStatus)
 {
 	switch (TargetStatus)
@@ -1205,9 +1223,15 @@ void FCookWorkerClient::HandleRetractionMessage(FMPCollectorClientMessageContext
 	UE_LOG(LogCook, Display, TEXT("Retraction message received from director. %d packages retracted."),
 		PackageNames.Num());
 
+	// Queue the message for sending rather than sending it immediately. We need it to be processed on the director
+	// AFTER any PendingResults or GeneratorEvent messages. Most notably, if we have sent
+	// EGeneratorEvent::QueuedGeneratedPackages for a retracted generator package, we need the director to process
+	// that before it processes the retraction message, otherwise it might assign the generator to itself, finish the
+	// save of the generator package, and then receive our QueuedGeneratedPackages and call
+	// EndQueueGeneratedPackagesOnDirector, which is invalid to receive after a generator package has finished saving.
 	FRetractionResultsMessage ResultsMessage;
 	ResultsMessage.ReturnedPackages = MoveTemp(PackageNames);
-	SendMessage(ResultsMessage);
+	QueueMessage(MarshalToCompactBinaryTCP(ResultsMessage));
 }
 
 void FCookWorkerClient::HandleHeartbeatMessage(FMPCollectorClientMessageContext& Context, bool bReadSuccessful,
diff --git a/Engine/Source/Editor/UnrealEd/Private/Cooker/CookWorkerClient.h b/Engine/Source/Editor/UnrealEd/Private/Cooker/CookWorkerClient.h
--- a/Engine/Source/Editor/UnrealEd/Private/Cooker/CookWorkerClient.h
+++ b/Engine/Source/Editor/UnrealEd/Private/Cooker/CookWorkerClient.h
@@ -112,10 +112,13 @@ private:
 	/** Wait for the Config Message, switch state when it succeeds or fails. */
 	void PollReceiveConfigMessage();
 	void LogConnected();
-	/** Helper for Tick, pump send messages to the Server. */
-	void PumpSendMessages();
-	/** Helper for Tick, send a message for any pending package results that we have. */
-	void SendPendingResults();
+	/**
+	 * Helper for Tick. Gather messages from pending package results, queued messages, and any other sources, and add
+	 * them to the SendBuffer.
+	 */
+	void SendPendingMessages();
+	/** Helper for Tick, pump the messages that have been added to the buffer to send to the Server. */
+	void PumpSendBuffer();
 	/** Helper for Tick, pump receive messages from the Server. */
 	void PumpReceiveMessages();
 	/** Helper for PumpReceiveMessages: dispatch the messages received from the socket. */
@@ -124,6 +127,12 @@ private:
 	void PumpDisconnect(FTickStackData& StackData);
 	/** Send the message immediately to the Socket. If cannot complete immediately, it will be finished during Tick. */
 	void SendMessage(const IMPCollectorMessage& Message);
+	void SendMessage(UE::CompactBinaryTCP::FMarshalledMessage&& Message);
+	/**
+	 * Queue a message for sending during the next SendPendingMessages. @see QueuedMessages for how it is ordered
+	 * with other messages.
+	 */
+	void QueueMessage(UE::CompactBinaryTCP::FMarshalledMessage&& Message);
 	/** Send this into the given state. Update any state-dependent variables. */
 	void SendToState(EConnectStatus TargetStatus);
 	void LogInvalidMessage(const TCHAR* MessageTypeName);
@@ -186,6 +195,13 @@ private:
 	FCriticalSection PendingResultsLock;
 	TArray<TUniquePtr<FPackageRemoteResult>> PendingResults;
 	TMap<FPackageRemoteResult*, FPendingResultNeedingAsyncWork> PendingResultsNeedingAsyncWork;
+	/**
+	 * QueuedMessages are added to the SendBuffer during Tick, after PendingResults, GeneratorEvents, and
+	 * FileTransfers. QueuedMessages are sent in the order in which they are queued, but they are unordered
+	 * with respect to PendingResults/GeneratorEvents/FileTransfers. A QueuedMessage can be sent after a
+	 * PendingResult/GeneratorEvent/FileTransfer even if it was queued before it.
+	 */
+	TArray<UE::CompactBinaryTCP::FMarshalledMessage> QueuedMessages;
 
 	// Variables Read/Write only within FileTransferLock
 	FCriticalSection FileTransferLock;
