Backed out 46059122  as it affects sub object duplication

#rnx

diff --git a/Engine/Source/Editor/UnrealEd/Private/ObjectTools.cpp b/Engine/Source/Editor/UnrealEd/Private/ObjectTools.cpp
--- a/Engine/Source/Editor/UnrealEd/Private/ObjectTools.cpp
+++ b/Engine/Source/Editor/UnrealEd/Private/ObjectTools.cpp
@@ -1019,8 +1019,6 @@ namespace ObjectTools
 				Params.CreatedObjects = &CreatedObjects;
 			}
 
-			// Skip postload so we can call it later - this ensures PostDuplicate is called on all subobjects before any PostLoad is performed
-			Params.bSkipPostLoad = true;
 			DupObject = StaticDuplicateObjectEx(Params);
 
 			if (DuplicatedObjects)
@@ -1036,9 +1034,6 @@ namespace ObjectTools
 
 		if( DupObject )
 		{
-			// Now call PostLoad (since we disabled it with bSkipPostLoad above)
-			DupObject->ConditionalPostLoad();
-
 			ReplacementMap.Add( Object, DupObject );
 			DupObject->MarkPackageDirty();
 
