Re-enable parallel localization report generation

* Downgraded the report generation failures to warning, as this step is optional.
* Updated GenerateTextLocalizationReportCommandlet to use structured logging.
* Updated Localize to skip GenerateReports if running Gather, as Gather runs a superset of the GenerateReports step.

#rnx

diff --git a/Engine/Source/Editor/UnrealEd/Classes/Commandlets/GenerateTextLocalizationReportCommandlet.h b/Engine/Source/Editor/UnrealEd/Classes/Commandlets/GenerateTextLocalizationReportCommandlet.h
--- a/Engine/Source/Editor/UnrealEd/Classes/Commandlets/GenerateTextLocalizationReportCommandlet.h
+++ b/Engine/Source/Editor/UnrealEd/Classes/Commandlets/GenerateTextLocalizationReportCommandlet.h
@@ -18,7 +18,7 @@ public:
 	//~ End UCommandlet Interface
 
 	virtual EGatherTextCommandletPhase GetPhase() const override { return EGatherTextCommandletPhase::Finalize; }
-	//virtual bool CanParallelize() const { return true; }
+	virtual bool CanParallelize() const { return true; }
 
 	virtual bool ShouldRunInPreview(const TArray<FString>& Switches, const TMap<FString, FString>& ParamVals) const
 	{
diff --git a/Engine/Source/Editor/UnrealEd/Private/Commandlets/GenerateTextLocalizationReportCommandlet.cpp b/Engine/Source/Editor/UnrealEd/Private/Commandlets/GenerateTextLocalizationReportCommandlet.cpp
--- a/Engine/Source/Editor/UnrealEd/Private/Commandlets/GenerateTextLocalizationReportCommandlet.cpp
+++ b/Engine/Source/Editor/UnrealEd/Private/Commandlets/GenerateTextLocalizationReportCommandlet.cpp
@@ -11,6 +11,7 @@
 #include "LocalizationSourceControlUtil.h"
 #include "Logging/LogCategory.h"
 #include "Logging/LogMacros.h"
+#include "Logging/StructuredLog.h"
 #include "Misc/CString.h"
 #include "Misc/ConfigCacheIni.h"
 #include "Misc/DateTime.h"
@@ -21,6 +22,10 @@
 #include UE_INLINE_GENERATED_CPP_BY_NAME(GenerateTextLocalizationReportCommandlet)
 
 DEFINE_LOG_CATEGORY_STATIC(LogGenerateTextLocalizationReportCommandlet, Log, All);
+namespace GenerateTextLocalizationReportCommandlet
+{
+	static constexpr int32 LocalizationLogIdentifier = 304;
+}
 
 UGenerateTextLocalizationReportCommandlet::UGenerateTextLocalizationReportCommandlet(const FObjectInitializer& ObjectInitializer)
 	: Super(ObjectInitializer)
@@ -45,7 +50,9 @@ int32 UGenerateTextLocalizationReportCommandlet::Main(const FString& Params)
 	}
 	else
 	{
-		UE_LOG(LogGenerateTextLocalizationReportCommandlet, Error, TEXT("No config specified."));
+		UE_LOGFMT(LogGenerateTextLocalizationReportCommandlet, Error, "No config specified.",
+			("id", GenerateTextLocalizationReportCommandlet::LocalizationLogIdentifier)
+		);
 	}
 
 	// Set config section.
@@ -57,7 +64,9 @@ int32 UGenerateTextLocalizationReportCommandlet::Main(const FString& Params)
 	}
 	else
 	{
-		UE_LOG(LogGenerateTextLocalizationReportCommandlet, Error, TEXT("No config section specified."));
+		UE_LOGFMT(LogGenerateTextLocalizationReportCommandlet, Error, "No config section specified.",
+			("id", GenerateTextLocalizationReportCommandlet::LocalizationLogIdentifier)
+		);
 		return -1;
 	}
 
@@ -74,14 +83,18 @@ int32 UGenerateTextLocalizationReportCommandlet::Main(const FString& Params)
 	// Get source path.
 	if( !( GetPathFromConfig( *SectionName, TEXT("SourcePath"), SourcePath, GatherTextConfigPath ) ) )
 	{
-		UE_LOG(LogGenerateTextLocalizationReportCommandlet, Error, TEXT("No source path specified."));
+		UE_LOGFMT(LogGenerateTextLocalizationReportCommandlet, Error, "No source path specified.",
+			("id", GenerateTextLocalizationReportCommandlet::LocalizationLogIdentifier)
+		);
 		return -1;
 	}
 
 	// Get destination path.
 	if( !( GetPathFromConfig( *SectionName, TEXT("DestinationPath"), DestinationPath, GatherTextConfigPath ) ) )
 	{
-		UE_LOG(LogGenerateTextLocalizationReportCommandlet, Error, TEXT("No destination path specified."));
+		UE_LOGFMT(LogGenerateTextLocalizationReportCommandlet, Error, "No destination path specified.",
+			("id", GenerateTextLocalizationReportCommandlet::LocalizationLogIdentifier)
+		);
 		return -1;
 	}
 
@@ -99,8 +112,9 @@ int32 UGenerateTextLocalizationReportCommandlet::Main(const FString& Params)
 	{
 		if( !ProcessWordCountReport( SourcePath, DestinationPath ) )
 		{
-			UE_LOG(LogGenerateTextLocalizationReportCommandlet, Error, TEXT("Failed to generate word count report."));
-			return -1;
+			UE_LOGFMT(LogGenerateTextLocalizationReportCommandlet, Warning, "Failed to generate word count report.",
+				("id", GenerateTextLocalizationReportCommandlet::LocalizationLogIdentifier)
+			);
 		}
 	}
 
@@ -108,8 +122,9 @@ int32 UGenerateTextLocalizationReportCommandlet::Main(const FString& Params)
 	{
 		if( !ProcessConflictReport( DestinationPath) )
 		{
-			UE_LOG(LogGenerateTextLocalizationReportCommandlet, Error, TEXT("Failed to generate localization conflict report."));
-			return -1;
+			UE_LOGFMT(LogGenerateTextLocalizationReportCommandlet, Warning, "Failed to generate localization conflict report.",
+				("id", GenerateTextLocalizationReportCommandlet::LocalizationLogIdentifier)
+			);
 		}
 	}
 
@@ -128,7 +143,9 @@ bool UGenerateTextLocalizationReportCommandlet::ProcessWordCountReport(const FSt
 	FString ManifestName;
 	if( !( GetStringFromConfig( *SectionName, TEXT("ManifestName"), ManifestName, GatherTextConfigPath ) ) )
 	{
-		UE_LOG(LogGenerateTextLocalizationReportCommandlet, Error, TEXT("No manifest name specified."));
+		UE_LOGFMT(LogGenerateTextLocalizationReportCommandlet, Error, "No manifest name specified.",
+			("id", GenerateTextLocalizationReportCommandlet::LocalizationLogIdentifier)
+		);
 		return false;
 	}
 
@@ -136,7 +153,9 @@ bool UGenerateTextLocalizationReportCommandlet::ProcessWordCountReport(const FSt
 	FString ArchiveName;
 	if( !( GetStringFromConfig( *SectionName, TEXT("ArchiveName"), ArchiveName, GatherTextConfigPath ) ) )
 	{
-		UE_LOG(LogGenerateTextLocalizationReportCommandlet, Error, TEXT("No archive name specified."));
+		UE_LOGFMT(LogGenerateTextLocalizationReportCommandlet, Error, "No archive name specified.",
+			("id", GenerateTextLocalizationReportCommandlet::LocalizationLogIdentifier)
+		);
 		return false;
 	}
 
@@ -144,7 +163,9 @@ bool UGenerateTextLocalizationReportCommandlet::ProcessWordCountReport(const FSt
 	FString WordCountReportName;
 	if( !( GetStringFromConfig( *SectionName, TEXT("WordCountReportName"), WordCountReportName, GatherTextConfigPath ) ) )
 	{
-		UE_LOG(LogGenerateTextLocalizationReportCommandlet, Error, TEXT("No word count report name specified."));
+		UE_LOGFMT(LogGenerateTextLocalizationReportCommandlet, Error, "No word count report name specified.",
+			("id", GenerateTextLocalizationReportCommandlet::LocalizationLogIdentifier)
+		);
 		return false;
 	}
 
@@ -167,7 +188,10 @@ bool UGenerateTextLocalizationReportCommandlet::ProcessWordCountReport(const FSt
 		FText LoadError;
 		if (!LocTextHelper.LoadAll(ELocTextHelperLoadFlags::LoadOrCreate, &LoadError))
 		{
-			UE_LOG(LogGenerateTextLocalizationReportCommandlet, Error, TEXT("%s"), *LoadError.ToString());
+			UE_LOGFMT(LogGenerateTextLocalizationReportCommandlet, Warning, "{error}",
+				("error", LoadError.ToString()),
+				("id", GenerateTextLocalizationReportCommandlet::LocalizationLogIdentifier)
+			);
 			return false;
 		}
 	}
@@ -177,7 +201,10 @@ bool UGenerateTextLocalizationReportCommandlet::ProcessWordCountReport(const FSt
 	FText ReportSaveError;
 	if (!LocTextHelper.SaveWordCountReport(Timestamp, ReportFilePath, &ReportSaveError))
 	{
-		UE_LOG(LogGenerateTextLocalizationReportCommandlet, Error, TEXT("%s"), *ReportSaveError.ToString());
+		UE_LOGFMT(LogGenerateTextLocalizationReportCommandlet, Warning, "{error}",
+			("error", ReportSaveError.ToString()),
+			("id", GenerateTextLocalizationReportCommandlet::LocalizationLogIdentifier)
+		);
 		return false;
 	}
 
@@ -190,7 +217,9 @@ bool UGenerateTextLocalizationReportCommandlet::ProcessConflictReport(const FStr
 	FString ConflictReportName;
 	if (!GetStringFromConfig(*SectionName, TEXT("ConflictReportName"), ConflictReportName, GatherTextConfigPath))
 	{
-		UE_LOG(LogGenerateTextLocalizationReportCommandlet, Error, TEXT("No conflict report name specified."));
+		UE_LOGFMT(LogGenerateTextLocalizationReportCommandlet, Error, "No conflict report name specified.",
+			("id", GenerateTextLocalizationReportCommandlet::LocalizationLogIdentifier)
+		);
 		return false;
 	}
 	
@@ -214,11 +243,17 @@ bool UGenerateTextLocalizationReportCommandlet::ProcessConflictReport(const FStr
 		// We found a file extension somewhere in the specified name. 
 		if (!FileExtension.IsEmpty())
 		{
-			UE_LOG(LogGenerateTextLocalizationReportCommandlet, Warning, TEXT("The conflict report filename %s has an unsupported extension. Only .txt and .csv is supported at this time."), *ConflictReportName);
+			UE_LOGFMT(LogGenerateTextLocalizationReportCommandlet, Warning, "The conflict report filename {filename} has an unsupported extension. Only .txt and .csv is supported at this time.",
+				("filename", ConflictReportName),
+				("id", GenerateTextLocalizationReportCommandlet::LocalizationLogIdentifier)
+			);
 		}
 		else
 		{
-			UE_LOG(LogGenerateTextLocalizationReportCommandlet, Warning, TEXT("The conflict report filename %s has no extension. Defaulting the report to be generated as a .csv file."), *ConflictReportName);
+			UE_LOGFMT(LogGenerateTextLocalizationReportCommandlet, Warning, "The conflict report filename {filename} has no extension. Defaulting the report to be generated as a .csv file.",
+				("filename", ConflictReportName),
+				("id", GenerateTextLocalizationReportCommandlet::LocalizationLogIdentifier)
+			);
 		}
 		ConflictReportName = FPaths::SetExtension(ConflictReportName, CSVExtension);
 	}
@@ -227,7 +262,10 @@ bool UGenerateTextLocalizationReportCommandlet::ProcessConflictReport(const FStr
 	FText ReportSaveError;
 	if (!GatherManifestHelper->SaveConflictReport(ReportFilePath, ConflictReportFormat , &ReportSaveError))
 	{
-		UE_LOG(LogGenerateTextLocalizationReportCommandlet, Error, TEXT("%s"), *ReportSaveError.ToString());
+		UE_LOGFMT(LogGenerateTextLocalizationReportCommandlet, Warning, "{error}",
+			("error", ReportSaveError.ToString()),
+			("id", GenerateTextLocalizationReportCommandlet::LocalizationLogIdentifier)
+		);
 		return false;
 	}
 
diff --git a/Engine/Source/Programs/AutomationTool/Scripts/Localisation.Automation.cs b/Engine/Source/Programs/AutomationTool/Scripts/Localisation.Automation.cs
--- a/Engine/Source/Programs/AutomationTool/Scripts/Localisation.Automation.cs
+++ b/Engine/Source/Programs/AutomationTool/Scripts/Localisation.Automation.cs
@@ -399,7 +399,7 @@ class Localize : BuildCommand
 			var LocalizationStepNamesStr = ParseParamValue("LocalizationSteps");
 			if (LocalizationStepNamesStr == null)
 			{
-				LocalizationStepNames.AddRange(new string[] { "Download", "Gather", "Import", "Export", "Compile", "GenerateReports", "Upload" });
+				LocalizationStepNames.AddRange([ "Download", "Gather", "Import", "Export", "Compile", "Upload" ]);
 			}
 			else
 			{
@@ -408,6 +408,11 @@ class Localize : BuildCommand
 					LocalizationStepNames.Add(StepName.Trim());
 				}
 			}
+			if (LocalizationStepNames.Contains("Gather"))
+			{
+				// Gather runs a superset of the GenerateReports step, so if we're running Gather, we can skip GenerateReports
+				LocalizationStepNames.Remove("GenerateReports");
+			}
 			LocalizationStepNames.Add("Monolithic"); // Always allow the monolithic scripts to run as we don't know which steps they do
 		}
 
