Fixed inverted logic

#rnx

diff --git a/Engine/Source/Editor/UnrealEd/Private/Commandlets/GenerateTextLocalizationReportCommandlet.cpp b/Engine/Source/Editor/UnrealEd/Private/Commandlets/GenerateTextLocalizationReportCommandlet.cpp
--- a/Engine/Source/Editor/UnrealEd/Private/Commandlets/GenerateTextLocalizationReportCommandlet.cpp
+++ b/Engine/Source/Editor/UnrealEd/Private/Commandlets/GenerateTextLocalizationReportCommandlet.cpp
@@ -173,11 +173,11 @@ bool UGenerateTextLocalizationReportCommandlet::ProcessWordCountReport(const FSt
 	TArray<FString> CulturesToGenerate;
 	GetStringArrayFromConfig( *SectionName, TEXT("CulturesToGenerate"), CulturesToGenerate, GatherTextConfigPath );
 
-	for(int32 i = 0; i < CulturesToGenerate.Num(); ++i)
+	for (const FString& CultureToGenerate : CulturesToGenerate)
 	{
-		if( FInternationalization::Get().GetCulture( CulturesToGenerate[i] ).IsValid() )
+		if (!FInternationalization::Get().GetCulture(CultureToGenerate))
 		{
-			UE_LOG(LogGenerateTextLocalizationReportCommandlet, Verbose, TEXT("Specified culture is not a valid runtime culture, but may be a valid base language: %s"), *(CulturesToGenerate[i]));
+			UE_LOG(LogGenerateTextLocalizationReportCommandlet, Verbose, TEXT("Specified culture is not a valid runtime culture, but may be a valid base language: %s"), *CultureToGenerate);
 		}
 	}
 
