Cooker: Fix typo in comment.
#rnx

diff --git a/Engine/Source/Editor/UnrealEd/Classes/CookOnTheSide/CookOnTheFlyServer.h b/Engine/Source/Editor/UnrealEd/Classes/CookOnTheSide/CookOnTheFlyServer.h
--- a/Engine/Source/Editor/UnrealEd/Classes/CookOnTheSide/CookOnTheFlyServer.h
+++ b/Engine/Source/Editor/UnrealEd/Classes/CookOnTheSide/CookOnTheFlyServer.h
@@ -1218,7 +1218,7 @@ private:
 	/** Finalize the package store */
 	void FinalizePackageStore(const ITargetPlatform* TargetPlatform, const TSet<FName>& CookedPackageNames);
 	/**
-	 * Calculate which Plugins (and also /Game and /Editor) had packages referenced in the current session,
+	 * Calculate which Plugins (and also /Game and /Engine) had packages referenced in the current session,
 	 * this is important to know for DLC cooks.
 	 */
 	void CalculateReferencedPlugins(const ITargetPlatform* TargetPlatform,
