Switched filename for file in some structured logs

#rnx

diff --git a/Engine/Source/Developer/Localization/Private/LocalizationSourceControlUtil.cpp b/Engine/Source/Developer/Localization/Private/LocalizationSourceControlUtil.cpp
--- a/Engine/Source/Developer/Localization/Private/LocalizationSourceControlUtil.cpp
+++ b/Engine/Source/Developer/Localization/Private/LocalizationSourceControlUtil.cpp
@@ -380,8 +380,8 @@ void FLocFileSCCNotifies::PreFileWrite(const FString& InFilename)
 		FText ErrorMsg;
 		if (!SourceControlInfo->CheckOutFile(InFilename, ErrorMsg))
 		{
-			UE_LOGFMT(LogLocalizationSourceControl, Error, "Failed to check out file '{filename}'. {error}",
-				("filename", InFilename),
+			UE_LOGFMT(LogLocalizationSourceControl, Error, "Failed to check out file '{file}'. {error}",
+				("file", InFilename),
 				("error", ErrorMsg.ToString()),
 				("id", LocalizationSourceControlUtil::LocalizationLogIdentifier)
 			);
@@ -397,8 +397,8 @@ void FLocFileSCCNotifies::PostFileWrite(const FString& InFilename)
 		FText ErrorMsg;
 		if (!SourceControlInfo->CheckOutFile(InFilename, ErrorMsg))
 		{
-			UE_LOGFMT(LogLocalizationSourceControl, Error, "Failed to check out file '{filename}'. {error}",
-				("filename", InFilename),
+			UE_LOGFMT(LogLocalizationSourceControl, Error, "Failed to check out file '{file}'. {error}",
+				("file", InFilename),
 				("error", ErrorMsg.ToString()),
 				("id", LocalizationSourceControlUtil::LocalizationLogIdentifier)
 			);
diff --git a/Engine/Source/Editor/UnrealEd/Private/Commandlets/GenerateTextLocalizationReportCommandlet.cpp b/Engine/Source/Editor/UnrealEd/Private/Commandlets/GenerateTextLocalizationReportCommandlet.cpp
--- a/Engine/Source/Editor/UnrealEd/Private/Commandlets/GenerateTextLocalizationReportCommandlet.cpp
+++ b/Engine/Source/Editor/UnrealEd/Private/Commandlets/GenerateTextLocalizationReportCommandlet.cpp
@@ -243,15 +243,15 @@ bool UGenerateTextLocalizationReportCommandlet::ProcessConflictReport(const FStr
 		// We found a file extension somewhere in the specified name. 
 		if (!FileExtension.IsEmpty())
 		{
-			UE_LOGFMT(LogGenerateTextLocalizationReportCommandlet, Warning, "The conflict report filename {filename} has an unsupported extension. Only .txt and .csv is supported at this time.",
-				("filename", ConflictReportName),
+			UE_LOGFMT(LogGenerateTextLocalizationReportCommandlet, Warning, "The conflict report filename {file} has an unsupported extension. Only .txt and .csv is supported at this time.",
+				("file", ConflictReportName),
 				("id", GenerateTextLocalizationReportCommandlet::LocalizationLogIdentifier)
 			);
 		}
 		else
 		{
-			UE_LOGFMT(LogGenerateTextLocalizationReportCommandlet, Warning, "The conflict report filename {filename} has no extension. Defaulting the report to be generated as a .csv file.",
-				("filename", ConflictReportName),
+			UE_LOGFMT(LogGenerateTextLocalizationReportCommandlet, Warning, "The conflict report filename {file} has no extension. Defaulting the report to be generated as a .csv file.",
+				("file", ConflictReportName),
 				("id", GenerateTextLocalizationReportCommandlet::LocalizationLogIdentifier)
 			);
 		}
