Improvements/fixes to the experimental multi-process localization asset gather

Improvements:
* Worker result ingestion now happens immediately on the MessageBus thread, rather than being queued for later processing on the game thread.
* This also allows the MessageBus thread to immediately request more work for "idle" workers (rather than wait until the next game thread loop), which helps keep worker processes saturated with work.
* Tweaked the worker idle threshold to provide better saturation, and decreased the idle sleep on the worker process so they can pick-up new work faster.
* The worker processes now disable shader compilation, to match the main process.
* The worker processes can now load the cached asset registry when not using `-StartGatherTextWorkersImmediately`.
* Worker result ingestion now uses the same global progress tracking as local work.
* Added stats for how many packages were processed locally vs remotely via a worker.

Fixes:
* Fixed FLoadPackageLogOutputRedirector not being thread-safe.
* It is now available permanently (as a static), as a background thread may read it as the current GWarn value and write to it after log capture has been disabled (and GWarn has been set back to the real GWarn).
* It now only captures logs from the threads associated with package loading (the game thread and async-loading thread), and only when log capture is enabled.
* Fixed the worker processes not ensuring they start with an empty editor world, as a partially initialized world can cause issues for actor localization gathering.
* Fixed a GC race-condition in the worker process.

#rnx

diff --git a/Engine/Source/Editor/UnrealEd/Classes/Commandlets/GatherTextFromAssetsCommandlet.h b/Engine/Source/Editor/UnrealEd/Classes/Commandlets/GatherTextFromAssetsCommandlet.h
--- a/Engine/Source/Editor/UnrealEd/Classes/Commandlets/GatherTextFromAssetsCommandlet.h
+++ b/Engine/Source/Editor/UnrealEd/Classes/Commandlets/GatherTextFromAssetsCommandlet.h
@@ -81,7 +81,11 @@ class UGatherTextFromAssetsCommandlet : public UGatherTextCommandletBase
 	GENERATED_UCLASS_BODY()
 
 	void StartWorkers(const int32 MinPackagesToUseWorkers);
+	void AssignPackagesToWorkers(TConstArrayView<FGuid> IdleWorkerIds); // Note: You must hold PackagesPendingGatherMutex when calling this function
+	void IngestPackageResultFromWorker(const FGatherTextFromAssetsWorkerMessage_PackageResult& PackageResult, const bool bSendWorkIfIdle = false);
+
 	void ProcessGatherableTextDataArray(const TArray<FGatherableTextData>& GatherableTextDataArray);
+	
 	void CalculateDependenciesForPackagesPendingGather();
 
 public:
@@ -177,11 +181,20 @@ private:
 	TArray<FString> ManifestDependenciesList;
 
 	TArray<FPackagePendingGather> PackagesPendingGather;
+	mutable UE::FMutex PackagesPendingGatherMutex;
 
 	TSet<FName> PackagesWithDuplicateLocalizationIds;
 
 	/** Record of which packages were sent to workers, so that we can re-add them to the main PackagesPendingGather queue if the worker crashes */
 	TMap<FName, FPackagePendingGather> PackagesDistributedToWorkers;
+	mutable UE::FMutex PackagesDistributedToWorkersMutex;
+
+	mutable UE::FMutex GatherManifestHelperMutex;
+
+	/** Stats for the uncached package loading */
+	int32 TotalNumUncachedPackages = 0;
+	std::atomic<int32> NumUncachedPackagesProcessedLocally = 0;
+	std::atomic<int32> NumUncachedPackagesProcessedRemotely = 0;
 
 	/** Number of worker processes to use */
 	int32 NumWorkers = 0;
@@ -346,6 +359,12 @@ public:
 	 */
 	bool StopWorkers();
 
+	/**
+	 * Does this director have workers?
+	 * @note They may be in any state and not ready to accept work (see GetAvailableWorkerIds).
+	 */
+	bool HasWorkers() const;
+
 	/**
 	 * Tick workers, to detect which have crashed or hung.
 	 */
@@ -378,6 +397,25 @@ public:
 	 */
 	bool AssignPackageToWorker(const FGuid& WorkerId, const FGatherTextFromAssetsWorkerMessage_PackageRequest& PackageRequest);
 
+	/**
+	 * Set whether or not worker processes are allowed to read the asset registry cache.
+	 */
+	void SetWorkersCanReadAssetRegistryCache(const bool bValue)
+	{
+		bWorkersCanReadAssetRegistryCache = bValue;
+	}
+
+	/**
+	 * Set the handler used to automatically ingest package results as they arrive (from any thread), rather than add them to PackageResults for deferred processing via IngestPackageResult.
+	 * @note The handler may return false if it cannot handle a given result, at which point it will still be added to PackageResults.
+	 */
+	void SetIngestPackageResultHandler(TFunction<bool(const FGatherTextFromAssetsWorkerMessage_PackageResult&)>&& Handler);
+
+	/**
+	 * Clear any currently set handler (@see SetIngestPackageResultHandler).
+	 */
+	void ClearIngestPackageResultHandler();
+
 	/**
 	 * Ingest the next package result received from any worker, if any.
 	 * 
@@ -428,9 +466,15 @@ private:
 	/** The queue of results received from any worker */
 	TMpscQueue<FGatherTextFromAssetsWorkerMessage_PackageResult> PackageResults;
 
+	/** Optional handler that will attempt to process results immediately, rather than add them to PackageResults */
+	TFunction<bool(const FGatherTextFromAssetsWorkerMessage_PackageResult&)> PackageResultHandler;
+
 	/** The queue of packages that failed to be processed due to their worker crashing */
 	TArray<FName> PackagesFromCrashedWorkers;
 
 	/** Holds the messaging endpoint we are sending from */
 	TSharedPtr<FMessageEndpoint> MessageEndpoint;
+
+	/** Are workers allowed to read the asset registry cache? */
+	bool bWorkersCanReadAssetRegistryCache = false;
 };
diff --git a/Engine/Source/Editor/UnrealEd/Private/Commandlets/GatherTextFromAssetsCommandlet.cpp b/Engine/Source/Editor/UnrealEd/Private/Commandlets/GatherTextFromAssetsCommandlet.cpp
--- a/Engine/Source/Editor/UnrealEd/Private/Commandlets/GatherTextFromAssetsCommandlet.cpp
+++ b/Engine/Source/Editor/UnrealEd/Private/Commandlets/GatherTextFromAssetsCommandlet.cpp
@@ -68,8 +68,8 @@ class FLoadPackageLogOutputRedirector : public FFeedbackContext
 public:
 	struct FScopedCapture
 	{
-		FScopedCapture(FLoadPackageLogOutputRedirector* InLogOutputRedirector, FStringView InPackageContext)
-			: LogOutputRedirector(InLogOutputRedirector)
+		explicit FScopedCapture(FStringView InPackageContext)
+			: LogOutputRedirector(&FLoadPackageLogOutputRedirector::Get())
 		{
 			LogOutputRedirector->BeginCapturingLogData(InPackageContext);
 		}
@@ -82,30 +82,66 @@ public:
 		FLoadPackageLogOutputRedirector* LogOutputRedirector;
 	};
 
-	explicit FLoadPackageLogOutputRedirector(FOutputDevice* InTargetOutputDevice = GWarn)
-		: TargetOutputDevice(InTargetOutputDevice)
+	virtual ~FLoadPackageLogOutputRedirector() = default;
+
+	static FLoadPackageLogOutputRedirector& Get()
 	{
+		static FLoadPackageLogOutputRedirector Instance;
+		return Instance;
 	}
 
-	virtual ~FLoadPackageLogOutputRedirector() = default;
+	virtual bool CanBeUsedOnAnyThread() const override
+	{
+		return true;
+	}
 
-	void BeginCapturingLogData(FStringView InPackageContext)
+	virtual bool CanBeUsedOnMultipleThreads() const override
+	{
+		return true;
+	}
+
+	void Hook(FOutputDevice* InTargetOutputDevice = GWarn)
 	{
+		UE::TScopeLock _(Mutex);
+
+		TargetOutputDevice = InTargetOutputDevice;
+
 		// Override GWarn so that we can capture any log data
-		check(!OriginalWarningContext);
+		checkf(!OriginalWarningContext, TEXT("Hook called recursively!"));
 		OriginalWarningContext = GWarn;
 		GWarn = this;
-
-		PackageContext = InPackageContext;
 	}
 
-	void EndCapturingLogData()
+	void Unhook()
 	{
+		UE::TScopeLock _(Mutex);
+
+		checkf(!bIsCapturing, TEXT("Unhook called while in a BeginCapturingLogData scope!"));
+
 		// Restore the original GWarn now that we've finished capturing log data
-		check(OriginalWarningContext);
+		checkf(OriginalWarningContext, TEXT("Unhook mismatch with Hook!"));
 		GWarn = OriginalWarningContext;
 		OriginalWarningContext = nullptr;
 
+		TargetOutputDevice = nullptr;
+	}
+
+	void BeginCapturingLogData(FStringView InPackageContext)
+	{
+		UE::TScopeLock _(Mutex);
+
+		checkf(OriginalWarningContext, TEXT("BeginCapturingLogData called before Hook!"));
+		checkf(!bIsCapturing, TEXT("BeginCapturingLogData called recursively!"));
+		PackageContext = InPackageContext;
+		bIsCapturing = true;
+	}
+
+	void EndCapturingLogData()
+	{
+		UE::TScopeLock _(Mutex);
+
+		checkf(bIsCapturing, TEXT("EndCapturingLogData mismatch with BeginCapturingLogData!"));
+
 		// Report any messages, and also report a warning if we silenced some warnings or errors when loading
 		if (TargetOutputDevice && (ErrorCount > 0 || WarningCount > 0))
 		{
@@ -119,6 +155,7 @@ public:
 		}
 
 		PackageContext.Reset();
+		bIsCapturing = false;
 
 		// Reset the counts and previous log output
 		ErrorCount = 0;
@@ -133,6 +170,16 @@ public:
 
 	virtual void Serialize(const TCHAR* V, ELogVerbosity::Type Verbosity, const FName& Category, double Time) override
 	{
+		UE::TScopeLock _(Mutex);
+
+		// We only want to capture logs while capture is active, and only from the GT and ALT
+		if (!bIsCapturing || !(IsInGameThread() || IsInAsyncLoadingThread()))
+		{
+			FFeedbackContext* TargetWarningContext = OriginalWarningContext ? OriginalWarningContext : GWarn;
+			TargetWarningContext->Serialize(V, Verbosity, Category, Time);
+			return;
+		}
+
 		if (Verbosity == ELogVerbosity::Error)
 		{
 			++ErrorCount;
@@ -159,6 +206,16 @@ public:
 
 	virtual void SerializeRecord(const UE::FLogRecord& Record) override
 	{
+		UE::TScopeLock _(Mutex);
+
+		// We only want to capture logs while capture is active, and only from the GT and ALT
+		if (!bIsCapturing || !(IsInGameThread() || IsInAsyncLoadingThread()))
+		{
+			FFeedbackContext* TargetWarningContext = OriginalWarningContext ? OriginalWarningContext : GWarn;
+			TargetWarningContext->SerializeRecord(Record);
+			return;
+		}
+
 		const ELogVerbosity::Type Verbosity = Record.GetVerbosity();
 		if (Verbosity == ELogVerbosity::Error)
 		{
@@ -195,6 +252,8 @@ public:
 	}
 
 private:
+	FLoadPackageLogOutputRedirector() = default;
+
 	int32 ErrorCount = 0;
 	int32 WarningCount = 0;
 	TArray<FString> FormattedErrorsAndWarningsList;
@@ -203,6 +262,10 @@ private:
 	FFeedbackContext* OriginalWarningContext = nullptr;
 
 	FOutputDevice* TargetOutputDevice = nullptr;
+
+	bool bIsCapturing = false;
+
+	mutable UE::FMutex Mutex;
 };
 
 class FAssetGatherCacheMetrics
@@ -289,13 +352,13 @@ namespace UE::Private::GatherTextFromAssetsCommandlet
 	static constexpr int32 MinPackagesToUseWorkers = 1000;	// How many packages should we have to process (in total) before we start the workers?
 	static constexpr int32 MinPackagesToKeepLocal = 50;		// How many packages should we keep in the main process, rather than distribute to workers? (to avoid waiting on workers at the end of the gather)
 	static constexpr int32 MaxPackagesToDistribute = 100;	// How many packages should we distribute to each worker in a single batch?
-	static constexpr int32 WorkerIdleThreshold = 10;		// How many packages can a worker have pending in its queue before we consider it "idle" and try and assign it more packages? (to avoid workers becoming idle)
+	static constexpr int32 WorkerIdleThreshold = 20;		// How many packages can a worker have pending in its queue before we consider it "idle" and try and assign it more packages? (to avoid workers becoming idle)
 
-	UPackage* LoadPackageToGather(const FString& PackageName, const bool bHasExternalActors, FLoadPackageLogOutputRedirector& LogOutputRedirector)
+	UPackage* LoadPackageToGather(const FString& PackageName, const bool bHasExternalActors)
 	{
 		UPackage* Package = nullptr;
 		{
-			FLoadPackageLogOutputRedirector::FScopedCapture ScopedCapture(&LogOutputRedirector, PackageName);
+			FLoadPackageLogOutputRedirector::FScopedCapture ScopedCapture(PackageName);
 			if (bHasExternalActors)
 			{
 				Package = LoadWorldPackageForEditor(PackageName, EWorldType::Editor, LOAD_NoWarn | LOAD_Quiet);
@@ -308,7 +371,7 @@ namespace UE::Private::GatherTextFromAssetsCommandlet
 		return Package;
 	}
 
-	bool GatherTextFromPackage(UPackage* Package, const TSet<FGuid>& ExternalActors, const bool bGatherFromPrimaryPackage, const bool bGatherFromExternalActorPackages, FLoadPackageLogOutputRedirector& LogOutputRedirector, TArray<FGatherableTextData>& OutGatherableTextDataArray)
+	bool GatherTextFromPackage(UPackage* Package, const TSet<FGuid>& ExternalActors, const bool bGatherFromPrimaryPackage, const bool bGatherFromExternalActorPackages, TArray<FGatherableTextData>& OutGatherableTextDataArray)
 	{
 		// Because packages may not have been resaved after this flagging was implemented, we may have added packages to load that weren't flagged - potential false positives.
 		// The loading process should have reflagged said packages so that only true positives will have this flag.
@@ -342,7 +405,7 @@ namespace UE::Private::GatherTextFromAssetsCommandlet
 					TOptional<FScopedEditorWorld> ScopeEditorWorld;
 					if (!World->IsInitialized())
 					{
-						FLoadPackageLogOutputRedirector::FScopedCapture ScopedCapture(&LogOutputRedirector, Package->GetName());
+						FLoadPackageLogOutputRedirector::FScopedCapture ScopedCapture(Package->GetName());
 						ScopeEditorWorld.Emplace(World, IVS); // Initializing FScopedEditorWorld can log warnings, so capture those like we do with loading errors
 					}
 
@@ -470,12 +533,12 @@ namespace UE::Private::GatherTextFromAssetsCommandlet
 	}
 
 	template <typename TPackagePendingGather>
-	void ConditionalPurgeGarbage(const uint64 MinFreeMemoryBytes, const uint64 MaxUsedMemoryBytes, TArray<TPackagePendingGather>& PackagesPendingGather)
+	void ConditionalPurgeGarbage(const uint64 MinFreeMemoryBytes, const uint64 MaxUsedMemoryBytes, TArray<TPackagePendingGather>* PackagesPendingGather)
 	{
 		if (HasExceededMemoryLimit(MinFreeMemoryBytes, MaxUsedMemoryBytes, /*bLog*/true))
 		{
 			// First try a minimal purge to only remove things that are no longer referenced or needed by other packages pending gather
-			PurgeGarbage<TPackagePendingGather>(&PackagesPendingGather);
+			PurgeGarbage<TPackagePendingGather>(PackagesPendingGather);
 
 			if (HasExceededMemoryLimit(MinFreeMemoryBytes, MaxUsedMemoryBytes, /*bLog*/false))
 			{
@@ -531,7 +594,6 @@ UGatherTextFromAssetsCommandlet::UGatherTextFromAssetsCommandlet(const FObjectIn
 
 void UGatherTextFromAssetsCommandlet::ProcessGatherableTextDataArray(const TArray<FGatherableTextData>& GatherableTextDataArray)
 {
-	
 	for (const FGatherableTextData& GatherableTextData : GatherableTextDataArray)
 	{
 		for (const FTextSourceSiteContext& TextSourceSiteContext : GatherableTextData.SourceSiteContexts)
@@ -1461,100 +1523,52 @@ void UGatherTextFromAssetsCommandlet::MergeInExternalActorsWithStaleOrMissingCac
 bool UGatherTextFromAssetsCommandlet::LoadAndProcessUncachedPackages(TArray<FName>& OutPackagesWithStaleGatherCache)
 {
 	UE_SCOPED_TIMER(TEXT("UGatherTextFromAssetsCommandlet::LoadAndProcessUncachedPackages"), LogGatherTextFromAssetsCommandlet, Display);
-	FLoadPackageLogOutputRedirector LogOutputRedirector;
 	TArray<FGatherableTextData> GatherableTextDataArray;
-	int32 NumPackagesProcessed = 0;
-	int32 PackageCount = PackagesPendingGather.Num();
-	int32 NumPackagesFailedLoading = 0;
+	int32 LocalNumPackagesProcessed = 0;
+	int32 LocalNumPackagesFailedLoading = 0;
+	int32 LocalPackageCount = PackagesPendingGather.Num();
 	
 	FGatherTextFromAssetsWorkerDirector& WorkerDirector = FGatherTextFromAssetsWorkerDirector::Get();
 
-	FScopedSlowTask SlowTask(PackageCount, NSLOCTEXT("GatherTextCommandlet", "LoadAndProcessUncachedPackages", "Loading and Gathering Packages..."));
-	UE_LOG(LogGatherTextFromAssetsCommandlet, Display, TEXT("Number of packages with duplicate loc ids: %d"), NumPackagesDupLocId);
-	UE_LOG(LogGatherTextFromAssetsCommandlet, Display, TEXT("Number of packages to load: %d"), PackagesPendingGather.Num());
+	FScopedSlowTask SlowTask(LocalPackageCount, NSLOCTEXT("GatherTextCommandlet", "LoadAndProcessUncachedPackages", "Loading and Gathering Packages..."));
 
 	// Auxiliary lambda function for resaving packages
-	auto ResavePackage = [this, &LogOutputRedirector](
+	auto ResavePackage = [this](
 		const FStringView PackageNameStr,
 		UPackage* Package,
 		const FString& PackageFileName
 		) -> bool
 	{
-		FLoadPackageLogOutputRedirector::FScopedCapture ScopedCapture(&LogOutputRedirector, PackageNameStr);
+		FLoadPackageLogOutputRedirector::FScopedCapture ScopedCapture(PackageNameStr);
 		return FLocalizedAssetSCCUtil::SavePackageWithSCC(SourceControlInfo, Package, PackageFileName);
 	};
 
-	auto AssignPackagesToWorker = [this, &WorkerDirector, &SlowTask, &NumPackagesProcessed](const FGuid& WorkerId, const int32 NumPackages)
+	for (;;)
 	{
-		check(PackagesPendingGather.Num() >= NumPackages);
-
-		// We assign worker packages from the tail of the queue (the front of the array) 
-		// as these packages have the fewest number of dependencies so are the simplest 
-		// to distribute and load, which helps to balance the memory usage between 
-		// the main process (which takes the heavy tasks) and the worker processes
-		TArray<FPackagePendingGather> PackagesToReinsert;
-		for (int32 Index = 0; Index < NumPackages; ++Index)
+		TOptional<FPackagePendingGather> MaybePackagePendingGather;
 		{
-			FPackagePendingGather PackagePendingGather = MoveTemp(PackagesPendingGather[Index]);
+			UE::TScopeLock _(PackagesPendingGatherMutex);
 
-			FGatherTextFromAssetsWorkerMessage_PackageRequest PackageRequest;
-			PackageRequest.PackageName = PackagePendingGather.PackageName;
-			PackageRequest.Dependencies = PackagePendingGather.Dependencies;
-			PackageRequest.ExternalActors = PackagePendingGather.ExternalActors;
+			// If we have (soon to be) idle workers, distribute some work to them now
+			AssignPackagesToWorkers(WorkerDirector.GetAvailableWorkerIds(UE::Private::GatherTextFromAssetsCommandlet::WorkerIdleThreshold));
 
-			if (WorkerDirector.AssignPackageToWorker(WorkerId, PackageRequest))
+			if (PackagesPendingGather.IsEmpty())
 			{
-				SlowTask.EnterProgressFrame();
-				++NumPackagesProcessed;
-
-				const FName PackageName = PackagePendingGather.PackageName;
-				PackagesDistributedToWorkers.Add(PackageName, MoveTemp(PackagePendingGather));
+				break;
 			}
-			else
-			{
-				PackagesToReinsert.Add(MoveTemp(PackagePendingGather));
-			}
-		}
-		PackagesPendingGather.RemoveAt(0, NumPackages, EAllowShrinking::No);
-		if (PackagesToReinsert.Num() > 0)
-		{
-			PackagesPendingGather.Insert(MoveTemp(PackagesToReinsert), 0);
-		}
-	};
 
-	while (PackagesPendingGather.Num() > 0)
-	{
-		// Should we try and distribute some work?
-		{
-			if (PackagesPendingGather.Num() > UE::Private::GatherTextFromAssetsCommandlet::MinPackagesToKeepLocal)
-			{
-				// If we have (soon to be) idle workers, distribute some work to them now
-				if (const TArray<FGuid> IdleWorkerIds = WorkerDirector.GetAvailableWorkerIds(UE::Private::GatherTextFromAssetsCommandlet::WorkerIdleThreshold);
-					IdleWorkerIds.Num() > 0)
-				{
-					if (const int32 NumPackagesToDistributeToEachWorker = FMath::Clamp((PackagesPendingGather.Num() - UE::Private::GatherTextFromAssetsCommandlet::MinPackagesToKeepLocal) / IdleWorkerIds.Num(), 0, UE::Private::GatherTextFromAssetsCommandlet::MaxPackagesToDistribute);
-						NumPackagesToDistributeToEachWorker > 0)
-					{
-						for (const FGuid& WorkerId : IdleWorkerIds)
-						{
-							AssignPackagesToWorker(WorkerId, NumPackagesToDistributeToEachWorker);
-						}
-					}
-				}
-				if (PackagesPendingGather.Num() == 0)
-				{
-					break;
-				}
-			}
+			MaybePackagePendingGather = PackagesPendingGather.Pop(EAllowShrinking::No);
 		}
 
-		const FPackagePendingGather PackagePendingGather = PackagesPendingGather.Pop(EAllowShrinking::No);
+		check(MaybePackagePendingGather);
+		const FPackagePendingGather& PackagePendingGather = MaybePackagePendingGather.GetValue();
 		FNameBuilder PackageNameStr(PackagePendingGather.PackageName);
 
-		const int32 CurrentPackageNum = ++NumPackagesProcessed;
-		const float PercentageComplete = static_cast<float>(CurrentPackageNum) / static_cast<float>(PackageCount) * 100.0f;
+		const int32 TotalNumUncachedPackagesProcessed = ++NumUncachedPackagesProcessedLocally + NumUncachedPackagesProcessedRemotely;
+		const float PercentageComplete = static_cast<float>(TotalNumUncachedPackagesProcessed) / static_cast<float>(TotalNumUncachedPackages) * 100.0f;
 		UE_LOG(LogGatherTextFromAssetsCommandlet, Display, TEXT("[%6.2f%%] Loading package: '%s'..."), PercentageComplete, *PackageNameStr);
 
+		++LocalNumPackagesProcessed;
 		SlowTask.EnterProgressFrame();
 
 		if (SlowTask.ShouldCancel() || (EmbeddedContext && EmbeddedContext->ShouldAbort()))
@@ -1565,14 +1579,14 @@ bool UGatherTextFromAssetsCommandlet::LoadAndProcessUncachedPackages(TArray<FNam
 			return false;
 		}
 
-		UPackage* Package = UE::Private::GatherTextFromAssetsCommandlet::LoadPackageToGather(*PackageNameStr, PackagePendingGather.ExternalActors.Num() > 0, LogOutputRedirector);
+		UPackage* Package = UE::Private::GatherTextFromAssetsCommandlet::LoadPackageToGather(*PackageNameStr, PackagePendingGather.ExternalActors.Num() > 0);
 		if (!Package)
 		{
 			UE_LOGFMT(LogGatherTextFromAssetsCommandlet, Warning, "Failed to load package: '{package}'.",
 				("package", *PackageNameStr),
 				("id", GatherTextFromAssetsCommandlet::LocalizationLogIdentifier)
 			);
-			++NumPackagesFailedLoading;
+			++LocalNumPackagesFailedLoading;
 			continue;
 		}
 
@@ -1612,7 +1626,7 @@ bool UGatherTextFromAssetsCommandlet::LoadAndProcessUncachedPackages(TArray<FNam
 		}
 		WorkerDirector.TickWorkers();
 
-		if (UE::Private::GatherTextFromAssetsCommandlet::GatherTextFromPackage(Package, PackagePendingGather.ExternalActors, /*bGatherFromPrimaryPackage*/true, /*bGatherFromExternalActorPackages*/false, LogOutputRedirector, GatherableTextDataArray))
+		if (UE::Private::GatherTextFromAssetsCommandlet::GatherTextFromPackage(Package, PackagePendingGather.ExternalActors, /*bGatherFromPrimaryPackage*/true, /*bGatherFromExternalActorPackages*/false, GatherableTextDataArray))
 		{
 			UE_LOG(LogGatherTextFromAssetsCommandlet, Display, TEXT("[%6.2f%%] Gathering package: '%s'..."), PercentageComplete, *PackageNameStr);
 
@@ -1667,13 +1681,14 @@ bool UGatherTextFromAssetsCommandlet::LoadAndProcessUncachedPackages(TArray<FNam
 			// If this is a WP world then query the localization for any external actors actors that were determined to be stale
 			if (PackagePendingGather.ExternalActors.Num() > 0)
 			{
-				UE::Private::GatherTextFromAssetsCommandlet::GatherTextFromPackage(Package, PackagePendingGather.ExternalActors, /*bGatherFromPrimaryPackage*/false, /*bGatherFromExternalActorPackages*/true, LogOutputRedirector, GatherableTextDataArray);
+				UE::Private::GatherTextFromAssetsCommandlet::GatherTextFromPackage(Package, PackagePendingGather.ExternalActors, /*bGatherFromPrimaryPackage*/false, /*bGatherFromExternalActorPackages*/true, GatherableTextDataArray);
 			}
 
 			// This package may have already been cached in cases where we're reporting or fixing assets with a stale gather cache
 			// This check prevents it being gathered a second time
-			if (PackagePendingGather.PackageLocCacheState != EPackageLocCacheState::Cached)
+			if (PackagePendingGather.PackageLocCacheState != EPackageLocCacheState::Cached && GatherableTextDataArray.Num() > 0)
 			{
+				UE::TScopeLock _(GatherManifestHelperMutex);
 				ProcessGatherableTextDataArray(GatherableTextDataArray);
 			}
 
@@ -1693,10 +1708,13 @@ bool UGatherTextFromAssetsCommandlet::LoadAndProcessUncachedPackages(TArray<FNam
 			}
 		}
 
-		UE::Private::GatherTextFromAssetsCommandlet::ConditionalPurgeGarbage<FPackagePendingGather>(MinFreeMemoryBytes, MaxUsedMemoryBytes, PackagesPendingGather);
+		{
+			UE::TScopeLock _(PackagesPendingGatherMutex);
+			UE::Private::GatherTextFromAssetsCommandlet::ConditionalPurgeGarbage<FPackagePendingGather>(MinFreeMemoryBytes, MaxUsedMemoryBytes, &PackagesPendingGather);
+		}
 	}
 
-	UE_LOG(LogGatherTextFromAssetsCommandlet, Display, TEXT("Loaded %d packages. %d failed."), NumPackagesProcessed, NumPackagesFailedLoading);
+	UE_LOG(LogGatherTextFromAssetsCommandlet, Display, TEXT("Loaded %d packages. %d failed."), LocalNumPackagesProcessed, LocalNumPackagesFailedLoading);
 	return true;
 }
 
@@ -1816,6 +1834,8 @@ int32 UGatherTextFromAssetsCommandlet::Main(const FString& Params)
 		return -1;
 	}
 	
+	FGatherTextFromAssetsWorkerDirector& WorkerDirector = FGatherTextFromAssetsWorkerDirector::Get();
+
 	UE_LOG(LogGatherTextFromAssetsCommandlet, Display, TEXT("Discovering assets to gather..."));
 	const double DiscoveringAssetsStartTime = FPlatformTime::Seconds();
 	{
@@ -1826,6 +1846,9 @@ int32 UGatherTextFromAssetsCommandlet::Main(const FString& Params)
 		if (bSearchAllAssets || FirstPassFilter.PackagePaths.IsEmpty())
 		{
 			AssetRegistry.SearchAllAssets(true);
+
+			// Workers are allowed to read the asset registry cache after we've performed a full asset scan
+			WorkerDirector.SetWorkersCanReadAssetRegistryCache(true);
 		}
 		else
 		{
@@ -1891,8 +1914,10 @@ int32 UGatherTextFromAssetsCommandlet::Main(const FString& Params)
 		return 0;
 	}
 
+	TotalNumUncachedPackages = PackagesPendingGather.Num();
+
 	const double PackageLoadingStartTime = FPlatformTime::Seconds();
-	UE_LOG(LogGatherTextFromAssetsCommandlet, Display, TEXT("Preparing to load %d packages..."), PackagesPendingGather.Num());
+	UE_LOG(LogGatherTextFromAssetsCommandlet, Display, TEXT("Preparing to load %d packages..."), TotalNumUncachedPackages);
 
 	StartWorkers(UE::Private::GatherTextFromAssetsCommandlet::MinPackagesToUseWorkers);
 
@@ -1910,18 +1935,22 @@ int32 UGatherTextFromAssetsCommandlet::Main(const FString& Params)
 		GShaderCompilingManager->SkipShaderCompilation(true);
 	}
 
-	FGatherTextFromAssetsWorkerDirector& WorkerDirector = FGatherTextFromAssetsWorkerDirector::Get();
+	WorkerDirector.SetIngestPackageResultHandler(
+		[this](const FGatherTextFromAssetsWorkerMessage_PackageResult& PackageResult)
+		{
+			IngestPackageResultFromWorker(PackageResult, /*bSendWorkIfIdle*/true);
+			return true;
+		});
+
+	FLoadPackageLogOutputRedirector::Get().Hook();
 
 	TArray<FName> PackagesWithStaleGatherCache;
 	for (;;)
 	{
 		// Process the main task work
-		if (PackagesPendingGather.Num() > 0)
+		if (!LoadAndProcessUncachedPackages(PackagesWithStaleGatherCache))
 		{
-			if (!LoadAndProcessUncachedPackages(PackagesWithStaleGatherCache))
-			{
-				return -1;
-			}
+			return -1;
 		}
 
 		WorkerDirector.TickWorkers();
@@ -1929,59 +1958,49 @@ int32 UGatherTextFromAssetsCommandlet::Main(const FString& Params)
 		// Ingest any results from workers
 		while (TOptional<FGatherTextFromAssetsWorkerMessage_PackageResult> PackageResult = WorkerDirector.IngestPackageResult())
 		{
-			if (!PackageResult->LoadLogCapture.IsEmpty())
-			{
-				GWarn->Log(NAME_None, ELogVerbosity::Display, PackageResult->LoadLogCapture);
-			}
-
-			if (PackageResult->bLoadError)
-			{
-				UE_LOGFMT(LogGatherTextFromAssetsCommandlet, Warning, "Failed to load package: '{package}'.",
-					("package", PackageResult->PackageName),
-					("id", GatherTextFromAssetsCommandlet::LocalizationLogIdentifier)
-				);
-			}
-			else
-			{
-				FMemoryReader MemoryReader(PackageResult->GatherableTextData);
-
-				TArray<FGatherableTextData> GatherableTextData;
-				MemoryReader << GatherableTextData;
-
-				UE_LOG(LogGatherTextFromAssetsCommandlet, Display, TEXT("Ingesting worker result for '%s'..."), *PackageResult->PackageName.ToString());
-				ProcessGatherableTextDataArray(GatherableTextData);
-			}
-
-			PackagesDistributedToWorkers.Remove(PackageResult->PackageName);
+			UE_LOG(LogGatherTextFromAssetsCommandlet, Display, TEXT("Ingesting worker result for '%s'..."), *PackageResult->PackageName.ToString());
+			IngestPackageResultFromWorker(*PackageResult);
 		}
 
 		// Re-add any jobs sent to crashed workers to the main task work queue
 		{
-			const TArray<FName> PackagesFromCrashedWorkers = WorkerDirector.IngestPackagesFromCrashedWorkers();
-			for (const FName PackageName : PackagesFromCrashedWorkers)
+			if (const TArray<FName> PackagesFromCrashedWorkers = WorkerDirector.IngestPackagesFromCrashedWorkers();
+				PackagesFromCrashedWorkers.Num() > 0)
 			{
-				if (FPackagePendingGather* PackagePendingGather = PackagesDistributedToWorkers.Find(PackageName))
+				UE::TScopeLock _1(PackagesPendingGatherMutex);
+				UE::TScopeLock _2(PackagesDistributedToWorkersMutex);
+				for (const FName PackageName : PackagesFromCrashedWorkers)
 				{
-					PackagesPendingGather.Add(MoveTemp(*PackagePendingGather));
-					PackagesDistributedToWorkers.Remove(PackageName);
+					if (FPackagePendingGather* PackagePendingGather = PackagesDistributedToWorkers.Find(PackageName))
+					{
+						PackagesPendingGather.Add(MoveTemp(*PackagePendingGather));
+						PackagesDistributedToWorkers.Remove(PackageName);
+					}
 				}
 			}
 		}
 
 		int32 NumPendingWorkerPackages = 0;
-		if (PackagesPendingGather.Num() == 0 && WorkerDirector.IsIdle(&NumPendingWorkerPackages))
 		{
-			break;
+			UE::TScopeLock _(PackagesPendingGatherMutex);
+			if (PackagesPendingGather.Num() == 0 && WorkerDirector.IsIdle(&NumPendingWorkerPackages))
+			{
+				break;
+			}
 		}
 
 		UE_LOG(LogGatherTextFromAssetsCommandlet, Display, TEXT("Waiting for workers to finish (%d pending package(s))..."), NumPendingWorkerPackages);
 		FPlatformProcess::SleepNoStats(1.0f);
 	}
 
+	FLoadPackageLogOutputRedirector::Get().Unhook();
+
+	WorkerDirector.ClearIngestPackageResultHandler();
+
 	// Clear list of packages with duplicate localization IDs
 	PackagesWithDuplicateLocalizationIds.Empty();
 
-	UE_LOG(LogGatherTextFromAssetsCommandlet, Display, TEXT("Loading packages took %.2f seconds."), FPlatformTime::Seconds() - PackageLoadingStartTime);
+	UE_LOG(LogGatherTextFromAssetsCommandlet, Display, TEXT("Loading %d packages took %.2f seconds (%d local, %d remote)."), TotalNumUncachedPackages, FPlatformTime::Seconds() - PackageLoadingStartTime, NumUncachedPackagesProcessedLocally.load(), NumUncachedPackagesProcessedRemotely.load());
 
 	// Collect garbage after loading all packages
 	// This reclaims as much memory as possible for the rest of the gather pipeline
@@ -2254,6 +2273,108 @@ void UGatherTextFromAssetsCommandlet::StartWorkers(const int32 MinPackagesToUseW
 	WorkerDirector.StartWorkers(NumWorkersToUse, /*bStopAdditionalWorkers*/false, /*NumRestartAttemptsIfCrashed*/4);
 }
 
+void UGatherTextFromAssetsCommandlet::AssignPackagesToWorkers(TConstArrayView<FGuid> IdleWorkerIds)
+{
+	if (IdleWorkerIds.Num() > 0)
+	{
+		if (PackagesPendingGather.Num() > UE::Private::GatherTextFromAssetsCommandlet::MinPackagesToKeepLocal)
+		{
+			if (const int32 NumPackagesToDistributeToEachWorker = FMath::Clamp((PackagesPendingGather.Num() - UE::Private::GatherTextFromAssetsCommandlet::MinPackagesToKeepLocal) / IdleWorkerIds.Num(), 0, UE::Private::GatherTextFromAssetsCommandlet::MaxPackagesToDistribute);
+				NumPackagesToDistributeToEachWorker > 0)
+			{
+				FGatherTextFromAssetsWorkerDirector& WorkerDirector = FGatherTextFromAssetsWorkerDirector::Get();
+
+				for (const FGuid& WorkerId : IdleWorkerIds)
+				{
+					check(PackagesPendingGather.Num() >= NumPackagesToDistributeToEachWorker);
+
+					// We assign worker packages from the tail of the queue (the front of the array) 
+					// as these packages have the fewest number of dependencies so are the simplest 
+					// to distribute and load, which helps to balance the memory usage between 
+					// the main process (which takes the heavy tasks) and the worker processes
+					TArray<FPackagePendingGather> PackagesToReinsert;
+					for (int32 Index = 0; Index < NumPackagesToDistributeToEachWorker; ++Index)
+					{
+						FPackagePendingGather PackagePendingGather = MoveTemp(PackagesPendingGather[Index]);
+
+						FGatherTextFromAssetsWorkerMessage_PackageRequest PackageRequest;
+						PackageRequest.PackageName = PackagePendingGather.PackageName;
+						PackageRequest.Dependencies = PackagePendingGather.Dependencies;
+						PackageRequest.ExternalActors = PackagePendingGather.ExternalActors;
+
+						if (WorkerDirector.AssignPackageToWorker(WorkerId, PackageRequest))
+						{
+							UE::TScopeLock _(PackagesDistributedToWorkersMutex);
+
+							const FName PackageName = PackagePendingGather.PackageName;
+							PackagesDistributedToWorkers.Add(PackageName, MoveTemp(PackagePendingGather));
+						}
+						else
+						{
+							PackagesToReinsert.Add(MoveTemp(PackagePendingGather));
+						}
+					}
+					PackagesPendingGather.RemoveAt(0, NumPackagesToDistributeToEachWorker, EAllowShrinking::No);
+					if (PackagesToReinsert.Num() > 0)
+					{
+						PackagesPendingGather.Insert(MoveTemp(PackagesToReinsert), 0);
+					}
+				}
+			}
+		}
+	}
+}
+
+void UGatherTextFromAssetsCommandlet::IngestPackageResultFromWorker(const FGatherTextFromAssetsWorkerMessage_PackageResult& PackageResult, const bool bSendWorkIfIdle)
+{
+	const int32 TotalNumUncachedPackagesProcessed = NumUncachedPackagesProcessedLocally + ++NumUncachedPackagesProcessedRemotely;
+	const float PercentageComplete = static_cast<float>(TotalNumUncachedPackagesProcessed) / static_cast<float>(TotalNumUncachedPackages) * 100.0f;
+	UE_LOG(LogGatherTextFromAssetsCommandlet, Display, TEXT("[%6.2f%%] Ingesting result for package: '%s' (from worker: %s)..."), PercentageComplete, *PackageResult.PackageName.ToString(), *PackageResult.WorkerId.ToString());
+
+	if (!PackageResult.LoadLogCapture.IsEmpty())
+	{
+		GWarn->Log(NAME_None, ELogVerbosity::Display, PackageResult.LoadLogCapture);
+	}
+
+	if (PackageResult.bLoadError)
+	{
+		UE_LOGFMT(LogGatherTextFromAssetsCommandlet, Warning, "Failed to load package: '{package}'.",
+			("package", PackageResult.PackageName),
+			("id", GatherTextFromAssetsCommandlet::LocalizationLogIdentifier)
+		);
+	}
+	else
+	{
+		FMemoryReader MemoryReader(PackageResult.GatherableTextData);
+
+		TArray<FGatherableTextData> GatherableTextData;
+		MemoryReader << GatherableTextData;
+
+		if (GatherableTextData.Num() > 0)
+		{
+			UE::TScopeLock _(GatherManifestHelperMutex);
+			ProcessGatherableTextDataArray(GatherableTextData);
+		}
+	}
+
+	{
+		UE::TScopeLock _(PackagesDistributedToWorkersMutex);
+		PackagesDistributedToWorkers.Remove(PackageResult.PackageName);
+	}
+
+	if (bSendWorkIfIdle)
+	{
+		UE::TScopeLock _(PackagesPendingGatherMutex);
+
+		// If we are (soon to be) idle, distribute some work now
+		if (TArray<FGuid> IdleWorkerIds = FGatherTextFromAssetsWorkerDirector::Get().GetAvailableWorkerIds(UE::Private::GatherTextFromAssetsCommandlet::WorkerIdleThreshold);
+			IdleWorkerIds.Contains(PackageResult.WorkerId))
+		{
+			AssignPackagesToWorkers(MakeArrayView(&PackageResult.WorkerId, 1));
+		}
+	}
+}
+
 int32 UGatherTextFromAssetsWorkerCommandlet::Main(const FString& Params)
 {
 	FProcHandle DirectorProc;
@@ -2346,9 +2467,31 @@ int32 UGatherTextFromAssetsWorkerCommandlet::Main(const FString& Params)
 		MessageEndpoint->Subscribe<FGatherTextFromAssetsWorkerMessage_Ping>();
 	}
 
+	// If the editor has loaded a persistent world then create an empty world prior to starting the asset gather
+	// This avoids any issues when loading and initializing worlds during the gather, as WP needs to re-initialize the world
+	if (GEditor)
+	{
+		if (UWorld* EditorWorld = GEditor->GetEditorWorldContext().World();
+			EditorWorld && !FPackageName::IsTempPackage(FNameBuilder(EditorWorld->GetPackage()->GetFName()).ToView()))
+		{
+			GEditor->CreateNewMapForEditing(/*bPromptForSave*/false);
+		}
+	}
+
 	// Run a GC now that we're ready to start accepting work, to clean-up anything loaded during start-up that is no longer needed
 	UE::Private::GatherTextFromAssetsCommandlet::PurgeGarbage<FPackagePendingGather>(nullptr);
 
+	// We don't need to have compiled shaders to gather text
+	bool bWasShaderCompilationEnabled = false;
+	if (GShaderCompilingManager)
+	{
+		bWasShaderCompilationEnabled = !GShaderCompilingManager->IsShaderCompilationSkipped();
+		GShaderCompilingManager->SkipShaderCompilation(true);
+	}
+
+	FStringOutputDevice LogOutputTargetDevice;
+	FLoadPackageLogOutputRedirector::Get().Hook(&LogOutputTargetDevice);
+
 	// Wait for work, and process each package requested
 	for (;;)
 	{
@@ -2371,17 +2514,15 @@ int32 UGatherTextFromAssetsWorkerCommandlet::Main(const FString& Params)
 			const FString PackageNameStr = PackagePendingGather->PackageName.ToString();
 			UE_LOG(LogGatherTextFromAssetsCommandlet, Display, TEXT("Gathering package: '%s'."), *PackageNameStr);
 
-			FStringOutputDevice LogOutputTargetDevice;
-			FLoadPackageLogOutputRedirector LogOutputRedirector(&LogOutputTargetDevice);
-
 			FGatherTextFromAssetsWorkerMessage_PackageResult* PackageResultMessage = FMessageEndpoint::MakeMessage<FGatherTextFromAssetsWorkerMessage_PackageResult>();
 			PackageResultMessage->WorkerId = WorkerId;
 			PackageResultMessage->PackageName = PackagePendingGather->PackageName;
 
-			if (UPackage* Package = UE::Private::GatherTextFromAssetsCommandlet::LoadPackageToGather(PackageNameStr, PackagePendingGather->ExternalActors.Num() > 0, LogOutputRedirector))
+			LogOutputTargetDevice.Reset();
+			if (UPackage* Package = UE::Private::GatherTextFromAssetsCommandlet::LoadPackageToGather(PackageNameStr, PackagePendingGather->ExternalActors.Num() > 0))
 			{
 				TArray<FGatherableTextData> GatherableTextDataArray;
-				UE::Private::GatherTextFromAssetsCommandlet::GatherTextFromPackage(Package, PackagePendingGather->ExternalActors, /*bGatherFromPrimaryPackage*/true, /*bGatherFromExternalActorPackages*/true, LogOutputRedirector, GatherableTextDataArray);
+				UE::Private::GatherTextFromAssetsCommandlet::GatherTextFromPackage(Package, PackagePendingGather->ExternalActors, /*bGatherFromPrimaryPackage*/true, /*bGatherFromExternalActorPackages*/true, GatherableTextDataArray);
 
 				FMemoryWriter MemoryWriter(PackageResultMessage->GatherableTextData);
 				MemoryWriter << GatherableTextDataArray;
@@ -2403,9 +2544,10 @@ int32 UGatherTextFromAssetsWorkerCommandlet::Main(const FString& Params)
 
 			UE_LOG(LogGatherTextFromAssetsCommandlet, Display, TEXT("Sent package result for '%s'."), *PackageNameStr);
 
-			UE::Private::GatherTextFromAssetsCommandlet::TickBackgroundTasks();
-
-			UE::Private::GatherTextFromAssetsCommandlet::ConditionalPurgeGarbage<FPackagePendingGather>(MinFreeMemoryBytes, MaxUsedMemoryBytes, PackagesPendingGather);
+			{
+				UE::TScopeLock _(PackagesPendingGatherMutex);
+				UE::Private::GatherTextFromAssetsCommandlet::ConditionalPurgeGarbage<FPackagePendingGather>(MinFreeMemoryBytes, MaxUsedMemoryBytes, &PackagesPendingGather);
+			}
 		}
 		else
 		{
@@ -2416,9 +2558,10 @@ int32 UGatherTextFromAssetsWorkerCommandlet::Main(const FString& Params)
 					IdleStartTimeUtc = FDateTime::UtcNow();
 				}
 			}
-			FPlatformProcess::SleepNoStats(1.0f);
+			FPlatformProcess::SleepNoStats(0.1f);
 		}
 
+		UE::Private::GatherTextFromAssetsCommandlet::TickBackgroundTasks();
 		UE::Private::GatherTextFromAssetsCommandlet::TickMessageBusGT();
 
 		// If our director process has stopped, then we can also exit
@@ -2429,11 +2572,18 @@ int32 UGatherTextFromAssetsWorkerCommandlet::Main(const FString& Params)
 		}
 	}
 
+	FLoadPackageLogOutputRedirector::Get().Unhook();
+
 	if (DirectorProc.IsValid())
 	{
 		FPlatformProcess::CloseProc(DirectorProc);
 	}
 
+	if (GShaderCompilingManager)
+	{
+		GShaderCompilingManager->SkipShaderCompilation(!bWasShaderCompilationEnabled);
+	}
+
 	// Disable the Endpoint message handling since the message could keep it alive a bit.
 	if (MessageEndpoint)
 	{
@@ -2665,6 +2815,12 @@ bool FGatherTextFromAssetsWorkerDirector::StopWorkers()
 	return true;
 }
 
+bool FGatherTextFromAssetsWorkerDirector::HasWorkers() const
+{
+	UE::TScopeLock _(CurrentWorkersMutex);
+	return CurrentWorkers.Num() > 0;
+}
+
 void FGatherTextFromAssetsWorkerDirector::TickWorkers()
 {
 	if (MessageEndpoint)
@@ -2808,6 +2964,16 @@ bool FGatherTextFromAssetsWorkerDirector::AssignPackageToWorker(const FGuid& Wor
 	return true;
 }
 
+void FGatherTextFromAssetsWorkerDirector::SetIngestPackageResultHandler(TFunction<bool(const FGatherTextFromAssetsWorkerMessage_PackageResult&)>&& Handler)
+{
+	PackageResultHandler = MoveTemp(Handler);
+}
+
+void FGatherTextFromAssetsWorkerDirector::ClearIngestPackageResultHandler()
+{
+	PackageResultHandler.Reset();
+}
+
 TOptional<FGatherTextFromAssetsWorkerMessage_PackageResult> FGatherTextFromAssetsWorkerDirector::IngestPackageResult()
 {
 	return PackageResults.Dequeue();
@@ -2868,13 +3034,22 @@ FString FGatherTextFromAssetsWorkerDirector::GenerateWorkerCommandLine(const FGu
 			Tokens.Add(MoveTemp(Token));
 		}
 	}
-	Tokens.Add(TEXT("-NoAssetRegistryCache")); // Disable the AR cache to avoid IO conflicts with the main process; the worker doesn't directly use the AR
+	
 	Tokens.Add(TEXT("-unattended"));
 	Tokens.Add(TEXT("-messaging"));
 	Tokens.Add(TEXT("-multiprocess"));
 	Tokens.Add(TEXT("-GatherWorkerId=") + WorkerId.ToString());
 	Tokens.Add(FString::Printf(TEXT("-GatherDirectorPid=%u"), FPlatformProcess::GetCurrentProcessId()));
 
+	if (bWorkersCanReadAssetRegistryCache)
+	{
+		Tokens.Add(TEXT("-NoAssetRegistryCacheWrite")); // Disable writing the AR cache as the main process has already updated it
+	}
+	else
+	{
+		Tokens.Add(TEXT("-NoAssetRegistryCache")); // Disable the AR cache to avoid IO conflicts with the main process; the worker doesn't directly use the AR
+	}
+
 	// We are joining the tokens back into a commandline string; wrap tokens with whitespace in quotes
 	for (FString& Token : Tokens)
 	{
@@ -3021,19 +3196,26 @@ void FGatherTextFromAssetsWorkerDirector::HandlePackageResultMessage(const FGath
 		{
 			CurrentWorker->LastMessageReceivedUtc = FDateTime::UtcNow();
 			CurrentWorker->IdleStartTimeUtc.Reset(); // Don't consider it idle if we're still receiving results from it
-
-			if (CurrentWorker->PendingPackageRequests.Contains(Message.PackageName))
-			{
-				bWasSolicitedResult = true;
-				CurrentWorker->PendingPackageRequests.Remove(Message.PackageName);
-			}
+			bWasSolicitedResult = CurrentWorker->PendingPackageRequests.Contains(Message.PackageName);
 		}
 	}
 
 	if (bWasSolicitedResult)
 	{
-		UE_LOG(LogGatherTextFromAssetsCommandlet, Display, TEXT("Received package result from worker '%s' for '%s' and queued it for ingestion."), *Message.WorkerId.ToString(), *Message.PackageName.ToString());
-		PackageResults.Enqueue(Message);
+		if (!PackageResultHandler || !PackageResultHandler(Message))
+		{
+			UE_LOG(LogGatherTextFromAssetsCommandlet, Display, TEXT("Received package result from worker '%s' for '%s' and queued it for ingestion."), *Message.WorkerId.ToString(), *Message.PackageName.ToString());
+			PackageResults.Enqueue(Message);
+		}
+
+		// Note: We do the remove after the ingestion to avoid a race condition with IsIdle
+		{
+			UE::TScopeLock _(CurrentWorkersMutex);
+			if (TSharedPtr<FWorkerInfo> CurrentWorker = CurrentWorkers.FindRef(Message.WorkerId))
+			{
+				CurrentWorker->PendingPackageRequests.Remove(Message.PackageName);
+			}
+		}
 	}
 	else
 	{
