Downgraded warning to avoid false positives when running multiple asset gather workers

#rnx

diff --git a/Engine/Source/Editor/UnrealEd/Private/Commandlets/GatherTextFromAssetsCommandlet.cpp b/Engine/Source/Editor/UnrealEd/Private/Commandlets/GatherTextFromAssetsCommandlet.cpp
--- a/Engine/Source/Editor/UnrealEd/Private/Commandlets/GatherTextFromAssetsCommandlet.cpp
+++ b/Engine/Source/Editor/UnrealEd/Private/Commandlets/GatherTextFromAssetsCommandlet.cpp
@@ -545,8 +545,8 @@ namespace UE::Private::GatherTextFromAssetsCommandlet
 				// If we're still over the memory limit after a minimal purge, then attempt a full purge
 				PurgeGarbage<TPackagePendingGather>(nullptr);
 
-				// If we're still over the memory limit after both purges, then log a warning as we may be about to OOM
-				UE_CLOG(HasExceededMemoryLimit(MinFreeMemoryBytes, MaxUsedMemoryBytes, /*bLog*/false), LogGatherTextFromAssetsCommandlet, Warning, TEXT("Flushing failed to reduce process memory to within the requested limits; this process may OOM!"));
+				// If we're still over the memory limit after both purges, then log a message that we may be about to OOM
+				UE_CLOG(HasExceededMemoryLimit(MinFreeMemoryBytes, MaxUsedMemoryBytes, /*bLog*/false), LogGatherTextFromAssetsCommandlet, Display, TEXT("Flushing failed to reduce process memory to within the requested limits; this process may OOM!"));
 			}
 		}
 	}
