Improved customization of the experimental multi-process localization asset gather

* Added `-StartGatherTextWorkersAlways` to skip the usual MinPackagesToUseWorkers threshold.
* This can be useful when your gather contains multiple smaller localization targets, where none meet MinPackagesToUseWorkers on their own, but can still benefit as many are over the MinPackagesToKeepLocal threshold.
* Added CVars for the worker threshold values, so that projects can tweak them to better fit their gather size.
* Localization.GatherTextFromAssetsCommandlet.MinPackagesToUseWorkers, Localization.GatherTextFromAssetsCommandlet.MinPackagesToKeepLocal, Localization.GatherTextFromAssetsCommandlet.MaxPackagesToDistribute, Localization.GatherTextFromAssetsCommandlet.WorkerIdleThreshold.

#rnx

diff --git a/Engine/Source/Editor/UnrealEd/Classes/Commandlets/GatherTextFromAssetsCommandlet.h b/Engine/Source/Editor/UnrealEd/Classes/Commandlets/GatherTextFromAssetsCommandlet.h
--- a/Engine/Source/Editor/UnrealEd/Classes/Commandlets/GatherTextFromAssetsCommandlet.h
+++ b/Engine/Source/Editor/UnrealEd/Classes/Commandlets/GatherTextFromAssetsCommandlet.h
@@ -202,6 +202,9 @@ private:
 	/** True if we should start the worker processes immediately, rather than wait until we know we have packages to load */
 	bool bStartWorkersImmediately = false;
 
+	/** True if we should start the worker processes always, even if we don't meet the MinPackagesToUseWorkers threshold */
+	bool bStartWorkersAlways = false;
+
 	/** Run a GC if the free system memory is less than this value (or zero to disable) */
 	uint64 MinFreeMemoryBytes;
 
diff --git a/Engine/Source/Editor/UnrealEd/Private/Commandlets/GatherTextFromAssetsCommandlet.cpp b/Engine/Source/Editor/UnrealEd/Private/Commandlets/GatherTextFromAssetsCommandlet.cpp
--- a/Engine/Source/Editor/UnrealEd/Private/Commandlets/GatherTextFromAssetsCommandlet.cpp
+++ b/Engine/Source/Editor/UnrealEd/Private/Commandlets/GatherTextFromAssetsCommandlet.cpp
@@ -349,10 +349,14 @@ namespace UE::Private::GatherTextFromAssetsCommandlet
 	static constexpr int32 WorkerProtocolVersion = 1;
 
 	// Values controlling how packages are distributed to workers
-	static constexpr int32 MinPackagesToUseWorkers = 1000;	// How many packages should we have to process (in total) before we start the workers?
-	static constexpr int32 MinPackagesToKeepLocal = 50;		// How many packages should we keep in the main process, rather than distribute to workers? (to avoid waiting on workers at the end of the gather)
-	static constexpr int32 MaxPackagesToDistribute = 100;	// How many packages should we distribute to each worker in a single batch?
-	static constexpr int32 WorkerIdleThreshold = 20;		// How many packages can a worker have pending in its queue before we consider it "idle" and try and assign it more packages? (to avoid workers becoming idle)
+	static int32 MinPackagesToUseWorkers = 1000;
+	static FAutoConsoleVariableRef CVarMinPackagesToUseWorkers(TEXT("Localization.GatherTextFromAssetsCommandlet.MinPackagesToUseWorkers"), MinPackagesToUseWorkers, TEXT("How many packages should we have to process (in total) before we start the workers?"));
+	static int32 MinPackagesToKeepLocal = 50;
+	static FAutoConsoleVariableRef CVarMinPackagesToKeepLocal(TEXT("Localization.GatherTextFromAssetsCommandlet.MinPackagesToKeepLocal"), MinPackagesToKeepLocal, TEXT("How many packages should we keep in the main process, rather than distribute to workers? (to avoid waiting on workers at the end of the gather)"));
+	static int32 MaxPackagesToDistribute = 100;
+	static FAutoConsoleVariableRef CVarMaxPackagesToDistribute(TEXT("Localization.GatherTextFromAssetsCommandlet.MaxPackagesToDistribute"), MaxPackagesToDistribute, TEXT("How many packages should we distribute to each worker in a single batch?"));
+	static int32 WorkerIdleThreshold = 20;
+	static FAutoConsoleVariableRef CVarWorkerIdleThreshold(TEXT("Localization.GatherTextFromAssetsCommandlet.WorkerIdleThreshold"), WorkerIdleThreshold, TEXT("How many packages can a worker have pending in its queue before we consider it 'idle' and try and assign it more packages? (to avoid workers becoming idle)"));
 
 	UPackage* LoadPackageToGather(const FString& PackageName, const bool bHasExternalActors)
 	{
@@ -1919,7 +1923,7 @@ int32 UGatherTextFromAssetsCommandlet::Main(const FString& Params)
 	const double PackageLoadingStartTime = FPlatformTime::Seconds();
 	UE_LOG(LogGatherTextFromAssetsCommandlet, Display, TEXT("Preparing to load %d packages..."), TotalNumUncachedPackages);
 
-	StartWorkers(UE::Private::GatherTextFromAssetsCommandlet::MinPackagesToUseWorkers);
+	StartWorkers(bStartWorkersAlways ? 0 : UE::Private::GatherTextFromAssetsCommandlet::MinPackagesToUseWorkers);
 
 	CalculateDependenciesForPackagesPendingGather();
 
@@ -2205,6 +2209,7 @@ bool UGatherTextFromAssetsCommandlet::ConfigureFromScript(const FString& GatherT
 
 	{
 		bStartWorkersImmediately = FParse::Param(FCommandLine::Get(), TEXT("StartGatherTextWorkersImmediately"));
+		bStartWorkersAlways = FParse::Param(FCommandLine::Get(), TEXT("StartGatherTextWorkersAlways"));
 
 		NumWorkers = 0;
 		FParse::Value(FCommandLine::Get(), TEXT("NumGatherTextWorkers="), NumWorkers);
