[Viewport] Fix a regression that caused middle clicking in th level editor viewport to select objects.

[RN] Editor, Minor

diff --git a/Engine/Source/Editor/UnrealEd/Private/LevelEditorViewport.cpp b/Engine/Source/Editor/UnrealEd/Private/LevelEditorViewport.cpp
--- a/Engine/Source/Editor/UnrealEd/Private/LevelEditorViewport.cpp
+++ b/Engine/Source/Editor/UnrealEd/Private/LevelEditorViewport.cpp
@@ -2814,13 +2814,16 @@ void FLevelEditorViewportClient::ProcessClick(FSceneView& View, HHitProxy* HitPr
 	}
 
 	const FViewportClick Click(&View,this,Key,Event,HitX,HitY);
-	if (Click.GetKey() == EKeys::MiddleMouseButton && !Click.IsAltDown() && !Click.IsShiftDown() && Click.IsControlDown())
+	if (Click.GetKey() == EKeys::MiddleMouseButton && !Click.IsAltDown() && !Click.IsShiftDown())
 	{
-		if (TSharedPtr<ILevelEditor> ParentLevelEditorPinned = ParentLevelEditor.Pin())
+		if (Click.IsControlDown())
 		{
-			ParentLevelEditorPinned->SummonLevelViewportViewOptionMenu(LVT_Perspective);
+			if (TSharedPtr<ILevelEditor> ParentLevelEditorPinned = ParentLevelEditor.Pin())
+			{
+				ParentLevelEditorPinned->SummonLevelViewportViewOptionMenu(LVT_Perspective);
+			}
 		}
-
+	
 		return;
 	}
 	if (!ModeTools->HandleClick(this, HitProxy,Click))
