StaticMeshEditor: [Github PR 12594] Add support for a search box to filter static mesh sockets

#12594

diff --git a/Engine/Source/Editor/UnrealEd/Private/SSocketManager.cpp b/Engine/Source/Editor/UnrealEd/Private/SSocketManager.cpp
--- a/Engine/Source/Editor/UnrealEd/Private/SSocketManager.cpp
+++ b/Engine/Source/Editor/UnrealEd/Private/SSocketManager.cpp
@@ -8,6 +8,8 @@
 #include "UObject/UObjectHash.h"
 #include "UObject/UObjectIterator.h"
 #include "Widgets/Layout/SSeparator.h"
+#include "Widgets/Input/SButton.h"
+#include "Widgets/Input/SSearchBox.h"
 #include "Developer/ToolWidgets/Public/SPositiveActionButton.h"
 #include "Widgets/Views/SListView.h"
 #include "Styling/AppStyle.h"
@@ -291,7 +293,23 @@ void SSocketManager::Construct(const FArguments& InArgs)
 				.BorderImage(FAppStyle::GetBrush("ToolPanel.GroupBorder"))
 				[
 					SNew(SVerticalBox)
-
+						
+					+ SVerticalBox::Slot()
+					.AutoHeight()
+					.Padding(0, 2)
+					[
+						SNew(SHorizontalBox)
+							// Filter entry
+							+ SHorizontalBox::Slot()
+							.FillWidth(1)
+							[
+								SAssignNew(NameFilterBox, SSearchBox)
+									.SelectAllTextWhenFocused(true)
+									.OnTextChanged(this, &SSocketManager::OnFilterTextChanged)
+									.OnTextCommitted(this, &SSocketManager::OnFilterTextCommitted)
+							]
+					]
+		
 					+ SVerticalBox::Slot()
 					.AutoHeight()
 					[
@@ -740,10 +758,11 @@ void SSocketManager::RefreshSocketList()
 			bIsSameStaticMesh = false;
 		}
 		// Only rebuild the socket list if it differs from the static meshes socket list
+		// or the filter text changed since the last refresh 
 		// This is done so that an undo on a socket property doesn't cause the selected
 		// socket to be de-selected, thus hiding the socket properties on the detail view.
 		// NB: Also force a rebuild if the underlying StaticMesh has been changed.
-		bool bRefreshList = StaticMesh->Sockets.Num() != SocketList.Num() || !bIsSameStaticMesh;
+		bool bRefreshList = StaticMesh->Sockets.Num() != SocketList.Num() || !bIsSameStaticMesh || !FilterTextAtLastListRefresh.EqualTo(FilterText);
 		if (!bRefreshList)
 		{
 			//Make sure list is in sync (UObject are the sames)
@@ -761,13 +780,25 @@ void SSocketManager::RefreshSocketList()
 		{
 			SocketListView->ClearSelection();
 			SocketList.Empty();
+			
+			const FString FilterString = *FilterText.ToString();
+
 			for (int32 i = 0; i < StaticMesh->Sockets.Num(); i++)
 			{
 				UStaticMeshSocket* Socket = StaticMesh->Sockets[i];
+				const FString SocketString = Socket->SocketName.ToString();
+
+				// See if we pass the search filter
+				if (!FilterText.IsEmpty() && !SocketString.Contains(FilterString))
+				{
+					continue;
+				}
+
 				SocketList.Add(MakeShareable(new SocketListItem(Socket)));
 			}
-
+			
 			SocketListView->RequestListRefresh();
+			FilterTextAtLastListRefresh = FilterText;
 		}
 
 		// Set the socket on the detail view to keep it in sync with the sockets properties
@@ -1011,6 +1042,19 @@ void SSocketManager::OnSocketPropertyChanged( const UStaticMeshSocket* Socket, c
 	}
 }
 
+void SSocketManager::OnFilterTextChanged(const FText& SearchText)
+{
+	FilterText = SearchText;
+
+	RefreshSocketList();
+}
+
+void SSocketManager::OnFilterTextCommitted(const FText& SearchText, ETextCommit::Type CommitInfo)
+{
+	// Just do the same as if the user typed in the box
+	OnFilterTextChanged(SearchText);
+}
+
 void SSocketManager::PostUndo()
 {
 	RefreshSocketList();
diff --git a/Engine/Source/Editor/UnrealEd/Private/SSocketManager.h b/Engine/Source/Editor/UnrealEd/Private/SSocketManager.h
--- a/Engine/Source/Editor/UnrealEd/Private/SSocketManager.h
+++ b/Engine/Source/Editor/UnrealEd/Private/SSocketManager.h
@@ -117,6 +117,24 @@ private:
 	/** Called when a socket property has changed. */
 	void OnSocketPropertyChanged( const UStaticMeshSocket* Socket, const FProperty* ChangedProperty );
 
+	/**
+	* Filters the SListView when the user changes the search text box (NameFilterBox)
+	*
+	* @param SearchText - The text the user has typed
+	*
+	*/
+	void OnFilterTextChanged(const FText& SearchText);
+
+	/**
+	* Filters the SListView when the user hits enter or clears the search box
+	* Simply calls OnFilterTextChanged
+	*
+	* @param SearchText - The text the user has typed
+	* @param CommitInfo - Not used
+	*
+	*/
+	void OnFilterTextCommitted(const FText& SearchText, ETextCommit::Type CommitInfo);
+
 	/** Called when socket selection changes */
 	FSimpleDelegate OnSocketSelectionChanged;
 
@@ -132,6 +150,15 @@ private:
 	/** Listview for displaying the sockets. */
 	TSharedPtr< SListView<TSharedPtr< SocketListItem > > > SocketListView;
 
+	/** Box to filter to a specific static mesh name */
+	TSharedPtr< class SSearchBox > NameFilterBox;
+
+	/** Current text typed into NameFilterBox */
+	FText FilterText;
+
+	/** Text typed into NameFilterBox at the time of the last socket list refresh */
+	FText FilterTextAtLastListRefresh;
+
 	/** Helper variable for rotating in world space. */
 	FVector WorldSpaceRotation;
 
