Cooker: CookPackageSplitter: Change the return value of GetGenerateList to a struct (and change name to ReportGenerationManifest) so that generators can request that the package containing the generator be not saved.
[RN] Minor, Cooking

#FailedValidations Preflight

diff --git a/Engine/Plugins/Mutable/Source/CustomizableObjectEditor/Private/MuCOE/CustomizableObjectCookPackageSplitter.cpp b/Engine/Plugins/Mutable/Source/CustomizableObjectEditor/Private/MuCOE/CustomizableObjectCookPackageSplitter.cpp
--- a/Engine/Plugins/Mutable/Source/CustomizableObjectEditor/Private/MuCOE/CustomizableObjectCookPackageSplitter.cpp
+++ b/Engine/Plugins/Mutable/Source/CustomizableObjectEditor/Private/MuCOE/CustomizableObjectCookPackageSplitter.cpp
@@ -168,7 +168,7 @@ bool FCustomizableObjectCookPackageSplitter::ShouldSplit(UObject* SplitData)
 	return false;
 }
 
-TArray<ICookPackageSplitter::FGeneratedPackage> FCustomizableObjectCookPackageSplitter::GetGenerateList(
+ICookPackageSplitter::FGenerationManifest FCustomizableObjectCookPackageSplitter::ReportGenerationManifest(
 	const UPackage* OwnerPackage,
 	const UObject* OwnerObject)
 {
@@ -181,18 +181,18 @@ TArray<ICookPackageSplitter::FGeneratedPackage> FCustomizableObjectCookPackageSp
 	const UModelResources* ModelResources = FindModelResources(*Object);
 	check(ModelResources);
 
-	TArray<ICookPackageSplitter::FGeneratedPackage> Result;
+	ICookPackageSplitter::FGenerationManifest Result;
 
 	// Generate a new package for each streamed Resource Data
 	for (const FCustomizableObjectStreamedResourceData& StreamedData : ModelResources->StreamedResourceData)
 	{
-		GenerateNewPackage(StreamedData, OwnerPackage, OwnerObject, Result);
+		GenerateNewPackage(StreamedData, OwnerPackage, OwnerObject, Result.GeneratedPackages);
 	}
 
 	// Generate a new package for each streamed Extension Data
 	for (const FCustomizableObjectStreamedResourceData& StreamedData : ModelResources->StreamedExtensionData)
 	{
-		GenerateNewPackage(StreamedData, OwnerPackage, OwnerObject, Result);
+		GenerateNewPackage(StreamedData, OwnerPackage, OwnerObject, Result.GeneratedPackages);
 	}
 
 	return Result;
diff --git a/Engine/Plugins/Mutable/Source/CustomizableObjectEditor/Private/MuCOE/CustomizableObjectCookPackageSplitter.h b/Engine/Plugins/Mutable/Source/CustomizableObjectEditor/Private/MuCOE/CustomizableObjectCookPackageSplitter.h
--- a/Engine/Plugins/Mutable/Source/CustomizableObjectEditor/Private/MuCOE/CustomizableObjectCookPackageSplitter.h
+++ b/Engine/Plugins/Mutable/Source/CustomizableObjectEditor/Private/MuCOE/CustomizableObjectCookPackageSplitter.h
@@ -15,7 +15,7 @@ public:
 	static FString GetSplitterDebugName() { return TEXT("FCustomizableObjectCookPackageSplitter"); }
 	static bool RequiresCachedCookedPlatformDataBeforeSplit() { return true; }
 
-	virtual TArray<ICookPackageSplitter::FGeneratedPackage> GetGenerateList(
+	virtual ICookPackageSplitter::FGenerationManifest ReportGenerationManifest(
 		const UPackage* OwnerPackage,
 		const UObject* OwnerObject) override;
 	
diff --git a/Engine/Source/Editor/UnrealEd/Private/CookOnTheFlyServer.cpp b/Engine/Source/Editor/UnrealEd/Private/CookOnTheFlyServer.cpp
--- a/Engine/Source/Editor/UnrealEd/Private/CookOnTheFlyServer.cpp
+++ b/Engine/Source/Editor/UnrealEd/Private/CookOnTheFlyServer.cpp
@@ -3185,7 +3185,7 @@ void UCookOnTheFlyServer::LoadPackageInQueue(UE::Cook::FPackageData& PackageData
 			RejectPackageToLoad(PackageData, TEXT("is an orphaned generated package"), ESuppressCookReason::OrphanedGenerated);
 			return;
 		}
-		if (!GenerationHelper->TryGenerateList())
+		if (!GenerationHelper->TryReportGenerationManifest())
 		{
 			RejectPackageToLoad(PackageData, TEXT("is an orphaned generated package"), ESuppressCookReason::OrphanedGenerated);
 			return;
@@ -4059,7 +4059,7 @@ UE::Cook::EPollStatus UCookOnTheFlyServer::TryPopulateGeneratedPackage(UE::Cook:
 	bool bPackageIsMap = GeneratedPackage->ContainsMap();
 	if (bPackageIsMap != GeneratedInfo.IsCreateAsMap())
 	{
-		UE_LOG(LogCook, Error, TEXT("PackageSplitter specified generated package is %s in GetGenerateList results, but then in PreSaveGeneratedPackage created it as %s. Splitter=%s, Generated=%s."),
+		UE_LOG(LogCook, Error, TEXT("PackageSplitter specified generated package is %s in ReportGenerationManifest, but then in PreSaveGeneratedPackage created it as %s. Splitter=%s, Generated=%s."),
 			(GeneratedInfo.IsCreateAsMap() ? TEXT("map") : TEXT("uasset")), (bPackageIsMap ? TEXT("map") : TEXT("uasset")),
 			*GenerationHelper.GetSplitDataObjectName().ToString(), *GeneratedPackageName);
 		return EPollStatus::Error;
@@ -4220,13 +4220,13 @@ UE::Cook::EPollStatus UCookOnTheFlyServer::PrepareSaveInternal(UE::Cook::FPackag
 
 		if (GenerationHelper)
 		{
-			if (PackageData.GetSaveSubState() <= ESaveSubState::Generation_TryGenerateList)
+			if (PackageData.GetSaveSubState() <= ESaveSubState::Generation_TryReportGenerationManifest)
 			{
 				// Keep it referenced even if we are only precaching, so we do not recreate it
 				GenerationHelper->SetKeepForGeneratorSaveAllPlatforms();
 				if (bPrecaching)
 				{
-					// Do not proceed to GetGenerateList when precaching; do that only when we're ready to save the package
+					// Do not proceed to ReportGenerationManifest when precaching; do that only when we're ready to save the package
 					return EPollStatus::Incomplete;
 				}
 				else
@@ -4253,12 +4253,12 @@ UE::Cook::EPollStatus UCookOnTheFlyServer::PrepareSaveInternal(UE::Cook::FPackag
 						return EPollStatus::Incomplete;
 					}
 
-					if (!GenerationHelper->TryGenerateList())
+					if (!GenerationHelper->TryReportGenerationManifest())
 					{
 						return EPollStatus::Error;
 					}
 					GenerationHelper->StartOwnerSave();
-					PackageData.SetSaveSubStateComplete(ESaveSubState::Generation_TryGenerateList);
+					PackageData.SetSaveSubStateComplete(ESaveSubState::Generation_TryReportGenerationManifest);
 				}
 			}
 
diff --git a/Engine/Source/Editor/UnrealEd/Private/Cooker/CookGenerationHelper.cpp b/Engine/Source/Editor/UnrealEd/Private/Cooker/CookGenerationHelper.cpp
--- a/Engine/Source/Editor/UnrealEd/Private/Cooker/CookGenerationHelper.cpp
+++ b/Engine/Source/Editor/UnrealEd/Private/Cooker/CookGenerationHelper.cpp
@@ -327,7 +327,8 @@ void FGenerationHelper::Uninitialize()
 	// Keep DoesGeneratedRequireGeneratorValue; it is allowed in the uninitialized state
 	// Keep bUseInternalReferenceToAvoidGarbageCollect; it is allowed in the uninitialized state
 	// Keep bRequiresGeneratorPackageDestructBeforeResplit; it is allowed in the uninitialized state
-	bGeneratedList = false;
+	bReportedManifest = false;
+	bSaveGenerator = true;
 	bCurrentGCHasKeptGeneratorPackage = false;
 	bCurrentGCHasKeptGeneratorKeepPackages = false;
 	// Keep bKeepForAllSavedOrGC ; it is allowed in the uninitialized state
@@ -946,9 +947,9 @@ UPackage* FGenerationHelper::FindOrLoadOwnerPackage(UCookOnTheFlyServer& COTFS)
 	return Result;
 }
 
-bool FGenerationHelper::TryGenerateList()
+bool FGenerationHelper::TryReportGenerationManifest()
 {
-	if (bGeneratedList)
+	if (bReportedManifest)
 	{
 		return true;
 	}
@@ -957,7 +958,8 @@ bool FGenerationHelper::TryGenerateList()
 	if (!IsValid())
 	{
 		// Unexpected, caller should not call in this case
-		UE_LOG(LogCook, Error, TEXT("TryGenerateList failed for package %s: Called on an invalid FGenerationHelper."),
+		UE_LOG(LogCook, Error,
+			TEXT("TryReportGenerationManifest failed for package %s: Called on an invalid FGenerationHelper."),
 			*OwnerPackageName.ToString());
 		FDebug::DumpStackTraceToLog(ELogVerbosity::Warning);
 		return false;
@@ -969,7 +971,8 @@ bool FGenerationHelper::TryGenerateList()
 	if (!OwnerObject)
 	{
 		// Unexpected, we found it earlier when we marked valid.
-		UE_LOG(LogCook, Error, TEXT("TryGenerateList failed for package %s: Valid GenerationHelper but could not find OwnerObject."),
+		UE_LOG(LogCook, Error,
+			TEXT("TryReportGenerationManifest failed for package %s: Valid GenerationHelper but could not find OwnerObject."),
 			*OwnerPackageName.ToString());
 		FDebug::DumpStackTraceToLog(ELogVerbosity::Warning);
 		return false;
@@ -977,7 +980,7 @@ bool FGenerationHelper::TryGenerateList()
 
 	UPackage* LocalOwnerPackage = OwnerObject->GetPackage();
 
-	TArray<ICookPackageSplitter::FGeneratedPackage> GeneratorDatas;
+	ICookPackageSplitter::FGenerationManifest Manifest;
 	{
 		UCookOnTheFlyServer::FScopedActivePackage ScopedActivePackage(COTFS, OwnerPackageName,
 #if UE_WITH_OBJECT_HANDLE_TRACKING
@@ -985,23 +988,40 @@ bool FGenerationHelper::TryGenerateList()
 #else
 			FName()
 #endif
-			);
-		GeneratorDatas = GetCookPackageSplitterInstance()->GetGenerateList(LocalOwnerPackage, OwnerObject);
+		);
+		Manifest = GetCookPackageSplitterInstance()->ReportGenerationManifest(LocalOwnerPackage, OwnerObject);
+		if (Manifest.GeneratedPackages.IsEmpty())
+		{
+			PRAGMA_DISABLE_DEPRECATION_WARNINGS;
+			TArray<ICookPackageSplitter::FGeneratedPackage> GeneratorDatas;
+			GeneratorDatas = GetCookPackageSplitterInstance()->GetGenerateList(LocalOwnerPackage, OwnerObject);
+			PRAGMA_ENABLE_DEPRECATION_WARNINGS;
+			if (!GeneratorDatas.IsEmpty())
+			{
+				// Deprecated in UE 5.7
+				UE_LOG(LogCook, Warning,
+					TEXT("%s is overriding deprecated ICookPackageSplitter::GetGenerateList. Implement ReportGenerationManifest instead. - Please update your code to the new API before upgrading to the next release, otherwise your project will no longer compile."),
+					*GetRegisteredSplitterType()->GetSplitterDebugName());
+				Manifest.GeneratedPackages = MoveTemp(GeneratorDatas);
+			}
+		}
 	}
 
+	bSaveGenerator = Manifest.bCookSaveGeneratorPackage;
+
 	TMap<FPackageData*, int32> AlreadyExistingInfoPackageToIndex;
 	int32 NumAlreadyExisting = PackagesToGenerate.Num();
 	for (int32 ExistingIndex = 0; ExistingIndex < PackagesToGenerate.Num(); ++ExistingIndex)
 	{
 		AlreadyExistingInfoPackageToIndex.Add(PackagesToGenerate[ExistingIndex].PackageData, ExistingIndex);
 	}
-	PackagesToGenerate.Reserve(GeneratorDatas.Num());
+	PackagesToGenerate.Reserve(Manifest.GeneratedPackages.Num());
 
 	for (TPair<const ITargetPlatform*, FPlatformData>& PlatformPair : PlatformDatas)
 	{
 		PlatformPair.Value.NumSaved = 0;
 	}
-	for (ICookPackageSplitter::FGeneratedPackage& SplitterData : GeneratorDatas)
+	for (ICookPackageSplitter::FGeneratedPackage& SplitterData : Manifest.GeneratedPackages)
 	{
 		if (!SplitterData.GetCreateAsMap().IsSet())
 		{
@@ -1112,8 +1132,9 @@ bool FGenerationHelper::TryGenerateList()
 		for (TPair<FPackageData*, int32>& Pair : AlreadyExistingInfoPackageToIndex)
 		{
 			UnusedExistingIndexes.Add(Pair.Value);
-			UE_LOG(LogCook, Warning, TEXT("Unexpected generated package (local TryGenerateList). A remote cookworker reported generated package %s for generator %s,")
-				TEXT(" but when TryGenerateList was called on the CookDirector, this package was not listed.")
+			UE_LOG(LogCook, Warning,
+				TEXT("Unexpected generated package (local TryReportGenerationManifest). A remote cookworker reported generated package %s for generator %s,")
+				TEXT(" but when TryReportGenerationManifest was called on the CookDirector, this package was not listed.")
 				TEXT(" This is unexpected and causes minor performance problems in the cook."),
 				*Pair.Key->GetPackageName().ToString(), *OwnerPackageData.GetPackageName().ToString());
 		}
@@ -1124,13 +1145,13 @@ bool FGenerationHelper::TryGenerateList()
 		}
 	}
 	UE_LOG(LogCookGenerationHelper, Verbose,
-		TEXT("%s: TryGenerateList: %d packages to generate, %d previously recorded, %d/%d already saved."),
+		TEXT("%s: TryReportGenerationManifest: %d packages to generate, %d previously recorded, %d/%d already saved."),
 		*WriteToString<256>(GetOwner().GetPackageName()), PackagesToGenerate.Num(),
 		NumAlreadyExisting, MinNumSaved, PackagesToGenerate.Num() + 1);
 
 	OnNumSavedUpdated();
 
-	bGeneratedList = true;
+	bReportedManifest = true;
 	return true;
 }
 
@@ -1144,15 +1165,16 @@ bool FGenerationHelper::TryCallPopulateGeneratorPackage(
 	FPackageData& OwnerPackageData = GetOwner();
 	FName OwnerPackageName = OwnerPackageData.GetPackageName();
 	UCookOnTheFlyServer& COTFS = OwnerPackageData.GetPackageDatas().GetCookOnTheFlyServer();
-	if (!bGeneratedList)
+	if (!bReportedManifest)
 	{
 		// Unexpected, caller should not call in this case
-		UE_LOG(LogCook, Error, TEXT("TryCallPopulateGeneratorPackage called for package %s without a previous successful call to TryGenerateList."),
+		UE_LOG(LogCook, Error,
+			TEXT("TryCallPopulateGeneratorPackage called for package %s without a previous successful call to TryReportGenerationManifest."),
 			*OwnerPackageName.ToString());
 		FDebug::DumpStackTraceToLog(ELogVerbosity::Warning);
 		return false;
 	}
-	check(IsValid()); // Could not have set bGeneratedList=true without being valid. 
+	check(IsValid()); // Could not have set bReportedManifest=true without being valid. 
 	UObject* LocalSplitDataObject = this->FindOrLoadSplitDataObject();
 	if (!LocalSplitDataObject)
 	{
@@ -1227,15 +1249,16 @@ bool FGenerationHelper::TryCallPopulateGeneratedPackage(UE::Cook::FCookGeneratio
 	FName OwnerPackageName = OwnerPackageData.GetPackageName();
 	UCookOnTheFlyServer& COTFS = OwnerPackageData.GetPackageDatas().GetCookOnTheFlyServer();
 
-	if (!bGeneratedList)
+	if (!bReportedManifest)
 	{
 		// Unexpected, caller should not call in this case
-		UE_LOG(LogCook, Error, TEXT("TryCallPopulateGeneratedPackage called for package %s without a previous successful call to TryGenerateList."),
+		UE_LOG(LogCook, Error,
+			TEXT("TryCallPopulateGeneratedPackage called for package %s without a previous successful call to TryReportGenerationManifest."),
 			*Info.GetPackageName());
 		FDebug::DumpStackTraceToLog(ELogVerbosity::Warning);
 		return false;
 	}
-	check(IsValid()); // Could not have set bGeneratedList=true without being valid. 
+	check(IsValid()); // Could not have set bReportedManifest=true without being valid. 
 	UObject* LocalSplitDataObject = this->FindOrLoadSplitDataObject();
 	if (!LocalSplitDataObject)
 	{
@@ -1410,7 +1433,7 @@ void FGenerationHelper::NotifyStartQueueGeneratedPackages(UCookOnTheFlyServer& C
 	// Note this function can be called on an uninitialized Generator; the generator is only needed
 	// on the director so it can serve as the passer of messages. We have to keep ourselves referenced after
 	// this call, until after we send EGeneratorEvent::QueuedGeneratedPackagesFencePassed, so that we don't destruct
-	// and lose the information from SavedOnWorker or TryGenerateList.
+	// and lose the information from SavedOnWorker or TryReportGenerationManifest.
 	if (!COTFS.CookWorkerClient)
 	{
 		GetOwnerInfo().SavedOnWorker = SourceWorkerId;
@@ -2018,7 +2041,7 @@ void FGenerationHelper::FetchExternalActorDependencies()
 	// The Generator package declares all its ExternalActor dependencies in its AssetRegistry dependencies
 	// The Generator's generated packages can also include ExternalActors from other maps due to level instancing,
 	// these are included in the dependencies reported by the Generator for each GeneratedPackage in the data
-	// returned from GetGenerateList. These sets will overlap; take the union.
+	// returned from ReportGenerationManifest. These sets will overlap; take the union.
 	ExternalActorDependencies.Reset();
 	IAssetRegistry::GetChecked().GetDependencies(GetOwner().GetPackageName(), ExternalActorDependencies,
 		UE::AssetRegistry::EDependencyCategory::Package);
@@ -2252,12 +2275,13 @@ void FGenerationHelper::PostGarbageCollectGCLifetimeData(FCookGCDiagnosticContex
 void FGenerationHelper::TrackGeneratedPackageListedRemotely(UCookOnTheFlyServer& COTFS, FPackageData& PackageData,
 	const FIoHash& CurrentPackageHash)
 {
-	if (bGeneratedList)
+	if (bReportedManifest)
 	{
 		if (!FindInfo(PackageData))
 		{
-			UE_LOG(LogCook, Warning, TEXT("Unexpected generated package (discovery replication). A remote cookworker reported generated package %s for generator %s,")
-				TEXT(" but when TryGenerateList was called on the CookDirector, this package was not listed.")
+			UE_LOG(LogCook, Warning,
+				TEXT("Unexpected generated package (discovery replication). A remote cookworker reported generated package %s for generator %s,")
+				TEXT(" but when TryReportGenerationManifest was called on the CookDirector, this package was not listed.")
 				TEXT(" This is unexpected and causes minor performance problems in the cook."),
 				*PackageData.GetPackageName().ToString(), *GetOwner().GetPackageName().ToString());
 		}
diff --git a/Engine/Source/Editor/UnrealEd/Private/Cooker/CookGenerationHelper.h b/Engine/Source/Editor/UnrealEd/Private/Cooker/CookGenerationHelper.h
--- a/Engine/Source/Editor/UnrealEd/Private/Cooker/CookGenerationHelper.h
+++ b/Engine/Source/Editor/UnrealEd/Private/Cooker/CookGenerationHelper.h
@@ -249,7 +249,9 @@ public:
 	 */
 	bool IsValid();
 
-	/** Accessor for the packages to generate, will be empty if invalid or if TryGenerateList was not yet called. */
+	/** True if ReportGenerationManifest requested that the generator package should be cooksaved as well. */
+	bool IsSaveGenerator() const;
+	/** Accessor for the packages to generate, empty if invalid or TryReportGenerationManifest was not yet called. */
 	TArrayView<FCookGenerationInfo> GetPackagesToGenerate();
 	/** Return the GenerationInfo used to save the generator package's UPackage. Does not call Initialize. */
 	FCookGenerationInfo& GetOwnerInfo();
@@ -299,8 +301,11 @@ public:
 	 */
 	TArray<FName> ReleaseExternalActorDependencies();
 
-	/** Call the Splitter's GetGenerateList and create the PackageDatas. Logs errors and returns false on failure. */
-	bool TryGenerateList();
+	/**
+	 * Call the Splitter's ReportGenerationManifest and create the PackageDatas. Logs errors and returns false
+	 * on failure.
+	 */
+	bool TryReportGenerationManifest();
 	/**
 	 * Call the Splitter's PopulateGeneratorPackage if not yet called. Assumes GenerateList has been called. Logs
 	 * errors and returns false on failure.
@@ -450,7 +455,7 @@ public:
 	int32& GetMPCookNextAssignmentIndex();
 
 	/**
-	 * Called for each of the generated packages that were discovered when TryGenerateList was called on a
+	 * Called for each of the generated packages that were discovered when TryReportGenerationManifest was called on a
 	 * remote CookWorker. Does not call Initialize.
 	 */
 	void TrackGeneratedPackageListedRemotely(UCookOnTheFlyServer& COTFS, FPackageData& PackageData,
@@ -592,7 +597,8 @@ private:
 		ICookPackageSplitter::EGeneratedRequiresGenerator::None;
 	bool bUseInternalReferenceToAvoidGarbageCollect = false;
 	bool bRequiresGeneratorPackageDestructBeforeResplit = false;
-	bool bGeneratedList = false;
+	bool bReportedManifest = false;
+	bool bSaveGenerator = true;
 	bool bCurrentGCHasKeptGeneratorPackage = false;
 	bool bCurrentGCHasKeptGeneratorKeepPackages = false;
 	bool bKeepForAllSavedOrGC = false;
@@ -791,6 +797,11 @@ inline bool FGenerationHelper::IsValid()
 	return InitializeStatus == EInitializeStatus::Valid;
 }
 
+inline bool FGenerationHelper::IsSaveGenerator() const
+{
+	return bSaveGenerator;
+}
+
 inline TArrayView<UE::Cook::FCookGenerationInfo> FGenerationHelper::GetPackagesToGenerate()
 {
 	ConditionalInitialize();
diff --git a/Engine/Source/Editor/UnrealEd/Private/Cooker/CookTypes.cpp b/Engine/Source/Editor/UnrealEd/Private/Cooker/CookTypes.cpp
--- a/Engine/Source/Editor/UnrealEd/Private/Cooker/CookTypes.cpp
+++ b/Engine/Source/Editor/UnrealEd/Private/Cooker/CookTypes.cpp
@@ -130,7 +130,7 @@ const TCHAR* LexToString(UE::Cook::ESaveSubState State)
 	case ESaveSubState::FirstCookedPlatformData_CallingBegin: return TEXT("FirstCookedPlatformData_CallingBegin");
 	case ESaveSubState::FirstCookedPlatformData_CheckForGenerator: return TEXT("FirstCookedPlatformData_CheckForGenerator");
 	case ESaveSubState::FirstCookedPlatformData_CheckForGeneratorAfterWaitingForIsLoaded: return TEXT("FirstCookedPlatformData_CheckForGeneratorAfterWaitingForIsLoaded");
-	case ESaveSubState::Generation_TryGenerateList: return TEXT("Generation_TryGenerateList");
+	case ESaveSubState::Generation_TryReportGenerationManifest: return TEXT("Generation_TryReportGenerationManifest");
 	case ESaveSubState::Generation_QueueGeneratedPackages: return TEXT("Generation_QueueGeneratedPackages");
 	case ESaveSubState::CheckForIsGenerated: return TEXT("CheckForIsGenerated");
 	case ESaveSubState::Generation_PreMoveCookedPlatformData_WaitingForIsLoaded: return TEXT("Generation_PreMoveCookedPlatformData_WaitingForIsLoaded");
diff --git a/Engine/Source/Editor/UnrealEd/Private/Cooker/CookTypes.h b/Engine/Source/Editor/UnrealEd/Private/Cooker/CookTypes.h
--- a/Engine/Source/Editor/UnrealEd/Private/Cooker/CookTypes.h
+++ b/Engine/Source/Editor/UnrealEd/Private/Cooker/CookTypes.h
@@ -237,7 +237,7 @@ namespace UE::Cook
 		FirstCookedPlatformData_CallingBegin,
 		FirstCookedPlatformData_CheckForGenerator,
 		FirstCookedPlatformData_CheckForGeneratorAfterWaitingForIsLoaded,
-		Generation_TryGenerateList,
+		Generation_TryReportGenerationManifest,
 		Generation_QueueGeneratedPackages,
 
 		CheckForIsGenerated,
diff --git a/Engine/Source/Editor/UnrealEd/Public/CookPackageSplitter.h b/Engine/Source/Editor/UnrealEd/Public/CookPackageSplitter.h
--- a/Engine/Source/Editor/UnrealEd/Public/CookPackageSplitter.h
+++ b/Engine/Source/Editor/UnrealEd/Public/CookPackageSplitter.h
@@ -41,7 +41,7 @@ public:
 	// before creating an instance of the class.
 	/**
 	 * Return whether IsCachedCookedPlatformDataLoaded needs to return true for all UObjects in the
-	 * generator package before ShouldSplit or GetGenerateList can be called. If true this slows down
+	 * generator package before ShouldSplit or ReportGenerationManifest can be called. If true this slows down
 	 * our ability to parallelize the cook of the generated packages.
 	 */
 	static bool RequiresCachedCookedPlatformDataBeforeSplit() { return false; }
@@ -97,20 +97,20 @@ public:
 	enum class EGeneratedRequiresGenerator : uint8
 	{
 		/**
-		 * GetGenerateList will be called before PopulateGeneratedPackage. PopulateGenerator and PreSaveGenerator 
-		 * might or might not be called before. OutKeepReferencedPackages from PopulateGenerator will not be kept
-		 * referenced after PostSaveGenerator. Best for performance.
+		 * ReportGenerationManifest will be called before PopulateGeneratedPackage. PopulateGenerator and
+		 * PreSaveGenerator might or might not be called before. OutKeepReferencedPackages from PopulateGenerator
+		 * will not be kept referenced after PostSaveGenerator. Best for performance.
 		 */
 		None,
 		/**
-		 * GetGenerateList and PopulateGenerator will be called before PopulateGeneratedPackage.
+		 * ReportGenerationManifest and PopulateGenerator will be called before PopulateGeneratedPackage.
 		 * OutKeepReferencedPackages from PopulateGenerator will be kept referenced until all generated and generator
 		 * packages call PostSave or until the splitter is destroyed. Performance cost: Possible extra calls to
 		 * PopulateGeneratedPackage, possible unnecessary memory increase due to OutKeepReferencedPackages.
 		 */
 		Populate,
 		/**
-		 * GetGenerateList PopulateGenerator, PreSaveGenerator, and PostSaveGenerator will be called before
+		 * ReportGenerationManifest, PopulateGenerator, PreSaveGenerator, and PostSaveGenerator will be called before
 		 * PopulateGeneratedPackage. Performance cost: Progress on generated packages will be delayed until generator
 		 * finishes saving. Possible unnecessary memory increase due to OutKeepReferencedPackages. Retraction is not
 		 * possible in MPCook for the generated packages; they must all be saved on the same CookWorker that saves the
@@ -136,7 +136,7 @@ public:
 		 return EGeneratedRequiresGenerator::None;
 	}
 
-	/** Data sent to the cooker to describe each desired generated package */
+	/** Data sent to the cooker to describe each desired generated package. */
 	struct FGeneratedPackage
 	{
 		/** Parent path for the generated package. If empty, uses the generator's package path. */
@@ -161,7 +161,11 @@ public:
 		 * Changes to this hash will cause invalidation of the package during incremental cooks.
 		 */
 		FBlake3Hash GenerationHash;
-		/* GetGenerateList must specify true if the package will be a map (.umap, contains a UWorld or ULevel), else false */
+		/*
+		 * Must be called with true or false for each GeneratedPackage reported in the GenerationManifest.
+		 * Report whether PopulateGeneratorPackage will populate the package as a map (contains a UWorld or ULevel),
+		 * and it should therefore be assigned a .umap extension, or as a non-map package, with a .uasset extension.
+		 */
 		void SetCreateAsMap(bool bInCreateAsMap) { bCreateAsMap = bInCreateAsMap; }
 		const TOptional<bool>& GetCreateAsMap() const { return bCreateAsMap; }
 
@@ -182,15 +186,35 @@ public:
 		TOptional<bool> bCreateAsMap;
 	};
 
+	/** Data sent to the cooker to describe what to generate. */
+	struct FGenerationManifest
+	{
+		/** List of generated packages. */
+		TArray<FGeneratedPackage> GeneratedPackages;
+		/** Whether the package containing the generator should also be saved and added to the cook's output. */
+		bool bCookSaveGeneratorPackage = true;
+	};
+
+	/** Return the list of packages to generate. */
+	UE_DEPRECATED(5.7, "Implement ReportGenerationManifest instead.")
+	virtual TArray<FGeneratedPackage> GetGenerateList(const UPackage* OwnerPackage, const UObject* OwnerObject)
+	{
+		return TArray<FGeneratedPackage>();
+	}
+
 	/** Return the list of packages to generate. */
-	virtual TArray<FGeneratedPackage> GetGenerateList(const UPackage* OwnerPackage, const UObject* OwnerObject) = 0;
+	virtual FGenerationManifest ReportGenerationManifest(const UPackage* OwnerPackage, const UObject* OwnerObject)
+	{
+		// This function will become pure virtual after deprecation handling is removed.
+		return FGenerationManifest();
+	}
 
 	/** Representation of generated packages prepared by the cooker. */
 	struct FGeneratedPackageForPopulate
 	{
-		/** RelativePath returned from GetGenerateList */
+		/** RelativePath reported in the GenerationManifest. */
 		FString RelativePath;
-		/** Root returned from GetGenerateList */
+		/** Root reported in the GenerationManifest. */
 		FString GeneratedRootPath;
 		/**
 		 * Non-null UPackage. Possibly an empty placeholder package, but may contain modifications
@@ -198,7 +222,10 @@ public:
 		 * can create import references to objects that will be stored in the generated package.
 		 */
 		UPackage* Package = nullptr;
-		/** *GetCreateAsMap returned from GetGenerateList. The package filename extension has already been set based on this. */
+		/**
+		 * *GetCreateAsMap reported in GenerationManifest. The package filename extension has already been set based
+		 * on this.
+		 */
 		bool bCreatedAsMap = false;
 	};
 
@@ -288,6 +315,8 @@ public:
 	 * PopulateGeneratorPackage is guaranteed to not be called again until the splitter has been destroyed and the
 	 * generator package has been garbage collected.
 	 *
+	 * This function will not be called if bCookSaveGeneratorPackage is false in the GenerationManifest.
+	 *
 	 * @return True if successfully populated, false on error (this will cause a cook error).
 	 */
 	virtual bool PopulateGeneratorPackage(FPopulateContext& PopulateContext)
@@ -296,8 +325,12 @@ public:
 	}
 
 	/**
-	 * Called before saving the parent generator package, after PopulateGeneratorPackage but before PopulateGeneratedPackage for any
-	 * generated packages. Make any required adjustments to the parent package before it is saved into the target domain.
+	 * Called before saving the parent generator package, after PopulateGeneratorPackage but before
+	 * PopulateGeneratedPackage for any generated packages. Make any required adjustments to the parent package before
+	 * it is saved into the target domain.
+	 * 
+	 * This function will not be called if bCookSaveGeneratorPackage is false in the GenerationManifest.
+	 * 
 	 * @return True if successfully presaved, false on error (this will cause a cook error).
 	 */
 	virtual bool PreSaveGeneratorPackage(FPopulateContext& PopulateContext)
@@ -308,7 +341,9 @@ public:
 	/**
 	 * Called after saving the parent generator package. Undo any required adjustments to the parent package that
 	 * were made in PreSaveGeneratorPackage, so that the package is once again ready for use in the editor or in
-	 * future GetGenerateList or PreSaveGeneratedPackage calls
+	 * future ReportGenerationManifest or PreSaveGeneratedPackage calls.
+
+	 * This function will not be called if bCookSaveGeneratorPackage is false in the GenerationManifest.
 	 */
 	virtual void PostSaveGeneratorPackage(FPopulateContext& PopulateContext)
 	{
@@ -317,7 +352,7 @@ public:
 	/**
 	 * Try to populate a generated package.
 	 *
-	 * Receive an empty UPackage generated from an element in GetGenerateList and populate it.
+	 * Receive an empty UPackage generated from a GeneratedPackage in the GenerationManifest and populate it.
 	 * Return a list of all the objects that will be moved into the Generated package during its save, so the cooker
 	 * can call BeginCacheForCookedPlatformData on them before the move
 	 * After returning, the given package will be queued for saving into the TargetDomain
@@ -364,8 +399,8 @@ public:
 
 	/**
 	 * Return the full packagename that will be used for a GeneratedPackage, based on the GeneratorPackage's name and
-	 * on the RelPath and optional GeneratedRootPath that the splitter provides in the FGeneratedPackage it returns from
-	 * GetGenerateList.
+	 * on the RelPath and optional GeneratedRootPath that the splitter provides in the FGeneratedPackage it reports in
+	 * the GenerationManifest.
 	 */
 	UNREALED_API static FString ConstructGeneratedPackageName(FName OwnerPackageName, FStringView RelPath,
 		FStringView GeneratedRootOverride = FStringView());
diff --git a/Engine/Source/Runtime/Engine/Private/WorldCookPackageSplitter.cpp b/Engine/Source/Runtime/Engine/Private/WorldCookPackageSplitter.cpp
--- a/Engine/Source/Runtime/Engine/Private/WorldCookPackageSplitter.cpp
+++ b/Engine/Source/Runtime/Engine/Private/WorldCookPackageSplitter.cpp
@@ -72,7 +72,8 @@ bool FWorldCookPackageSplitter::RequiresGeneratorPackageDestructBeforeResplit()
 	return true;
 }
 
-TArray<ICookPackageSplitter::FGeneratedPackage> FWorldCookPackageSplitter::GetGenerateList(const UPackage* OwnerPackage, const UObject* OwnerObject)
+ICookPackageSplitter::FGenerationManifest FWorldCookPackageSplitter::ReportGenerationManifest(
+	const UPackage* OwnerPackage, const UObject* OwnerObject)
 {
 	TNotNull<const UWorld*> OwnerWorld = CastChecked<const UWorld>(OwnerObject);
 
@@ -92,13 +93,13 @@ TArray<ICookPackageSplitter::FGeneratedPackage> FWorldCookPackageSplitter::GetGe
 		});
 	check(CookPackageSubSplitters.Num());
 
-	TArray<ICookPackageSplitter::FGeneratedPackage> PackagesToGenerate;
+	ICookPackageSplitter::FGenerationManifest Manifest;
 	ForEachCookPackageSubSplitters(
-		[this, &PackagesToGenerate, OwnerPackage, OwnerWorld](FWorldCookPackageSplitter::ISubSplitter* WorldCookPackageSplitter)
+		[this, &Manifest, OwnerPackage, OwnerWorld](FWorldCookPackageSplitter::ISubSplitter* WorldCookPackageSplitter)
 		{
 			const TArray<ICookPackageSplitter::FGeneratedPackage> GenerateList = WorldCookPackageSplitter->GetGenerateList(OwnerPackage);
 			
-			PackagesToGenerate.Reserve(PackagesToGenerate.Num() + GenerateList.Num());
+			Manifest.GeneratedPackages.Reserve(Manifest.GeneratedPackages.Num() + GenerateList.Num());
 			for (const ICookPackageSplitter::FGeneratedPackage& GeneratePackage : GenerateList)
 			{
 				ICookPackageSplitter::FGeneratedPackage LocalGeneratePackage(GeneratePackage);
@@ -108,7 +109,7 @@ TArray<ICookPackageSplitter::FGeneratedPackage> FWorldCookPackageSplitter::GetGe
 					LocalGeneratePackage.GeneratedRootPath = OwnerPackage->GetName();
 				}
 
-				PackagesToGenerate.Add(LocalGeneratePackage);
+				Manifest.GeneratedPackages.Add(LocalGeneratePackage);
 
 				const TPair<FName, FName> FullPackageName = { *LocalGeneratePackage.GeneratedRootPath, *LocalGeneratePackage.RelativePath };
 				check(!SplittersGenerateListMap.Contains(FullPackageName));
@@ -118,7 +119,7 @@ TArray<ICookPackageSplitter::FGeneratedPackage> FWorldCookPackageSplitter::GetGe
 
 	FWorldDelegates::OnWorldCleanup.AddRaw(this, &FWorldCookPackageSplitter::OnWorldCleanup);
 
-	return PackagesToGenerate;
+	return Manifest;
 }
 
 bool FWorldCookPackageSplitter::PopulateGeneratedPackage(FPopulateContext& PopulateContext)
diff --git a/Engine/Source/Runtime/Engine/Public/WorldCookPackageSplitter.h b/Engine/Source/Runtime/Engine/Public/WorldCookPackageSplitter.h
--- a/Engine/Source/Runtime/Engine/Public/WorldCookPackageSplitter.h
+++ b/Engine/Source/Runtime/Engine/Public/WorldCookPackageSplitter.h
@@ -67,7 +67,8 @@ protected:
 	virtual bool UseInternalReferenceToAvoidGarbageCollect() override;
 	virtual EGeneratedRequiresGenerator DoesGeneratedRequireGenerator() override;
 	virtual bool RequiresGeneratorPackageDestructBeforeResplit() override;
-	virtual TArray<ICookPackageSplitter::FGeneratedPackage> GetGenerateList(const UPackage* OwnerPackage, const UObject* OwnerObject) override;
+	virtual ICookPackageSplitter::FGenerationManifest ReportGenerationManifest(
+		const UPackage* OwnerPackage, const UObject* OwnerObject) override;
 	virtual bool PopulateGeneratedPackage(FPopulateContext& PopulateContext) override;
 	virtual bool PopulateGeneratorPackage(FPopulateContext& PopulateContext) override;
 	virtual void OnOwnerReloaded(UPackage* OwnerPackage, UObject* OwnerObject) override;
