Fix Json preview not being saved on editor restart
- Fix the Friendly name of the Json Preview not being correctly setup at restart.
- Add AspectRatio and SafeZones to the settings we save

diff --git a/Engine/Source/Editor/UnrealEd/Classes/Editor/EditorPerProjectUserSettings.h b/Engine/Source/Editor/UnrealEd/Classes/Editor/EditorPerProjectUserSettings.h
--- a/Engine/Source/Editor/UnrealEd/Classes/Editor/EditorPerProjectUserSettings.h
+++ b/Engine/Source/Editor/UnrealEd/Classes/Editor/EditorPerProjectUserSettings.h
@@ -187,6 +187,14 @@ public:
 	UPROPERTY(config)
 	FName PreviewShaderPlatformName;
 	
+	/** Constrained Aspect Ratio of the preview platform */
+	UPROPERTY(config)
+	float PreviewConstrainedAspectRatio;
+
+	/** Safe Zones of the preview platform */
+	UPROPERTY(config)
+	FVector4f PreviewSafeZones;
+
 	/** Is feature level preview currently active */
 	UPROPERTY(config)
 	bool bPreviewFeatureLevelActive;
diff --git a/Engine/Source/Editor/UnrealEd/Private/EditorEngine.cpp b/Engine/Source/Editor/UnrealEd/Private/EditorEngine.cpp
--- a/Engine/Source/Editor/UnrealEd/Private/EditorEngine.cpp
+++ b/Engine/Source/Editor/UnrealEd/Private/EditorEngine.cpp
@@ -6714,7 +6714,18 @@ void UEditorEngine::LoadEditorShaderPlatformPreview()
 		{
 			PreviewShaderPlatformName = *(LexToString(GShaderPlatformForFeatureLevel[Settings->PreviewFeatureLevel]));
 		}
-		SetPreviewPlatform(FPreviewPlatformInfo((ERHIFeatureLevel::Type)Settings->PreviewFeatureLevel, ShaderPlatformToPreview, Settings->PreviewPlatformName, Settings->PreviewShaderFormatName, Settings->PreviewDeviceProfileName, Settings->bPreviewFeatureLevelActive, PreviewShaderPlatformName), false);
+		SetPreviewPlatform(FPreviewPlatformInfo(
+			(ERHIFeatureLevel::Type)Settings->PreviewFeatureLevel, 
+			ShaderPlatformToPreview, 
+			Settings->PreviewPlatformName, 
+			Settings->PreviewShaderFormatName, 
+			Settings->PreviewDeviceProfileName, 
+			Settings->bPreviewFeatureLevelActive, 
+			PreviewShaderPlatformName, 
+			FText(),
+			Settings->PreviewConstrainedAspectRatio,
+			Settings->PreviewSafeZones),
+			false);
 	}
 }
 
@@ -6728,6 +6739,8 @@ void UEditorEngine::SaveEditorShaderPlatformPreview()
 	Settings->bPreviewFeatureLevelWasDefault = (PreviewPlatform.PreviewFeatureLevel == GMaxRHIFeatureLevel);
 	Settings->PreviewDeviceProfileName = PreviewPlatform.DeviceProfileName;
 	Settings->PreviewShaderPlatformName = PreviewPlatform.PreviewShaderPlatformName;
+	Settings->PreviewConstrainedAspectRatio = PreviewPlatform.GetConstrainedAspectRatio();
+	Settings->PreviewSafeZones = PreviewPlatform.GetSafeZones();
 	Settings->PostEditChange();
 }
 
diff --git a/Engine/Source/Runtime/RHI/Private/DataDrivenShaderPlatformInfo.cpp b/Engine/Source/Runtime/RHI/Private/DataDrivenShaderPlatformInfo.cpp
--- a/Engine/Source/Runtime/RHI/Private/DataDrivenShaderPlatformInfo.cpp
+++ b/Engine/Source/Runtime/RHI/Private/DataDrivenShaderPlatformInfo.cpp
@@ -620,7 +620,7 @@ FText FGenericDataDrivenShaderPlatformInfo::GetFriendlyName(const FStaticShaderP
 		return *FriendlyName;
 	}
 
-	return FText::FromString(TEXT(""));
+	return FText::FromName(DeviceProfileName);
 }
 
 const EShaderPlatform FGenericDataDrivenShaderPlatformInfo::GetPreviewShaderPlatformParent(const FStaticShaderPlatform Platform)
