FindOrCreatePackageForAssetType now loads packages with LOAD_NoWarn

#rnx

diff --git a/Engine/Source/Editor/UnrealEd/Private/PackageTools.cpp b/Engine/Source/Editor/UnrealEd/Private/PackageTools.cpp
--- a/Engine/Source/Editor/UnrealEd/Private/PackageTools.cpp
+++ b/Engine/Source/Editor/UnrealEd/Private/PackageTools.cpp
@@ -297,10 +297,11 @@ UPackageTools::UPackageTools(const FObjectInitializer& ObjectInitializer)
 	 * Loads the specified package file (or returns an existing package if it's already loaded.)
 	 *
 	 * @param	InFilename	File name of package to load
+	 * @param	InLoadFlags	Flags controlling loading behavior, from the ELoadFlags enum
 	 *
 	 * @return	The loaded package (or NULL if something went wrong.)
 	 */
-	UPackage* UPackageTools::LoadPackage( FString InFilename )
+	UPackage* UPackageTools::LoadPackage( FString InFilename, uint32 InLoadFlags )
 	{
 		// Detach all components while loading a package.
 		// This is necessary for the cases where the load replaces existing objects which may be referenced by the attached components.
@@ -314,7 +315,7 @@ UPackageTools::UPackageTools(const FObjectInitializer& ObjectInitializer)
 		Arguments.Add(TEXT("PackageName"), FText::FromString(InFilename));
 		FMessageLog("LoadErrors").NewPage(FText::Format(LOCTEXT("LoadPackageLogPage", "Loading package: {PackageName}"), Arguments));
 
-		UPackage* Package = ::LoadPackage( NULL, *InFilename, 0 );
+		UPackage* Package = ::LoadPackage( NULL, *InFilename, InLoadFlags );
 
 		// display any load errors that happened while loading the package
 		FEditorDelegates::DisplayLoadErrors.Broadcast();
@@ -1580,7 +1581,7 @@ UPackageTools::UPackageTools(const FObjectInitializer& ObjectInitializer)
 			{
 				// Create or return the existing package
 				FString PackageNameAsString = LongPackageName.ToString();
-				Package = LoadPackage(PackageNameAsString);
+				Package = LoadPackage(PackageNameAsString, LOAD_NoWarn);
 
 				if (Package)
 				{
diff --git a/Engine/Source/Editor/UnrealEd/Public/PackageTools.h b/Engine/Source/Editor/UnrealEd/Public/PackageTools.h
--- a/Engine/Source/Editor/UnrealEd/Public/PackageTools.h
+++ b/Engine/Source/Editor/UnrealEd/Public/PackageTools.h
@@ -68,10 +68,12 @@ public:
 	 * Loads the specified package file (or returns an existing package if it's already loaded.)
 	 *
 	 * @param	InFilename	File name of package to load
+	 * @param	InLoadFlags	Flags controlling loading behavior, from the ELoadFlags enum
+
 	 *
 	 * @return	The loaded package (or NULL if something went wrong.)
 	 */
-	static UNREALED_API UPackage* LoadPackage( FString InFilename );
+	static UNREALED_API UPackage* LoadPackage( FString InFilename, uint32 InLoadFlags = 0);
 
 	/**
 	 * Helper function that attempts to unload the specified top-level packages.
