Cooker: CookPackageSplitter: Update some variable names and comments related to generation to match the current code behavior after tweaks to CookGenerationHelper flow over the past year.
#rnx

diff --git a/Engine/Source/Editor/UnrealEd/Classes/CookOnTheSide/CookOnTheFlyServer.h b/Engine/Source/Editor/UnrealEd/Classes/CookOnTheSide/CookOnTheFlyServer.h
--- a/Engine/Source/Editor/UnrealEd/Classes/CookOnTheSide/CookOnTheFlyServer.h
+++ b/Engine/Source/Editor/UnrealEd/Classes/CookOnTheSide/CookOnTheFlyServer.h
@@ -1497,8 +1497,8 @@ private:
 	UE::Cook::EPollStatus BeginCachePostMove(UE::Cook::FGenerationHelper& GenerationHelper,
 		UE::Cook::FCookGenerationInfo& Info, UE::Cook::FCookerTimer& Timer);
 
-	/** Try calling the CookPackageSplitter's populate to create the package */
-	UE::Cook::EPollStatus TryPopulateGeneratedPackage(UE::Cook::FGenerationHelper& GenerationHelper,
+	/** Try calling the CookPackageSplitter's PreSave hook on the package. */
+	UE::Cook::EPollStatus TryPreSaveGeneratedPackage(UE::Cook::FGenerationHelper& GenerationHelper,
 		UE::Cook::FCookGenerationInfo& GeneratedInfo);
 
 	ICookArtifactReader& FindOrCreateCookArtifactReader(const ITargetPlatform* TargetPlatform);
diff --git a/Engine/Source/Editor/UnrealEd/Private/CookOnTheFlyServer.cpp b/Engine/Source/Editor/UnrealEd/Private/CookOnTheFlyServer.cpp
--- a/Engine/Source/Editor/UnrealEd/Private/CookOnTheFlyServer.cpp
+++ b/Engine/Source/Editor/UnrealEd/Private/CookOnTheFlyServer.cpp
@@ -3712,13 +3712,12 @@ UE::Cook::EPollStatus UCookOnTheFlyServer::PrepareSaveGenerationPackage(UE::Cook
 		PackageData.SetSaveSubStateComplete(ESaveSubState::Generation_FinishCacheObjectsToMove);
 	}
 
-	if (PackageData.GetSaveSubState() <= ESaveSubState::Generation_CallPopulate)
+	if (PackageData.GetSaveSubState() <= ESaveSubState::Generation_CallPreSave)
 	{
 		if (bPrecaching)
 		{
-			// We're not allowed to populate when precaching, because we want to avoid 
-			// garbagecollection in between Populating and PostSaving the populated package,
-			// so we need to not Populate until we're ready to save
+			// We're not allowed to PreSave when precaching, because we want to avoid garbagecollection in between
+			// PreSaving and PostSaving the package, so we need to not PreSave until we're ready to save.
 			return EPollStatus::Incomplete;
 		}
 
@@ -3729,13 +3728,13 @@ UE::Cook::EPollStatus UCookOnTheFlyServer::PrepareSaveGenerationPackage(UE::Cook
 		}
 		else
 		{
-			Result = TryPopulateGeneratedPackage(GenerationHelper, Info);
+			Result = TryPreSaveGeneratedPackage(GenerationHelper, Info);
 		}
 		if (Result != EPollStatus::Success)
 		{
 			return Result;
 		}
-		PackageData.SetSaveSubStateComplete(ESaveSubState::Generation_CallPopulate);
+		PackageData.SetSaveSubStateComplete(ESaveSubState::Generation_CallPreSave);
 	}
 
 	if (PackageData.GetSaveSubState() <= ESaveSubState::LastCookedPlatformData_WaitingForIsLoaded)
@@ -3904,7 +3903,7 @@ bool UCookOnTheFlyServer::TryConstructGeneratedPackagesForPopulate(UE::Cook::FPa
 	}
 
 	// We need to find or (create empty stub packages for) each of the PackagesToGenerate so that PreSaveGeneratorPackage
-	// can refer to them to create hardlinks in the cooked Generator package
+	// can refer to them to create hardlinks in the cooked Generator package.
 	TArrayView<FCookGenerationInfo> PackagesToGenerate = GenerationHelper.GetPackagesToGenerate();
 	GeneratedPackagesForPopulate.Reserve(PackagesToGenerate.Num());
 	for (FCookGenerationInfo& Info : PackagesToGenerate)
@@ -3992,7 +3991,7 @@ UE::Cook::EPollStatus UCookOnTheFlyServer::BeginCachePostMove(UE::Cook::FGenerat
 	return EPollStatus::Success;
 }
 
-UE::Cook::EPollStatus UCookOnTheFlyServer::TryPopulateGeneratedPackage(UE::Cook::FGenerationHelper& GenerationHelper,
+UE::Cook::EPollStatus UCookOnTheFlyServer::TryPreSaveGeneratedPackage(UE::Cook::FGenerationHelper& GenerationHelper,
 	UE::Cook::FCookGenerationInfo& GeneratedInfo)
 {
 	using namespace UE::Cook;
@@ -4002,7 +4001,8 @@ UE::Cook::EPollStatus UCookOnTheFlyServer::TryPopulateGeneratedPackage(UE::Cook:
 	UPackage* OwnerPackage = GenerationHelper.FindOrLoadOwnerPackage(*this);
 	if (!OwnerPackage)
 	{
-		UE_LOG(LogCook, Error, TEXT("TryPopulateGeneratedPackage: could not load ParentGeneratorPackage %s for GeneratedPackage %s"),
+		UE_LOG(LogCook, Error,
+			TEXT("TryPreSaveGeneratedPackage: could not load ParentGeneratorPackage %s for GeneratedPackage %s"),
 			*GenerationHelper.GetOwner().GetPackageName().ToString(), *GeneratedPackageName);
 		return EPollStatus::Error;
 	}
@@ -4012,15 +4012,17 @@ UE::Cook::EPollStatus UCookOnTheFlyServer::TryPopulateGeneratedPackage(UE::Cook:
 	UObject* OwnerObject = GenerationHelper.FindOrLoadSplitDataObject();
 	if (!OwnerObject)
 	{
-		UE_LOG(LogCook, Error, TEXT("PopulateGeneratedPackage could not find the original splitting object. Generated package can not be created. Splitter=%s, Generated=%s."),
+		UE_LOG(LogCook, Error,
+			TEXT("TryPreSaveGeneratedPackage could not find the original splitting object. Generated package can not be created. Splitter=%s, Generated=%s."),
 			*GenerationHelper.GetSplitDataObjectName().ToString(), *GeneratedPackageName);
 		return EPollStatus::Error;
 	}
 
 	ICookPackageSplitter* Splitter = GenerationHelper.GetCookPackageSplitterInstance();
 
-	// Populate package using CookPackageSplitterInstance and pass GeneratedPackage's cooked name for it to
-	// properly setup any internal reference to this package (SoftObjectPaths or others)
+	// Call the CookPackageSplitterInstance's PreSave hook for the package and pass name and other information about
+	// the GeneratedPackage so the splitter can properly setup any internal reference to this package (e.g.
+	// SoftObjectPath to the package).
 	ICookPackageSplitter::FGeneratedPackageForPopulate GeneratedPackagePopulateData;
 	GeneratedPackagePopulateData.RelativePath = GeneratedInfo.RelativePath;
 	GeneratedPackagePopulateData.GeneratedRootPath = GeneratedInfo.GeneratedRootPath;
@@ -4059,7 +4061,8 @@ UE::Cook::EPollStatus UCookOnTheFlyServer::TryPopulateGeneratedPackage(UE::Cook:
 	bool bPackageIsMap = GeneratedPackage->ContainsMap();
 	if (bPackageIsMap != GeneratedInfo.IsCreateAsMap())
 	{
-		UE_LOG(LogCook, Error, TEXT("PackageSplitter specified generated package is %s in ReportGenerationManifest, but then in PreSaveGeneratedPackage created it as %s. Splitter=%s, Generated=%s."),
+		UE_LOG(LogCook, Error,
+			TEXT("PackageSplitter specified generated package is %s in ReportGenerationManifest, but then in PreSaveGeneratedPackage created it as %s. Splitter=%s, Generated=%s."),
 			(GeneratedInfo.IsCreateAsMap() ? TEXT("map") : TEXT("uasset")), (bPackageIsMap ? TEXT("map") : TEXT("uasset")),
 			*GenerationHelper.GetSplitDataObjectName().ToString(), *GeneratedPackageName);
 		return EPollStatus::Error;
@@ -4313,8 +4316,9 @@ UE::Cook::EPollStatus UCookOnTheFlyServer::PrepareSaveInternal(UE::Cook::FPackag
 		{
 			if (PackageData.GetSaveSubState() <= ESaveSubState::CheckForIsGenerated)
 			{
-				// Skip over the LastCookedPlatformData_CallingBegin state; we only need to enter that
-				// state if RefreshObjectCache finds some new objects
+				// Skip over the LastCookedPlatformData_CallingBegin state; we only need to enter that state if
+				// RefreshObjectCache finds some new objects, and for this block (non-generation packages that reach
+				// CheckForIsGenerated for the first time) we have not called RefreshObjectCache yet.
 				PackageData.SetSaveSubState(ESaveSubState::LastCookedPlatformData_WaitingForIsLoaded);
 			}
 
diff --git a/Engine/Source/Editor/UnrealEd/Private/Cooker/CookGenerationHelper.cpp b/Engine/Source/Editor/UnrealEd/Private/Cooker/CookGenerationHelper.cpp
--- a/Engine/Source/Editor/UnrealEd/Private/Cooker/CookGenerationHelper.cpp
+++ b/Engine/Source/Editor/UnrealEd/Private/Cooker/CookGenerationHelper.cpp
@@ -1869,7 +1869,7 @@ void FGenerationHelper::ResetSaveState(FCookGenerationInfo& Info, UPackage* Pack
 	// We release references to *this in this function so keep a local reference to avoid deletion during the function.
 	TRefCountPtr<FGenerationHelper> LocalRefCount = this;
 
-	if (Info.PackageData->GetSaveSubState() > ESaveSubState::Generation_CallPopulate)
+	if (Info.PackageData->GetSaveSubState() > ESaveSubState::Generation_CallPreSave)
 	{
 		UObject* SplitObject = GetWeakSplitDataObject();
 		UPackage* LocalOwnerPackage = Info.IsGenerator() ? Package : GetOwnerPackage();
@@ -2124,7 +2124,7 @@ void FGenerationHelper::PreGarbageCollect(const TRefCountPtr<FGenerationHelper>&
 	{
 		// If we don't have a contract to keep the packagedata referenced during GC, don't report
 		// anything to garbage collection, and demote the package if it has progressed too far
-		if (Info.PackageData->GetSaveSubState() > ESaveSubState::Generation_CallPopulate)
+		if (Info.PackageData->GetSaveSubState() > ESaveSubState::Generation_CallPreSave)
 		{
 			bOutShouldDemote = true;
 		}
@@ -2155,7 +2155,7 @@ void FGenerationHelper::PreGarbageCollect(const TRefCountPtr<FGenerationHelper>&
 		}
 	}
 
-	// Keep the objects returned from GetObjectsToMove* functions referenced
+	// Keep the objects returned from GetObjectsToMove* functions referenced.
 	if (Info.HasTakenOverCachedCookedPlatformData())
 	{
 		bKeepingAnyObjects = true;
@@ -2169,9 +2169,9 @@ void FGenerationHelper::PreGarbageCollect(const TRefCountPtr<FGenerationHelper>&
 		}
 	}
 
-	// Keep the generator and generated package referenced if we've passed the call to populate, or if we are keeping
-	// any other objects referenced
-	if (bKeepingAnyObjects || Info.PackageData->GetSaveSubState() > ESaveSubState::Generation_CallPopulate)
+	// Keep the generator and generated package referenced if we've passed the call to presave, or if we are keeping
+	// any other objects referenced.
+	if (bKeepingAnyObjects || Info.PackageData->GetSaveSubState() > ESaveSubState::Generation_CallPreSave)
 	{
 		bNeedsGeneratorPackage = true;
 		if (&Info != &OwnerInfo)
@@ -2405,10 +2405,10 @@ void FGenerationHelper::PostGarbageCollect(const TRefCountPtr<FGenerationHelper>
 	}
 	else if (!IsUseInternalReferenceToAvoidGarbageCollect())
 	{
-		// After the Generator Package is saved, we drop our references to it and it can be garbage collected
-		// If we have any packages left to populate, our splitter contract requires that it be garbage collected
-		// because we promise that the package is not partially GC'd during calls to TryPopulateGeneratedPackage
-		// The splitter can opt-out of this contract and keep it referenced itself if it desires.
+		// After the Generator Package is saved, we drop our references to it and it can be garbage collected. If we
+		// have any packages left to populate, our splitter contract requires that it be garbage collected because
+		// we promise that the package is not partially GC'd during calls to PopulateGeneratedPackage or
+		// PreSaveGeneratedPackage. The splitter can opt-out of this contract and keep it referenced itself.
 		if (!Owner.IsInProgress() && !Owner.IsKeepReferencedDuringGC())
 		{
 			VerifyGeneratorPackageGarbageCollected(Context);
diff --git a/Engine/Source/Editor/UnrealEd/Private/Cooker/CookGenerationHelper.h b/Engine/Source/Editor/UnrealEd/Private/Cooker/CookGenerationHelper.h
--- a/Engine/Source/Editor/UnrealEd/Private/Cooker/CookGenerationHelper.h
+++ b/Engine/Source/Editor/UnrealEd/Private/Cooker/CookGenerationHelper.h
@@ -249,7 +249,10 @@ public:
 	 */
 	bool IsValid();
 
-	/** True if ReportGenerationManifest requested that the generator package should be cooksaved as well. */
+	/**
+	 * True if ReportGenerationManifest requested that the generator package should be cooksaved as well. Returns true
+	 * if not initialized. Does not call Initialize. 
+	 */
 	bool IsSaveGenerator() const;
 	/** Accessor for the packages to generate, empty if invalid or TryReportGenerationManifest was not yet called. */
 	TArrayView<FCookGenerationInfo> GetPackagesToGenerate();
@@ -307,8 +310,9 @@ public:
 	 */
 	bool TryReportGenerationManifest();
 	/**
-	 * Call the Splitter's PopulateGeneratorPackage if not yet called. Assumes GenerateList has been called. Logs
-	 * errors and returns false on failure.
+	 * Call the Splitter's PopulateGeneratorPackage if not yet called. Assumes GenerateList has been called.
+	 *
+	 * Logs errors and returns false on failure.
 	 */
 	bool TryCallPopulateGeneratorPackage(
 		TArray<ICookPackageSplitter::FGeneratedPackageForPopulate>& InOutGeneratedPackagesForPresave);
diff --git a/Engine/Source/Editor/UnrealEd/Private/Cooker/CookSavePackage.cpp b/Engine/Source/Editor/UnrealEd/Private/Cooker/CookSavePackage.cpp
--- a/Engine/Source/Editor/UnrealEd/Private/Cooker/CookSavePackage.cpp
+++ b/Engine/Source/Editor/UnrealEd/Private/Cooker/CookSavePackage.cpp
@@ -270,7 +270,7 @@ void FSaveCookedPackageContext::SetupPlatform(const ITargetPlatform* InTargetPla
 
 	if (CommitType == EReachability::Runtime)
 	{
-		// don't save Editor resources from the Engine if the target doesn't have editoronly data
+		// Don't save Editor resources from the Engine if the target doesn't have editoronly data.
 		if (COTFS.IsCookFlagSet(ECookInitializationFlags::SkipEditorContent) &&
 			(PackageName.StartsWith(TEXT("/Engine/Editor")) || PackageName.StartsWith(TEXT("/Engine/VREditor"))) &&
 			!TargetPlatform->HasEditorOnlyData())
@@ -309,11 +309,11 @@ void FSaveCookedPackageContext::SetupPlatform(const ITargetPlatform* InTargetPla
 		{
 			if (TSet<FName>* NeverCookPackages = COTFS.PackageTracker->PlatformSpecificNeverCookPackages.Find(TargetPlatform))
 			{
-				FGenerationHelper* GenerationHelper =
+				FGenerationHelper* ParentGenerationHelper =
 					PackageData.IsGenerated() ? PackageData.GetParentGenerationHelper().GetReference() : nullptr;
 
 				if (NeverCookPackages->Find(Package->GetFName()) ||
-					(GenerationHelper && NeverCookPackages->Find(GenerationHelper->GetOwner().GetPackageName())))
+					(ParentGenerationHelper && NeverCookPackages->Find(ParentGenerationHelper->GetOwner().GetPackageName())))
 				{
 					SavePackageResult = ESavePackageResult::ContainsEditorOnlyData;
 					const TCHAR* RejectedReason = TEXT("PlatformSpecificNeverCook");
@@ -555,14 +555,14 @@ void FSaveCookedPackageContext::FinishPackage()
 	// If any save succeeded, add all dependencies from all platforms to the cook for the platform
 	if (CommitType == EReachability::Build || bAnySaveSucceeded)
 	{
-		bool bAddSoftReferences = CommitType == EReachability::Runtime && bAnySaveSucceeded
+		bool bAddRuntimeReferences = CommitType == EReachability::Runtime && bAnySaveSucceeded
 			&& !COTFS.CookByTheBookOptions->bSkipSoftReferences;
 
 		FName PackageFName = Package->GetFName();
 		if (PlatformsForPackage.Num() == 1)
 		{
 			TConstArrayView<const ITargetPlatform*> ReachablePlatforms(&TargetPlatform, 1);
-			if (bAddSoftReferences)
+			if (bAddRuntimeReferences)
 			{
 				for (const TPair<FPackageData*, EInstigator>& DependencyPair : PlatformDependencies[0].RuntimeDependencies)
 				{
@@ -584,7 +584,7 @@ void FSaveCookedPackageContext::FinishPackage()
 				const FPlatformDiscoveryData& DData = PlatformDependencies[LocalIndex];
 				const ITargetPlatform* CurrentPlatform = PlatformsForPackage[LocalIndex];
 
-				if (bAddSoftReferences)
+				if (bAddRuntimeReferences)
 				{
 					// Merge the runtime dependencies for each package into a single QueueDiscoveredPackage call, if possible
 					// It will not be possible if the different platforms have different instigators, so track a list of
diff --git a/Engine/Source/Editor/UnrealEd/Private/Cooker/CookTypes.cpp b/Engine/Source/Editor/UnrealEd/Private/Cooker/CookTypes.cpp
--- a/Engine/Source/Editor/UnrealEd/Private/Cooker/CookTypes.cpp
+++ b/Engine/Source/Editor/UnrealEd/Private/Cooker/CookTypes.cpp
@@ -137,7 +137,7 @@ const TCHAR* LexToString(UE::Cook::ESaveSubState State)
 	case ESaveSubState::Generation_CallObjectsToMove: return TEXT("Generation_CallObjectsToMove");
 	case ESaveSubState::Generation_BeginCacheObjectsToMove: return TEXT("Generation_BeginCacheObjectsToMove");
 	case ESaveSubState::Generation_FinishCacheObjectsToMove: return TEXT("Generation_FinishCacheObjectsToMove");
-	case ESaveSubState::Generation_CallPopulate: return TEXT("Generation_CallPopulate");
+	case ESaveSubState::Generation_CallPreSave: return TEXT("Generation_CallPreSave");
 	case ESaveSubState::Generation_CallGetPostMoveObjects: return TEXT("Generation_CallGetPostMoveObjects");
 	case ESaveSubState::LastCookedPlatformData_CallingBegin: return TEXT("LastCookedPlatformData_CallingBegin");
 	case ESaveSubState::LastCookedPlatformData_WaitingForIsLoaded: return TEXT("LastCookedPlatformData_WaitingForIsLoaded");
diff --git a/Engine/Source/Editor/UnrealEd/Private/Cooker/CookTypes.h b/Engine/Source/Editor/UnrealEd/Private/Cooker/CookTypes.h
--- a/Engine/Source/Editor/UnrealEd/Private/Cooker/CookTypes.h
+++ b/Engine/Source/Editor/UnrealEd/Private/Cooker/CookTypes.h
@@ -246,7 +246,7 @@ namespace UE::Cook
 		Generation_CallObjectsToMove,
 		Generation_BeginCacheObjectsToMove,
 		Generation_FinishCacheObjectsToMove,
-		Generation_CallPopulate,
+		Generation_CallPreSave,
 		Generation_CallGetPostMoveObjects,
 
 		LastCookedPlatformData_CallingBegin,
diff --git a/Engine/Source/Editor/UnrealEd/Public/CookPackageSplitter.h b/Engine/Source/Editor/UnrealEd/Public/CookPackageSplitter.h
--- a/Engine/Source/Editor/UnrealEd/Public/CookPackageSplitter.h
+++ b/Engine/Source/Editor/UnrealEd/Public/CookPackageSplitter.h
@@ -310,12 +310,15 @@ public:
 
 	/**
 	 * Called before presaving the parent generator package, to give the generator a chance to inform the cooker which
-	 * objects will be moved into the generator package that are not already present in it.
+	 * objects will be moved into the generator package that are not already present in it. This function will also
+	 * be called before calling Populate and PreSave on generated packages, if that behavior is requested via
+	 * DoesGeneratedRequireGenerator.
 	 *
-	 * PopulateGeneratorPackage is guaranteed to not be called again until the splitter has been destroyed and the
-	 * generator package has been garbage collected.
+	 * After being called once, PopulateGeneratorPackage is guaranteed to not be called again until the splitter has
+	 * been destroyed and the generator package has been garbage collected.
 	 *
-	 * This function will not be called if bCookSaveGeneratorPackage is false in the GenerationManifest.
+	 * This function will not be called if bCookSaveGeneratorPackage is false in the GenerationManifest, unless it
+	 * is needed for generated packages as well via DoesGeneratedRequireGenerator.
 	 *
 	 * @return True if successfully populated, false on error (this will cause a cook error).
 	 */
