Cooker: CookPackageSplitter: Skip the cooksave of the generated package when the CookPackageSplitter sets bCookSaveGeneratorPackage=false.
#rnx
#FailedValidations Preflight

diff --git a/Engine/Source/Editor/UnrealEd/Private/CookOnTheFlyServer.cpp b/Engine/Source/Editor/UnrealEd/Private/CookOnTheFlyServer.cpp
--- a/Engine/Source/Editor/UnrealEd/Private/CookOnTheFlyServer.cpp
+++ b/Engine/Source/Editor/UnrealEd/Private/CookOnTheFlyServer.cpp
@@ -3793,10 +3793,13 @@ UE::Cook::EPollStatus UCookOnTheFlyServer::BeginCacheObjectsToMove(UE::Cook::FGe
 	if (PackageData.GetSaveSubState() <= ESaveSubState::Generation_CallObjectsToMove)
 	{
 		bool bPopulateSucceeded = false;
-		if (Info.IsGenerator() || GenerationHelper.DoesGeneratedRequireGenerator()
-			>= ICookPackageSplitter::EGeneratedRequiresGenerator::Populate)
+		bool bGeneratedRequiresGeneratorPopulate = GenerationHelper.DoesGeneratedRequireGenerator()
+			>= ICookPackageSplitter::EGeneratedRequiresGenerator::Populate;
+		if (Info.IsGenerator() || bGeneratedRequiresGeneratorPopulate)
 		{
-			if (!GenerationHelper.TryCallPopulateGeneratorPackage(GeneratedPackagesForPopulate))
+			bool bPopulateCallAllowed = bGeneratedRequiresGeneratorPopulate ||
+				(Info.IsGenerator() && GenerationHelper.IsSaveGenerator());
+			if (!GenerationHelper.TryCallPopulateGeneratorPackage(GeneratedPackagesForPopulate, bPopulateCallAllowed))
 			{
 				return EPollStatus::Error;
 			}
@@ -3860,6 +3863,9 @@ UE::Cook::EPollStatus UCookOnTheFlyServer::PreSaveGeneratorPackage(UE::Cook::FPa
 	}
 
 	UE::Cook::CookPackageSplitter::FPopulateContextData PopulateData;
+	// By Contract we do not call PreSaveGeneratorPackage if not saving generator, but we execute the rest of this
+	// function so that other GenerationHelper code can rely on the support tasks for the PreSave being done.
+	if (GenerationHelper.IsSaveGenerator())
 	{
 		FScopedActivePackage ScopedActivePackage(*this, GenerationHelper.GetOwner().GetPackageName(),
 #if UE_WITH_OBJECT_HANDLE_TRACKING
diff --git a/Engine/Source/Editor/UnrealEd/Private/Cooker/CookGenerationHelper.cpp b/Engine/Source/Editor/UnrealEd/Private/Cooker/CookGenerationHelper.cpp
--- a/Engine/Source/Editor/UnrealEd/Private/Cooker/CookGenerationHelper.cpp
+++ b/Engine/Source/Editor/UnrealEd/Private/Cooker/CookGenerationHelper.cpp
@@ -1156,7 +1156,8 @@ bool FGenerationHelper::TryReportGenerationManifest()
 }
 
 bool FGenerationHelper::TryCallPopulateGeneratorPackage(
-	TArray<ICookPackageSplitter::FGeneratedPackageForPopulate>& InOutGeneratedPackagesForPopulate)
+	TArray<ICookPackageSplitter::FGeneratedPackageForPopulate>& InOutGeneratedPackagesForPopulate,
+	bool bPopulateCallAllowed)
 {
 	if (OwnerInfo.HasCalledPopulate())
 	{
@@ -1191,32 +1192,52 @@ bool FGenerationHelper::TryCallPopulateGeneratorPackage(
 			*GetSplitDataObjectName().ToString());
 		return false;
 	}
-	UCookOnTheFlyServer::FScopedActivePackage ScopedActivePackage(COTFS, OwnerPackageName,
+
+	// By Contract we do not call PopulateGeneratorPackage if not saving generator and not required by generated
+	// packages, but we execute the rest of this function so that other GenerationHelper code can rely on the support
+	// tasks for the Populate being done.
+	UE::Cook::CookPackageSplitter::FPopulateContextData PopulateData;
+	if (bPopulateCallAllowed)
+	{
+		UCookOnTheFlyServer::FScopedActivePackage ScopedActivePackage(COTFS, OwnerPackageName,
 #if UE_WITH_OBJECT_HANDLE_TRACKING
-		PackageAccessTrackingOps::NAME_CookerBuildObject
+			PackageAccessTrackingOps::NAME_CookerBuildObject
 #else
-		FName()
+			FName()
 #endif
-	);
-	UE::Cook::CookPackageSplitter::FPopulateContextData PopulateData;
-	ICookPackageSplitter::FPopulateContext PopulateContext(PopulateData);
-	PopulateData.OwnerPackage = LocalOwnerPackage;
-	PopulateData.OwnerObject = LocalSplitDataObject;
-	PopulateData.GeneratedPackages = InOutGeneratedPackagesForPopulate;
-
-	PRAGMA_DISABLE_DEPRECATION_WARNINGS;
-	bool bPopulateSucceeded = CookPackageSplitterInstance->PopulateGeneratorPackage(PopulateData.OwnerPackage,
-		PopulateData.OwnerObject, InOutGeneratedPackagesForPopulate, PopulateData.ObjectsToMove,
-		PopulateData.KeepReferencedPackages);
-	CookPackageSplitterInstance->WarnIfDeprecatedVirtualNotCalled(TEXT("PopulateGeneratorPackage"));
-	PRAGMA_ENABLE_DEPRECATION_WARNINGS;
-	bPopulateSucceeded = CookPackageSplitterInstance->PopulateGeneratorPackage(PopulateContext) && bPopulateSucceeded;
+		);
+		ICookPackageSplitter::FPopulateContext PopulateContext(PopulateData);
+		PopulateData.OwnerPackage = LocalOwnerPackage;
+		PopulateData.OwnerObject = LocalSplitDataObject;
+		PopulateData.GeneratedPackages = InOutGeneratedPackagesForPopulate;
+
+		PRAGMA_DISABLE_DEPRECATION_WARNINGS;
+		bool bPopulateSucceeded = CookPackageSplitterInstance->PopulateGeneratorPackage(PopulateData.OwnerPackage,
+			PopulateData.OwnerObject, InOutGeneratedPackagesForPopulate, PopulateData.ObjectsToMove,
+			PopulateData.KeepReferencedPackages);
+		CookPackageSplitterInstance->WarnIfDeprecatedVirtualNotCalled(TEXT("PopulateGeneratorPackage"));
+		PRAGMA_ENABLE_DEPRECATION_WARNINGS;
+		bPopulateSucceeded = CookPackageSplitterInstance->PopulateGeneratorPackage(PopulateContext) && bPopulateSucceeded;
+
+		if (!bPopulateSucceeded)
+		{
+			UE_LOG(LogCook, Error, TEXT("CookPackageSplitter returned false from PopulateGeneratorPackage. Splitter=%s"),
+				*GetSplitDataObjectName().ToString());
+			return false;
+		}
+		UE_LOG(LogCookGenerationHelper, Verbose, TEXT("%s: PopulateGeneratorPackage."),
+			*WriteToString<256>(GetOwner().GetPackageName()));
 
-	if (!bPopulateSucceeded)
+		// Contract Point 1: We will not call populate again until the splitter has been destroyed
+		// Contract Point 2: We will not call populate again without garbage collecting the generator package
+		OwnerInfo.SetHasCalledPopulate(true);
+		SetKeepForAllSavedOrGC();
+	}
+	else
 	{
-		UE_LOG(LogCook, Error, TEXT("CookPackageSplitter returned false from PopulateGeneratorPackage. Splitter=%s"),
-			*GetSplitDataObjectName().ToString());
-		return false;
+		UE_LOG(LogCookGenerationHelper, Verbose, TEXT("%s: PopulateGeneratorPackage Skipped because !IsSaveGenerator()."),
+			*WriteToString<256>(GetOwner().GetPackageName()));
+
 	}
 	OwnerInfo.AddKeepReferencedPackages(*this, PopulateData.KeepReferencedPackages);
 	OwnerObjectsToMove.Reserve(PopulateData.ObjectsToMove.Num());
@@ -1228,13 +1249,6 @@ bool FGenerationHelper::TryCallPopulateGeneratorPackage(
 		}
 	}
 
-	UE_LOG(LogCookGenerationHelper, Verbose, TEXT("%s: PopulateGeneratorPackage."),
-		*WriteToString<256>(GetOwner().GetPackageName()));
-
-	// Contract Point 1: We will not call populate again until the splitter has been destroyed
-	// Contract Point 2: We will not call populate again without garbage collecting the generator package
-	OwnerInfo.SetHasCalledPopulate(true);
-	SetKeepForAllSavedOrGC();
 	return true;
 }
 
@@ -1900,11 +1914,15 @@ void FGenerationHelper::ResetSaveState(FCookGenerationInfo& Info, UPackage* Pack
 			ICookPackageSplitter* LocalSplitterInstance = GetCookPackageSplitterInstance();
 			if (Info.IsGenerator())
 			{
-				PRAGMA_DISABLE_DEPRECATION_WARNINGS;
-				LocalSplitterInstance->PostSaveGeneratorPackage(PopulateData.OwnerPackage, PopulateData.OwnerObject);
-				LocalSplitterInstance->WarnIfDeprecatedVirtualNotCalled(TEXT("PostSaveGeneratorPackage"));
-				PRAGMA_ENABLE_DEPRECATION_WARNINGS;
-				LocalSplitterInstance->PostSaveGeneratorPackage(PopulateContext);
+				// By contract, skip the call to PostSaveGeneratorPackage if !bCookSaveGeneratorPackage
+				if (IsSaveGenerator())
+				{
+					PRAGMA_DISABLE_DEPRECATION_WARNINGS;
+					LocalSplitterInstance->PostSaveGeneratorPackage(PopulateData.OwnerPackage, PopulateData.OwnerObject);
+					LocalSplitterInstance->WarnIfDeprecatedVirtualNotCalled(TEXT("PostSaveGeneratorPackage"));
+					PRAGMA_ENABLE_DEPRECATION_WARNINGS;
+					LocalSplitterInstance->PostSaveGeneratorPackage(PopulateContext);
+				}
 			}
 			else
 			{
diff --git a/Engine/Source/Editor/UnrealEd/Private/Cooker/CookGenerationHelper.h b/Engine/Source/Editor/UnrealEd/Private/Cooker/CookGenerationHelper.h
--- a/Engine/Source/Editor/UnrealEd/Private/Cooker/CookGenerationHelper.h
+++ b/Engine/Source/Editor/UnrealEd/Private/Cooker/CookGenerationHelper.h
@@ -311,11 +311,15 @@ public:
 	bool TryReportGenerationManifest();
 	/**
 	 * Call the Splitter's PopulateGeneratorPackage if not yet called. Assumes GenerateList has been called.
+	 * This function executes support operations, needed for Populate and other calls, even if
+	 * PopulateGeneratorPackage is not allowed to be called due to skipping of the Generator's package save.
+	 * bPopulateCallAllowed=false executes those support functions without calling Populate.
 	 *
 	 * Logs errors and returns false on failure.
 	 */
 	bool TryCallPopulateGeneratorPackage(
-		TArray<ICookPackageSplitter::FGeneratedPackageForPopulate>& InOutGeneratedPackagesForPresave);
+		TArray<ICookPackageSplitter::FGeneratedPackageForPopulate>& InOutGeneratedPackagesForPresave,
+		bool bPopulateCallAllowed = true);
 	/**
 	 * Call the Splitter's PopulateGeneratedPackage if not yet called. Assumes GenerateList has been called. Logs
 	 * errors and returns false on failure.
diff --git a/Engine/Source/Editor/UnrealEd/Private/Cooker/CookSavePackage.cpp b/Engine/Source/Editor/UnrealEd/Private/Cooker/CookSavePackage.cpp
--- a/Engine/Source/Editor/UnrealEd/Private/Cooker/CookSavePackage.cpp
+++ b/Engine/Source/Editor/UnrealEd/Private/Cooker/CookSavePackage.cpp
@@ -270,8 +270,26 @@ void FSaveCookedPackageContext::SetupPlatform(const ITargetPlatform* InTargetPla
 
 	if (CommitType == EReachability::Runtime)
 	{
+		// Skip saving the Package, and mark it as a benign savefail result, if it is a generator that the
+		// CookPackageSplitter requested not to save.
+		TRefCountPtr<FGenerationHelper> GenerationHelper = PackageData.GetGenerationHelper();
+		if (GenerationHelper && !GenerationHelper->IsSaveGenerator())
+		{
+			SavePackageResult = ESavePackageResult::ContainsEditorOnlyData;
+			const TCHAR* RejectedReason = TEXT("!bCookSaveGeneratorPackage");
+			if (GCookProgressDisplay & (int32)ECookProgressDisplayMode::Instigators)
+			{
+				UE_LOG(LogCook, Display, TEXT("Cooking %s, Instigator: { %s } -> Rejected %s"), *PackageName,
+					*(PackageData.GetInstigator(EReachability::Runtime).ToString()), RejectedReason);
+			}
+			else
+			{
+				UE_LOG(LogCook, Display, TEXT("Cooking %s -> Rejected %s"), *PackageName, RejectedReason);
+			}
+			return;
+		}
 		// Don't save Editor resources from the Engine if the target doesn't have editoronly data.
-		if (COTFS.IsCookFlagSet(ECookInitializationFlags::SkipEditorContent) &&
+		else if (COTFS.IsCookFlagSet(ECookInitializationFlags::SkipEditorContent) &&
 			(PackageName.StartsWith(TEXT("/Engine/Editor")) || PackageName.StartsWith(TEXT("/Engine/VREditor"))) &&
 			!TargetPlatform->HasEditorOnlyData())
 		{
@@ -459,8 +477,8 @@ void FSaveCookedPackageContext::FinishPlatform()
 		CookResult = ECookResult::NotAttempted;
 	}
 
-	// Commit the saved bytes and the incremental cook data to the PackageWriter
-	if (bPlatformSetupSuccessful)
+	// Calculate dependencies of the save if successful, or if we otherwise will need them.
+	if (bPlatformSetupSuccessful || (GenerationHelper && !GenerationHelper->IsSaveGenerator()))
 	{
 		// ProcessUnsolicitedPackages so that we record discovereddependencies for any hidden dependency
 		// packages loaded during Save,BeginCacheForCookedPlatformData, or Generator functions. These dependencies
@@ -479,7 +497,10 @@ void FSaveCookedPackageContext::FinishPlatform()
 			GeneratedPackageBuildResultDependencies.GetPtrOrNull() /** non-const, function sets it to empty */);
 		RecordPlatformBuildDependencies();
 		RecordCookImportsCheckerData();
-
+	}
+	// Commit the saved bytes and the incremental cook data to the PackageWriter.
+	if (bPlatformSetupSuccessful)
+	{
 		ICookedPackageWriter::FCommitPackageInfo Info;
 		Info.Attachments = GetCommitAttachments();
 		Info.Status = HasSavePackageResult() ? PackageResultToCommitStatus(SavePackageResult.Result)
@@ -553,11 +574,13 @@ void FSaveCookedPackageContext::FinishPlatform()
 void FSaveCookedPackageContext::FinishPackage()
 {
 	// If any save succeeded, add all dependencies from all platforms to the cook for the platform
-	if (CommitType == EReachability::Build || bAnySaveSucceeded)
+	TRefCountPtr<FGenerationHelper> GenerationHelper = PackageData.GetGenerationHelper();
+	bool bGenerationSkippedSave = GenerationHelper && !GenerationHelper->IsSaveGenerator();
+	bool bAddBuildReferences = CommitType == EReachability::Build || bAnySaveSucceeded;
+	bool bAddRuntimeReferences = CommitType == EReachability::Runtime && (bAnySaveSucceeded || bGenerationSkippedSave)
+		&& !COTFS.CookByTheBookOptions->bSkipSoftReferences;;
+	if (bAddBuildReferences || bAddRuntimeReferences)
 	{
-		bool bAddRuntimeReferences = CommitType == EReachability::Runtime && bAnySaveSucceeded
-			&& !COTFS.CookByTheBookOptions->bSkipSoftReferences;
-
 		FName PackageFName = Package->GetFName();
 		if (PlatformsForPackage.Num() == 1)
 		{
@@ -570,10 +593,13 @@ void FSaveCookedPackageContext::FinishPackage()
 						FInstigator(DependencyPair.Value, PackageFName), FDiscoveredPlatformSet(ReachablePlatforms));
 				}
 			}
-			for (FPackageData* BuildDependency : PlatformDependencies[0].BuildDependencies)
+			if (bAddBuildReferences)
 			{
-				COTFS.QueueDiscoveredPackage(*BuildDependency,
-					FInstigator(EInstigator::BuildDependency, PackageFName), FDiscoveredPlatformSet(ReachablePlatforms));
+				for (FPackageData* BuildDependency : PlatformDependencies[0].BuildDependencies)
+				{
+					COTFS.QueueDiscoveredPackage(*BuildDependency,
+						FInstigator(EInstigator::BuildDependency, PackageFName), FDiscoveredPlatformSet(ReachablePlatforms));
+				}
 			}
 		}
 		else
@@ -597,12 +623,15 @@ void FSaveCookedPackageContext::FinishPackage()
 					}
 				}
 
-				// Merge the build dependencies for each package into a single QueueDiscoveredPackage call.
-				for (FPackageData* BuildDependency : DData.BuildDependencies)
+				if (bAddBuildReferences)
 				{
-					TMap<EInstigator, TArray<const ITargetPlatform*>>& TargetMap =
-						PackagePlatformsForInstigator.FindOrAdd(BuildDependency);
-					TargetMap.FindOrAdd(EInstigator::BuildDependency).Add(CurrentPlatform);
+					// Merge the build dependencies for each package into a single QueueDiscoveredPackage call.
+					for (FPackageData* BuildDependency : DData.BuildDependencies)
+					{
+						TMap<EInstigator, TArray<const ITargetPlatform*>>& TargetMap =
+							PackagePlatformsForInstigator.FindOrAdd(BuildDependency);
+						TargetMap.FindOrAdd(EInstigator::BuildDependency).Add(CurrentPlatform);
+					}
 				}
 			}
 			for (const TPair<FPackageData*, TMap<EInstigator, TArray<const ITargetPlatform*>>>& PackagePair
