[Interchange] Normals, Tangents, BinormalSigns processing fixes for Use T0 As Ref Pose.

#rnx

diff --git a/Engine/Plugins/Interchange/Runtime/Source/Import/Private/Mesh/InterchangeSkeletalMeshFactory.cpp b/Engine/Plugins/Interchange/Runtime/Source/Import/Private/Mesh/InterchangeSkeletalMeshFactory.cpp
--- a/Engine/Plugins/Interchange/Runtime/Source/Import/Private/Mesh/InterchangeSkeletalMeshFactory.cpp
+++ b/Engine/Plugins/Interchange/Runtime/Source/Import/Private/Mesh/InterchangeSkeletalMeshFactory.cpp
@@ -400,6 +400,7 @@ namespace UE
 							for (const FVertexInstanceID VertexInstanceID : VertexInstanceIDs)
 							{
 								FVector3f& Normal = VertexInstanceNormals[VertexInstanceID];
+								FVector3f OriginalNormal = Normal;
 
 								if (Normal.IsZero())
 								{
@@ -422,16 +423,22 @@ namespace UE
 									continue;
 								}
 
-								float& BinormalSign = VertexInstanceBinormalSigns[VertexInstanceID];
-
-								//Calculate original binormal vector to correct handidness after transform applications:
-								FVector3f OriginalBinormal = FVector3f::CrossProduct(Normal, Tangent) * BinormalSign;
-
 								//Calculate transformed Tangent:
+								FVector3f OriginalTangent = Tangent;
 								FVector TangentTransformed = SkinDeformations[VertexIndex].TransformVector(FVector(Tangent));
 								TangentTransformed /= Weight;
 								Tangent = FVector4f(SceneNodeTransform.TransformFVector4(TangentTransformed)).GetSafeNormal();
 
+								if (!VertexInstanceBinormalSigns.IsValid())
+								{
+									continue;
+								}
+
+								float& BinormalSign = VertexInstanceBinormalSigns[VertexInstanceID];
+
+								//Calculate original binormal vector to correct handidness after transform applications:
+								FVector3f OriginalBinormal = FVector3f::CrossProduct(OriginalNormal, OriginalTangent) * BinormalSign;
+
 								//Calclate transformed binormal vector:
 								FVector BinormalTransformed = SkinDeformations[VertexIndex].TransformVector(FVector(OriginalBinormal));
 								BinormalTransformed /= Weight;
diff --git a/Engine/Source/Editor/UnrealEd/Private/Fbx/FbxSkeletalMeshImport.cpp b/Engine/Source/Editor/UnrealEd/Private/Fbx/FbxSkeletalMeshImport.cpp
--- a/Engine/Source/Editor/UnrealEd/Private/Fbx/FbxSkeletalMeshImport.cpp
+++ b/Engine/Source/Editor/UnrealEd/Private/Fbx/FbxSkeletalMeshImport.cpp
@@ -499,7 +499,6 @@ void FFbxImporter::SkinControlPointsToPose(FSkeletalMeshImportData& ImportData,
 
 					int32 NormalMapIndex = (NormalMappingMode == FbxLayerElement::eByControlPoint) ? ControlPointIndex : TmpIndex;
 					int32 NormalValueIndex = (NormalReferenceMode == FbxLayerElement::eDirect) ? NormalMapIndex : LayerElementNormal->GetIndexArray().GetAt(NormalMapIndex);
-					int32 TangentMapIndex = TmpIndex;
 
 					FbxAMatrix ClusterDeformation = ClusterDeformations[ControlPointIndex].Inverse().Transpose();
 
@@ -507,17 +506,45 @@ void FFbxImporter::SkinControlPointsToPose(FSkeletalMeshImportData& ImportData,
 					{
 						if (ImportData.bHasTangents)
 						{
-							FbxVector4 TangentX = LayerElementTangent->GetDirectArray().GetAt(TangentMapIndex);
-							TangentX = ClusterDeformation.MultT(TangentX);
-							TangentX = TotalMatrixForNormal.MultT(TangentX);
-							Triangle.TangentX[UnrealVertexIndex] = (FVector3f)Converter.ConvertDir(TangentX);
-							Triangle.TangentX[UnrealVertexIndex].Normalize();
-
-							FbxVector4 TangentY = LayerElementBinormal->GetDirectArray().GetAt(TangentMapIndex);
-							TangentY = ClusterDeformation.MultT(TangentY);
-							TangentY = TotalMatrixForNormal.MultT(TangentY);
-							Triangle.TangentY[UnrealVertexIndex] = (FVector3f)Converter.ConvertDir(TangentY);
-							Triangle.TangentY[UnrealVertexIndex].Normalize();
+							{
+								//Tangent
+								FbxLayerElement::EReferenceMode TangentReferenceMode(FbxLayerElement::eDirect);
+								FbxLayerElement::EMappingMode TangentMappingMode(FbxLayerElement::eByControlPoint);
+								if (LayerElementTangent)
+								{
+									TangentReferenceMode = LayerElementTangent->GetReferenceMode();
+									TangentMappingMode = LayerElementTangent->GetMappingMode();
+								}
+
+								int32 TangentMapIndex = (TangentMappingMode == FbxLayerElement::eByControlPoint) ? ControlPointIndex : TmpIndex;
+								int32 TangentValueIndex = (TangentReferenceMode == FbxLayerElement::eDirect) ? TangentMapIndex : LayerElementTangent->GetIndexArray().GetAt(TangentMapIndex);
+
+								FbxVector4 TangentX = LayerElementTangent->GetDirectArray().GetAt(TangentValueIndex);
+								TangentX = ClusterDeformation.MultT(TangentX);
+								TangentX = TotalMatrixForNormal.MultT(TangentX);
+								Triangle.TangentX[UnrealVertexIndex] = (FVector3f)Converter.ConvertDir(TangentX);
+								Triangle.TangentX[UnrealVertexIndex].Normalize();
+							}
+
+							{
+								//Binormal
+								FbxLayerElement::EReferenceMode BinormalReferenceMode(FbxLayerElement::eDirect);
+								FbxLayerElement::EMappingMode BinormalMappingMode(FbxLayerElement::eByControlPoint);
+								if (LayerElementBinormal)
+								{
+									BinormalReferenceMode = LayerElementBinormal->GetReferenceMode();
+									BinormalMappingMode = LayerElementBinormal->GetMappingMode();
+								}
+
+								int32 BinormalMapIndex = (BinormalMappingMode == FbxLayerElement::eByControlPoint) ? ControlPointIndex : TmpIndex;
+								int32 BinormalValueIndex = (BinormalReferenceMode == FbxLayerElement::eDirect) ? BinormalMapIndex : LayerElementBinormal->GetIndexArray().GetAt(BinormalMapIndex);
+
+								FbxVector4 TangentY = LayerElementBinormal->GetDirectArray().GetAt(BinormalValueIndex);
+								TangentY = ClusterDeformation.MultT(TangentY);
+								TangentY = TotalMatrixForNormal.MultT(TangentY);
+								Triangle.TangentY[UnrealVertexIndex] = (FVector3f)Converter.ConvertDir(TangentY);
+								Triangle.TangentY[UnrealVertexIndex].Normalize();
+							}
 						}
 
 						FbxVector4 TangentZ = LayerElementNormal->GetDirectArray().GetAt(NormalValueIndex);
