[Legacy FBX Importer] Normal,Tangent,BinormalSign dereferencing and optimization fixes for Use T0 as Ref Pose

#rnx

diff --git a/Engine/Source/Editor/UnrealEd/Private/Fbx/FbxSkeletalMeshImport.cpp b/Engine/Source/Editor/UnrealEd/Private/Fbx/FbxSkeletalMeshImport.cpp
--- a/Engine/Source/Editor/UnrealEd/Private/Fbx/FbxSkeletalMeshImport.cpp
+++ b/Engine/Source/Editor/UnrealEd/Private/Fbx/FbxSkeletalMeshImport.cpp
@@ -474,13 +474,14 @@ void FFbxImporter::SkinControlPointsToPose(FSkeletalMeshImportData& ImportData,
 
 			bool OddNegativeScale = IsOddNegativeScale(MeshMatrix);
 
-			FbxLayerElement::EReferenceMode NormalReferenceMode(FbxLayerElement::eDirect);
-			FbxLayerElement::EMappingMode NormalMappingMode(FbxLayerElement::eByControlPoint);
-			if (LayerElementNormal)
-			{
-				NormalReferenceMode = LayerElementNormal->GetReferenceMode();
-				NormalMappingMode = LayerElementNormal->GetMappingMode();
-			}
+			FbxLayerElement::EReferenceMode NormalReferenceMode(LayerElementNormal ? LayerElementNormal->GetReferenceMode() : FbxLayerElement::eDirect);
+			FbxLayerElement::EMappingMode NormalMappingMode(LayerElementNormal ? LayerElementNormal->GetMappingMode() : FbxLayerElement::eByControlPoint);
+			
+			FbxLayerElement::EReferenceMode TangentReferenceMode(LayerElementTangent ? LayerElementTangent->GetReferenceMode() : FbxLayerElement::eDirect);
+			FbxLayerElement::EMappingMode TangentMappingMode(LayerElementTangent ? LayerElementTangent->GetMappingMode() : FbxLayerElement::eByControlPoint);
+
+			FbxLayerElement::EReferenceMode BinormalReferenceMode(LayerElementBinormal ? LayerElementBinormal->GetReferenceMode() : FbxLayerElement::eDirect);
+			FbxLayerElement::EMappingMode BinormalMappingMode(LayerElementBinormal ? LayerElementBinormal->GetMappingMode() : FbxLayerElement::eByControlPoint);
 
 			const int32 TriangleCount = FbxMesh->GetPolygonCount();
 			const int32 ExistingTriangleCount = ImportData.Faces.Num();
@@ -506,16 +507,8 @@ void FFbxImporter::SkinControlPointsToPose(FSkeletalMeshImportData& ImportData,
 					{
 						if (ImportData.bHasTangents)
 						{
+							if (LayerElementTangent)
 							{
-								//Tangent
-								FbxLayerElement::EReferenceMode TangentReferenceMode(FbxLayerElement::eDirect);
-								FbxLayerElement::EMappingMode TangentMappingMode(FbxLayerElement::eByControlPoint);
-								if (LayerElementTangent)
-								{
-									TangentReferenceMode = LayerElementTangent->GetReferenceMode();
-									TangentMappingMode = LayerElementTangent->GetMappingMode();
-								}
-
 								int32 TangentMapIndex = (TangentMappingMode == FbxLayerElement::eByControlPoint) ? ControlPointIndex : TmpIndex;
 								int32 TangentValueIndex = (TangentReferenceMode == FbxLayerElement::eDirect) ? TangentMapIndex : LayerElementTangent->GetIndexArray().GetAt(TangentMapIndex);
 
@@ -525,17 +518,13 @@ void FFbxImporter::SkinControlPointsToPose(FSkeletalMeshImportData& ImportData,
 								Triangle.TangentX[UnrealVertexIndex] = (FVector3f)Converter.ConvertDir(TangentX);
 								Triangle.TangentX[UnrealVertexIndex].Normalize();
 							}
-
+							else
 							{
-								//Binormal
-								FbxLayerElement::EReferenceMode BinormalReferenceMode(FbxLayerElement::eDirect);
-								FbxLayerElement::EMappingMode BinormalMappingMode(FbxLayerElement::eByControlPoint);
-								if (LayerElementBinormal)
-								{
-									BinormalReferenceMode = LayerElementBinormal->GetReferenceMode();
-									BinormalMappingMode = LayerElementBinormal->GetMappingMode();
-								}
+								Triangle.TangentX[UnrealVertexIndex] = FVector3f();
+							}
 
+							if (LayerElementBinormal)
+							{
 								int32 BinormalMapIndex = (BinormalMappingMode == FbxLayerElement::eByControlPoint) ? ControlPointIndex : TmpIndex;
 								int32 BinormalValueIndex = (BinormalReferenceMode == FbxLayerElement::eDirect) ? BinormalMapIndex : LayerElementBinormal->GetIndexArray().GetAt(BinormalMapIndex);
 
@@ -545,13 +534,24 @@ void FFbxImporter::SkinControlPointsToPose(FSkeletalMeshImportData& ImportData,
 								Triangle.TangentY[UnrealVertexIndex] = (FVector3f)Converter.ConvertDir(TangentY);
 								Triangle.TangentY[UnrealVertexIndex].Normalize();
 							}
+							else
+							{
+								Triangle.TangentY[UnrealVertexIndex] = FVector3f();
+							}
 						}
 
-						FbxVector4 TangentZ = LayerElementNormal->GetDirectArray().GetAt(NormalValueIndex);
-						TangentZ = ClusterDeformation.MultT(TangentZ);
-						TangentZ = TotalMatrixForNormal.MultT(TangentZ);
-						Triangle.TangentZ[UnrealVertexIndex] = (FVector3f)Converter.ConvertDir(TangentZ);
-						Triangle.TangentZ[UnrealVertexIndex].Normalize();
+						if (LayerElementNormal)
+						{
+							FbxVector4 TangentZ = LayerElementNormal->GetDirectArray().GetAt(NormalValueIndex);
+							TangentZ = ClusterDeformation.MultT(TangentZ);
+							TangentZ = TotalMatrixForNormal.MultT(TangentZ);
+							Triangle.TangentZ[UnrealVertexIndex] = (FVector3f)Converter.ConvertDir(TangentZ);
+							Triangle.TangentZ[UnrealVertexIndex].Normalize();
+						}
+						else
+						{
+							Triangle.TangentZ[UnrealVertexIndex] = FVector3f();
+						}
 					}
 				}
 			}
