MultiProcessCook: Change CookGenerationHelper to require a message from the director that the generated packages have been queued for cooking before it is allowed to complete the save of the generator. This is required for the next point. It will have a small performance cost, but will also decrease complexity by reducing how many different orders we can encounter in the events of a GenerationHelper's lifetime during MPCook.

MultiProductBuild: Skip the save of the generator if any of its generated packages were already cooked in the basegame. This allows us to cook newly discovered generated packages in the MultiProductBuild's cook session without duplicating the cook of the generator package itself.
#rnx

diff --git a/Engine/Source/Editor/UnrealEd/Private/CookOnTheFlyServer.cpp b/Engine/Source/Editor/UnrealEd/Private/CookOnTheFlyServer.cpp
--- a/Engine/Source/Editor/UnrealEd/Private/CookOnTheFlyServer.cpp
+++ b/Engine/Source/Editor/UnrealEd/Private/CookOnTheFlyServer.cpp
@@ -3632,6 +3632,17 @@ UE::Cook::EPollStatus UCookOnTheFlyServer::PrepareSaveGenerationPackage(UE::Cook
 	}
 	FCookGenerationInfo& Info(*InfoPtr);
 
+	if (PackageData.GetSaveSubState() <= ESaveSubState::CheckForIsGenerated)
+	{
+		// We cannot proceed with Populate,PreSave,Save,PostSave on the generator until after
+		// QueuedGeneratedPackagesFencePassed, because we might discover during processing of that fence that we are
+		// not allowed to save the generator package.
+		if (Info.IsGenerator() && !GenerationHelper.IsQueuedGeneratedPackagesFencePassed())
+		{
+			return EPollStatus::Incomplete;
+		}
+	}
+
 	if (PackageData.GetSaveSubState() <= ESaveSubState::Generation_PreMoveCookedPlatformData_WaitingForIsLoaded)
 	{
 		// Both Generator packages and Generated packages should wait for all IsCachedCookedPlatformData
diff --git a/Engine/Source/Editor/UnrealEd/Private/Cooker/CookDirector.cpp b/Engine/Source/Editor/UnrealEd/Private/Cooker/CookDirector.cpp
--- a/Engine/Source/Editor/UnrealEd/Private/Cooker/CookDirector.cpp
+++ b/Engine/Source/Editor/UnrealEd/Private/Cooker/CookDirector.cpp
@@ -391,6 +391,8 @@ TMap<FPackageData*, FAssignPackageExtraData> FCookDirector::GetAssignPackageExtr
 		if (GenerationHelper)
 		{
 			ExtraData = GetOrAllocateExtraData(Request);
+			ExtraData->bSkipSaveExistingGenerator = GenerationHelper->IsSkipSaveExistingGenerator();
+			ExtraData->bQueuedGeneratedPackagesFencePassed = GenerationHelper->IsQueuedGeneratedPackagesFencePassed();
 			ExtraData->GeneratorPerPlatformPreviousGeneratedPackages.Reserve(SessionPlatforms.Num());
 			for (const ITargetPlatform* TargetPlatform : SessionPlatforms)
 			{
diff --git a/Engine/Source/Editor/UnrealEd/Private/Cooker/CookGenerationHelper.cpp b/Engine/Source/Editor/UnrealEd/Private/Cooker/CookGenerationHelper.cpp
--- a/Engine/Source/Editor/UnrealEd/Private/Cooker/CookGenerationHelper.cpp
+++ b/Engine/Source/Editor/UnrealEd/Private/Cooker/CookGenerationHelper.cpp
@@ -2,6 +2,7 @@
 
 #include "Cooker/CookGenerationHelper.h"
 
+#include "Algo/AnyOf.h"
 #include "Algo/Unique.h"
 #include "AssetRegistry/AssetData.h"
 #include "Cooker/CookDirector.h"
@@ -329,11 +330,15 @@ void FGenerationHelper::Uninitialize()
 	// Keep bRequiresGeneratorPackageDestructBeforeResplit; it is allowed in the uninitialized state
 	bReportedManifest = false;
 	bSaveGenerator = true;
+	// Keep bSkipSaveExistingGenerator; we need to preserve it across the uninitialized state
 	bCurrentGCHasKeptGeneratorPackage = false;
 	bCurrentGCHasKeptGeneratorKeepPackages = false;
 	// Keep bKeepForAllSavedOrGC ; it is allowed in the uninitialized state
 	// Keep bKeepForCompletedAllSavesMessage; it is allowed in the uninitialized state
 	// Keep bNeedConfirmGeneratorPackageDestroyed; it is allowed in the uninitialized state
+	// Keep bHasFinishedQueueGeneratedPackages; we need to preserve it across the uninitialized state
+	// Keep bQueuedGeneratedPackagesFencePassed; we need to preserve it across the uninitialized state
+	// Keep bDeferEvents; we need to preserve it across the uninitialized state
 	// Keep bSentAllSavesCompleted; it is allowed in the uninitialized state
 }
 
@@ -1007,7 +1012,7 @@ bool FGenerationHelper::TryReportGenerationManifest()
 		}
 	}
 
-	bSaveGenerator = Manifest.bCookSaveGeneratorPackage;
+	bSaveGenerator = Manifest.bCookSaveGeneratorPackage && !bSkipSaveExistingGenerator;
 
 	TMap<FPackageData*, int32> AlreadyExistingInfoPackageToIndex;
 	int32 NumAlreadyExisting = PackagesToGenerate.Num();
@@ -1548,9 +1553,52 @@ void FGenerationHelper::OnRequestFencePassed(UCookOnTheFlyServer& COTFS)
 		// out for saving before receiving the EndQueueGeneratedPackagesOnDirector message.
 		OnNumSavedUpdated();
 
+		FName PackageName = GetOwner().GetPackageName();
+		constexpr bool bShouldSkipCookSaveGeneratorPackageIfBaseGamePartiallyCooked = false;
+		if (bShouldSkipCookSaveGeneratorPackageIfBaseGamePartiallyCooked)
+		{
+			auto IsCookedElsewhere = [](const FPackageData& PackageData)
+				{
+					for (const TPair<const ITargetPlatform*, FPackagePlatformData>& PlatformPair
+						: PackageData.GetPlatformDatas())
+					{
+						if (PlatformPair.Value.GetWhereCooked() != EWhereCooked::ThisSession)
+						{
+							return true;
+						}
+					}
+					return false;
+				};
+			if (!IsCookedElsewhere(GetOwner()))
+			{
+				// This is a temporary solution for MultiProductBuilds that want to recook a Generator with extra
+				// GeneratedPackages inserted by plugins, but do not want to duplicate the cook of the generator
+				// package. MultiProductBuild implements that case by adding the generated packages into the
+				// basegame's DevelopmentAssetRegistry, but NOT adding the generator package, so that in
+				// RecordDLCPackagesFromBaseGame we marked the generated packages as
+				// SetWhereCooked(EWhereCooked::BaseGame), but did not mark the generator.
+				// Detect that case and suppress the save of the generator package via bSaveGenerator=false.
+				bool bAnyCookedElsewhere = Algo::AnyOf(PackagesToGenerate,
+					[&IsCookedElsewhere](const FCookGenerationInfo& Info)
+					{
+						return IsCookedElsewhere(*Info.PackageData);
+					});
+				if (bAnyCookedElsewhere)
+				{
+					// OnSkipSaveExistingGenerator, and the broadcast of it to CookWorkers, should be called before
+					// OnQueuedGeneratedPackagesFencePassed because Generators wait for
+					// OnQueuedGeneratedPackagesFencePassed and then execute code that uses this value if set.
+					OnSkipSaveExistingGenerator();
+					if (COTFS.CookDirector)
+					{
+						FGeneratorEventMessage Message(EGeneratorEvent::SkipSaveExistingGenerator, PackageName);
+						COTFS.CookDirector->BroadcastMessage(MoveTemp(Message), ECookBroadcastTiming::AfterAssignPackages);
+					}
+				}
+			}
+		}
 		if (COTFS.CookDirector)
 		{
-			FName PackageName = GetOwner().GetPackageName();
 			FGeneratorEventMessage Message(EGeneratorEvent::QueuedGeneratedPackagesFencePassed, PackageName);
 			COTFS.CookDirector->BroadcastMessage(MoveTemp(Message), ECookBroadcastTiming::AfterAssignPackages);
 		}
@@ -1566,11 +1614,20 @@ void FGenerationHelper::OnRequestFencePassed(UCookOnTheFlyServer& COTFS)
 	}
 }
 
+void FGenerationHelper::OnSkipSaveExistingGenerator()
+{
+	UE_LOG(LogCookGenerationHelper, Verbose, TEXT("%s: SkipSaveExistingGenerator."),
+		*WriteToString<256>(GetOwner().GetPackageName()));
+	bSkipSaveExistingGenerator = true;
+	bSaveGenerator = false;
+}
+
 void FGenerationHelper::OnQueuedGeneratedPackagesFencePassed(UCookOnTheFlyServer& COTFS)
 {
 	UE_LOG(LogCookGenerationHelper, Verbose, TEXT("%s: RequestFencePassed, and HasFinishedQueueGeneratedPackages."),
 		*WriteToString<256>(GetOwner().GetPackageName()));
 
+	bQueuedGeneratedPackagesFencePassed = true;
 	ClearKeepForQueueResults();
 	// We no longer need PreviousGeneratedPackages or KeepForIncremental, because they are used only in 
 	// StartQueueGeneratedPackages or the request cluster that they end up in in PumpRequests, both of which
diff --git a/Engine/Source/Editor/UnrealEd/Private/Cooker/CookGenerationHelper.h b/Engine/Source/Editor/UnrealEd/Private/Cooker/CookGenerationHelper.h
--- a/Engine/Source/Editor/UnrealEd/Private/Cooker/CookGenerationHelper.h
+++ b/Engine/Source/Editor/UnrealEd/Private/Cooker/CookGenerationHelper.h
@@ -254,6 +254,11 @@ public:
 	 * if not initialized. Does not call Initialize. 
 	 */
 	bool IsSaveGenerator() const;
+	/**
+	 * True if we detected that we need to skip save of the generator package because it was already cooked in the
+	 * basegame and we only need to cook any new generated packages that it has.
+	 */
+	bool IsSkipSaveExistingGenerator() const;
 	/** Accessor for the packages to generate, empty if invalid or TryReportGenerationManifest was not yet called. */
 	TArrayView<FCookGenerationInfo> GetPackagesToGenerate();
 	/** Return the GenerationInfo used to save the generator package's UPackage. Does not call Initialize. */
@@ -303,6 +308,11 @@ public:
 	 * Does not call initialize.
 	 */
 	TArray<FName> ReleaseExternalActorDependencies();
+	/**
+	 * Return whether PumpRequests on the director has processed all of the generated package requests, and we have
+	 * done any followup configuration based on data about those generated packages in previous cooks.
+	 */
+	bool IsQueuedGeneratedPackagesFencePassed() const;
 
 	/**
 	 * Call the Splitter's ReportGenerationManifest and create the PackageDatas. Logs errors and returns false
@@ -347,6 +357,11 @@ public:
 	 * Calls the proper followup events locally and on all Workers. Does not call Initialize.
 	 */
 	void OnRequestFencePassed(UCookOnTheFlyServer& COTFS);
+	/**
+	 * Called on Directors and CookWorkers when we discover that the generator was saved in the basegame and we should
+	 * skip its save in this session but still save any new generatedpackages.
+	 */
+	void OnSkipSaveExistingGenerator();
 	/**
 	 * Called on Directors and Workers when the RequestFence added from EndQueueGeneratedPackages has passed.
 	 * Does not call Initialize.
@@ -607,12 +622,14 @@ private:
 	bool bRequiresGeneratorPackageDestructBeforeResplit = false;
 	bool bReportedManifest = false;
 	bool bSaveGenerator = true;
+	bool bSkipSaveExistingGenerator = false;
 	bool bCurrentGCHasKeptGeneratorPackage = false;
 	bool bCurrentGCHasKeptGeneratorKeepPackages = false;
 	bool bKeepForAllSavedOrGC = false;
 	bool bKeepForCompletedAllSavesMessage = false;
 	bool bNeedConfirmGeneratorPackageDestroyed = false;
 	bool bHasFinishedQueueGeneratedPackages = false;
+	bool bQueuedGeneratedPackagesFencePassed = false;
 	bool bDeferEvents = false;
 	bool bSentAllSavesCompleted = false;
 
@@ -810,6 +827,11 @@ inline bool FGenerationHelper::IsSaveGenerator() const
 	return bSaveGenerator;
 }
 
+inline bool FGenerationHelper::IsSkipSaveExistingGenerator() const
+{
+	return bSkipSaveExistingGenerator;
+}
+
 inline TArrayView<UE::Cook::FCookGenerationInfo> FGenerationHelper::GetPackagesToGenerate()
 {
 	ConditionalInitialize();
@@ -883,6 +905,11 @@ inline TArray<FName> FGenerationHelper::ReleaseExternalActorDependencies()
 	return Result;
 }
 
+inline bool FGenerationHelper::IsQueuedGeneratedPackagesFencePassed() const
+{
+	return bQueuedGeneratedPackagesFencePassed;
+}
+
 inline const TMap<FName, FAssetPackageData>& FGenerationHelper::GetPreviousGeneratedPackages(
 	const ITargetPlatform* TargetPlatform) const
 {
diff --git a/Engine/Source/Editor/UnrealEd/Private/Cooker/CookWorkerClient.cpp b/Engine/Source/Editor/UnrealEd/Private/Cooker/CookWorkerClient.cpp
--- a/Engine/Source/Editor/UnrealEd/Private/Cooker/CookWorkerClient.cpp
+++ b/Engine/Source/Editor/UnrealEd/Private/Cooker/CookWorkerClient.cpp
@@ -376,6 +376,9 @@ void FCookWorkerClient::HandleGeneratorMessage(FGeneratorEventMessage&& Generato
 		{
 			switch (GeneratorMessage.Event)
 			{
+			case EGeneratorEvent::SkipSaveExistingGenerator:
+				GenerationHelper->OnSkipSaveExistingGenerator();
+				break;
 			case EGeneratorEvent::QueuedGeneratedPackagesFencePassed:
 				GenerationHelper->OnQueuedGeneratedPackagesFencePassed(COTFS);
 				break;
@@ -1010,6 +1013,16 @@ void FCookWorkerClient::AssignPackages(FAssignPackagesMessage& Message)
 					const ITargetPlatform* TargetPlatform = OrderedSessionPlatforms[PlatformPair.Key];
 					GenerationHelper->SetPreviousGeneratedPackages(TargetPlatform, MoveTemp(PlatformPair.Value));
 				}
+				if (AssignData.bSkipSaveExistingGenerator)
+				{
+					// Should be called before OnQueuedGeneratedPackagesFencePassed because Generators wait for
+					// OnQueuedGeneratedPackagesFencePassed and then execute code that uses this value if set.
+					GenerationHelper->OnSkipSaveExistingGenerator();
+				}
+				if (AssignData.bQueuedGeneratedPackagesFencePassed)
+				{
+					GenerationHelper->OnQueuedGeneratedPackagesFencePassed(COTFS);
+				}
 			}
 			if (!AssignData.PerPackageCollectorMessages.IsEmpty())
 			{
diff --git a/Engine/Source/Editor/UnrealEd/Private/Cooker/CookWorkerServer.cpp b/Engine/Source/Editor/UnrealEd/Private/Cooker/CookWorkerServer.cpp
--- a/Engine/Source/Editor/UnrealEd/Private/Cooker/CookWorkerServer.cpp
+++ b/Engine/Source/Editor/UnrealEd/Private/Cooker/CookWorkerServer.cpp
@@ -758,6 +758,8 @@ void FCookWorkerServer::SendPendingPackages()
 		FAssignPackageExtraData* ExtraData = PackagesToAssignExtraDatas.Find(PackageData);
 		if (ExtraData)
 		{
+			AssignData.bSkipSaveExistingGenerator = ExtraData->bSkipSaveExistingGenerator;
+			AssignData.bQueuedGeneratedPackagesFencePassed = ExtraData->bQueuedGeneratedPackagesFencePassed;
 			AssignData.GeneratorPerPlatformPreviousGeneratedPackages.Reserve(
 				ExtraData->GeneratorPerPlatformPreviousGeneratedPackages.Num());
 			for (TPair<const ITargetPlatform*, TMap<FName, FAssetPackageData>>& PlatformPair
@@ -1486,6 +1488,8 @@ void FAssignPackageData::Write(FCbWriter& Writer,
 	static_assert(sizeof(EUrgency) <= sizeof(uint8), "We are storing it in a uint8");
 	Writer << static_cast<uint8>(Reachability);
 	static_assert(sizeof(EReachability) <= sizeof(uint8), "We are storing it in a uint8");
+	Writer << bSkipSaveExistingGenerator;
+	Writer << bQueuedGeneratedPackagesFencePassed;
 	WriteToCompactBinary(Writer, NeedCommitPlatforms, OrderedSessionPlatforms);
 	{
 		Writer.BeginArray();
@@ -1538,6 +1542,8 @@ bool FAssignPackageData::TryRead(FCbFieldView Field, TConstArrayView<const ITarg
 	{
 		bOk = false;
 	}
+	bOk = LoadFromCompactBinary(*It++, bSkipSaveExistingGenerator);
+	bOk = LoadFromCompactBinary(*It++, bQueuedGeneratedPackagesFencePassed);
 	bOk = LoadFromCompactBinary(*It++, NeedCommitPlatforms, OrderedSessionPlatforms) & bOk;
 	{
 		FCbFieldView PlatformArrayFieldView = *It++;
diff --git a/Engine/Source/Editor/UnrealEd/Private/Cooker/CookWorkerServer.h b/Engine/Source/Editor/UnrealEd/Private/Cooker/CookWorkerServer.h
--- a/Engine/Source/Editor/UnrealEd/Private/Cooker/CookWorkerServer.h
+++ b/Engine/Source/Editor/UnrealEd/Private/Cooker/CookWorkerServer.h
@@ -53,6 +53,8 @@ struct FAssignPackageExtraData
 {
 	TMap<const ITargetPlatform*, TMap<FName, FAssetPackageData>> GeneratorPerPlatformPreviousGeneratedPackages;
 	TArray<UE::CompactBinaryTCP::FMarshalledMessage> PerPackageCollectorMessages;
+	bool bSkipSaveExistingGenerator = false;
+	bool bQueuedGeneratedPackagesFencePassed = false;
 };
 
 /** Class in a Director process that communicates over a Socket with FCookWorkerClient in a CookWorker process. */
@@ -267,6 +269,8 @@ struct FAssignPackageData
 	TArray<UE::CompactBinaryTCP::FMarshalledMessage> PerPackageCollectorMessages;
 	ICookPackageSplitter::EGeneratedRequiresGenerator DoesGeneratedRequireGenerator
 		= ICookPackageSplitter::EGeneratedRequiresGenerator::None;
+	bool bSkipSaveExistingGenerator = false;
+	bool bQueuedGeneratedPackagesFencePassed = false;
 
 private:
 	void Write(FCbWriter& Writer, TConstArrayView<const ITargetPlatform*> OrderedSessionPlatforms) const;
@@ -467,7 +471,12 @@ enum class EGeneratorEvent
 	Invalid,
 	/** Sent from Worker to Director to let the Director know to start waiting for AllSavesCompleted. */
 	QueuedGeneratedPackages,
-	/** Broadcast to Workers so any waiting on the fence to pass before clearing their data can clear their data. */
+	/** Broadcast to Workers to set their value of bSkipSaveExistingGenerator when it is first set. */
+	SkipSaveExistingGenerator,
+	/**
+	 * Broadcast to Workers so any waiting to save or waiting on the fence to pass before clearing their data can save
+	 * or clear their data.
+	 */
 	QueuedGeneratedPackagesFencePassed,
 	/** Broadcast to Workers to unconditionally clear their data. */
 	AllSavesCompleted,
