MultiProcessCook: Change CookGenerationHelper to require a message from the director: Fixes from static analysis.
#rnx
#FailedValidations Preflight

diff --git a/Engine/Source/Editor/UnrealEd/Private/Cooker/CookWorkerClient.cpp b/Engine/Source/Editor/UnrealEd/Private/Cooker/CookWorkerClient.cpp
--- a/Engine/Source/Editor/UnrealEd/Private/Cooker/CookWorkerClient.cpp
+++ b/Engine/Source/Editor/UnrealEd/Private/Cooker/CookWorkerClient.cpp
@@ -1004,7 +1004,8 @@ void FCookWorkerClient::AssignPackages(FAssignPackagesMessage& Message)
 				PackageData.SetGenerated(AssignData.ParentGenerator);
 				PackageData.SetDoesGeneratedRequireGenerator(AssignData.DoesGeneratedRequireGenerator);
 			}
-			if (!AssignData.GeneratorPerPlatformPreviousGeneratedPackages.IsEmpty())
+			if (!AssignData.GeneratorPerPlatformPreviousGeneratedPackages.IsEmpty() ||
+				AssignData.bSkipSaveExistingGenerator || AssignData.bQueuedGeneratedPackagesFencePassed)
 			{
 				TRefCountPtr<FGenerationHelper> GenerationHelper = PackageData.CreateUninitializedGenerationHelper();
 				for (TPair<uint8, TMap<FName, FAssetPackageData>>& PlatformPair
diff --git a/Engine/Source/Editor/UnrealEd/Private/Cooker/CookWorkerServer.cpp b/Engine/Source/Editor/UnrealEd/Private/Cooker/CookWorkerServer.cpp
--- a/Engine/Source/Editor/UnrealEd/Private/Cooker/CookWorkerServer.cpp
+++ b/Engine/Source/Editor/UnrealEd/Private/Cooker/CookWorkerServer.cpp
@@ -1542,8 +1542,8 @@ bool FAssignPackageData::TryRead(FCbFieldView Field, TConstArrayView<const ITarg
 	{
 		bOk = false;
 	}
-	bOk = LoadFromCompactBinary(*It++, bSkipSaveExistingGenerator);
-	bOk = LoadFromCompactBinary(*It++, bQueuedGeneratedPackagesFencePassed);
+	bOk = LoadFromCompactBinary(*It++, bSkipSaveExistingGenerator) & bOk;
+	bOk = LoadFromCompactBinary(*It++, bQueuedGeneratedPackagesFencePassed) & bOk;
 	bOk = LoadFromCompactBinary(*It++, NeedCommitPlatforms, OrderedSessionPlatforms) & bOk;
 	{
 		FCbFieldView PlatformArrayFieldView = *It++;
