[AutoRTFM] Disable AutoRTFM when cooking with Verse BPVM.

We haven't battle-hardened this yet, and AutoRTFM's triggering a handful of issues during Aegis cook jobs. The AutoRTFM issues will need to be fixed for VVM regardless, but we don't want to rock the boat for BPVM until we've solved the known issues.

#FailedValidations Preflight

diff --git a/Engine/Source/Editor/UnrealEd/Private/Commandlets/CookCommandlet.cpp b/Engine/Source/Editor/UnrealEd/Private/Commandlets/CookCommandlet.cpp
--- a/Engine/Source/Editor/UnrealEd/Private/Commandlets/CookCommandlet.cpp
+++ b/Engine/Source/Editor/UnrealEd/Private/Commandlets/CookCommandlet.cpp
@@ -8,6 +8,7 @@
 
 #include "AssetRegistry/AssetRegistryModule.h"
 #include "Async/TaskGraphInterfaces.h"
+#include "AutoRTFM.h"
 #include "CookOnTheSide/CookOnTheFlyServer.h"
 #include "Cooker/CookImportsChecker.h"
 #include "Cooker/CookProfiling.h"
@@ -175,6 +176,12 @@ bool UCookCommandlet::CookOnTheFly( FGuid InstanceId, int32 Port, int32 Timeout,
 
 int32 UCookCommandlet::Main(const FString& CmdLineParams)
 {
+#if WITH_VERSE_VM
+	AutoRTFM::ForTheRuntime::SetAutoRTFMRuntime(AutoRTFM::ForTheRuntime::EAutoRTFMEnabledState::AutoRTFM_Enabled);
+#else
+	AutoRTFM::ForTheRuntime::SetAutoRTFMRuntime(AutoRTFM::ForTheRuntime::EAutoRTFMEnabledState::AutoRTFM_Disabled);
+#endif
+
 	COOK_STAT(DetailedCookStats::CookStartTime = FPlatformTime::Seconds());
 	Params = CmdLineParams;
 	ParseCommandLine(*Params, Tokens, Switches);
@@ -207,7 +214,7 @@ int32 UCookCommandlet::Main(const FString& CmdLineParams)
 
 	COOK_STAT(FParse::Value(*Params, TEXT("CookCultures="), DetailedCookStats::CookCultures));
 	COOK_STAT(FParse::Value(*Params, TEXT("CookLabel="), DetailedCookStats::CookLabel));
-	
+
 	ITargetPlatformManagerModule& TPM = GetTargetPlatformManagerRef();
 	if ( bCookOnTheFly )
 	{
@@ -312,7 +319,7 @@ bool UCookCommandlet::CookByTheBook(const TArray<ITargetPlatform*>& Platforms)
 	CookFlags |= Switches.Contains(TEXT("IgnoreScriptPackagesOutOfDate")) || CookerSettings->bIgnoreScriptPackagesOutOfDateForIteration ? ECookInitializationFlags::IgnoreScriptPackagesOutOfDate : ECookInitializationFlags::None;
 
 	//////////////////////////////////////////////////////////////////////////
-	// parse commandline options 
+	// parse commandline options
 
 	FString DLCName;
 	FParse::Value(*Params, TEXT("DLCNAME="), DLCName);
@@ -419,7 +426,7 @@ bool UCookCommandlet::CookByTheBook(const TArray<ITargetPlatform*>& Platforms)
 		(Switches.Contains(TEXT("ZenStore")) || PackagingSettings->GetUseZenStoreEffective());
 
 	//////////////////////////////////////////////////////////////////////////
-	// start cook by the book 
+	// start cook by the book
 	ECookByTheBookOptions CookOptions = ECookByTheBookOptions::None;
 
 	CookOptions |= bCookAll ? ECookByTheBookOptions::CookAll : ECookByTheBookOptions::None;
