Fix for -nosound crash

#rnx

diff --git a/Engine/Source/Editor/UnrealEd/Private/EditorEngine.cpp b/Engine/Source/Editor/UnrealEd/Private/EditorEngine.cpp
--- a/Engine/Source/Editor/UnrealEd/Private/EditorEngine.cpp
+++ b/Engine/Source/Editor/UnrealEd/Private/EditorEngine.cpp
@@ -2904,7 +2904,11 @@ UAudioComponent* UEditorEngine::ResetPreviewAudioComponent(USoundBase* Sound, US
 		PreviewAudioComponent = NewObject<UAudioComponent>();
 	}
 
-	check(PreviewAudioComponent);
+	// This may happen when the engine is launched with -nosound
+	if (!PreviewAudioComponent)
+	{
+		return nullptr;
+	}
 
 	// Mark as a preview component so the distance calculations is ignored
 	PreviewAudioComponent->bPreviewComponent = true;
