MultiProductBuild: Enable the skipsave of the generator if any of its generated packages were already cooked in the basegame.
#rnx

diff --git a/Engine/Source/Editor/UnrealEd/Private/Cooker/CookGenerationHelper.cpp b/Engine/Source/Editor/UnrealEd/Private/Cooker/CookGenerationHelper.cpp
--- a/Engine/Source/Editor/UnrealEd/Private/Cooker/CookGenerationHelper.cpp
+++ b/Engine/Source/Editor/UnrealEd/Private/Cooker/CookGenerationHelper.cpp
@@ -1554,7 +1554,7 @@ void FGenerationHelper::OnRequestFencePassed(UCookOnTheFlyServer& COTFS)
 		OnNumSavedUpdated();
 
 		FName PackageName = GetOwner().GetPackageName();
-		constexpr bool bShouldSkipCookSaveGeneratorPackageIfBaseGamePartiallyCooked = false;
+		constexpr bool bShouldSkipCookSaveGeneratorPackageIfBaseGamePartiallyCooked = true;
 		if (bShouldSkipCookSaveGeneratorPackageIfBaseGamePartiallyCooked)
 		{
 			auto IsCookedElsewhere = [](const FPackageData& PackageData)
