More direct fix for case of -nosound

#rnx
#DisabledValidations CheckUnacceptableWords

diff --git a/Engine/Source/Editor/UnrealEd/Private/EditorEngine.cpp b/Engine/Source/Editor/UnrealEd/Private/EditorEngine.cpp
--- a/Engine/Source/Editor/UnrealEd/Private/EditorEngine.cpp
+++ b/Engine/Source/Editor/UnrealEd/Private/EditorEngine.cpp
@@ -2884,6 +2884,12 @@ UAudioComponent* UEditorEngine::ResetPreviewAudioComponent(USoundBase* Sound, US
 		PreviewSoundCue = nullptr;
 	}
 
+	// The audio engine may be disabled
+	if (!UseSound())
+	{
+		return nullptr;
+	}
+
 	// 1. Initialize & overwrite new component for previewing
 	if (PreviewAudioComponent)
 	{
@@ -2904,11 +2910,7 @@ UAudioComponent* UEditorEngine::ResetPreviewAudioComponent(USoundBase* Sound, US
 		PreviewAudioComponent = NewObject<UAudioComponent>();
 	}
 
-	// This may happen when the engine is launched with -nosound
-	if (!PreviewAudioComponent)
-	{
-		return nullptr;
-	}
+	check(PreviewAudioComponent);
 
 	// Mark as a preview component so the distance calculations is ignored
 	PreviewAudioComponent->bPreviewComponent = true;
