Cooker: Add commandline arguments -cookshowremaining,  -cookshowpackagenames, -cookshowinstigators to provide a clearer behavior than -dpcvars=cook.displaymode=1|2|4. Each flag can be specified with a value: -cookshowremaining=1|0, or with no value, defaulting to 1: -cookshowremaining. These flags override whatever value has been assigned via config or -dpcvars commandline.
[RN] Minor, Cooking

diff --git a/Engine/Source/Editor/UnrealEd/Private/CookOnTheFlyServer.cpp b/Engine/Source/Editor/UnrealEd/Private/CookOnTheFlyServer.cpp
--- a/Engine/Source/Editor/UnrealEd/Private/CookOnTheFlyServer.cpp
+++ b/Engine/Source/Editor/UnrealEd/Private/CookOnTheFlyServer.cpp
@@ -6324,23 +6324,29 @@ void UCookOnTheFlyServer::SetInitializeConfigSettings(UE::Cook::FInitializeConfi
 	// The rest of this function parses config settings that are reparsed on every CookDirector and CookWorker rather than
 	// being replicated from CookDirector to CookWorker
 
-	UE_LOG(LogCook, Display, TEXT("CookSettings for Memory:%s"), *GetCookSettingsForMemoryLogText());
 
 	const TCHAR* CommandLine = FCommandLine::Get();
 
+	// Parse -cookshow... flags
+	bool bShowRemaining = !!(GCookProgressDisplay & (int32)ECookProgressDisplayMode::RemainingPackages);
+	bool bShowPackageNames = !!(GCookProgressDisplay & (int32)ECookProgressDisplayMode::PackageNames);
+	bool bShowInstigators = !!(GCookProgressDisplay & (int32)ECookProgressDisplayMode::Instigators);
+	ParseBoolParamOrValue(CommandLine, TEXT("cookshowremaining"), bShowRemaining);
+	ParseBoolParamOrValue(CommandLine, TEXT("cookshowpackagenames"), bShowPackageNames);
+	ParseBoolParamOrValue(CommandLine, TEXT("cookshowinstigators"), bShowInstigators);
+	GCookProgressDisplay = 0;
+	GCookProgressDisplay |= bShowRemaining ? (int32)ECookProgressDisplayMode::RemainingPackages : 0;
+	GCookProgressDisplay |= bShowPackageNames ? (int32)ECookProgressDisplayMode::PackageNames : 0;
+	GCookProgressDisplay |= bShowInstigators ? (int32)ECookProgressDisplayMode::Instigators : 0;
+
+	// Report the memory settings
+	UE_LOG(LogCook, Display, TEXT("CookSettings for Memory:%s"), *GetCookSettingsForMemoryLogText());
+
 	// Debugging hidden dependencies
 	bOnlyEditorOnlyDebug = FParse::Param(CommandLine, TEXT("OnlyEditorOnlyDebug"));
 	bSkipOnlyEditorOnly = false;
 	GConfig->GetBool(TEXT("CookSettings"), TEXT("SkipOnlyEditorOnly"), bSkipOnlyEditorOnly, GEditorIni);
-	FString ParamText;
-	if (FParse::Value(CommandLine, TEXT("-SkipOnlyEditorOnly="), ParamText))
-	{
-		LexFromString(bSkipOnlyEditorOnly, *ParamText);
-	}
-	else if (FParse::Param(CommandLine, TEXT("SkipOnlyEditorOnly")))
-	{
-		bSkipOnlyEditorOnly = true;
-	}
+	ParseBoolParamOrValue(CommandLine, TEXT("SkipOnlyEditorOnly"), bSkipOnlyEditorOnly);
 	bSkipOnlyEditorOnly |= bOnlyEditorOnlyDebug;
 	if (bSkipOnlyEditorOnly)
 	{
diff --git a/Engine/Source/Editor/UnrealEd/Private/Cooker/CookTypes.cpp b/Engine/Source/Editor/UnrealEd/Private/Cooker/CookTypes.cpp
--- a/Engine/Source/Editor/UnrealEd/Private/Cooker/CookTypes.cpp
+++ b/Engine/Source/Editor/UnrealEd/Private/Cooker/CookTypes.cpp
@@ -1163,3 +1163,20 @@ const TCHAR* LexToString(UE::Cook::ECookResult CookResult)
 		return TEXT("Invalid");
 	}
 }
+
+void ParseBoolParamOrValue(const TCHAR* CommandLine, const TCHAR* ParamNameNoDashOrEquals, bool& bInOutValue)
+{
+	FString ParamText;
+	TStringBuilder<256> ParamStr(InPlace, TEXT("-"), ParamNameNoDashOrEquals, TEXT("="));
+	if (FParse::Value(CommandLine, *ParamStr, ParamText))
+	{
+		LexFromString(bInOutValue, *ParamText);
+		return;
+	}
+	bool bParamWithNoValuePresent = FParse::Param(CommandLine, ParamNameNoDashOrEquals);
+	if (bParamWithNoValuePresent)
+	{
+		bInOutValue = true;
+		return;
+	}
+}
diff --git a/Engine/Source/Editor/UnrealEd/Private/Cooker/CookTypes.h b/Engine/Source/Editor/UnrealEd/Private/Cooker/CookTypes.h
--- a/Engine/Source/Editor/UnrealEd/Private/Cooker/CookTypes.h
+++ b/Engine/Source/Editor/UnrealEd/Private/Cooker/CookTypes.h
@@ -828,3 +828,12 @@ const TCHAR* LexToString(UE::Cook::ECookResult CookResult);
  * @param Severity of the message
  */
 void LogCookerMessage(const FString& MessageText, EMessageSeverity::Type Severity);
+
+/**
+ * Parses commandline for -Param or -Param=<1|0|true|false|on|off>.
+ * If -Param=<CoercableToBoolValue> is found, bInOutValue is set to the value.
+ * Else If -Param is found, bInOutValue is set to true.
+ * If multiple occurrences of -Param= are found, the first value is used.
+ * If neither -Param nor -Param= is found, bInOutValue is unchanged.
+ */
+void ParseBoolParamOrValue(const TCHAR* CommandLine, const TCHAR* ParamNameNoDashOrEquals, bool& bInOutValue);
