Cooker: Deprecate cooking in editor.
[RN] Minor, Cooking

diff --git a/Engine/Config/BaseEditorPerProjectUserSettings.ini b/Engine/Config/BaseEditorPerProjectUserSettings.ini
--- a/Engine/Config/BaseEditorPerProjectUserSettings.ini
+++ b/Engine/Config/BaseEditorPerProjectUserSettings.ini
@@ -165,7 +165,7 @@ bEnableAsyncSkinnedAssetCompilation=True
 bEnableAsyncSoundWaveCompilation=True
 bEnableAsyncGroomBindingCompilation=True
 
-; Cooking in the editor is disabled by default now; WorldPartition does not yet support it and we want to deprecate the mode in the future to reduce the number of code flows that have to be supported
+; Deprecated in 5.7 and must now always be true; see comment in EditorExperimentalSettings.h
 bDisableCookInEditor=true
 
 [/Script/UnrealEd.EditorLoadingSavingSettings]
diff --git a/Engine/Source/Editor/UnrealEd/Classes/Settings/EditorExperimentalSettings.h b/Engine/Source/Editor/UnrealEd/Classes/Settings/EditorExperimentalSettings.h
--- a/Engine/Source/Editor/UnrealEd/Classes/Settings/EditorExperimentalSettings.h
+++ b/Engine/Source/Editor/UnrealEd/Classes/Settings/EditorExperimentalSettings.h
@@ -81,7 +81,10 @@ public:
 	UPROPERTY(EditAnywhere,config,Category=UserInterface,meta=(DisplayName="Allow ChunkID Assignments"))
 	bool bContextMenuChunkAssignments;
 
-	/** Disable cook in the editor */
+	/**
+	 * Cooking in the editor is deprecated in 5.7; it has never been updated to support WorldPartition and we are
+	 * removing support for it entirely in a future release. bDisableCookInEditor should be set to true.
+	 */
 	UPROPERTY(EditAnywhere, config, Category = Cooking, meta = (DisplayName = "Disable Cook In The Editor feature (cooks from launch on will be run in a separate process if disabled)", ConfigRestartRequired=true))
 	bool bDisableCookInEditor;
 
diff --git a/Engine/Source/Editor/UnrealEd/Private/UnrealEdEngine.cpp b/Engine/Source/Editor/UnrealEd/Private/UnrealEdEngine.cpp
--- a/Engine/Source/Editor/UnrealEd/Private/UnrealEdEngine.cpp
+++ b/Engine/Source/Editor/UnrealEd/Private/UnrealEdEngine.cpp
@@ -256,6 +256,13 @@ void UUnrealEdEngine::Init(IEngineLoop* InEngineLoop)
 		}
 		else if (!ExperimentalSettings->bDisableCookInEditor)
 		{
+			FMessageLog EditorLog("PackagingResults");
+			EditorLog.Warning(NSLOCTEXT("UnrealEd", "DisableCookingInEditor=true is required",
+				"bDisableCookInEditor=false is deprecated in 5.7; cooks launched from editor must be configured to call an external commandlet. "
+				"See comment in EditorExperimentalSettings.h. "
+				"Remove your config setting for EditorPerProjectUserSettings:[/Script/UnrealEd.EditorExperimentalSettings]:bDisableCookInEditor. "
+				"Support will be removed entirely in a future release and the config setting will no longer be used."));
+
 			CookServer = NewObject<UCookOnTheFlyServer>();
 			CookServer->Initialize(ECookMode::CookByTheBookFromTheEditor, BaseCookingFlags);
 		}
