Editor: make legacy bsp conversion to static mesh set the lightmap index appropriately.
PR: #7396

diff --git a/Engine/Source/Editor/UnrealEd/Private/StaticMeshEdit.cpp b/Engine/Source/Editor/UnrealEd/Private/StaticMeshEdit.cpp
--- a/Engine/Source/Editor/UnrealEd/Private/StaticMeshEdit.cpp
+++ b/Engine/Source/Editor/UnrealEd/Private/StaticMeshEdit.cpp
@@ -294,6 +294,10 @@ UStaticMesh* CreateStaticMeshFromBrush(UObject* Outer, FName Name, ABrush* Brush
 		StaticMesh->GetSectionInfoMap().Set(0, SectionIdx, Info);
 		StaticMesh->GetOriginalSectionInfoMap().Set(0, SectionIdx, Info);
 	}
+	// The default settings auto generate lightmaps at channel 1
+	StaticMesh->SetLightMapCoordinateIndex(1);
+	// 64 is a common default, used in import
+	StaticMesh->SetLightMapResolution(64);
 
 	//Set the Imported version before calling the build
 	StaticMesh->SetImportVersion(EImportStaticMeshVersion::LastVersion);
