Editor: Fixes crash bug where, under conditions of reloading a sublevel from which an asset editor had two assets opened from, the toolkit would be referenced after it's destruction.

diff --git a/Engine/Source/Editor/UnrealEd/Private/Toolkits/AssetEditorToolkit.cpp b/Engine/Source/Editor/UnrealEd/Private/Toolkits/AssetEditorToolkit.cpp
--- a/Engine/Source/Editor/UnrealEd/Private/Toolkits/AssetEditorToolkit.cpp
+++ b/Engine/Source/Editor/UnrealEd/Private/Toolkits/AssetEditorToolkit.cpp
@@ -496,7 +496,7 @@ bool FAssetEditorToolkit::CloseWindow()
 
 bool FAssetEditorToolkit::CloseWindow(EAssetEditorCloseReason InCloseReason)
 {
-	if (OnRequestClose(InCloseReason))
+	if (DoesSharedInstanceExist() && OnRequestClose(InCloseReason))
 	{
 		// We are closing, unbind OnRequestClose since we're past that point and we want to make sure we don't redo the request close process when closing the host tab
 		if (TSharedPtr<SStandaloneAssetEditorToolkitHost> StandaloneHostPtr = StandaloneHost.Pin())
