When DlcReevaluateUncookedAssets is specified traverse the package dependencies.
This allows for DLCs to cook editor only package dependencies which might have been skipped in a base cook. Useful for Cooked Editor applications which might depend upon cooked client data.

diff --git a/Engine/Source/Editor/UnrealEd/Private/Cooker/CookRequestCluster.cpp b/Engine/Source/Editor/UnrealEd/Private/Cooker/CookRequestCluster.cpp
--- a/Engine/Source/Editor/UnrealEd/Private/Cooker/CookRequestCluster.cpp
+++ b/Engine/Source/Editor/UnrealEd/Private/Cooker/CookRequestCluster.cpp
@@ -1306,9 +1306,13 @@ void FRequestCluster::FGraphSearch::VisitVertex(FVertexData& Vertex)
 			// session). Both of those questions require expensively and asynchronously fetching dependencies about
 			// the package from the previous cook's oplog. Mark that we need to do that fetch based on whether this
 			// cluster needs to answer those questions about the package.
+			const bool bWasCookedThisSessionOrReevaluateRequested = PlatformData.GetWhereCooked() == EWhereCooked::ThisSession || EnumHasAnyFlags(Cluster.COTFS.CookByTheBookOptions->StartupOptions, ECookByTheBookOptions::DlcReevaluateUncookedAssets);
+			const bool bShouldExploreDependencies = (Cluster.TraversalExploreRuntimeDependencies()
+				&& PlatformData.IsExplorable()
+				&& bWasCookedThisSessionOrReevaluateRequested);
+
 			if ((Cluster.TraversalExploreIncremental() && bPlatformNeedsCook)
-				|| (Cluster.TraversalExploreRuntimeDependencies() && PlatformData.IsExplorable()
-					&& PlatformData.GetWhereCooked() == EWhereCooked::ThisSession))
+				|| bShouldExploreDependencies)
 			{
 				ShouldFetchPlatforms[PlatformIndex] = true;
 				Vertex.GetPlatformData()[PlatformIndex].bExploreRequested = true;
