PR #12834: AnimSequence FBX exporter use USkeleton::FindCompatibleMesh

#rnx

diff --git a/Engine/Source/Editor/UnrealEd/Private/EditorExporters.cpp b/Engine/Source/Editor/UnrealEd/Private/EditorExporters.cpp
--- a/Engine/Source/Editor/UnrealEd/Private/EditorExporters.cpp
+++ b/Engine/Source/Editor/UnrealEd/Private/EditorExporters.cpp
@@ -2462,7 +2462,17 @@ bool UAnimSequenceExporterFBX::ExportBinary( UObject* Object, const TCHAR* Type,
 {
 	UAnimSequence* AnimSequence = CastChecked<UAnimSequence>( Object );
 	USkeleton* AnimSkeleton = AnimSequence->GetSkeleton();
-	USkeletalMesh* PreviewMesh = AnimSkeleton ? AnimSkeleton->GetAssetPreviewMesh(AnimSequence) : NULL;
+	USkeletalMesh* PreviewMesh = nullptr;
+	if (AnimSkeleton)
+	{
+		// first try to get already set preview mesh
+		PreviewMesh = AnimSkeleton->GetAssetPreviewMesh(AnimSequence);
+		if (!PreviewMesh)
+		{
+			// find compatible one (as used in persona toolkit)
+			PreviewMesh = AnimSkeleton->FindCompatibleMesh();
+		}
+	}
 
 	if (AnimSkeleton && PreviewMesh)
 	{
