Fixed activity monitor not being enabled by default

diff --git a/Engine/Config/BaseEditorSettings.ini b/Engine/Config/BaseEditorSettings.ini
--- a/Engine/Config/BaseEditorSettings.ini
+++ b/Engine/Config/BaseEditorSettings.ini
@@ -69,7 +69,7 @@ TriggerPressedThreshold_Rift=0.500000
 Duration=(SubFrame=0.0)
 
 [/Script/UnrealEd.EditorPerformanceSettings]
-bEnableEditorPerformanceTool=True
+bEnableEditorPeformanceTool=True
 bEnableNotifications=True
 NotifyList=Memory_Memory_Pressure
 
diff --git a/Engine/Source/Editor/UnrealEd/Classes/Editor/EditorPerformanceSettings.h b/Engine/Source/Editor/UnrealEd/Classes/Editor/EditorPerformanceSettings.h
--- a/Engine/Source/Editor/UnrealEd/Classes/Editor/EditorPerformanceSettings.h
+++ b/Engine/Source/Editor/UnrealEd/Classes/Editor/EditorPerformanceSettings.h
@@ -35,7 +35,7 @@ class UEditorPerformanceSettings : public UDeveloperSettings
 
 	/** Whether to enable Editor Perf Tool Completely */
 	UPROPERTY(EditAnywhere, config, Category = EditorPerformanceTool, meta = (DisplayName = "Enable Tool", ConfigRestartRequired = true))
-	bool bEnableEditorPeformanceTool = false;
+	bool bEnableEditorPeformanceTool = true;
 
 	/** Whether to enable Editor Perf Warnings */
 	UPROPERTY(EditAnywhere, config, Category = EditorPerformanceTool, meta = (DisplayName = "Enable Notifications", ConfigRestartRequired = false, EditCondition = "bEnableEditorPeformanceTool"))
