Partially backout a previous change to ensure existing node graph comments aren't update and saved with your editor preferences color.

- Reverting the behavior back to the previously default value of White.

#tests opening an existing BP and ensuring the comments are white and not my editor preference color
#rnx

diff --git a/Engine/Source/Editor/UnrealEd/Private/EdGraphNode_Comment.cpp b/Engine/Source/Editor/UnrealEd/Private/EdGraphNode_Comment.cpp
--- a/Engine/Source/Editor/UnrealEd/Private/EdGraphNode_Comment.cpp
+++ b/Engine/Source/Editor/UnrealEd/Private/EdGraphNode_Comment.cpp
@@ -49,7 +49,7 @@ UEdGraphNode_Comment::UEdGraphNode_Comment(const FObjectInitializer& ObjectIniti
 	NodeWidth = 400;
 	NodeHeight = 100;
 	FontSize = 18;
-	CommentColor = GetDefault<UGraphEditorSettings>()->DefaultCommentNodeTitleColor;
+	CommentColor = FLinearColor::White;
 	bColorCommentBubble = false;
 	MoveMode = ECommentBoxMode::GroupMovement;
 
