Add logging for Preview Platform Set/Toggle
Add a FlushRenderingCommands when toggling preview platform  to  avoid overlapping Device Profile switch with rendering frames.

diff --git a/Engine/Source/Editor/UnrealEd/Private/EditorEngine.cpp b/Engine/Source/Editor/UnrealEd/Private/EditorEngine.cpp
--- a/Engine/Source/Editor/UnrealEd/Private/EditorEngine.cpp
+++ b/Engine/Source/Editor/UnrealEd/Private/EditorEngine.cpp
@@ -6494,6 +6494,16 @@ void UEditorEngine::SetPreviewPlatform(const FPreviewPlatformInfo& NewPreviewPla
 {
 	// Get the requested preview platform, make sure it is valid.
 	EShaderPlatform ShaderPlatform = NewPreviewPlatform.ShaderPlatform;
+	if (ShaderPlatform == GMaxRHIShaderPlatform)
+	{
+		UE_LOG(LogEditor, Log, TEXT("Disable Preview Platform"));
+	}
+	else
+	{
+		UE_LOG(LogEditor, Log, TEXT("Set Preview Platform: %s"), 
+			*NewPreviewPlatform.GetFriendlyName().ToString());
+	}
+
 	check(FDataDrivenShaderPlatformInfo::IsValid(ShaderPlatform));
 	ERHIFeatureLevel::Type MaxFeatureLevel = NewPreviewPlatform.PreviewShaderFormatName != NAME_None ? (ERHIFeatureLevel::Type)GetMaxSupportedFeatureLevel(ShaderPlatform) : ERHIFeatureLevel::SM5;
 	check(NewPreviewPlatform.PreviewShaderFormatName.IsNone() || MaxFeatureLevel == NewPreviewPlatform.PreviewFeatureLevel);
@@ -6610,7 +6620,13 @@ void UEditorEngine::SetPreviewPlatform(const FPreviewPlatformInfo& NewPreviewPla
 
 void UEditorEngine::ToggleFeatureLevelPreview()
 {
- 	PreviewPlatform.bPreviewFeatureLevelActive ^= 1;
+	UE_LOG(LogEditor, Log, TEXT("Toggle %s Preview ShaderPlatform: %s"),
+		!PreviewPlatform.bPreviewFeatureLevelActive ? TEXT("ON") : TEXT("OFF"),
+		*PreviewPlatform.GetFriendlyName().ToString());
+
+	FlushRenderingCommands();
+
+	PreviewPlatform.bPreviewFeatureLevelActive ^= 1;
 
 	ERHIFeatureLevel::Type NewPreviewFeatureLevel = PreviewPlatform.GetEffectivePreviewFeatureLevel();
 
