Cherrypick of CL 45576912 to aid with diagnosing Blueprint-related crash reports.

[Blueprint] added log for explicit Blueprint compilation requests.
Can be disabled using 'BP.bLogExplicitCompilationRequests = false'

diff --git a/Engine/Source/Editor/UnrealEd/Private/Kismet2/Kismet2.cpp b/Engine/Source/Editor/UnrealEd/Private/Kismet2/Kismet2.cpp
--- a/Engine/Source/Editor/UnrealEd/Private/Kismet2/Kismet2.cpp
+++ b/Engine/Source/Editor/UnrealEd/Private/Kismet2/Kismet2.cpp
@@ -85,10 +85,21 @@
 #include "Stats/StatsHierarchical.h"
 #include "Settings/EditorStyleSettings.h"
 #include "ToolMenus.h"
+#include "Logging/LogScopedVerbosityOverride.h"
 #include "Subsystems/AssetEditorSubsystem.h"
 
 #define LOCTEXT_NAMESPACE "UnrealEd.Editor"
 
+namespace UE::Private
+{
+	/** Flag to indicate whether a log should be produced when users explicitly request a compilation of a blueprint class */
+	bool bLogExplicitCompilationRequests = true;
+	static FAutoConsoleVariableRef CVarLogExplicitCompilationRequests(
+		TEXT("BP.bLogExplicitCompilationRequests"), bLogExplicitCompilationRequests,
+		TEXT("If true, CompileBlueprint will emit a log indicating the class getting recompiled."),
+		ECVF_Default);
+}
+
 //////////////////////////////////////////////////////////////////////////
 // FArchiveInvalidateTransientRefs
 
@@ -781,7 +792,12 @@ bool FKismetEditorUtilities::IsReferencedByUndoBuffer(UBlueprint* Blueprint)
 void FKismetEditorUtilities::CompileBlueprint(UBlueprint* BlueprintObj, EBlueprintCompileOptions CompileFlags, FCompilerResultsLog* pResults)
 {
 	DECLARE_SCOPE_HIERARCHICAL_COUNTER_FUNC()
-
+	if (UE::Private::bLogExplicitCompilationRequests
+		&& !EnumHasAnyFlags(CompileFlags, EBlueprintCompileOptions::IsRegeneratingOnLoad))
+	{
+		LOG_SCOPE_VERBOSITY_OVERRIDE(LogBlueprint, ELogVerbosity::Log);
+		UE_LOG(LogBlueprint, Log, TEXT("Compiling Blueprint '%s'"), *GetPathNameSafe(BlueprintObj));
+	}
 	FBlueprintCompilationManager::CompileSynchronously(FBPCompileRequest(BlueprintObj, CompileFlags, pResults));
 }
 
