Add new CVar for disabling/enabling fog in editor orthographic view modes.
Lit mode was a more recent fix for ortho rendering, and the default fog settings don't mix well with it completely occluding the scene ineditor debug views.

diff --git a/Engine/Source/Editor/UnrealEd/Private/EditorViewportClient.cpp b/Engine/Source/Editor/UnrealEd/Private/EditorViewportClient.cpp
--- a/Engine/Source/Editor/UnrealEd/Private/EditorViewportClient.cpp
+++ b/Engine/Source/Editor/UnrealEd/Private/EditorViewportClient.cpp
@@ -104,6 +104,12 @@ static TAutoConsoleVariable<int32> CVarAlignedOrthoZoom(
 	TEXT(" 1: All ortho viewport zoom are locked to each other to allow axis lines to be aligned with each other."),
 	ECVF_RenderThreadSafe);
 
+static TAutoConsoleVariable<bool> CVarEditorDisableFogInOrthoDebugViews(
+	TEXT("r.Editor.DisableFogInOrthoDebugViews"),
+	true,
+	TEXT("Fog can obstruct the view when displaying it in ortho debug views without the correct settings adjusted, so this override disables by default."),
+	ECVF_RenderThreadSafe);
+
 static TAutoConsoleVariable<int32> CVarEditorViewportTest(
 	TEXT("r.Test.EditorConstrainedView"),
 	0,
@@ -1464,7 +1470,15 @@ FSceneView* FEditorViewportClient::CalcSceneView(FSceneViewFamily* ViewFamily, c
 				ZScale,
 				ZOffset
 			);
-		}
+
+			/**
+			* Unmanaged fog obscures the Ortho lit view modes so we disable it for these views unless the user overrides the setting.
+			*/
+			if (CVarEditorDisableFogInOrthoDebugViews->GetBool())
+			{
+				ViewFamily->EngineShowFlags.SetFog(false);
+			}
+		};
 
 		if (bConstrainAspectRatio)
 		{
