Fixed issue introduced in CL 30717246 where IsHitIgnoredRenderingThread() now allows primitives with no SceneProxy to pass (returns false), which means non-rendered components are considered for placement. For editor actor positioning, components without a SceneProxy are effectively not renderable; they should be ignored unless bConsiderForActorPlacementWhenHidden is true.

diff --git a/Engine/Source/Editor/UnrealEd/Private/Editor/ObjectPositioning.cpp b/Engine/Source/Editor/UnrealEd/Private/Editor/ObjectPositioning.cpp
--- a/Engine/Source/Editor/UnrealEd/Private/Editor/ObjectPositioning.cpp
+++ b/Engine/Source/Editor/UnrealEd/Private/Editor/ObjectPositioning.cpp
@@ -94,20 +94,22 @@ namespace ObjectPositioningLocals
 		check(IsInParallelRenderingThread());
 
 		const UPrimitiveComponent* PrimitiveComponent = InWeakPrimitiveComponent.Get();
+
+		// Only use this component if it is visible in the specified scene views
+		bool bIsRenderedOnScreen = false;
+		bool bIgnoreTranslucentPrimitive = false;
+		const bool bConsiderInvisibleComponentForPlacement = PrimitiveComponent && PrimitiveComponent->bConsiderForActorPlacementWhenHidden;
+
 		if (PrimitiveComponent && PrimitiveComponent->SceneProxy)
 		{
-			const bool bConsiderInvisibleComponentForPlacement = PrimitiveComponent->bConsiderForActorPlacementWhenHidden;
-
 			// Only use this component if it is visible in the specified scene views
 			const FPrimitiveViewRelevance ViewRelevance = PrimitiveComponent->SceneProxy->GetViewRelevance(&InSceneView);
 			// BSP is a bit special in that its bDrawRelevance is false even when drawn as wireframe because InSceneView.Family->EngineShowFlags.BSPTriangles is off
-			const bool bIsRenderedOnScreen = ViewRelevance.bDrawRelevance || (PrimitiveComponent->IsA(UModelComponent::StaticClass()) && InSceneView.Family->EngineShowFlags.BSP);
-			const bool bIgnoreTranslucentPrimitive = ViewRelevance.HasTranslucency() && !GetDefault<UEditorPerProjectUserSettings>()->bAllowSelectTranslucent;
-		
-			return (!bIsRenderedOnScreen && !bConsiderInvisibleComponentForPlacement) || bIgnoreTranslucentPrimitive;
+			bIsRenderedOnScreen = ViewRelevance.bDrawRelevance || (PrimitiveComponent->IsA(UModelComponent::StaticClass()) && InSceneView.Family->EngineShowFlags.BSP);
+			bIgnoreTranslucentPrimitive = ViewRelevance.HasTranslucency() && !GetDefault<UEditorPerProjectUserSettings>()->bAllowSelectTranslucent;
 		}
 
-		return false;
+		return (!bIsRenderedOnScreen && !bConsiderInvisibleComponentForPlacement) || bIgnoreTranslucentPrimitive;
 	}
 }
 
@@ -203,4 +205,4 @@ UE::Positioning::FObjectPositioningTraceResult UE::Positioning::TraceWorldForPos
 	}
 
 	return Results;
-}
+}
\ No newline at end of file
