Updated UBrowseToAssetOverrideSubsystem to support FTopLevelAssetPaths.

diff --git a/Engine/Source/Editor/LevelInstanceEditor/Private/LevelInstanceEditorModule.cpp b/Engine/Source/Editor/LevelInstanceEditor/Private/LevelInstanceEditorModule.cpp
--- a/Engine/Source/Editor/LevelInstanceEditor/Private/LevelInstanceEditorModule.cpp
+++ b/Engine/Source/Editor/LevelInstanceEditor/Private/LevelInstanceEditorModule.cpp
@@ -928,11 +928,11 @@ void FLevelInstanceEditorModule::RegisterToFirstLevelEditor()
 	}
 }
 
-FName FLevelInstanceEditorModule::GetBrowseToLevelInstanceAsset(const UObject* Object) const
+FTopLevelAssetPath FLevelInstanceEditorModule::GetBrowseToLevelInstanceAsset(const UObject* Object) const
 {
 	// Level instances browse to both the level instance asset and the current level asset by default, while we only want to browse to the former
 	const ILevelInstanceInterface* LevelInstanceInterface = CastChecked<ILevelInstanceInterface>(Object);
-	return FName(LevelInstanceInterface->GetWorldAssetPackage());
+	return LevelInstanceInterface->GetWorldAsset().GetUniqueID().GetAssetPath();
 }
 
 void FLevelInstanceEditorModule::StartupModule()
@@ -1004,7 +1004,7 @@ void FLevelInstanceEditorModule::StartupModule()
 	if (UBrowseToAssetOverrideSubsystem* BrowseToAssetOverrideSubsystem = UBrowseToAssetOverrideSubsystem::Get())
 	{
 		BrowseToAssetOverrideSubsystem->RegisterBrowseToAssetOverrideForInterface<ILevelInstanceInterface>(
-			FBrowseToAssetOverrideDelegate::CreateRaw(this, &FLevelInstanceEditorModule::GetBrowseToLevelInstanceAsset));
+			FGetBrowseToAssetOverrideDelegate::CreateRaw(this, &FLevelInstanceEditorModule::GetBrowseToLevelInstanceAsset));
 	}
 }
 
diff --git a/Engine/Source/Editor/LevelInstanceEditor/Private/LevelInstanceEditorModule.h b/Engine/Source/Editor/LevelInstanceEditor/Private/LevelInstanceEditorModule.h
--- a/Engine/Source/Editor/LevelInstanceEditor/Private/LevelInstanceEditorModule.h
+++ b/Engine/Source/Editor/LevelInstanceEditor/Private/LevelInstanceEditorModule.h
@@ -61,7 +61,7 @@ private:
 		
 	void OnLevelEditorCreated(TSharedPtr<ILevelEditor> InLevelEditor);
 	void RegisterToFirstLevelEditor();
-	FName GetBrowseToLevelInstanceAsset(const UObject* Object) const;
+	FTopLevelAssetPath GetBrowseToLevelInstanceAsset(const UObject* Object) const;
 
 	void RegisterLevelInstanceColumn();
 	void UnregisterLevelInstanceColumn();
diff --git a/Engine/Source/Editor/UnrealEd/Private/AssetSelection.cpp b/Engine/Source/Editor/UnrealEd/Private/AssetSelection.cpp
--- a/Engine/Source/Editor/UnrealEd/Private/AssetSelection.cpp
+++ b/Engine/Source/Editor/UnrealEd/Private/AssetSelection.cpp
@@ -320,7 +320,7 @@ namespace AssetSelectionUtils
 						++ActorInfo.NumSimulationChanges;
 					}
 
-					if (!BrowseToAssetOverrideSubsystem->GetBrowseToAssetOverride(CurrentActor).IsNone())
+					if (!ActorInfo.bHaveBrowseOverride && BrowseToAssetOverrideSubsystem->HasBrowseToAssetOverride(CurrentActor))
 					{
 						ActorInfo.bHaveBrowseOverride = true;
 					}
diff --git a/Engine/Source/Editor/UnrealEd/Private/EditorEngine.cpp b/Engine/Source/Editor/UnrealEd/Private/EditorEngine.cpp
--- a/Engine/Source/Editor/UnrealEd/Private/EditorEngine.cpp
+++ b/Engine/Source/Editor/UnrealEd/Private/EditorEngine.cpp
@@ -3387,18 +3387,11 @@ void UEditorEngine::GetAssetsToSyncToContentBrowser(TArray<FAssetData>& Assets,
 						if (AActor* Actor = ObjectInterface.GetObjectAs<AActor>())
 						{
 							// If BrowseToAssetOverride is set, then use the asset it points to instead of the selected asset
-							const FName BrowseToAssetOverride = BrowseToAssetOverrideSubsystem->GetBrowseToAssetOverride(Actor);
-							if (!BrowseToAssetOverride.IsNone())
+							FAssetData AssetData;
+							if (BrowseToAssetOverrideSubsystem->TryGetBrowseToAssetOverride(Actor, AssetData))
 							{
-								if (IAssetRegistry* AssetRegistry = IAssetRegistry::Get())
-								{
-									TArray<FAssetData> FoundAssets;
-									if (AssetRegistry->GetAssetsByPackageName(BrowseToAssetOverride, FoundAssets) && FoundAssets.Num() > 0)
-									{
-										Assets.Add(FoundAssets[0]);
-										return true;
-									}
-								}
+								Assets.Add(MoveTemp(AssetData));
+								return true;
 							}
 						}
 					}
diff --git a/Engine/Source/Editor/UnrealEd/Private/Subsystems/BrowseToAssetOverrideSubsystem.cpp b/Engine/Source/Editor/UnrealEd/Private/Subsystems/BrowseToAssetOverrideSubsystem.cpp
--- a/Engine/Source/Editor/UnrealEd/Private/Subsystems/BrowseToAssetOverrideSubsystem.cpp
+++ b/Engine/Source/Editor/UnrealEd/Private/Subsystems/BrowseToAssetOverrideSubsystem.cpp
@@ -2,6 +2,7 @@
 
 #include "Subsystems/BrowseToAssetOverrideSubsystem.h"
 
+#include "AssetRegistry/IAssetRegistry.h"
 #include "Editor.h"
 #include "GameFramework/Actor.h"
 
@@ -16,7 +17,13 @@ UBrowseToAssetOverrideSubsystem* UBrowseToAssetOverrideSubsystem::Get()
 	return nullptr;
 }
 
-FName UBrowseToAssetOverrideSubsystem::GetBrowseToAssetOverride(const UObject* Object)
+bool UBrowseToAssetOverrideSubsystem::HasBrowseToAssetOverride(const UObject* Object) const
+{
+	FTopLevelAssetPath AssetPath;
+	return TryGetBrowseToAssetOverride(Object, AssetPath);
+}
+
+bool UBrowseToAssetOverrideSubsystem::TryGetBrowseToAssetOverride(const UObject* Object, FTopLevelAssetPath& OutAssetPath) const
 {
 	// Actors also allow this to be overridden per-instance via meta-data
 	// If set, that takes priority over any per-class overrides
@@ -27,7 +34,8 @@ FName UBrowseToAssetOverrideSubsystem::GetBrowseToAssetOverride(const UObject* O
 		PRAGMA_ENABLE_DEPRECATION_WARNINGS
 		if (!ActorBrowseToAssetOverride.IsEmpty())
 		{
-			return *ActorBrowseToAssetOverride;
+			OutAssetPath = FTopLevelAssetPath(ActorBrowseToAssetOverride);
+			return OutAssetPath.IsValid();
 		}
 	}
 
@@ -36,12 +44,15 @@ FName UBrowseToAssetOverrideSubsystem::GetBrowseToAssetOverride(const UObject* O
 	{
 		for (const UClass* Class = Object->GetClass(); Class; Class = Class->GetSuperClass())
 		{
-			if (const FBrowseToAssetOverrideDelegate* CallbackPtr = PerClassOverrides.Find(Class->GetClassPathName()))
+			if (const FGetBrowseToAssetOverrideDelegate* CallbackPtr = PerClassOverrides.Find(Class->GetClassPathName()))
 			{
-				FName CallbackBrowseToAssetOverride = CallbackPtr->IsBound() ? CallbackPtr->Execute(Object) : FName();
-				if (!CallbackBrowseToAssetOverride.IsNone())
+				if (CallbackPtr->IsBound())
 				{
-					return CallbackBrowseToAssetOverride;
+					OutAssetPath = CallbackPtr->Execute(Object);
+					if (OutAssetPath.IsValid())
+					{
+						return true;
+					}
 				}
 			}
 		}
@@ -54,38 +65,105 @@ FName UBrowseToAssetOverrideSubsystem::GetBrowseToAssetOverride(const UObject* O
 		{
 			for (const FImplementedInterface& Interface : ObjectClass->Interfaces)
 			{
-				if (const FBrowseToAssetOverrideDelegate* CallbackPtr = PerInterfaceOverrides.Find(Interface.Class->GetClassPathName()))
+				if (const FGetBrowseToAssetOverrideDelegate* CallbackPtr = PerInterfaceOverrides.Find(Interface.Class->GetClassPathName()))
 				{
-					FName CallbackBrowseToAssetOverride = CallbackPtr->IsBound() ? CallbackPtr->Execute(Object) : FName();
-					if (!CallbackBrowseToAssetOverride.IsNone())
+					if (CallbackPtr->IsBound())
 					{
-						return CallbackBrowseToAssetOverride;
+						OutAssetPath = CallbackPtr->Execute(Object);
+						if (OutAssetPath.IsValid())
+						{
+							return true;
+						}
 					}
 				}
 			}
 		}
-				
+
 	}
 
+	return false;
+}
+
+bool UBrowseToAssetOverrideSubsystem::TryGetBrowseToAssetOverride(const UObject* Object, FAssetData& OutAssetData) const
+{
+	FTopLevelAssetPath AssetPath;
+	if (TryGetBrowseToAssetOverride(Object, AssetPath))
+	{
+		if (AssetPath.GetAssetName().IsNone())
+		{
+			TArray<FAssetData> AssetDatas;
+			if (IAssetRegistry::GetChecked().GetAssetsByPackageName(AssetPath.GetPackageName(), AssetDatas) && !AssetDatas.IsEmpty())
+			{
+				OutAssetData = MoveTemp(AssetDatas[0]);
+				return true;
+			}
+		}
+		else
+		{
+			OutAssetData = IAssetRegistry::GetChecked().GetAssetByObjectPath(FSoftObjectPath::ConstructFromAssetPath(AssetPath));
+			return OutAssetData.IsValid();
+		}
+	}
+
+	return false;
+}
+
+FName UBrowseToAssetOverrideSubsystem::GetBrowseToAssetOverride(const UObject* Object)
+{
+	FTopLevelAssetPath AssetPath;
+	if (TryGetBrowseToAssetOverride(Object, AssetPath))
+	{
+		return AssetPath.GetPackageName();
+	}
 
 	return FName();
 }
 
-void UBrowseToAssetOverrideSubsystem::RegisterBrowseToAssetOverrideForClass(const FTopLevelAssetPath& Class, FBrowseToAssetOverrideDelegate&& Callback)
+void UBrowseToAssetOverrideSubsystem::RegisterBrowseToAssetOverrideForClass(const FTopLevelAssetPath& Class, FGetBrowseToAssetOverrideDelegate&& Callback)
 {
 	PerClassOverrides.Add(Class, MoveTemp(Callback));
 }
+
+void UBrowseToAssetOverrideSubsystem::RegisterBrowseToAssetOverrideForClass(const FTopLevelAssetPath& Class, FBrowseToAssetOverrideDelegate&& Callback)
+{
+	PerClassOverrides.Add(Class, FGetBrowseToAssetOverrideDelegate::CreateLambda([Callback = MoveTemp(Callback)](const UObject* Object)
+	{
+		FName PackageName;
+
+		if (Callback.IsBound())
+		{
+			PackageName = Callback.Execute(Object);
+		}
+
+		return FTopLevelAssetPath(PackageName, FName());
+	}));
+}
 	
 void UBrowseToAssetOverrideSubsystem::UnregisterBrowseToAssetOverrideForClass(const FTopLevelAssetPath& Class)
 {
 	PerClassOverrides.Remove(Class);
 }
 
-void UBrowseToAssetOverrideSubsystem::RegisterBrowseToAssetOverrideForInterface(const FTopLevelAssetPath& Interface, FBrowseToAssetOverrideDelegate&& Callback)
+void UBrowseToAssetOverrideSubsystem::RegisterBrowseToAssetOverrideForInterface(const FTopLevelAssetPath& Interface, FGetBrowseToAssetOverrideDelegate&& Callback)
 {
 	PerInterfaceOverrides.Add(Interface, MoveTemp(Callback));
 }
 
+void UBrowseToAssetOverrideSubsystem::RegisterBrowseToAssetOverrideForInterface(const FTopLevelAssetPath& Interface, FBrowseToAssetOverrideDelegate&& Callback)
+{
+	PerInterfaceOverrides.Add(Interface, FGetBrowseToAssetOverrideDelegate::CreateLambda([Callback = MoveTemp(Callback)](const UObject* Object)
+	{
+		FName PackageName;
+
+		if (Callback.IsBound())
+		{
+			PackageName = Callback.Execute(Object);
+		}
+
+		return FTopLevelAssetPath(PackageName, FName());
+	}));
+}
+
 void UBrowseToAssetOverrideSubsystem::UnregisterBrowseToAssetOverrideForInterface(const FTopLevelAssetPath& Interface)
 {
 	PerInterfaceOverrides.Remove(Interface);
diff --git a/Engine/Source/Editor/UnrealEd/Public/Subsystems/BrowseToAssetOverrideSubsystem.h b/Engine/Source/Editor/UnrealEd/Public/Subsystems/BrowseToAssetOverrideSubsystem.h
--- a/Engine/Source/Editor/UnrealEd/Public/Subsystems/BrowseToAssetOverrideSubsystem.h
+++ b/Engine/Source/Editor/UnrealEd/Public/Subsystems/BrowseToAssetOverrideSubsystem.h
@@ -5,11 +5,16 @@
 #include "Containers/Map.h"
 #include "Delegates/Delegate.h"
 #include "EditorSubsystem.h"
+#include "UObject/TopLevelAssetPath.h"
 
 #include "BrowseToAssetOverrideSubsystem.generated.h"
 
+struct FAssetData;
+
 using FBrowseToAssetOverrideDelegate = TDelegate<FName(const UObject*)>;
 
+using FGetBrowseToAssetOverrideDelegate = TDelegate<FTopLevelAssetPath(const UObject*)>;
+
 UCLASS(MinimalAPI)
 class UBrowseToAssetOverrideSubsystem : public UEditorSubsystem
 {
@@ -18,12 +23,42 @@ class UBrowseToAssetOverrideSubsystem : public UEditorSubsystem
 public:
 	static UNREALED_API UBrowseToAssetOverrideSubsystem* Get();
 
+	/**
+	 * Given an object, see if it has a "browse to asset" override.
+	 * @return Returns true if there is an override.
+	 */
+	UNREALED_API bool HasBrowseToAssetOverride(const UObject* Object) const;
+
+	/**
+	 * Given an object, see if it has a "browse to asset" override.
+	 * @return Returns true if there is an override.
+	 */
+	UNREALED_API bool TryGetBrowseToAssetOverride(const UObject* Object, FTopLevelAssetPath& OutAssetPath) const;
+
+	/**
+	 * Given an object, see if it has a "browse to asset" override.
+	 * @return Returns true if there is an override.
+	 */
+	UNREALED_API bool TryGetBrowseToAssetOverride(const UObject* Object, FAssetData& OutAssetData) const;
+
 	/**
 	 * Given an object, see if it has a "browse to asset" package name override.
 	 * @return The package name of override, or None if there is no override.
 	 */
 	UNREALED_API FName GetBrowseToAssetOverride(const UObject* Object);
 
+	/**
+	 * Register a per-class override for the "browse to asset" resolution.
+	 * @note Callback should return a valid top level asset path if there is an override.
+	 */
+	UNREALED_API void RegisterBrowseToAssetOverrideForClass(const FTopLevelAssetPath& Class, FGetBrowseToAssetOverrideDelegate&& Callback);
+
+	template <typename ClassType>
+	void RegisterBrowseToAssetOverrideForClass(FGetBrowseToAssetOverrideDelegate&& Callback)
+	{
+		RegisterBrowseToAssetOverrideForClass(ClassType::StaticClass()->GetClassPathName(), MoveTemp(Callback));
+	}
+
 	/**
 	 * Register a per-class override for the "browse to asset" resolution.
 	 * @note Callback should return a package name, or None if there is no override.
@@ -47,6 +82,18 @@ public:
 		UnregisterBrowseToAssetOverrideForClass(ClassType::StaticClass()->GetClassPathName());
 	}
 
+	/**
+	 * Register a per-interface override for the "browse to asset" resolution.
+	 * @note Callback should return a valid top level asset path if there is an override.
+	 */
+	UNREALED_API void RegisterBrowseToAssetOverrideForInterface(const FTopLevelAssetPath& Interface, FGetBrowseToAssetOverrideDelegate&& Callback);
+
+	template <typename InterfaceType>
+	void RegisterBrowseToAssetOverrideForInterface(FGetBrowseToAssetOverrideDelegate&& Callback)
+	{
+		RegisterBrowseToAssetOverrideForInterface(InterfaceType::UClassType::StaticClass()->GetClassPathName(), MoveTemp(Callback));
+	}
+
 	/**
 	 * Register a per-interface override for the "browse to asset" resolution.
 	 * @note Callback should return a package name, or None if there is no override.
@@ -71,6 +118,6 @@ public:
 	}
 	
 private:
-	TMap<FTopLevelAssetPath, FBrowseToAssetOverrideDelegate> PerClassOverrides;
-	TMap<FTopLevelAssetPath, FBrowseToAssetOverrideDelegate> PerInterfaceOverrides;
+	TMap<FTopLevelAssetPath, FGetBrowseToAssetOverrideDelegate> PerClassOverrides;
+	TMap<FTopLevelAssetPath, FGetBrowseToAssetOverrideDelegate> PerInterfaceOverrides;
 };
