Deprecated package name functions in UBrowseToAssetOverrideSubsystem in favor of FTopLevelAssetPath functions.

diff --git a/Engine/Source/Editor/UnrealEd/Public/Subsystems/BrowseToAssetOverrideSubsystem.h b/Engine/Source/Editor/UnrealEd/Public/Subsystems/BrowseToAssetOverrideSubsystem.h
--- a/Engine/Source/Editor/UnrealEd/Public/Subsystems/BrowseToAssetOverrideSubsystem.h
+++ b/Engine/Source/Editor/UnrealEd/Public/Subsystems/BrowseToAssetOverrideSubsystem.h
@@ -45,6 +45,7 @@ public:
 	 * Given an object, see if it has a "browse to asset" package name override.
 	 * @return The package name of override, or None if there is no override.
 	 */
+	UE_DEPRECATED(5.8, "Call TryGetBrowseToAssetOverride instead.")
 	UNREALED_API FName GetBrowseToAssetOverride(const UObject* Object);
 
 	/**
@@ -63,9 +64,11 @@ public:
 	 * Register a per-class override for the "browse to asset" resolution.
 	 * @note Callback should return a package name, or None if there is no override.
 	 */
+	UE_DEPRECATED(5.8, "Use FGetBrowseToAssetOverrideDelegate overload instead.")
 	UNREALED_API void RegisterBrowseToAssetOverrideForClass(const FTopLevelAssetPath& Class, FBrowseToAssetOverrideDelegate&& Callback);
 
 	template <typename ClassType>
+	UE_DEPRECATED(5.8, "Use FGetBrowseToAssetOverrideDelegate overload instead.")
 	void RegisterBrowseToAssetOverrideForClass(FBrowseToAssetOverrideDelegate&& Callback)
 	{
 		RegisterBrowseToAssetOverrideForClass(ClassType::StaticClass()->GetClassPathName(), MoveTemp(Callback));
@@ -98,9 +101,11 @@ public:
 	 * Register a per-interface override for the "browse to asset" resolution.
 	 * @note Callback should return a package name, or None if there is no override.
 	 */
+	UE_DEPRECATED(5.8, "Use FGetBrowseToAssetOverrideDelegate overload instead.")
 	UNREALED_API void RegisterBrowseToAssetOverrideForInterface(const FTopLevelAssetPath& Interface, FBrowseToAssetOverrideDelegate&& Callback);
 
 	template <typename InterfaceType>
+	UE_DEPRECATED(5.8, "Use FGetBrowseToAssetOverrideDelegate overload instead.")
 	void RegisterBrowseToAssetOverrideForInterface(FBrowseToAssetOverrideDelegate&& Callback)
 	{
 		RegisterBrowseToAssetOverrideForInterface(InterfaceType::UClassType::StaticClass()->GetClassPathName(), MoveTemp(Callback));
