[ThumbnailHelper]

Changes TClassInstanceThumbnailScene to no longer clear all thumbnails when it's max threshold is hit and instead:
- The least recently used thumbnails are enqueued for async deletion
- Async deletion only deletes thumbnails scenes in batches that try to take no more than 250ms (since this code needs to run on the game thread)
- Async deletion will occur every couple seconds until all enqueued scenes are deleted

This reduces what can be 12s+ spikes on the game thread while browsing through the content browser to be much lower on average (tens of ms). The large destruction expense is cleaningup the world and since this does a global scan of objects this time can be especially bad when object dense maps are loaded.

#rnx

diff --git a/Engine/Source/Editor/UnrealEd/Public/ThumbnailHelpers.h b/Engine/Source/Editor/UnrealEd/Public/ThumbnailHelpers.h
--- a/Engine/Source/Editor/UnrealEd/Public/ThumbnailHelpers.h
+++ b/Engine/Source/Editor/UnrealEd/Public/ThumbnailHelpers.h
@@ -4,6 +4,8 @@
 #pragma once
 
 #include "CoreMinimal.h"
+#include "Containers/MRUArray.h"
+#include "Editor.h"
 #include "UObject/ObjectMacros.h"
 #include "GameFramework/Actor.h"
 #include "Engine/Blueprint.h"
@@ -357,6 +359,7 @@ public:
 	TClassInstanceThumbnailScene()
 	{
 		InstancedThumbnailScenes.Reserve(MaxNumScenes);
+		MRUClassNames.Reserve(MaxNumScenes);
 	}
 
 	/** Find an existing thumbnail scene instance for this class type. */
@@ -365,7 +368,13 @@ public:
 		check(InClass);
 		const FName ClassName = InClass->GetFName();
 
-		return InstancedThumbnailScenes.FindRef(ClassName);
+		TSharedPtr<ThumbnailSceneType> Scene = InstancedThumbnailScenes.FindRef(ClassName);
+		if (Scene.IsValid())
+		{
+			MRUClassNames.AddUnique(ClassName);
+		}
+
+		return Scene;
 	}
 
 	/** Find or create a thumbnail scene instance for this class type. */
@@ -379,13 +388,56 @@ public:
 		{
 			if (InstancedThumbnailScenes.Num() >= MaxNumScenes)
 			{
-				InstancedThumbnailScenes.Reset();
-				// Will hitch but is better than a crash
-				CollectGarbage(GARBAGE_COLLECTION_KEEPFLAGS);
+				FName ClassNameToRemove = MRUClassNames.Pop();
+				TSharedPtr<ThumbnailSceneType> SceneToDelete;
+				InstancedThumbnailScenes.RemoveAndCopyValue(ClassNameToRemove, SceneToDelete);
+				ScenesToDelete.Add(MoveTemp(SceneToDelete));
+
+				// Now that we have scenes to delete, schedule them to be destroyed in batches periodically until we have none left.
+				// Destroying many scenes at once can result in large stalls when iterating through content so we attempt to 
+				// smooth that out here by deferring destruction, doing destruction in small, timed, batches and destroying the 
+				// least recently used first to avoid deleting something we are likely to immediately recreate
+				const int32 MinBatchSize = 8;
+				if (!SceneDeleteTimerHandle.IsValid() && ScenesToDelete.Num() >= MinBatchSize)
+				{
+					checkf(IsInGameThread(), TEXT("Must be running on the game thread."));
+					GEditor->GetTimerManager()->SetTimer(SceneDeleteTimerHandle, FTimerDelegate::CreateLambda([this]() mutable
+						{
+							checkf(IsInGameThread(), TEXT("Must be running on the game thread."));
+							const double TimeoutSeconds = 0.25;
+
+							int32 Iteration = 0;
+							double StartTime = FPlatformTime::Seconds();
+							while (!ScenesToDelete.IsEmpty())
+							{
+								if (++Iteration > MinBatchSize)
+								{
+									if ((FPlatformTime::Seconds() - StartTime) > TimeoutSeconds)
+									{
+										break;
+									}
+									Iteration = 1;
+								}
+
+								// The destruction of these scenes can be quite expensive if the world has many objects
+								ScenesToDelete.Pop(EAllowShrinking::No);
+							}
+
+							CollectGarbage(GARBAGE_COLLECTION_KEEPFLAGS);
+
+							if (ScenesToDelete.IsEmpty())
+							{
+								GEditor->GetTimerManager()->ClearTimer(SceneDeleteTimerHandle);
+								ScenesToDelete.Shrink();
+								SceneDeleteTimerHandle.Invalidate();
+							}
+						}), 2.0f/*InRate*/, true /*bInLoop*/, 1.0f/*InFirstDelay*/);
+				}
 			}
 
 			ExistingThumbnailScene = MakeShareable(new ThumbnailSceneType());
 			InstancedThumbnailScenes.Add(ClassName, ExistingThumbnailScene);
+			MRUClassNames.Add(ClassName);
 		}
 
 		return ExistingThumbnailScene.ToSharedRef();
@@ -395,13 +447,16 @@ public:
 	void RemoveThumbnailScene(const UClass* InClass)
 	{
 		check(InClass);
-		InstancedThumbnailScenes.Remove(InClass->GetFName());
+		const FName ClassName = InClass->GetFName();
+		InstancedThumbnailScenes.Remove(ClassName);
+		MRUClassNames.RemoveSingle(ClassName);
 	}
 
 	/** Clears all thumbnail scenes */
 	void Clear()
 	{
 		InstancedThumbnailScenes.Reset();
+		MRUClassNames.Reset();
 	}
 
 private:
@@ -410,6 +465,13 @@ private:
 	 * @note This uses the class name rather than the class pointer to avoid leaving behind stale class instances as Blueprints are re-compiled.
 	 */
 	TMap<FName, TSharedPtr<ThumbnailSceneType>> InstancedThumbnailScenes;
+
+	/** Tracks the most recently used class names (most recent in the front) so we may prefer deleting the least recently used thumbnail scenes	*/
+	mutable TMRUArray<FName, TFixedAllocator<MaxNumScenes>> MRUClassNames;
+	/** Thumbnail scenes enqueued for deletion. Must only be modified by the game thread. */
+	TArray<TSharedPtr<ThumbnailSceneType>> ScenesToDelete;
+	/** Handle to time to periodically delete batches of enqueued scenes for deletion. */
+	FTimerHandle SceneDeleteTimerHandle;
 };
 
 
