MPCook: Add earlier diagnostic for InternalCookerError "Package %s was retracted from %s, but already has no platforms needing commit...". Validate the contract during every package transition. Also mark the existing warning as InternalCookerError.
#rnx

diff --git a/Engine/Source/Editor/UnrealEd/Private/Cooker/CookDirector.cpp b/Engine/Source/Editor/UnrealEd/Private/Cooker/CookDirector.cpp
--- a/Engine/Source/Editor/UnrealEd/Private/Cooker/CookDirector.cpp
+++ b/Engine/Source/Editor/UnrealEd/Private/Cooker/CookDirector.cpp
@@ -2328,7 +2328,7 @@ FCookDirector::FRetractionHandler::ReassignPackages(const FWorkerId& FromWorker,
 					PlatformPair.Value.IsCommitted() ? TEXT("true") : TEXT("false"));
 			}
 			UE_LOG(LogCook, Warning,
-				TEXT("Package %s was retracted from %s, but already has no platforms needing commit. Demoting the package to idle instead of reassinging. State = %s. CookPhase = %s. WorkerAssignment = %s. [ %s ]"),
+				TEXT("InternalCookerError: Package %s was retracted from %s, but already has no platforms needing commit. Demoting the package to idle instead of reassinging. State = %s. CookPhase = %s. WorkerAssignment = %s. [ %s ]"),
 				*PackageData->GetPackageName().ToString(), *Director.GetDisplayName(FromWorker), LexToString(PackageData->GetState()),
 				LexToString(CookPhase), *Director.GetDisplayName(PackageData->GetWorkerAssignment()), *PlatformDataText);
 			PackageData->SetWorkerAssignment(FWorkerId::Invalid());
diff --git a/Engine/Source/Editor/UnrealEd/Private/Cooker/CookPackageData.cpp b/Engine/Source/Editor/UnrealEd/Private/Cooker/CookPackageData.cpp
--- a/Engine/Source/Editor/UnrealEd/Private/Cooker/CookPackageData.cpp
+++ b/Engine/Source/Editor/UnrealEd/Private/Cooker/CookPackageData.cpp
@@ -810,6 +810,44 @@ void FPackageData::SendToState(EPackageState NextState, ESendFlags SendFlags, ES
 		}
 	}
 
+	if (NextState != EPackageState::Idle && NextState != EPackageState::Request)
+	{
+		// Add more diagnostics for an unexpected edge case we have seen: retracted packages already have been
+		// committed. In-progress packages after the request state have a contract that they have a platform
+		// needing commit; give a warning if this contract is broken during a state transition.
+		int32 NumNeedCommitPlatforms = 0;
+		ECookPhase CookPhase = PackageDatas.GetCookOnTheFlyServer().GetCookPhase();
+		for (const TPair<const ITargetPlatform*, FPackagePlatformData>& PlatformPair : GetPlatformDatas())
+		{
+			if (PlatformPair.Key == CookerLoadingPlatformKey)
+			{
+				continue;
+			}
+			if (PlatformPair.Value.NeedsCommit(PlatformPair.Key, CookPhase))
+			{
+				++NumNeedCommitPlatforms;
+			}
+		}
+		if (NumNeedCommitPlatforms == 0)
+		{
+			TStringBuilder<256> PlatformDataText;
+			for (const TPair<const ITargetPlatform*, FPackagePlatformData>& PlatformPair : GetPlatformDatas())
+			{
+				PlatformDataText.Appendf(TEXT("{ %s: Reachable=%s, Committed=%s }, "),
+					PlatformPair.Key == CookerLoadingPlatformKey ? TEXT("CookerLoadingPlatform") : *PlatformPair.Key->PlatformName(),
+					PlatformPair.Value.IsReachable(CookPhase == ECookPhase::Cook
+						? EReachability::Runtime : EReachability::Build) ? TEXT("true") : TEXT("false"),
+					PlatformPair.Value.IsCommitted() ? TEXT("true") : TEXT("false"));
+			}
+			UE_LOG(LogCook, Warning,
+				TEXT("InternalCookerError: Package %s is entering InProgress state %s, but already has no platforms needing commit. CookPhase = %s. [ %s ]"),
+				*GetPackageName().ToString(), LexToString(NextState), LexToString(CookPhase), *PlatformDataText);
+			UE_CALL_ONCE([]()
+				{
+					FDebug::DumpStackTraceToLog(ELogVerbosity::Display);
+				});
+		}
+	}
 
 	SetState(NextState);
 	switch (NextState)
