* Use the class display name when showing the asset type in the generic asset viewer (as we do in the blueprint editor)
* The tooltip now for the "go to source" hyperlink shows the native class name for both the blueprint editor and asset viewer

diff --git a/Engine/Source/Editor/Kismet/Private/BlueprintEditor.cpp b/Engine/Source/Editor/Kismet/Private/BlueprintEditor.cpp
--- a/Engine/Source/Editor/Kismet/Private/BlueprintEditor.cpp
+++ b/Engine/Source/Editor/Kismet/Private/BlueprintEditor.cpp
@@ -2618,6 +2618,8 @@ void FBlueprintEditor::PostRegenerateMenusAndToolbars()
 
 	if ( BluePrint && !FBlueprintEditorUtils::IsLevelScriptBlueprint( BluePrint ) )
 	{
+		const FText NativeClassNameText = GetParentClassNativeNameText();
+
 		// build and attach the menu overlay
 		TSharedRef<SHorizontalBox> MenuOverlayBox = SNew(SHorizontalBox)
 			+SHorizontalBox::Slot()
@@ -2691,7 +2693,8 @@ void FBlueprintEditor::PostRegenerateMenusAndToolbars()
 				.Visibility(this, &FBlueprintEditor::GetNativeParentClassButtonsVisibility)
 				.OnNavigate(this, &FBlueprintEditor::OnEditParentClassNativeCodeClicked)
 				.Text(this, &FBlueprintEditor::GetTextForNativeParentClassHeaderLink)
-				.ToolTipText(FText::Format(LOCTEXT("GoToCode_ToolTip", "Click to open this source file in {0}"), FSourceCodeNavigation::GetSelectedSourceCodeIDE()))
+				.ToolTipText(FText::Format(LOCTEXT("GoToCode_ToolTip", "Click to open the source file for the underlying class ({0}) in {1}"), 
+					NativeClassNameText, FSourceCodeNavigation::GetSelectedSourceCodeIDE()))
 			]
 			+ SHorizontalBox::Slot()
 			.AutoWidth()
@@ -2711,6 +2714,12 @@ FText FBlueprintEditor::GetParentClassNameText() const
 	return (ParentClass != nullptr) ? ParentClass->GetDisplayNameText() : LOCTEXT("BlueprintEditor_NoParentClass", "None");
 }
 
+FText FBlueprintEditor::GetParentClassNativeNameText() const
+{
+	UClass* ParentClass = (GetBlueprintObj() != nullptr) ? *GetBlueprintObj()->ParentClass : nullptr;
+	return (ParentClass != nullptr) ? FText::FromName(ParentClass->GetFName()) : LOCTEXT("BlueprintEditor_NoParentClass", "None");
+}
+
 bool FBlueprintEditor::IsParentClassABlueprint() const
 {
 	return FBlueprintEditorUtils::IsParentClassABlueprint(GetBlueprintObj());
@@ -2777,7 +2786,7 @@ void FBlueprintEditor::OnEditParentClassNativeCodeClicked()
 FText FBlueprintEditor::GetTextForNativeParentClassHeaderLink() const
 {
 	// it could be done using FSourceCodeNavigation, but it could be slow
-	return FText::FromString(*GetParentClassNameText().ToString());
+	return GetParentClassNameText();
 }
 
 FReply FBlueprintEditor::OnFindParentClassInContentBrowserClicked()
diff --git a/Engine/Source/Editor/Kismet/Public/BlueprintEditor.h b/Engine/Source/Editor/Kismet/Public/BlueprintEditor.h
--- a/Engine/Source/Editor/Kismet/Public/BlueprintEditor.h
+++ b/Engine/Source/Editor/Kismet/Public/BlueprintEditor.h
@@ -846,6 +846,9 @@ protected:
 	/** Returns the name of the Blueprint's parent class */
 	UE_API FText GetParentClassNameText() const;
 
+	/** Returns the native name of the Blueprint's parent class */
+	UE_API FText GetParentClassNativeNameText() const;
+
 	/** Handler for "Find parent class in CB" button */
 	UE_API FReply OnFindParentClassInContentBrowserClicked();
 
diff --git a/Engine/Source/Editor/UnrealEd/Private/EditorClassUtils.cpp b/Engine/Source/Editor/UnrealEd/Private/EditorClassUtils.cpp
--- a/Engine/Source/Editor/UnrealEd/Private/EditorClassUtils.cpp
+++ b/Engine/Source/Editor/UnrealEd/Private/EditorClassUtils.cpp
@@ -233,21 +233,22 @@ TSharedRef<SWidget> FEditorClassUtils::GetSourceLink(const UClass* Class, const
 			return false;
 		};
 
-		FText ClassNameText = FText::FromName(Class->GetFName());
+		const FText ClassDisplayNameText = Class->GetDisplayNameText();
+		const FText ClassNameText = FText::FromName(Class->GetFName());
 
 		FText FormattedText;
 		if (Params.CodeFormat)
 		{
-			FormattedText = FText::Format(*Params.CodeFormat, ClassNameText);
+			FormattedText = FText::Format(*Params.CodeFormat, ClassDisplayNameText);
 		}
 		else if (Params.bUseDefaultFormat)
 		{
 			const FText DefaultCodeFormat = NSLOCTEXT("SourceHyperlink", "GoToCode", "Open {0}");
-			FormattedText = FText::Format(DefaultCodeFormat, ClassNameText);
+			FormattedText = FText::Format(DefaultCodeFormat, ClassDisplayNameText);
 		}
 		else
 		{
-			FormattedText = ClassNameText;
+			FormattedText = ClassDisplayNameText;
 		}
 
 		TSharedRef<SWidget> NoLinkWidget = SNullWidget::NullWidget;
@@ -259,7 +260,7 @@ TSharedRef<SWidget> FEditorClassUtils::GetSourceLink(const UClass* Class, const
 		else
 		{
 			NoLinkWidget = SNew(STextBlock)
-				.Text(Params.bUseFormatIfNoLink ? FormattedText : ClassNameText)
+				.Text(Params.bUseFormatIfNoLink ? FormattedText : ClassDisplayNameText)
 				.ColorAndOpacity(FSlateColor::UseSubduedForeground())
 				.Visibility_Lambda([CanNavigateToClassCode]() { return CanNavigateToClassCode() ? EVisibility::Collapsed : EVisibility::Visible; });
 		}
@@ -273,7 +274,8 @@ TSharedRef<SWidget> FEditorClassUtils::GetSourceLink(const UClass* Class, const
 				.Style(FAppStyle::Get(), "Common.GotoNativeCodeHyperlink")
 				.OnNavigate_Lambda(OnNavigateToClassCode)
 				.Text(FormattedText)
-				.ToolTipText(FText::Format(NSLOCTEXT("SourceHyperlink", "GoToCode_ToolTip", "Click to open this source file in {0}"), FSourceCodeNavigation::GetSelectedSourceCodeIDE()))
+				.ToolTipText(FText::Format(NSLOCTEXT("SourceHyperlink", "GoToCode_ToolTip", "Click to open the source file for the underlying class ({0}) in {1}"), 
+					ClassNameText, FSourceCodeNavigation::GetSelectedSourceCodeIDE()))
 				.Visibility_Lambda([CanNavigateToClassCode]() { return CanNavigateToClassCode() ? EVisibility::Visible : EVisibility::Collapsed; })
 			]
 			+ SHorizontalBox::Slot()
