Clamp ?Play Number Of Clients? to MAX_PIE_INSTANCES to prevent PIE assert

Problem

Editor crashed with Assertion failed: PlayInEditorID < 10 when Level Editor > Play > Client > Play Number Of Clients exceeded 10 (MAX_PIE_INSTANCES)

Impact: setting persisted >10 caused immediate assert/crash on PIE

Solution

Clamp the Editor Preference UProperty for Play Number Of Clients to a maximum of 10

Behavior now:

Values >10 are automatically clamped to 10 in UI/config

?10 unchanged; PIE launches safely without assertions

Notes

No functional changes beyond input validation; config values above 10 are coerced at load/edit time

Risk: Low (bounds check only)

diff --git a/Engine/Source/Editor/UnrealEd/Classes/Settings/LevelEditorPlaySettings.h b/Engine/Source/Editor/UnrealEd/Classes/Settings/LevelEditorPlaySettings.h
--- a/Engine/Source/Editor/UnrealEd/Classes/Settings/LevelEditorPlaySettings.h
+++ b/Engine/Source/Editor/UnrealEd/Classes/Settings/LevelEditorPlaySettings.h
@@ -463,6 +463,8 @@ private:
 
 private:
 	UNREALED_API void PushDebugDrawingSettings();
+	
+	void ClampPIEClientValues();
 
 public:
 	bool ShowServerDebugDrawingByDefault() const
diff --git a/Engine/Source/Editor/UnrealEd/Private/Settings/SettingsClasses.cpp b/Engine/Source/Editor/UnrealEd/Private/Settings/SettingsClasses.cpp
--- a/Engine/Source/Editor/UnrealEd/Private/Settings/SettingsClasses.cpp
+++ b/Engine/Source/Editor/UnrealEd/Private/Settings/SettingsClasses.cpp
@@ -535,6 +535,12 @@ void ULevelEditorPlaySettings::PushDebugDrawingSettings()
 #endif
 }
 
+void ULevelEditorPlaySettings::ClampPIEClientValues()
+{
+	PlayNumberOfClients = FMath::Clamp(PlayNumberOfClients, 1, MAX_PIE_INSTANCES);
+	PrimaryPIEClientIndex = FMath::Clamp(PrimaryPIEClientIndex, INDEX_NONE, PlayNumberOfClients - 1);
+}
+
 void FPlayScreenResolution::PostInitProperties()
 {
 	ScaleFactor = 1.0f;
@@ -591,6 +597,11 @@ void ULevelEditorPlaySettings::PostEditChangeProperty(struct FPropertyChangedEve
 		}
 	}
 
+	if (PropertyName == GET_MEMBER_NAME_CHECKED(ULevelEditorPlaySettings, PlayNumberOfClients) || PropertyName == GET_MEMBER_NAME_CHECKED(ULevelEditorPlaySettings, PrimaryPIEClientIndex))
+	{
+		ClampPIEClientValues();
+	}
+
 	Super::PostEditChangeProperty(PropertyChangedEvent);
 }
 
@@ -638,6 +649,8 @@ void ULevelEditorPlaySettings::PostInitProperties()
 	}
 
 	PushDebugDrawingSettings();
+
+	ClampPIEClientValues();
 }
 
 bool ULevelEditorPlaySettings::CanEditChange(const FProperty* InProperty) const
diff --git a/Engine/Source/Runtime/CoreUObject/Private/UObject/LazyObjectPtr.cpp b/Engine/Source/Runtime/CoreUObject/Private/UObject/LazyObjectPtr.cpp
--- a/Engine/Source/Runtime/CoreUObject/Private/UObject/LazyObjectPtr.cpp
+++ b/Engine/Source/Runtime/CoreUObject/Private/UObject/LazyObjectPtr.cpp
@@ -14,7 +14,6 @@
 /** Annotation associating objects with their guids **/
 static FUObjectAnnotationSparseSearchable<FUniqueObjectGuid,true> GuidAnnotation;
 
-#define MAX_PIE_INSTANCES 10
 static TMap<FGuid, FGuid> PIEGuidMap[MAX_PIE_INSTANCES];
 
 static FGuid RemapGuid(const FUniqueObjectGuid& Guid, int32 PIEInstanceID)
diff --git a/Engine/Source/Runtime/CoreUObject/Public/UObject/LazyObjectPtr.h b/Engine/Source/Runtime/CoreUObject/Public/UObject/LazyObjectPtr.h
--- a/Engine/Source/Runtime/CoreUObject/Public/UObject/LazyObjectPtr.h
+++ b/Engine/Source/Runtime/CoreUObject/Public/UObject/LazyObjectPtr.h
@@ -26,6 +26,8 @@ class UObject;
 template <typename T> struct TIsPODType;
 template <typename T> struct TIsWeakPointerType;
 
+#define MAX_PIE_INSTANCES 10
+
 /**
  * Wrapper structure for a GUID that uniquely identifies registered UObjects.
  * The actual GUID is stored in an object annotation that is updated when a new reference is made.
