[Editor] Editor Primitive occlusion exposed as editor settings - dither pattern size, opacity, and occluded color multiplier

diff --git a/Engine/Shaders/Private/PostProcessCompositePrimitives.usf b/Engine/Shaders/Private/PostProcessCompositePrimitives.usf
--- a/Engine/Shaders/Private/PostProcessCompositePrimitives.usf
+++ b/Engine/Shaders/Private/PostProcessCompositePrimitives.usf
@@ -107,6 +107,7 @@ float2 DepthTextureJitter;
 uint bCameraCut;
 uint bProcessAlpha;
 
+float OccludedDitheringPatternSize;
 float OccludedDithering;
 float OccludedBrightness;
 
@@ -443,8 +444,11 @@ void MainCompositePostProcessPrimitivesPS(
 		EditorPrimitiveColor.rgb *= EditorPrimitiveColor.a;
 	}
 
-	// Generate 2x2 square tiling pattern for foreground primitive that end up behind scene opaque primitives.
-	float PatternMask = ((EditorPrimitivesPixelPos.x/2 + EditorPrimitivesPixelPos.y/2) % 2) * 0.7f;
+	// Generate NxN square tiling pattern for foreground primitive that end up behind scene opaque primitives.
+	float PatternMask = fmod(
+		floor(EditorPrimitivesPixelPos.x / OccludedDitheringPatternSize)
+		+ floor(EditorPrimitivesPixelPos.y / OccludedDitheringPatternSize),
+		2.0f) * OccludedBrightness;
 
 	// the contants express the two opacity values we see when the primitive is hidden
 	float LowContrastPatternMask = lerp(1.0 - OccludedDithering, 1.0f, PatternMask);
diff --git a/Engine/Source/Editor/UnrealEd/Classes/Settings/EditorStyleSettings.h b/Engine/Source/Editor/UnrealEd/Classes/Settings/EditorStyleSettings.h
--- a/Engine/Source/Editor/UnrealEd/Classes/Settings/EditorStyleSettings.h
+++ b/Engine/Source/Editor/UnrealEd/Classes/Settings/EditorStyleSettings.h
@@ -96,6 +96,18 @@ public:
 	UPROPERTY(EditAnywhere, config, Category=UserInterface, meta=(DisplayName="Additional Viewport Selection Colors"))
 	FLinearColor AdditionalSelectionColors[6];
 
+	/** Size (in pixels) of the dithering pattern applied to occluded primitives. */
+	UPROPERTY(EditAnywhere, config, Category = UserInterface, meta = (ClampMin = 0.25, UIMin = 0.25))
+	float OccludedDitheringPatternSize = 2.0f;
+
+	/** The strength of the dithering applied to occluded primitives. */
+	UPROPERTY(EditAnywhere, config, Category = UserInterface, meta = (ClampMin = 0.0, UIMin = 0.0, ClampMax = 1.0, UIMax = 1.0, DisplayName = "Occluded Pattern Opacity"))
+	float OccludedDithering = 0.8f;
+
+	/** The brightness of occluded primitives. */
+	UPROPERTY(EditAnywhere, config, Category = UserInterface, meta = (ClampMin = 0.0, UIMin = 0.0, ClampMax = 1.0, UIMax = 1.0, DisplayName = "Occluded Primitive Brightness"))
+	float OccludedBrightness = 0.7f;
+
 	/** The color used for overlay tools inside of the viewport, like the measure tool */
 	UPROPERTY(EditAnywhere, config, Category = UserInterface)
 	FLinearColor ViewportToolOverlayColor;
diff --git a/Engine/Source/Editor/UnrealEd/Private/EditorViewportClient.cpp b/Engine/Source/Editor/UnrealEd/Private/EditorViewportClient.cpp
--- a/Engine/Source/Editor/UnrealEd/Private/EditorViewportClient.cpp
+++ b/Engine/Source/Editor/UnrealEd/Private/EditorViewportClient.cpp
@@ -1595,6 +1595,10 @@ FSceneView* FEditorViewportClient::CalcSceneView(FSceneViewFamily* ViewFamily, c
 		View->AdditionalSelectionOutlineColors[OutlineColorIndex] = EditorStyle->AdditionalSelectionColors[OutlineColorIndex];
 	}
 
+	View->CompositePrimitiveParams.OccludedDitheringPatternSize = EditorStyle->OccludedDitheringPatternSize;
+	View->CompositePrimitiveParams.OccludedDithering = EditorStyle->OccludedDithering;
+	View->CompositePrimitiveParams.OccludedBrightness = EditorStyle->OccludedBrightness;
+
 	int32 FamilyIndex = ViewFamily->Views.Add(View);
 	check(FamilyIndex == View->StereoViewIndex || View->StereoViewIndex == INDEX_NONE);
 
diff --git a/Engine/Source/Runtime/Engine/Public/SceneView.h b/Engine/Source/Runtime/Engine/Public/SceneView.h
--- a/Engine/Source/Runtime/Engine/Public/SceneView.h
+++ b/Engine/Source/Runtime/Engine/Public/SceneView.h
@@ -1761,6 +1761,16 @@ public:
 	/** True if any components are selected in isolation (independent of actor selection) */
 	bool bHasSelectedComponents;
 
+	struct FPostProcessCompositePrimitiveParams
+	{
+		/** Size (in pixels) of the dithering pattern applied to occluded primitives. */
+		float OccludedDitheringPatternSize = 2.0f;
+		/** The strength of the dithering applied to occluded primitives. */
+		float OccludedDithering = 0.8f;
+		/** The brightness of occluded primitives. */
+		float OccludedBrightness = 0.7f;
+	} CompositePrimitiveParams;
+
 	/** In editor, this is used to feed rendering info back to the viewport. May be null. */
 	TSharedPtr<UE::RenderViewportFeedback::FCollector> ViewportFeedback;
 #endif
diff --git a/Engine/Source/Runtime/Renderer/Private/PostProcess/PostProcessCompositeDebugPrimitives.cpp b/Engine/Source/Runtime/Renderer/Private/PostProcess/PostProcessCompositeDebugPrimitives.cpp
--- a/Engine/Source/Runtime/Renderer/Private/PostProcess/PostProcessCompositeDebugPrimitives.cpp
+++ b/Engine/Source/Runtime/Renderer/Private/PostProcess/PostProcessCompositeDebugPrimitives.cpp
@@ -226,6 +226,7 @@ static void ComposeDebugPrimitives(
 	PassParameters->bCompositeAnyNonNullDepth = bCompositeAnyNonNullDepth;
 	PassParameters->DepthTextureJitter = SceneDepthJitter;
 	PassParameters->bProcessAlpha = IsPostProcessingWithAlphaChannelSupported();
+	PassParameters->OccludedDitheringPatternSize = Inputs.OccludedDitheringPatternSize; // Initialize, but this won't have any visual effect (dithering is disabled below)
 	PassParameters->OccludedDithering = 0.0f;
 	PassParameters->OccludedBrightness = 1.0f;
 
diff --git a/Engine/Source/Runtime/Renderer/Private/PostProcess/PostProcessCompositeEditorPrimitives.cpp b/Engine/Source/Runtime/Renderer/Private/PostProcess/PostProcessCompositeEditorPrimitives.cpp
--- a/Engine/Source/Runtime/Renderer/Private/PostProcess/PostProcessCompositeEditorPrimitives.cpp
+++ b/Engine/Source/Runtime/Renderer/Private/PostProcess/PostProcessCompositeEditorPrimitives.cpp
@@ -457,8 +457,9 @@ FScreenPassTexture AddEditorPrimitivePass(
 		PassParameters->bCompositeAnyNonNullDepth = Context.bCompositeAnyNonNullDepth;
 		PassParameters->DepthTextureJitter = SceneDepthJitter;
 		PassParameters->bProcessAlpha = IsPostProcessingWithAlphaChannelSupported();
-		PassParameters->OccludedDithering = 0.8f;
-		PassParameters->OccludedBrightness = 0.7f;
+		PassParameters->OccludedDitheringPatternSize = Inputs.OccludedDitheringPatternSize;
+		PassParameters->OccludedDithering = Inputs.OccludedDithering;
+		PassParameters->OccludedBrightness = Inputs.OccludedBrightness;
 
 		for (int32 i = 0; i < int32(NumMSAASamples); i++)
 		{
diff --git a/Engine/Source/Runtime/Renderer/Private/PostProcess/PostProcessCompositePrimitivesCommon.h b/Engine/Source/Runtime/Renderer/Private/PostProcess/PostProcessCompositePrimitivesCommon.h
--- a/Engine/Source/Runtime/Renderer/Private/PostProcess/PostProcessCompositePrimitivesCommon.h
+++ b/Engine/Source/Runtime/Renderer/Private/PostProcess/PostProcessCompositePrimitivesCommon.h
@@ -77,6 +77,15 @@ struct FCompositePrimitiveInputs
 	// [Optional] Lens distortion applied on the scene color.
 	FLensDistortionLUT LensDistortionLUT;
 
+	// [Optional] Size (in pixels) of the dithering pattern applied to occluded primitives.
+	float OccludedDitheringPatternSize = 2.0f;
+
+	// [Optional] The strength of the dithering applied to occluded primitives.
+	float OccludedDithering = 0.8f;
+
+	// [Optional] The brightness of occluded primitives.
+	float OccludedBrightness = 0.7f;
+
 	bool bUseMetalMSAAHDRDecode = false;
 };
 
@@ -136,6 +145,7 @@ public:
 		SHADER_PARAMETER(uint32, bCompositeAnyNonNullDepth)
 		SHADER_PARAMETER(FVector2f, DepthTextureJitter)
 		SHADER_PARAMETER(uint32, bProcessAlpha)
+		SHADER_PARAMETER(float, OccludedDitheringPatternSize)
 		SHADER_PARAMETER(float, OccludedDithering)
 		SHADER_PARAMETER(float, OccludedBrightness)
 		RENDER_TARGET_BINDING_SLOTS()
diff --git a/Engine/Source/Runtime/Renderer/Private/PostProcess/PostProcessing.cpp b/Engine/Source/Runtime/Renderer/Private/PostProcess/PostProcessing.cpp
--- a/Engine/Source/Runtime/Renderer/Private/PostProcess/PostProcessing.cpp
+++ b/Engine/Source/Runtime/Renderer/Private/PostProcess/PostProcessing.cpp
@@ -1741,6 +1741,9 @@ void AddPostProcessingPasses(
 		{
 			PassInputs.LensDistortionLUT = View.LensDistortionLUT;
 		}
+		PassInputs.OccludedDitheringPatternSize = View.CompositePrimitiveParams.OccludedDitheringPatternSize;
+		PassInputs.OccludedDithering = View.CompositePrimitiveParams.OccludedDithering;
+		PassInputs.OccludedBrightness = View.CompositePrimitiveParams.OccludedBrightness;
 		SceneColor = AddEditorPrimitivePass(GraphBuilder, View, PassInputs, InstanceCullingManager);
 	}
 #endif //WITH_EDITOR
