[Transform Gizmo] Disable old rotation gizmo when nothing selected, ctrl + LMB combo

diff --git a/Engine/Source/Editor/UnrealEd/Private/MouseDeltaTracker.cpp b/Engine/Source/Editor/UnrealEd/Private/MouseDeltaTracker.cpp
--- a/Engine/Source/Editor/UnrealEd/Private/MouseDeltaTracker.cpp
+++ b/Engine/Source/Editor/UnrealEd/Private/MouseDeltaTracker.cpp
@@ -11,6 +11,7 @@
 #include "SnappingUtils.h"
 #include "UnrealWidget.h"
 #include "EditorDragTools/EditorDragToolBehaviorTarget.h"
+#include "EditorInteractiveGizmoManager.h"
 #include "Misc/AxisDisplayInfo.h"
 
 #define LOCTEXT_NAMESPACE "MouseDeltaTracker"
@@ -216,18 +217,23 @@ void FMouseDeltaTracker::StartTracking(FEditorViewportClient* InViewportClient,
 	const bool RightMouseButtonDown = InViewportClient->Viewport->KeyState(EKeys::RightMouseButton);
 	const bool MiddleMouseButtonDown = InViewportClient->Viewport->KeyState(EKeys::MiddleMouseButton);
 
+	constexpr bool bAllowWidgetManipulationWithNewTRSGizmo = false;
+	const bool bUsesNewGizmo = UEditorInteractiveGizmoManager::UsesNewTRSGizmos();
+
+	const bool bAllowWidgetManipulation = !bUsesNewGizmo || bAllowWidgetManipulationWithNewTRSGizmo;
+
 	bool bIsDragging = ((ControlDown || ShiftDown) && (LeftMouseButtonDown || RightMouseButtonDown || MiddleMouseButtonDown)) ||
 		(InViewportClient->GetCurrentWidgetAxis() != EAxisList::None) || bNudge;
 
 	// Update bWidgetAxisControlledByDrag since we now know that we have begun dragging an object with the mouse.
-	if ( bIsDragging )
+	if ( bIsDragging && bAllowWidgetManipulation )
 	{
 		InViewportClient->bWidgetAxisControlledByDrag = true;
 	}
 
 	InViewportClient->TrackingStarted( InInputState, bIsDragging, bNudge );
 
-	if (InViewportClient->Widget)
+	if (InViewportClient->Widget && bAllowWidgetManipulation)
 	{
 		InViewportClient->Widget->SetDragStartPosition(FVector2D(InX, InY));
 		InViewportClient->Widget->SetDragging(bIsDragging);
@@ -287,8 +293,11 @@ void FMouseDeltaTracker::StartTracking(FEditorViewportClient* InViewportClient,
 	EndSnapped = StartSnapped;
 
 	bExternalMovement = false;	//no external movement has occurred yet.
-	InViewportClient->Widget->ResetDeltaRotation();
 
+	if (bAllowWidgetManipulation)
+	{
+		InViewportClient->Widget->ResetDeltaRotation();	
+	}
 }
 
 /**
