Bump cook incremental version to clear out old shader cruft.

diff --git a/Engine/Source/Editor/UnrealEd/Private/Cooker/CookIncrementalVersion.cpp b/Engine/Source/Editor/UnrealEd/Private/Cooker/CookIncrementalVersion.cpp
--- a/Engine/Source/Editor/UnrealEd/Private/Cooker/CookIncrementalVersion.cpp
+++ b/Engine/Source/Editor/UnrealEd/Private/Cooker/CookIncrementalVersion.cpp
@@ -8,6 +8,6 @@ namespace UE::Cook
 {
 
 // Change CookIncrementalVersion to a new guid when all packages in an incremental cook need to be invalidated
-FGuid CookIncrementalVersion(0x7106a584, 0x94714f4d, 0xabf8d8a8, 0xebe28cc6);
+FGuid CookIncrementalVersion( 0xA47D92F1, 0x3C5E18AB, 0x7FB2046D, 0xE89C51D7 );
 
 }
