ShaderPipelineCacheTools: add a mode to extract a shader from ushaderbytecode by its hash.

diff --git a/Engine/Source/Editor/UnrealEd/Private/Commandlets/ShaderPipelineCacheToolsCommandlet.cpp b/Engine/Source/Editor/UnrealEd/Private/Commandlets/ShaderPipelineCacheToolsCommandlet.cpp
--- a/Engine/Source/Editor/UnrealEd/Private/Commandlets/ShaderPipelineCacheToolsCommandlet.cpp
+++ b/Engine/Source/Editor/UnrealEd/Private/Commandlets/ShaderPipelineCacheToolsCommandlet.cpp
@@ -52,6 +52,8 @@
 #include "Trace/Detail/Channel.h"
 #include "UObject/NameTypes.h"
 #include "UObject/UnrealNames.h"
+#include "ShaderCodeArchive.h"
+#include "Serialization/LargeMemoryReader.h"
 
 #include UE_INLINE_GENERATED_CPP_BY_NAME(ShaderPipelineCacheToolsCommandlet)
 
@@ -727,6 +729,78 @@ int32 DumpStableKeysFile(const FString& Token, const TMap<FString, FString>& Par
 	return 0;
 }
 
+int32 ExtractShaderFromUshaderbytecode(const FString& UshaderbytecodeName, const FString& ShaderHash)
+{
+	TUniquePtr<FArchive> FileReader(IFileManager::Get().CreateFileReader(*UshaderbytecodeName));
+	if (!FileReader)
+	{
+		UE_LOG(LogShaderPipelineCacheTools, Error, TEXT("Cannot open ushaderbytecode file at %s"), *UshaderbytecodeName);
+		return 1;
+	}
+
+	TArray64<uint8> Data;
+	Data.SetNumUninitialized(FileReader->TotalSize());
+	FileReader->Serialize(Data.GetData(), Data.Num());
+	if (FileReader->IsError())
+	{
+		UE_LOG(LogShaderPipelineCacheTools, Error, TEXT("Error reading ushaderbytecode file at %s"), *UshaderbytecodeName);
+		return 1;
+	}
+
+	FLargeMemoryReader LibraryAr(Data.GetData(), Data.Num());
+	uint32 Version;
+	LibraryAr << Version;
+
+	UE_LOG(LogShaderPipelineCacheTools, Display, TEXT("Reading ushaderbytecode file of version %d"), Version);
+
+	FSerializedShaderArchive SerializedShaders;
+	LibraryAr << SerializedShaders;
+	int64 OffsetToShaderCode = LibraryAr.Tell();
+
+	UE_LOG(LogShaderPipelineCacheTools, Display, TEXT("Number of shaders %d"), SerializedShaders.GetNumShaders());
+
+	FSHAHash HashToLookFor;
+	HashToLookFor.FromString(ShaderHash);
+	UE_LOG(LogShaderPipelineCacheTools, Display, TEXT("Looking for shader hash %s"), *HashToLookFor.ToString());	// print so we know the conversion was successful
+
+	bool bHaveAtLeastOneMatch = false;
+	for (int IdxShader = 0, NumShaders = SerializedShaders.GetNumShaders(); IdxShader < NumShaders; ++IdxShader)
+	{
+		if (SerializedShaders.ShaderHashes[IdxShader] == HashToLookFor)
+		{
+			UE_LOG(LogShaderPipelineCacheTools, Display, TEXT("Shader %d matches hash %s"), IdxShader, *SerializedShaders.ShaderHashes[IdxShader].ToString());	// print so we know the conversion was successful
+			bHaveAtLeastOneMatch = true;
+			const FShaderCodeEntry& Entry = SerializedShaders.ShaderEntries[IdxShader];
+
+			// extract compressed first
+			FString SavePath = FPaths::ProjectSavedDir() / TEXT("ExtractedShaders") / FPaths::GetBaseFilename(UshaderbytecodeName);
+			FString ShaderName = FString::Printf(TEXT("%s_Shader_%d.bin"), *SerializedShaders.ShaderHashes[IdxShader].ToString(), IdxShader);
+			FString FinalSavePath = SavePath / ShaderName;
+			UE_LOG(LogShaderPipelineCacheTools, Display, TEXT("Extracting to %s"), *FinalSavePath);
+			TUniquePtr<FArchive> FileWriter(IFileManager::Get().CreateFileWriter(*FinalSavePath));
+
+			// small shaders might be stored without compression, handle them here
+			uint8* PtrToShader = Data.GetData() + OffsetToShaderCode + Entry.Offset;
+			if (Entry.Size == Entry.UncompressedSize)
+			{
+				FileWriter->Serialize(PtrToShader, Entry.Size);
+			}
+			else
+			{
+				TUniquePtr<uint8[]> Buffer = MakeUnique<uint8[]>(Entry.UncompressedSize);
+				ShaderCodeArchive::DecompressShaderWithOodle(Buffer.Get(), Entry.UncompressedSize, PtrToShader, Entry.Size);
+				FileWriter->Serialize(Buffer.Get(), Entry.UncompressedSize);
+			}
+		}
+	}
+
+	if (!bHaveAtLeastOneMatch)
+	{
+		UE_LOG(LogShaderPipelineCacheTools, Warning, TEXT("Could find a shader with hash %s among all %d shaders"), *HashToLookFor.ToString(), SerializedShaders.GetNumShaders());	// print so we know the conversion was successful
+	}
+	return 0;
+}
+
 int32 CheckStableKeyAliasing(const TArray<FString>& Tokens)
 {
 	if (Tokens.Num() < 2)
@@ -2896,8 +2970,21 @@ int32 UShaderPipelineCacheToolsCommandlet::StaticMain(const FString& Params)
 			Tokens.RemoveAt(0);
 			return DecompressCSV(Tokens);
 		}
+		else if (Tokens[0] == TEXT("Extract") && Tokens.Num() >= 3)
+		{
+			Tokens.RemoveAt(0);
+			if (Tokens[0].EndsWith(TEXT(".ushaderbytecode")))
+			{
+				return ExtractShaderFromUshaderbytecode(Tokens[0], Tokens[1]);
+			}
+			else
+			{
+				UE_LOG(LogShaderPipelineCacheTools, Warning, TEXT("Extract command expects a ushaderbytecode parameters as its first argument.\n"));
+			}
+		}
 	}
 	
+	UE_LOG(LogShaderPipelineCacheTools, Warning, TEXT("Usage: Extract Foo.ushaderbytecode 123ABC (shader hash) - extracts all shaders that match the hash from this ushaderbytecode (and decompresses them).\n"));
 	UE_LOG(LogShaderPipelineCacheTools, Warning, TEXT("Usage: Dump SCLInfo.%s [...] [-outputhash=LIST] [-objname=LIST] - dumps stable keys file, optionally only dumping entries that match the specified output hashes and/or object names (hash and name lists are comma-separated).\n"), ShaderStableKeysFileExt);
 	UE_LOG(LogShaderPipelineCacheTools, Warning, TEXT("Usage: Dump PSOCache.upipelinecache [SCLInfo.%s] - dumps recorded PSOs. If optional stable keys file, converts shader hashes to readable stable shader descriptions.\n"), ShaderStableKeysFileExt);
 	UE_LOG(LogShaderPipelineCacheTools, Warning, TEXT("Usage: Dump StablePSOCache.csv[.compressed] - dumps content of a stable PSO.\n"));
