Adds callback to RunWorldPartitionBuilder to update slowtask progress while the builder is running.

diff --git a/Engine/Source/Editor/UnrealEd/Private/EditorBuildUtils.cpp b/Engine/Source/Editor/UnrealEd/Private/EditorBuildUtils.cpp
--- a/Engine/Source/Editor/UnrealEd/Private/EditorBuildUtils.cpp
+++ b/Engine/Source/Editor/UnrealEd/Private/EditorBuildUtils.cpp
@@ -1035,7 +1035,8 @@ bool FEditorBuildUtils::RunWorldPartitionBuilder(
 	const FText& ProgressText,
 	const FText& CancelledText,
 	const FText& FailureText,
-	const FString& CommandLineArguments
+	const FString& CommandLineArguments,
+	TFunctionRef<void(FScopedSlowTask&)> UpdateSlowTask
 	)
 {
 	// Ask user to save dirty packages
@@ -1081,6 +1082,8 @@ bool FEditorBuildUtils::RunWorldPartitionBuilder(
 				break;
 			}
 
+			UpdateSlowTask(SlowTask);
+
 			SlowTask.EnterProgressFrame(0);
 			FPlatformProcess::Sleep(0.1);
 		}
diff --git a/Engine/Source/Editor/UnrealEd/Public/EditorBuildUtils.h b/Engine/Source/Editor/UnrealEd/Public/EditorBuildUtils.h
--- a/Engine/Source/Editor/UnrealEd/Public/EditorBuildUtils.h
+++ b/Engine/Source/Editor/UnrealEd/Public/EditorBuildUtils.h
@@ -11,6 +11,8 @@
 #include "SceneTypes.h"
 #include "RHIDefinitions.h"
 
+struct FScopedSlowTask;
+
 /** Names of the built-in editor build types. */
 struct FBuildOptions
 {
@@ -276,6 +278,7 @@ public:
 	 * @param CancelledText Message shown to the user after cancelling the task 
 	 * @param FailureText Message shown to the user if the process returned an error code
 	 * @param CommandLineArguments Arguments passed to the other instance commandline
+	 * @param OnProcRunningTick Function that will be called periodically while the builder process is running (can be used to update progress).
 	 * 
 	 * @return True if process completed without returning any error code.  
 	 */
@@ -284,7 +287,8 @@ public:
 		const FText& ProgressText,
 		const FText& CancelledText,
 		const FText& FailureText,
-		const FString& CommandLineArguments
+		const FString& CommandLineArguments,
+		TFunctionRef<void(FScopedSlowTask&)> UpdateSlowTask = TFunction<void(FScopedSlowTask&)>()
 		);
 private:
 
