Fixes default functor binding in RunWorldPartitionBuilder

diff --git a/Engine/Source/Editor/UnrealEd/Private/EditorBuildUtils.cpp b/Engine/Source/Editor/UnrealEd/Private/EditorBuildUtils.cpp
--- a/Engine/Source/Editor/UnrealEd/Private/EditorBuildUtils.cpp
+++ b/Engine/Source/Editor/UnrealEd/Private/EditorBuildUtils.cpp
@@ -1036,7 +1036,7 @@ bool FEditorBuildUtils::RunWorldPartitionBuilder(
 	const FText& CancelledText,
 	const FText& FailureText,
 	const FString& CommandLineArguments,
-	TFunctionRef<void(FScopedSlowTask&)> UpdateSlowTask
+	TFunction<void(FScopedSlowTask&)> UpdateSlowTask
 	)
 {
 	// Ask user to save dirty packages
@@ -1082,7 +1082,10 @@ bool FEditorBuildUtils::RunWorldPartitionBuilder(
 				break;
 			}
 
-			UpdateSlowTask(SlowTask);
+			if (UpdateSlowTask.IsSet())
+			{
+				UpdateSlowTask(SlowTask);
+			}
 
 			SlowTask.EnterProgressFrame(0);
 			FPlatformProcess::Sleep(0.1);
diff --git a/Engine/Source/Editor/UnrealEd/Public/EditorBuildUtils.h b/Engine/Source/Editor/UnrealEd/Public/EditorBuildUtils.h
--- a/Engine/Source/Editor/UnrealEd/Public/EditorBuildUtils.h
+++ b/Engine/Source/Editor/UnrealEd/Public/EditorBuildUtils.h
@@ -288,7 +288,7 @@ public:
 		const FText& CancelledText,
 		const FText& FailureText,
 		const FString& CommandLineArguments,
-		TFunctionRef<void(FScopedSlowTask&)> UpdateSlowTask = TFunction<void(FScopedSlowTask&)>()
+		TFunction<void(FScopedSlowTask&)> UpdateSlowTask = TFunction<void(FScopedSlowTask&)>()
 		);
 private:
 
