[ITF Viewport Interactions] Refactor ITF Viewport Camera control logic flow so that input values no longer interact with input checking in the Editor Viewport Client.
- Add new behavior-only camera movement functions & behavior to EVPC.
- Update and add new Viewport Interactions to fix bugs and fill behavioral gaps.
- Fix an issue where the ViewportDragInteraction would exit too early from modifier changes.

diff --git a/Engine/Source/Editor/Experimental/EditorInteractiveToolsFramework/Private/Behaviors/ViewportClickDragBehavior.cpp b/Engine/Source/Editor/Experimental/EditorInteractiveToolsFramework/Private/Behaviors/ViewportClickDragBehavior.cpp
--- a/Engine/Source/Editor/Experimental/EditorInteractiveToolsFramework/Private/Behaviors/ViewportClickDragBehavior.cpp
+++ b/Engine/Source/Editor/Experimental/EditorInteractiveToolsFramework/Private/Behaviors/ViewportClickDragBehavior.cpp
@@ -2,6 +2,8 @@
 
 #include "Behaviors/ViewportClickDragBehavior.h"
 
+#include "ViewportInteractions/ViewportInteractionsBehaviorSource.h"
+
 void UViewportClickDragBehavior::Initialize(IViewportClickDragBehaviorTarget* InTarget)
 {
 	check(InTarget != nullptr);
@@ -30,7 +32,6 @@ FInputCaptureRequest UViewportClickDragBehavior::WantsCapture(const FInputDevice
 {
 	if (InInputState.IsFromDevice(EInputDevices::Mouse) && !HasButtonsDown(InInputState))
 	{
-		// We don't want this event, reset state
 		MouseTraveledDistance = TOptional<float>();
 		return FInputCaptureRequest::Ignore();
 	}
@@ -40,6 +41,7 @@ FInputCaptureRequest UViewportClickDragBehavior::WantsCapture(const FInputDevice
 	
 	if (ModifierMode == EModifierMode::Continuous && !bModifiersValid)
 	{
+		MouseTraveledDistance = TOptional<float>();
 		return FInputCaptureRequest::Ignore();
 	}
 	
@@ -88,6 +90,13 @@ FInputCaptureUpdate UViewportClickDragBehavior::BeginCapture(const FInputDeviceS
 
 FInputCaptureUpdate UViewportClickDragBehavior::UpdateCapture(const FInputDeviceState& InputState, const FInputCaptureData& CaptureData)
 {
+	// This happens when a keyboard event is routed to the active mouse capture for updating modifier states
+	// The behavior should always continue here so that subsequent mouse updates can end capture more correctly	
+	if (!InputState.IsFromDevice(EInputDevices::Mouse))
+	{
+		return FInputCaptureUpdate::Continue();
+	}
+
 	MouseTraveledDistance = static_cast<float>(FVector2D::Distance(InputState.Mouse.Position2D, MouseDownPosition));
 
 	if (!bDragging)
diff --git a/Engine/Source/Editor/Experimental/EditorInteractiveToolsFramework/Private/ViewportInteractions/ViewportDollyInteraction.cpp b/Engine/Source/Editor/Experimental/EditorInteractiveToolsFramework/Private/ViewportInteractions/ViewportDollyInteraction.cpp
new file mode 100644
--- /dev/null
+++ b/Engine/Source/Editor/Experimental/EditorInteractiveToolsFramework/Private/ViewportInteractions/ViewportDollyInteraction.cpp
@@ -0,0 +1,72 @@
+﻿// Copyright Epic Games, Inc. All Rights Reserved.
+
+#include "ViewportInteractions/ViewportDollyInteraction.h"
+
+#include "ViewportClientNavigationHelper.h"
+#include "ViewportInteractions/ViewportInteractionsBehaviorSource.h"
+#include "Settings/LevelEditorViewportSettings.h"
+#include "EditorViewportClient.h"
+
+UViewportDollyInteraction::UViewportDollyInteraction()
+{
+	ToolType = UE::Editor::ViewportInteractions::CameraDrag;
+
+	InteractionName = TEXT("Perspective Dolly");
+	
+	// TODO: Make behaviors with more buttons automatically have higher priority
+	ViewportClickDragBehavior->SetDefaultPriority(UE::Editor::ViewportInteractions::MULTIBUTTON_DRAG_PRIORITY);
+	ViewportClickDragBehavior->SetDragConfirmedPriority(UE::Editor::ViewportInteractions::MULTIBUTTON_DRAG_CONFIRMED_PRIORITY);
+
+	EnableButton(EKeys::RightMouseButton);
+}
+
+FInputRayHit UViewportDollyInteraction::CanBeginClickDragSequence(const FInputDeviceRay& InPressPos)
+{
+	return FInputRayHit(TNumericLimits<float>::Max());
+}
+
+bool UViewportDollyInteraction::CanBeActivated(const FInputDeviceState& InInputDeviceState) const
+{
+	if (IsEnabled() && InInputDeviceState.bAltKeyDown)
+	{
+		if (FEditorViewportClient* Client = GetEditorViewportClient())
+		{
+			return Client->IsPerspective();
+		}
+	}
+	return false;
+}
+
+void UViewportDollyInteraction::OnDragStart(const FInputDeviceRay& InDragStartPos)
+{
+	UViewportDragInteraction::OnDragStart(InDragStartPos);
+	
+	if (UViewportInteractionsBehaviorSource* BehaviorSource = GetViewportInteractionsBehaviorSource())
+	{
+		BehaviorSource->SetIsMouseLooking(true);
+	}
+}
+
+void UViewportDollyInteraction::OnDragEnd(const FInputDeviceRay& InDragEndPos)
+{
+	UViewportDragInteraction::OnDragEnd(InDragEndPos);
+	
+	if (UViewportInteractionsBehaviorSource* BehaviorSource = GetViewportInteractionsBehaviorSource())
+	{
+		BehaviorSource->SetIsMouseLooking(false);
+	}
+}
+
+void UViewportDollyInteraction::OnDragDelta(float InMouseDeltaX, float InMouseDeltaY)
+{
+	if (FEditorViewportClient* Client = GetEditorViewportClient())
+	{
+		if (FViewportClientNavigationHelper* Helper = Client->GetViewportNavigationHelper())
+		{
+			const ULevelEditorViewportSettings* Settings = GetDefault<ULevelEditorViewportSettings>(); 
+			const float Direction = Settings->bInvertRightMouseDollyYAxis ? -1.0f : 1.0f;
+			// TODO: Switch to using Settings->MouseSensitivty?
+			Helper->OrbitDelta.Z += (InMouseDeltaX + InMouseDeltaY * Direction);
+		}
+	}
+}
diff --git a/Engine/Source/Editor/Experimental/EditorInteractiveToolsFramework/Private/ViewportInteractions/ViewportDragInteraction.cpp b/Engine/Source/Editor/Experimental/EditorInteractiveToolsFramework/Private/ViewportInteractions/ViewportDragInteraction.cpp
--- a/Engine/Source/Editor/Experimental/EditorInteractiveToolsFramework/Private/ViewportInteractions/ViewportDragInteraction.cpp
+++ b/Engine/Source/Editor/Experimental/EditorInteractiveToolsFramework/Private/ViewportInteractions/ViewportDragInteraction.cpp
@@ -11,32 +11,12 @@ UViewportDragInteraction::UViewportDragInteraction()
 	MouseX = 0;
 	MouseY = 0;
 
-	ULocalViewportClickDragBehavior* ClickDragBehavior = NewObject<ULocalViewportClickDragBehavior>();
-	ClickDragBehavior->Initialize();
+	UViewportClickDragBehavior* ClickDragBehavior = NewObject<UViewportClickDragBehavior>();
+	ClickDragBehavior->Initialize(this);
 	ClickDragBehavior->SetTriggerMode(UViewportClickDragBehavior::ETriggerMode::Down);
 	ClickDragBehavior->SetModifierMode(UViewportClickDragBehavior::EModifierMode::Continuous);
 	ClickDragBehavior->SetRequireAllButtons(true);
 
-	ClickDragBehavior->CanBeginClickDragFunc = [this](const FInputDeviceRay& InPressPos)
-	{
-		return CanBeginClickDragSequence(InPressPos);
-	};
-
-	ClickDragBehavior->OnDragStartFunc = [this](const FInputDeviceRay& InPressPos)
-	{
-		OnDragStart(InPressPos);
-	};
-
-	ClickDragBehavior->OnDragEndFunc = [this](const FInputDeviceRay& InReleasePos)
-	{
-		OnDragEnd(InReleasePos);
-	};
-
-	ClickDragBehavior->OnClickDragFunc = [this](const FInputDeviceRay& InDragPos)
-	{
-		OnDrag(InDragPos);
-	};
-
 	ClickDragBehavior->SetDefaultPriority(UE::Editor::ViewportInteractions::DRAG_PRIORITY);
 	ClickDragBehavior->SetDragConfirmedPriority(UE::Editor::ViewportInteractions::DRAG_CONFIRMED_PRIORITY);
 
@@ -114,11 +94,6 @@ void UViewportDragInteraction::Initialize(UViewportInteractionsBehaviorSource* c
 	}
 }
 
-FDeviceButtonState UViewportDragInteraction::GetActiveMouseButtonState(const FInputDeviceState& Input)
-{
-	return Input.Mouse.Left;
-}
-
 void UViewportDragInteraction::EnableButton(const FKey& InButton) const
 {
 	if (TStrongObjectPtr<UViewportClickDragBehavior> ClickDragBehavior = ViewportClickDragBehavior.Pin())
diff --git a/Engine/Source/Editor/Experimental/EditorInteractiveToolsFramework/Private/ViewportInteractions/ViewportInteractionsBehaviorSource.cpp b/Engine/Source/Editor/Experimental/EditorInteractiveToolsFramework/Private/ViewportInteractions/ViewportInteractionsBehaviorSource.cpp
--- a/Engine/Source/Editor/Experimental/EditorInteractiveToolsFramework/Private/ViewportInteractions/ViewportInteractionsBehaviorSource.cpp
+++ b/Engine/Source/Editor/Experimental/EditorInteractiveToolsFramework/Private/ViewportInteractions/ViewportInteractionsBehaviorSource.cpp
@@ -258,6 +258,7 @@ UViewportInteraction* UViewportInteractionsBehaviorSource::AddInteraction(const
 	UViewportInteraction* Interaction = NewObject<UViewportInteraction>(GetTransientPackage(), InInteractionClass);
 	Interaction->Initialize(this);
 
+	UE::Editor::ViewportInteractions::LOG(*FString::Printf(TEXT("Initializing Viewport Interaction \"%s\""), *InInteractionClass->GetName()));
 	ViewportInteractions.Add(Interaction->GetInteractionName(), Interaction);
 
 	for (UInputBehavior* const InputBehavior : Interaction->GetInputBehaviors())
diff --git a/Engine/Source/Editor/Experimental/EditorInteractiveToolsFramework/Private/ViewportInteractions/ViewportInteractionsUtilities.cpp b/Engine/Source/Editor/Experimental/EditorInteractiveToolsFramework/Private/ViewportInteractions/ViewportInteractionsUtilities.cpp
--- a/Engine/Source/Editor/Experimental/EditorInteractiveToolsFramework/Private/ViewportInteractions/ViewportInteractionsUtilities.cpp
+++ b/Engine/Source/Editor/Experimental/EditorInteractiveToolsFramework/Private/ViewportInteractions/ViewportInteractionsUtilities.cpp
@@ -5,6 +5,7 @@
 #include "ViewportInteractions/ViewportCameraSpeedMouseWheelInteraction.h"
 #include "ViewportInteractions/ViewportCameraRotateInteraction.h"
 #include "ViewportInteractions/ViewportCameraTranslateInteraction.h"
+#include "ViewportInteractions/ViewportDollyInteraction.h"
 #include "ViewportInteractions/ViewportFOVInteraction.h"
 #include "ViewportInteractions/ViewportInteractionsBehaviorSource.h"
 #include "ViewportInteractions/ViewportMoveYawInteraction.h"
@@ -24,9 +25,11 @@ void UE::Editor::ViewportInteractions::AddDefaultCameraMovementInteractions(
 			{UViewportCameraSpeedMouseWheelInteraction::StaticClass(),
 				UViewportOrbitInteraction::StaticClass(),
 				UViewportMoveYawInteraction::StaticClass(),
+				UViewportMovePanInteraction::StaticClass(),
 				UViewportViewAngleInteraction::StaticClass(),
 				UViewportOrthoPanInteraction::StaticClass(),
 				UViewportPanInteraction::StaticClass(),
+				UViewportDollyInteraction::StaticClass(),
 				UViewportZoomInteraction::StaticClass(),
 				UViewportCameraTranslateInteraction::StaticClass(),
 				UViewportCameraRotateInteraction::StaticClass(),
diff --git a/Engine/Source/Editor/Experimental/EditorInteractiveToolsFramework/Private/ViewportInteractions/ViewportMoveYawInteraction.cpp b/Engine/Source/Editor/Experimental/EditorInteractiveToolsFramework/Private/ViewportInteractions/ViewportMoveYawInteraction.cpp
--- a/Engine/Source/Editor/Experimental/EditorInteractiveToolsFramework/Private/ViewportInteractions/ViewportMoveYawInteraction.cpp
+++ b/Engine/Source/Editor/Experimental/EditorInteractiveToolsFramework/Private/ViewportInteractions/ViewportMoveYawInteraction.cpp
@@ -59,25 +59,18 @@ void UViewportMoveYawInteraction::OnDragDelta(float InMouseDeltaX, float InMouse
 		FVector LocationDelta = FVector::ZeroVector;
 		FRotator RotDelta = FRotator::ZeroRotator;
 
-		LocationDelta.X = -InMouseDeltaY * FMath::Cos(EditorViewportClient->GetViewRotation().Yaw * PI / 180.f);
-		LocationDelta.Y = -InMouseDeltaY * FMath::Sin(EditorViewportClient->GetViewRotation().Yaw * PI / 180.f);
+		const double YawRadians = FMath::DegreesToRadians(EditorViewportClient->GetViewRotation().Yaw);
+		LocationDelta.X = -InMouseDeltaY * FMath::Cos(YawRadians);
+		LocationDelta.Y = -InMouseDeltaY * FMath::Sin(YawRadians);
 
 		const ULevelEditorViewportSettings* ViewportSettings = GetDefault<ULevelEditorViewportSettings>();
 		RotDelta.Yaw = InMouseDeltaX * ViewportSettings->MouseSensitivty;
 
-		FVector& ViewportLocationDelta = EditorViewportClient->GetViewportNavigationHelper()->LocationDelta;
-		ViewportLocationDelta += LocationDelta;
-
-		FRotator& ViewportRotationDelta = EditorViewportClient->GetViewportNavigationHelper()->RotationDelta;
-		ViewportRotationDelta += RotDelta;
+		EditorViewportClient->GetViewportNavigationHelper()->LocationDelta += LocationDelta;
+		EditorViewportClient->GetViewportNavigationHelper()->RotationDelta += RotDelta;
 	}
 }
 
-FDeviceButtonState UViewportMoveYawInteraction::GetActiveMouseButtonState(const FInputDeviceState& Input)
-{
-	return Super::GetActiveMouseButtonState(Input);
-}
-
 bool UViewportMoveYawInteraction::CanBeActivated(const FInputDeviceState& InInputDeviceState) const
 {
 	if (const FEditorViewportClient* const EditorViewportClient = GetEditorViewportClient())
diff --git a/Engine/Source/Editor/Experimental/EditorInteractiveToolsFramework/Private/ViewportInteractions/ViewportOrbitInteraction.cpp b/Engine/Source/Editor/Experimental/EditorInteractiveToolsFramework/Private/ViewportInteractions/ViewportOrbitInteraction.cpp
--- a/Engine/Source/Editor/Experimental/EditorInteractiveToolsFramework/Private/ViewportInteractions/ViewportOrbitInteraction.cpp
+++ b/Engine/Source/Editor/Experimental/EditorInteractiveToolsFramework/Private/ViewportInteractions/ViewportOrbitInteraction.cpp
@@ -5,6 +5,7 @@
 #include "EditorViewportClient.h"
 #include "ViewportClientNavigationHelper.h"
 #include "Behaviors/ViewportClickDragBehavior.h"
+#include "Settings/LevelEditorViewportSettings.h"
 #include "ViewportInteractions/ViewportInteractionsBehaviorSource.h"
 
 UViewportOrbitInteraction::UViewportOrbitInteraction()
@@ -12,15 +13,12 @@ UViewportOrbitInteraction::UViewportOrbitInteraction()
 	ToolType = UE::Editor::ViewportInteractions::Orbit;
 
 	InteractionName = TEXT("Orbit");
+	
+	// TODO: Make behaviors with more buttons automatically have higher priority
+	ViewportClickDragBehavior->SetDefaultPriority(UE::Editor::ViewportInteractions::MULTIBUTTON_DRAG_PRIORITY);
+	ViewportClickDragBehavior->SetDragConfirmedPriority(UE::Editor::ViewportInteractions::MULTIBUTTON_DRAG_CONFIRMED_PRIORITY);
 
 	EnableButton(EKeys::LeftMouseButton);
-	EnableButton(EKeys::RightMouseButton);
-
-	// LMB and RMB work separately
-	if (const TStrongObjectPtr<UViewportClickDragBehavior> ClickDragBehavior = ViewportClickDragBehavior.Pin())
-	{
-		ClickDragBehavior->SetRequireAllButtons(false);
-	}
 }
 
 FInputRayHit UViewportOrbitInteraction::CanBeginClickDragSequence(const FInputDeviceRay& InPressPos)
@@ -64,17 +62,15 @@ void UViewportOrbitInteraction::OnDragDelta(float InMouseDeltaX, float InMouseDe
 {
 	if (FEditorViewportClient* const EditorViewportClient = GetEditorViewportClient())
 	{
-		constexpr bool bNudge = false;
-
-		EditorViewportClient->SetCurrentWidgetAxis(EAxisList::None);
-
-		// TODO: we can probably extract just the part dedicated to perspective view from EditorViewportClient->TranslateDelta
-		FVector Delta = FVector::ZeroVector;
-		Delta += EditorViewportClient->TranslateDelta(EKeys::MouseX, InMouseDeltaX, bNudge);
-		Delta += EditorViewportClient->TranslateDelta(EKeys::MouseY, -InMouseDeltaY, bNudge);
-
-		FVector& ViewportOrbitDelta = EditorViewportClient->GetViewportNavigationHelper()->OrbitDelta;
-		ViewportOrbitDelta += Delta;
+		const ULevelEditorViewportSettings* ViewportSettings = GetDefault<ULevelEditorViewportSettings>();
+		const float Direction = ViewportSettings->bInvertOrbitYAxis ? -1.0f : 1.0f;
+
+		// Change the viewing angle
+		FVector OrbitDelta = FVector::ZeroVector;
+		OrbitDelta.X = InMouseDeltaX * ViewportSettings->MouseSensitivty;
+		OrbitDelta.Y = -InMouseDeltaY * ViewportSettings->MouseSensitivty * Direction;
+	
+		EditorViewportClient->GetViewportNavigationHelper()->OrbitDelta += OrbitDelta;
 	}
 }
 
diff --git a/Engine/Source/Editor/Experimental/EditorInteractiveToolsFramework/Private/ViewportInteractions/ViewportOrthoPanInteraction.cpp b/Engine/Source/Editor/Experimental/EditorInteractiveToolsFramework/Private/ViewportInteractions/ViewportOrthoPanInteraction.cpp
--- a/Engine/Source/Editor/Experimental/EditorInteractiveToolsFramework/Private/ViewportInteractions/ViewportOrthoPanInteraction.cpp
+++ b/Engine/Source/Editor/Experimental/EditorInteractiveToolsFramework/Private/ViewportInteractions/ViewportOrthoPanInteraction.cpp
@@ -96,10 +96,6 @@ void UViewportOrthoPanInteraction::OnDragDelta(float InMouseDeltaX, float InMous
 	}
 }
 
-FDeviceButtonState UViewportOrthoPanInteraction::GetActiveMouseButtonState(const FInputDeviceState& Input)
-{
-	return Input.Mouse.Right;
-}
 bool UViewportOrthoPanInteraction::CanBeActivated(const FInputDeviceState& InInputDeviceState) const
 {
 	return IsEnabled();
diff --git a/Engine/Source/Editor/Experimental/EditorInteractiveToolsFramework/Private/ViewportInteractions/ViewportPanInteraction.cpp b/Engine/Source/Editor/Experimental/EditorInteractiveToolsFramework/Private/ViewportInteractions/ViewportPanInteraction.cpp
--- a/Engine/Source/Editor/Experimental/EditorInteractiveToolsFramework/Private/ViewportInteractions/ViewportPanInteraction.cpp
+++ b/Engine/Source/Editor/Experimental/EditorInteractiveToolsFramework/Private/ViewportInteractions/ViewportPanInteraction.cpp
@@ -15,48 +15,6 @@ UViewportPanInteraction::UViewportPanInteraction()
 	InteractionName = TEXT("Perspective Panning");
 
 	EnableButton(EKeys::MiddleMouseButton);
-
-	// LMB + RMB + Drag to Pan
-	ULocalViewportClickDragBehavior* MultiButtonClickDragBehavior = NewObject<ULocalViewportClickDragBehavior>();
-	MultiButtonClickDragBehavior->Initialize();
-	MultiButtonClickDragBehavior->EnableButton(EKeys::LeftMouseButton);
-	MultiButtonClickDragBehavior->EnableButton(EKeys::RightMouseButton);
-	MultiButtonClickDragBehavior->SetRequireAllButtons(true);
-	MultiButtonClickDragBehavior->SetTriggerMode(UViewportClickDragBehavior::ETriggerMode::Down);
-	MultiButtonClickDragBehavior->SetModifierMode(UViewportClickDragBehavior::EModifierMode::Continuous);
-
-	MultiButtonClickDragBehavior->CanBeginClickDragFunc = [this](const FInputDeviceRay& InPressPos)
-	{
-		bool bSuccess = CanBeActivated();
-		if (FEditorViewportClient* EditorViewportClient = GetEditorViewportClient())
-		{
-			bSuccess &= EditorViewportClient->IsPerspective();
-		}
-
-		return bSuccess ? FInputRayHit(TNumericLimits<float>::Max()) : FInputRayHit();
-	};
-
-	MultiButtonClickDragBehavior->OnDragStartFunc = [this](const FInputDeviceRay& InDragStartPos)
-	{
-		OnDragStart(InDragStartPos);
-		bMultiButtonPanInteraction = true;
-	};
-
-	MultiButtonClickDragBehavior->OnDragEndFunc = [this](const FInputDeviceRay& InDragEndPos)
-	{
-		OnDragEnd(InDragEndPos);
-	};
-
-	MultiButtonClickDragBehavior->OnClickDragFunc = [this](const FInputDeviceRay& InDragPos)
-	{
-		OnDrag(InDragPos);
-	};
-
-	MultiButtonClickDragBehavior->SetDefaultPriority(UE::Editor::ViewportInteractions::MULTIBUTTON_DRAG_PRIORITY);
-	MultiButtonClickDragBehavior->SetDragConfirmedPriority(UE::Editor::ViewportInteractions::MULTIBUTTON_DRAG_CONFIRMED_PRIORITY);
-	MultiButtonClickDragBehavior->SetAllowsCaptureStealing(true);
-
-	RegisterInputBehavior(MultiButtonClickDragBehavior);
 }
 
 FInputRayHit UViewportPanInteraction::CanBeginClickDragSequence(const FInputDeviceRay& InPressPos)
@@ -81,8 +39,6 @@ void UViewportPanInteraction::OnDragStart(const FInputDeviceRay& InDragStartPos)
 	{
 		BehaviorSource->SetIsMouseLooking(true);
 	}
-
-	bMultiButtonPanInteraction = false;
 }
 
 void UViewportPanInteraction::OnDragEnd(const FInputDeviceRay& InDragEndPos)
@@ -101,32 +57,57 @@ void UViewportPanInteraction::OnDragDelta(float InMouseDeltaX, float InMouseDelt
 {
 	if (FEditorViewportClient* const EditorViewportClient = GetEditorViewportClient())
 	{
-		// Pan left/right/up/down
-		// TODO: In legacy, this bool will also consider if using trackpad or not - possibly need to handle this in a different way with ITF?
-
-		bool bInvert = GetDefault<ULevelEditorViewportSettings>()->bInvertMiddleMousePan;
-
-		FVector LocationDelta;
-
-		const FRotator& ViewRotation = EditorViewportClient->GetViewRotation();
-		if (bMultiButtonPanInteraction)
-		{
-			float Direction = bInvert ? 1.0f : -1.0f;
-			LocationDelta.X = -InMouseDeltaX * Direction * FMath::Sin( ViewRotation.Yaw * PI / 180.0f );
-			LocationDelta.Y = -InMouseDeltaX * -Direction * FMath::Cos( ViewRotation.Yaw * PI / 180.0f );
-			LocationDelta.Z = -Direction * InMouseDeltaY;
-		}
-		else
-		{
-			LocationDelta = FQuat(ViewRotation) * FVector( 0.0f, InMouseDeltaX, -InMouseDeltaY) * (bInvert ? 1.0f : -1.0f);
-		}
-
-		FVector& ViewportLocationDelta = EditorViewportClient->GetViewportNavigationHelper()->LocationDelta;
-		ViewportLocationDelta += LocationDelta;
+		// TODO: In legacy, invert would also occur when using the trackpad - possibly need to handle this in a different way with ITF?
+		const ULevelEditorViewportSettings* Settings = GetDefault<ULevelEditorViewportSettings>();
+		PanCamera(*EditorViewportClient,
+			InMouseDeltaX, InMouseDeltaY,
+			Settings->bInvertMiddleMousePan,
+			// Alt should always make the pan relative to the camera
+			!IsAltDown() && Settings->bWorldSpaceVerticalPan
+		);
+	}
+}
+	
+void UViewportPanInteraction::PanCamera(FEditorViewportClient& Client, float MouseDeltaX, float MouseDeltaY, bool bInvert, bool bUseWorldSpaceUp)
+{
+	FVector LocationDelta;
+
+	const FRotator& ViewRotation = Client.GetViewRotation();
+	if (bUseWorldSpaceUp)
+	{
+		const float Direction = bInvert ? 1.0f : -1.0f;
+		const double YawRadians = FMath::DegreesToRadians(ViewRotation.Yaw);
+		LocationDelta.X = MouseDeltaX * Direction * FMath::Sin(YawRadians);
+		LocationDelta.Y = MouseDeltaX * -Direction * FMath::Cos(YawRadians);
+		LocationDelta.Z = Direction * MouseDeltaY;
 	}
+	else
+	{
+		LocationDelta = FQuat(ViewRotation) * FVector( 0.0f, MouseDeltaX, -MouseDeltaY) * (bInvert ? 1.0f : -1.0f);
+	}
+
+	FVector& ViewportLocationDelta = Client.GetViewportNavigationHelper()->LocationDelta;
+	ViewportLocationDelta += LocationDelta;
 }
 
-FDeviceButtonState UViewportPanInteraction::GetActiveMouseButtonState(const FInputDeviceState& Input)
+UViewportMovePanInteraction::UViewportMovePanInteraction()
 {
-	return Input.Mouse.Middle;
+	InteractionName = TEXT("Perspective Move/Pan");
+	
+	ViewportClickDragBehavior->DisableButton(EKeys::MiddleMouseButton);
+	EnableButton(EKeys::LeftMouseButton);
+	EnableButton(EKeys::RightMouseButton);
+	
+	// TODO: Make behaviors with more buttons automatically have higher priority
+	ViewportClickDragBehavior->SetDefaultPriority(UE::Editor::ViewportInteractions::MULTIBUTTON_DRAG_PRIORITY);
+	ViewportClickDragBehavior->SetDragConfirmedPriority(UE::Editor::ViewportInteractions::MULTIBUTTON_DRAG_CONFIRMED_PRIORITY);
 }
+
+void UViewportMovePanInteraction::OnDragDelta(float InMouseDeltaX, float InMouseDeltaY)
+{
+	if (FEditorViewportClient* const EditorViewportClient = GetEditorViewportClient())
+	{
+		PanCamera(*EditorViewportClient, InMouseDeltaX, InMouseDeltaY, /* Invert */ false, /* World Space up */ true);
+	}
+}
+
diff --git a/Engine/Source/Editor/Experimental/EditorInteractiveToolsFramework/Private/ViewportInteractions/ViewportViewAngleInteraction.cpp b/Engine/Source/Editor/Experimental/EditorInteractiveToolsFramework/Private/ViewportInteractions/ViewportViewAngleInteraction.cpp
--- a/Engine/Source/Editor/Experimental/EditorInteractiveToolsFramework/Private/ViewportInteractions/ViewportViewAngleInteraction.cpp
+++ b/Engine/Source/Editor/Experimental/EditorInteractiveToolsFramework/Private/ViewportInteractions/ViewportViewAngleInteraction.cpp
@@ -73,9 +73,3 @@ void UViewportViewAngleInteraction::OnDragDelta(float InMouseDeltaX, float InMou
 		ViewportRotationDelta += RotDelta;
 	}
 }
-
-FDeviceButtonState UViewportViewAngleInteraction::GetActiveMouseButtonState(const FInputDeviceState& Input)
-{
-	return Input.Mouse.Right;
-}
-
diff --git a/Engine/Source/Editor/Experimental/EditorInteractiveToolsFramework/Private/ViewportInteractions/ViewportZoomInteraction.cpp b/Engine/Source/Editor/Experimental/EditorInteractiveToolsFramework/Private/ViewportInteractions/ViewportZoomInteraction.cpp
--- a/Engine/Source/Editor/Experimental/EditorInteractiveToolsFramework/Private/ViewportInteractions/ViewportZoomInteraction.cpp
+++ b/Engine/Source/Editor/Experimental/EditorInteractiveToolsFramework/Private/ViewportInteractions/ViewportZoomInteraction.cpp
@@ -322,8 +322,7 @@ void UViewportZoomInteraction::OnClickDrag(const FInputDeviceRay& InDragPos)
 	if (FEditorViewportClient* const EditorViewportClient = GetEditorViewportClient())
 	{
 		FVector LocationDelta(0.0f, 0.0f, -DeltaY);
-		FVector& ViewportLocationDelta = EditorViewportClient->GetViewportNavigationHelper()->LocationDelta;
-		ViewportLocationDelta += LocationDelta;
+		EditorViewportClient->GetViewportNavigationHelper()->OrbitDelta += LocationDelta;
 	}
 
 	LastZoomDragY = ScreenPosY;
diff --git a/Engine/Source/Editor/Experimental/EditorInteractiveToolsFramework/Public/ViewportInteractions/ViewportDollyInteraction.h b/Engine/Source/Editor/Experimental/EditorInteractiveToolsFramework/Public/ViewportInteractions/ViewportDollyInteraction.h
new file mode 100644
--- /dev/null
+++ b/Engine/Source/Editor/Experimental/EditorInteractiveToolsFramework/Public/ViewportInteractions/ViewportDollyInteraction.h
@@ -0,0 +1,26 @@
+﻿// Copyright Epic Games, Inc. All Rights Reserved.
+
+#pragma once
+
+#include "ViewportDragInteraction.h"
+
+#include "ViewportDollyInteraction.generated.h"
+
+#define UE_API EDITORINTERACTIVETOOLSFRAMEWORK_API
+
+UCLASS(MinimalAPI, Transient)
+class UViewportDollyInteraction : public UViewportDragInteraction
+{
+	GENERATED_BODY()
+	
+public:
+	UE_API UViewportDollyInteraction();
+	
+	UE_API virtual FInputRayHit CanBeginClickDragSequence(const FInputDeviceRay& InPressPos) override;
+	UE_API virtual void OnDragStart(const FInputDeviceRay& InDragStartPos) override;
+	UE_API virtual void OnDragEnd(const FInputDeviceRay& InDragEndPos) override;
+	UE_API virtual void OnDragDelta(float InMouseDeltaX, float InMouseDeltaY) override;
+	UE_API virtual bool CanBeActivated(const FInputDeviceState& InInputDeviceState = FInputDeviceState()) const override;
+};
+
+#undef UE_API
\ No newline at end of file
diff --git a/Engine/Source/Editor/Experimental/EditorInteractiveToolsFramework/Public/ViewportInteractions/ViewportDragInteraction.h b/Engine/Source/Editor/Experimental/EditorInteractiveToolsFramework/Public/ViewportInteractions/ViewportDragInteraction.h
--- a/Engine/Source/Editor/Experimental/EditorInteractiveToolsFramework/Public/ViewportInteractions/ViewportDragInteraction.h
+++ b/Engine/Source/Editor/Experimental/EditorInteractiveToolsFramework/Public/ViewportInteractions/ViewportDragInteraction.h
@@ -52,10 +52,6 @@ public:
 	}
 
 protected:
-	/**
-	 * Returns the state of the mouse button currently used by this Interaction. Left mouse by default, override for other buttons.
-	 */
-	UE_API virtual FDeviceButtonState GetActiveMouseButtonState(const FInputDeviceState& Input);
 
 	UE_API void EnableButton(const FKey& InButton) const;
 
diff --git a/Engine/Source/Editor/Experimental/EditorInteractiveToolsFramework/Public/ViewportInteractions/ViewportMoveYawInteraction.h b/Engine/Source/Editor/Experimental/EditorInteractiveToolsFramework/Public/ViewportInteractions/ViewportMoveYawInteraction.h
--- a/Engine/Source/Editor/Experimental/EditorInteractiveToolsFramework/Public/ViewportInteractions/ViewportMoveYawInteraction.h
+++ b/Engine/Source/Editor/Experimental/EditorInteractiveToolsFramework/Public/ViewportInteractions/ViewportMoveYawInteraction.h
@@ -30,7 +30,6 @@ public:
 
 	//~ Begin UViewportDragInteraction
 	UE_API virtual void OnDragDelta(float InMouseDeltaX, float InMouseDeltaY) override;
-	UE_API virtual FDeviceButtonState GetActiveMouseButtonState(const FInputDeviceState& Input) override;
 	UE_API virtual bool CanBeActivated(const FInputDeviceState& InInputDeviceState = FInputDeviceState()) const override;
 	//~ End UViewportDragInteraction
 };
diff --git a/Engine/Source/Editor/Experimental/EditorInteractiveToolsFramework/Public/ViewportInteractions/ViewportOrthoPanInteraction.h b/Engine/Source/Editor/Experimental/EditorInteractiveToolsFramework/Public/ViewportInteractions/ViewportOrthoPanInteraction.h
--- a/Engine/Source/Editor/Experimental/EditorInteractiveToolsFramework/Public/ViewportInteractions/ViewportOrthoPanInteraction.h
+++ b/Engine/Source/Editor/Experimental/EditorInteractiveToolsFramework/Public/ViewportInteractions/ViewportOrthoPanInteraction.h
@@ -27,7 +27,6 @@ public:
 
 	//~ Begin UViewportDragInteraction
 	UE_API virtual void OnDragDelta(float InMouseDeltaX, float InMouseDeltaY) override;
-	UE_API virtual FDeviceButtonState GetActiveMouseButtonState(const FInputDeviceState& Input) override;
 	//~ End UViewportDragInteraction
 
 	UE_API virtual bool CanBeActivated(const FInputDeviceState& InInputDeviceState = FInputDeviceState()) const override;
diff --git a/Engine/Source/Editor/Experimental/EditorInteractiveToolsFramework/Public/ViewportInteractions/ViewportPanInteraction.h b/Engine/Source/Editor/Experimental/EditorInteractiveToolsFramework/Public/ViewportInteractions/ViewportPanInteraction.h
--- a/Engine/Source/Editor/Experimental/EditorInteractiveToolsFramework/Public/ViewportInteractions/ViewportPanInteraction.h
+++ b/Engine/Source/Editor/Experimental/EditorInteractiveToolsFramework/Public/ViewportInteractions/ViewportPanInteraction.h
@@ -11,7 +11,7 @@
 class UViewportInteractionsBehaviorSource;
 
 /**
- * Implements the LMB + Drag Horizontal plane move + Yaw interaction
+ * Implements the middle mouse pan interaction
  */
 UCLASS(MinimalAPI, Transient)
 class UViewportPanInteraction : public UViewportDragInteraction
@@ -29,12 +29,25 @@ public:
 
 	//~ Begin UViewportDragInteraction
 	UE_API virtual void OnDragDelta(float InMouseDeltaX, float InMouseDeltaY) override;
-	UE_API virtual FDeviceButtonState GetActiveMouseButtonState(const FInputDeviceState& Input) override;
 	//~ End UViewportDragInteraction
+	
+protected:
+	static void PanCamera(FEditorViewportClient& Client, float MouseDeltaX, float MouseDeltaY, bool bInvert, bool bUseWorldSpaceUp);
+};
 
-private:
-	// Pan direction can be inverted e.g. with a key modifier
-	bool bMultiButtonPanInteraction = false;
+/**
+ * Implements the LMB + RMB Drag world-space-up Pan interaction
+ */
+UCLASS(MinimalAPI, Transient)
+class UViewportMovePanInteraction : public UViewportPanInteraction
+{
+	GENERATED_BODY()
+	
+public:
+	UE_API UViewportMovePanInteraction();
+	
+protected:
+	virtual void OnDragDelta(float InMouseDeltaX, float InMouseDeltaY) override;
 };
 
 #undef UE_API
diff --git a/Engine/Source/Editor/Experimental/EditorInteractiveToolsFramework/Public/ViewportInteractions/ViewportViewAngleInteraction.h b/Engine/Source/Editor/Experimental/EditorInteractiveToolsFramework/Public/ViewportInteractions/ViewportViewAngleInteraction.h
--- a/Engine/Source/Editor/Experimental/EditorInteractiveToolsFramework/Public/ViewportInteractions/ViewportViewAngleInteraction.h
+++ b/Engine/Source/Editor/Experimental/EditorInteractiveToolsFramework/Public/ViewportInteractions/ViewportViewAngleInteraction.h
@@ -28,7 +28,6 @@ public:
 
 	//~ Begin UViewportDragInteraction
 	UE_API virtual void OnDragDelta(float InMouseDeltaX, float InMouseDeltaY) override;
-	UE_API virtual FDeviceButtonState GetActiveMouseButtonState(const FInputDeviceState& Input) override;
 	//~ End UViewportDragInteraction
 
 };
diff --git a/Engine/Source/Editor/UnrealEd/Private/EditorViewportClient.cpp b/Engine/Source/Editor/UnrealEd/Private/EditorViewportClient.cpp
--- a/Engine/Source/Editor/UnrealEd/Private/EditorViewportClient.cpp
+++ b/Engine/Source/Editor/UnrealEd/Private/EditorViewportClient.cpp
@@ -141,6 +141,18 @@ static bool GetDefaultLowDPIPreviewValue()
 	return CVarEditorViewportHighDPIPtr->GetInt() == 0;
 }
 
+bool UsesViewportInteractions(const FEditorModeTools* InModeTools)
+{
+	if (InModeTools && UE::Editor::ViewportInteractions::UseEditorViewportInteractions())
+	{
+		if (UModeManagerInteractiveToolsContext* InputToolsBehaviorSource = InModeTools->GetInteractiveToolsContext())
+		{
+			return InputToolsBehaviorSource->UsesViewportInteractions();
+		}
+	}
+	return false;
+}
+
 /**
  * Returns true if any viewport interaction of the specified type is active
  */
@@ -1595,10 +1607,6 @@ FSceneView* FEditorViewportClient::CalcSceneView(FSceneViewFamily* ViewFamily, c
 		View->AdditionalSelectionOutlineColors[OutlineColorIndex] = EditorStyle->AdditionalSelectionColors[OutlineColorIndex];
 	}
 
-	View->CompositePrimitiveParams.OccludedDitheringPatternSize = EditorStyle->OccludedDitheringPatternSize;
-	View->CompositePrimitiveParams.OccludedDithering = EditorStyle->OccludedDithering;
-	View->CompositePrimitiveParams.OccludedBrightness = EditorStyle->OccludedBrightness;
-
 	int32 FamilyIndex = ViewFamily->Views.Add(View);
 	check(FamilyIndex == View->StereoViewIndex || View->StereoViewIndex == INDEX_NONE);
 
@@ -1723,8 +1731,15 @@ void FEditorViewportClient::Tick(float DeltaTime)
 		// Update any real-time camera movement
 		UpdateCameraMovement( DeltaTime );
 
-		UpdateMouseDelta();
-
+		if (UsesViewportInteractions(ModeTools.Get()))
+		{
+			ConsumeExternalCameraMovement();
+		}
+		else
+		{
+			UpdateMouseDelta();
+		}
+		
 		UpdateGestureDelta();
 
 		EndCameraMovement();
@@ -2088,8 +2103,10 @@ bool FEditorViewportClient::IsActiveViewportType(ELevelViewportType InViewportTy
 // Updates real-time camera movement.  Should be called every viewport tick!
 void FEditorViewportClient::UpdateCameraMovement( float DeltaTime )
 {
+	const bool bIsITFInteractions = UsesViewportInteractions(ModeTools.Get());
+	
 	// We only want to move perspective cameras around like this
-	if( Viewport != nullptr && IsPerspective() && !ShouldOrbitCamera() )
+	if( Viewport != nullptr && IsPerspective() && (bIsITFInteractions ? ViewportNavigationHelper->OrbitDelta.IsNearlyZero() : !ShouldOrbitCamera()))
 	{
 		const bool bEnable = false;
 		ToggleOrbitCamera(bEnable);
@@ -2097,26 +2114,7 @@ void FEditorViewportClient::UpdateCameraMovement( float DeltaTime )
 		const bool bIsUsingTrackpad = FSlateApplication::Get().IsUsingTrackpad();
 
 		// Certain keys are only available while the flight camera input mode is active
-		const bool bUsingFlightInput = IsFlightCameraInputModeActive() || bIsUsingTrackpad || ViewportNavigationHelper->bIsMouseLooking;
-
-		// Is the current press unmodified?
-		const bool bUnmodifiedPress = !IsAltPressed() && !IsShiftPressed() && !IsCtrlPressed() && !IsCmdPressed();
-
-		// Do we want to use the regular arrow keys for flight input?
-		// Because the arrow keys are used for things like nudging actors, we'll only do this while the press is unmodified
-		const bool bRemapArrowKeys = bUnmodifiedPress || IsViewportInteractionActive(ModeTools.Get(), UE::Editor::ViewportInteractions::CameraTranslate);
-
-		const bool bITFCameraTranslateActive = IsViewportInteractionActive(ModeTools.Get(),UE::Editor::ViewportInteractions::CameraTranslate);
-		const bool bITFCameraRotateActive = IsViewportInteractionActive(ModeTools.Get(),UE::Editor::ViewportInteractions::CameraRotate);
-		const bool bITFFOVActive = IsViewportInteractionActive(ModeTools.Get(),UE::Editor::ViewportInteractions::FOV);
-
-		// Do we want to remap the various WASD keys for flight input?
-		const bool bRemapWASDKeys =
-			(bUnmodifiedPress || (GetDefault<ULevelEditorViewportSettings>()->FlightCameraControlExperimentalNavigation && IsShiftPressed())) &&
-			(GetDefault<ULevelEditorViewportSettings>()->FlightCameraControlType == WASD_Always ||
-			( bUsingFlightInput &&
-			( GetDefault<ULevelEditorViewportSettings>()->FlightCameraControlType == WASD_RMBOnly && (Viewport->KeyState(EKeys::RightMouseButton ) ||Viewport->KeyState(EKeys::MiddleMouseButton) || Viewport->KeyState(EKeys::LeftMouseButton) || bIsUsingTrackpad ) ) ) ) &&
-			!MouseDeltaTracker->UsingDragTool();
+		const bool bUsingFlightInput = (bIsITFInteractions ? ViewportNavigationHelper->bIsMouseLooking : IsFlightCameraInputModeActive()) || bIsUsingTrackpad;
 
 		// Apply impulse from magnify gesture and reset impulses if we're using WASD keys
 		CameraUserImpulseData->MoveForwardBackwardImpulse = GestureMoveForwardBackwardImpulse;
@@ -2129,53 +2127,66 @@ void FEditorViewportClient::UpdateCameraMovement( float DeltaTime )
 		CameraUserImpulseData->RotateRollImpulse = 0.0f;
 
 		GestureMoveForwardBackwardImpulse = 0.0f;
-
-		bool bForwardKeyState = false;
-		bool bBackwardKeyState = false;
-		bool bRightKeyState = false;
-		bool bLeftKeyState = false;
-
-		bool bWorldUpKeyState = false;
-		bool bWorldDownKeyState = false;
-		bool bLocalUpKeyState = false;
-		bool bLocalDownKeyState = false;
-		bool bZoomOutKeyState = false;
-		bool bZoomInKeyState = false;
-
-		bool bRotateUpKeyState = false;
-		bool bRotateDownKeyState = false;
-		bool bRotateLeftKeyState = false;
-		bool bRotateRightKeyState = false;
-		// Iterate through all key mappings to generate key state flags
-		for (uint32 i = 0; i < static_cast<uint8>(EMultipleKeyBindingIndex::NumChords); ++i)
-		{
-			EMultipleKeyBindingIndex ChordIndex = static_cast<EMultipleKeyBindingIndex> (i);
-			bForwardKeyState |= Viewport->KeyState(FViewportNavigationCommands::Get().Forward->GetActiveChord(ChordIndex)->Key);
-			bBackwardKeyState |= Viewport->KeyState(FViewportNavigationCommands::Get().Backward->GetActiveChord(ChordIndex)->Key);
-			bRightKeyState |= Viewport->KeyState(FViewportNavigationCommands::Get().Right->GetActiveChord(ChordIndex)->Key);
-			bLeftKeyState |= Viewport->KeyState(FViewportNavigationCommands::Get().Left->GetActiveChord(ChordIndex)->Key);
-
-			bWorldUpKeyState |= Viewport->KeyState(FViewportNavigationCommands::Get().WorldUp->GetActiveChord(ChordIndex)->Key);
-			bWorldDownKeyState |= Viewport->KeyState(FViewportNavigationCommands::Get().WorldDown->GetActiveChord(ChordIndex)->Key);
-			bLocalUpKeyState |= Viewport->KeyState(FViewportNavigationCommands::Get().LocalUp->GetActiveChord(ChordIndex)->Key);
-			bLocalDownKeyState |= Viewport->KeyState(FViewportNavigationCommands::Get().LocalDown->GetActiveChord(ChordIndex)->Key);
-			bZoomOutKeyState |= Viewport->KeyState(FViewportNavigationCommands::Get().FovZoomOut->GetActiveChord(ChordIndex)->Key);
-			bZoomInKeyState |= Viewport->KeyState(FViewportNavigationCommands::Get().FovZoomIn->GetActiveChord(ChordIndex)->Key);
-
-			bRotateUpKeyState |= Viewport->KeyState(FViewportNavigationCommands::Get().RotateUp->GetActiveChord(ChordIndex)->Key);
-			bRotateDownKeyState |= Viewport->KeyState(FViewportNavigationCommands::Get().RotateDown->GetActiveChord(ChordIndex)->Key);
-			bRotateLeftKeyState |= Viewport->KeyState(FViewportNavigationCommands::Get().RotateLeft->GetActiveChord(ChordIndex)->Key);
-			bRotateRightKeyState |= Viewport->KeyState(FViewportNavigationCommands::Get().RotateRight->GetActiveChord(ChordIndex)->Key);
-		}
-
-		if (!CameraController->IsRotating())
+		
+		if (bIsITFInteractions && ViewportNavigationHelper)
 		{
-			CameraController->GetConfig().bForceRotationalPhysics = false;
+			ViewportNavigationHelper->ConsumeImpulseData(*CameraUserImpulseData);
 		}
-
-		// Check if we need to use legacy translation input
-		if (!bITFCameraTranslateActive)
+		else
 		{
+			// Is the current press unmodified?
+			const bool bUnmodifiedPress = !IsAltPressed() && !IsShiftPressed() && !IsCtrlPressed() && !IsCmdPressed();
+
+			// Do we want to use the regular arrow keys for flight input?
+			// Because the arrow keys are used for things like nudging actors, we'll only do this while the press is unmodified
+			const bool bRemapArrowKeys = bUnmodifiedPress;
+
+			// Do we want to remap the various WASD keys for flight input?
+			const bool bRemapWASDKeys =
+				(bUnmodifiedPress || (GetDefault<ULevelEditorViewportSettings>()->FlightCameraControlExperimentalNavigation && IsShiftPressed())) &&
+				(GetDefault<ULevelEditorViewportSettings>()->FlightCameraControlType == WASD_Always ||
+				( bUsingFlightInput &&
+				( GetDefault<ULevelEditorViewportSettings>()->FlightCameraControlType == WASD_RMBOnly && (Viewport->KeyState(EKeys::RightMouseButton ) ||Viewport->KeyState(EKeys::MiddleMouseButton) || Viewport->KeyState(EKeys::LeftMouseButton) || bIsUsingTrackpad ) ) ) ) &&
+				!MouseDeltaTracker->UsingDragTool();
+				
+			bool bForwardKeyState = false;
+			bool bBackwardKeyState = false;
+			bool bRightKeyState = false;
+			bool bLeftKeyState = false;
+
+			bool bWorldUpKeyState = false;
+			bool bWorldDownKeyState = false;
+			bool bLocalUpKeyState = false;
+			bool bLocalDownKeyState = false;
+			bool bZoomOutKeyState = false;
+			bool bZoomInKeyState = false;
+
+			bool bRotateUpKeyState = false;
+			bool bRotateDownKeyState = false;
+			bool bRotateLeftKeyState = false;
+			bool bRotateRightKeyState = false;
+			// Iterate through all key mappings to generate key state flags
+			for (uint32 i = 0; i < static_cast<uint8>(EMultipleKeyBindingIndex::NumChords); ++i)
+			{
+				EMultipleKeyBindingIndex ChordIndex = static_cast<EMultipleKeyBindingIndex> (i);
+				bForwardKeyState |= Viewport->KeyState(FViewportNavigationCommands::Get().Forward->GetActiveChord(ChordIndex)->Key);
+				bBackwardKeyState |= Viewport->KeyState(FViewportNavigationCommands::Get().Backward->GetActiveChord(ChordIndex)->Key);
+				bRightKeyState |= Viewport->KeyState(FViewportNavigationCommands::Get().Right->GetActiveChord(ChordIndex)->Key);
+				bLeftKeyState |= Viewport->KeyState(FViewportNavigationCommands::Get().Left->GetActiveChord(ChordIndex)->Key);
+
+				bWorldUpKeyState |= Viewport->KeyState(FViewportNavigationCommands::Get().WorldUp->GetActiveChord(ChordIndex)->Key);
+				bWorldDownKeyState |= Viewport->KeyState(FViewportNavigationCommands::Get().WorldDown->GetActiveChord(ChordIndex)->Key);
+				bLocalUpKeyState |= Viewport->KeyState(FViewportNavigationCommands::Get().LocalUp->GetActiveChord(ChordIndex)->Key);
+				bLocalDownKeyState |= Viewport->KeyState(FViewportNavigationCommands::Get().LocalDown->GetActiveChord(ChordIndex)->Key);
+				bZoomOutKeyState |= Viewport->KeyState(FViewportNavigationCommands::Get().FovZoomOut->GetActiveChord(ChordIndex)->Key);
+				bZoomInKeyState |= Viewport->KeyState(FViewportNavigationCommands::Get().FovZoomIn->GetActiveChord(ChordIndex)->Key);
+
+				bRotateUpKeyState |= Viewport->KeyState(FViewportNavigationCommands::Get().RotateUp->GetActiveChord(ChordIndex)->Key);
+				bRotateDownKeyState |= Viewport->KeyState(FViewportNavigationCommands::Get().RotateDown->GetActiveChord(ChordIndex)->Key);
+				bRotateLeftKeyState |= Viewport->KeyState(FViewportNavigationCommands::Get().RotateLeft->GetActiveChord(ChordIndex)->Key);
+				bRotateRightKeyState |= Viewport->KeyState(FViewportNavigationCommands::Get().RotateRight->GetActiveChord(ChordIndex)->Key);
+			}
+
 			// Forward/back
 			if( ( bRemapWASDKeys && bForwardKeyState ) ||
 				( bRemapArrowKeys && Viewport->KeyState( EKeys::Up ) ) ||
@@ -2227,11 +2238,7 @@ void FEditorViewportClient::UpdateCameraMovement( float DeltaTime )
 			{
 				CameraUserImpulseData->MoveLocalUpDownImpulse -= 1.0f;
 			}
-		}
 
-		// Check if we need to use legacy FOV input
-		if (!bITFFOVActive)
-		{
 			// Zoom FOV out/in
 			if( ( bRemapWASDKeys && bZoomOutKeyState) ||
 				( bUnmodifiedPress && bUseNumpadCameraControl && Viewport->KeyState( EKeys::NumPadOne ) ) )
@@ -2243,11 +2250,7 @@ void FEditorViewportClient::UpdateCameraMovement( float DeltaTime )
 			{
 				CameraUserImpulseData->ZoomOutInImpulse -= 1.0f;
 			}
-		}
 
-		// Check if we need to use legacy rotation input
-		if (!bITFCameraRotateActive)
-		{
 			// Rotate up/down
 			if (bRemapWASDKeys && bRotateUpKeyState)
 			{
@@ -2273,10 +2276,13 @@ void FEditorViewportClient::UpdateCameraMovement( float DeltaTime )
 			}
 		}
 
-		// Add camera inputs coming from ITF, if any
-		if (bITFCameraTranslateActive || bITFCameraRotateActive || bITFFOVActive)
+		if (!CameraController->IsRotating())
+		{
+			CameraController->GetConfig().bForceRotationalPhysics = false;
+		}
+		
+		if (CameraUserImpulseData->HasRotationalImpulse())
 		{
-			ViewportNavigationHelper->ConsumeImpulseData(*CameraUserImpulseData);
 			CameraController->GetConfig().bForceRotationalPhysics = true;
 		}
 
@@ -2378,8 +2384,8 @@ void FEditorViewportClient::UpdateCameraMovement( float DeltaTime )
 			// Something has changed!
 			const bool bInvalidateChildViews=true;
 
-			// When flying the camera around the hit proxies dont need to be invalidated since we are flying around and not clicking on anything
-			const bool bInvalidateHitProxies= UE::Editor::ViewportInteractions::UseEditorViewportInteractions()
+			// When flying the camera around the hit proxies don't need to be invalidated since we are flying around and not clicking on anything
+			const bool bInvalidateHitProxies = bIsITFInteractions
             	? !ViewportNavigationHelper->bIsMouseLooking
             	: !IsFlightCameraActive();
 			Invalidate(bInvalidateChildViews,bInvalidateHitProxies);
@@ -2404,6 +2410,186 @@ void FEditorViewportClient::UpdateCameraMovement( float DeltaTime )
 	}
 }
 
+bool FEditorViewportClient::ConsumeExternalCameraMovement()
+{
+	if (!Viewport || !ViewportNavigationHelper || !ViewportNavigationHelper->HasTransformDelta())
+	{
+		return false;
+	}
+	
+	FVector Translation = FVector::ZeroVector;
+	FRotator Rotation = FRotator::ZeroRotator;
+	FVector Orbit = FVector::ZeroVector;
+	
+	BeginCameraMovement(true);
+	
+	ViewportNavigationHelper->ConsumeTransformDelta(Translation, Rotation, Orbit);
+	
+	if (IsPerspective())
+	{
+		if (!FMath::IsNearlyZero(Orbit.X) || !FMath::IsNearlyZero(Orbit.Y))
+		{
+			FRotator NewRotation = FRotator::ZeroRotator;
+            NewRotation.Yaw = Orbit.X;
+            NewRotation.Pitch = Orbit.Y;
+    
+            OrbitPerspectiveCamera(NewRotation);
+		}
+		
+		if (!FMath::IsNearlyZero(Orbit.Z))
+		{
+			DollyPerspectiveCamera(Orbit.Z);
+		}
+		
+		if (!Translation.IsNearlyZero() || !Rotation.IsNearlyZero())
+		{
+			MoveViewportPerspectiveCamera(Translation * GetCameraSpeed(), Rotation);
+		}
+	}
+	else
+	{
+		if (!FMath::IsNearlyZero(Orbit.Z))
+		{
+			float NewOrthoRatio = 1 + (FMath::Abs(Orbit.Z) / CAMERA_ZOOM_DAMPEN);
+			if ( Orbit.Z < 0 )
+			{
+				NewOrthoRatio = 1 / NewOrthoRatio;
+			}
+		
+			const float Zoom = GetOrthoZoom() * NewOrthoRatio;
+			SetOrthoZoom( FMath::Clamp<float>( Zoom, GetMinimumOrthoZoom(), MAX_ORTHOZOOM ) );
+		}
+		
+		if (!Translation.IsNearlyZero())
+		{
+			SetViewLocation(GetViewLocation() + Translation);
+		}
+	}
+	
+	return true;
+}
+
+void FEditorViewportClient::OrbitPerspectiveCamera(const FRotator& InOrbit)
+{
+	// This is pulled from InputAxisForOrbit(), with the input checking code removed
+	if (GetDefault<ULevelEditorViewportSettings>()->bOrbitCameraAroundSelection)
+	{
+		FRotator NewRotation = InOrbit;
+		
+		FVector OrbitPoint;
+		if (GetPivotForOrbit(OrbitPoint))
+		{
+			DefaultOrbitLocation = OrbitPoint;
+		}
+
+		// Override the default orbit behavior to allow orbiting around a given pivot
+		// This uses different computations from the default orbit behavior so it must not ToggleOrbitCamera
+		const bool bEnable = false;
+		ToggleOrbitCamera(bEnable);
+		
+		const FRotator ViewRotation = GetViewRotation();
+		
+		// Compute the look-at and view location centered on the orbit point
+		const FVector LookAtOffset = GetLookAtLocation() - DefaultOrbitLocation;
+		const FVector ViewLocationOffset = GetViewLocation() - DefaultOrbitLocation;
+		
+		// When the roll is at 180 degrees, it means the view is upside down, so invert the yaw rotation
+		if (ViewRotation.Roll == 180.f)
+		{
+			NewRotation.Yaw = -NewRotation.Yaw;
+		}
+		
+		// Compute the delta rotation to apply as a transform
+		FRotator DeltaRotation(NewRotation);
+		FRotator ViewRotationNoPitch = ViewRotation;
+		ViewRotationNoPitch.Pitch = 0;
+
+		FTransform ViewRotationTransform = FTransform(ViewRotationNoPitch);
+		FTransform DeltaRotationTransform = ViewRotationTransform.Inverse() * FTransform(DeltaRotation) * ViewRotationTransform;
+
+		// Apply the delta rotation to the view rotation
+		FRotator RotatedView = (FTransform(ViewRotation) * DeltaRotationTransform).Rotator();
+	
+		// Correct the rotation to remove drift in the roll
+		if (FMath::IsNearlyEqual(FMath::Abs(RotatedView.Roll), 180.f, 1.f))
+		{
+			RotatedView.Roll = 180.f;
+		}
+		else if (FMath::IsNearlyZero(RotatedView.Roll, 1.f))
+		{
+			RotatedView.Roll = 0.f;
+		}
+		else
+		{
+			// FTransform::Rotator() returns an invalid RotatedView with roll in it when the initial pitch is at +/-90 degrees due to a singularity
+			// FVector::ToOrientationRotator uses an alternate computation that doesn't suffer from the singularity. However, the roll it returns
+			// is always 0 so it's not possible to tell if it's upside down and its yaw is flipped 180 degrees
+			FVector ViewVector = ViewRotation.Vector();
+			FVector RotatedViewVector = DeltaRotationTransform.TransformVector(ViewVector);
+			FRotator RotatedViewRotator = RotatedViewVector.ToOrientationRotator();
+
+			RotatedView = RotatedViewRotator + FRotator(0.f, -180.f, ViewRotation.Roll);
+		}
+
+		SetViewRotation(RotatedView);
+
+		// Set the new rotated look-at
+		FVector RotatedLookAtOffset = DeltaRotationTransform.TransformVector(LookAtOffset);
+		FVector RotatedLookAt = DefaultOrbitLocation + RotatedLookAtOffset;
+		SetLookAtLocation(RotatedLookAt);
+
+		// Set the new rotated view location
+		FVector RotatedViewOffset = DeltaRotationTransform.TransformVector(ViewLocationOffset);
+		FVector RotatedViewLocation = DefaultOrbitLocation + RotatedViewOffset;
+		SetViewLocation(RotatedViewLocation);
+	}
+	else
+	{
+		ToggleOrbitCamera(true);
+
+		FViewportCameraTransform& ViewTransform = GetViewTransform();
+		
+		SetViewRotation( GetViewRotation() + FRotator( InOrbit.Pitch, -InOrbit.Yaw, InOrbit.Roll ) );
+
+		/*
+		 * Recalculates the view location according to the new SetViewRotation() did earlier.
+		 */
+		SetViewLocation(ViewTransform.ComputeOrbitMatrix().Inverse().GetOrigin());
+
+		ToggleOrbitCamera(false);
+	}
+	
+	PerspectiveCameraMoved();
+}
+
+void FEditorViewportClient::DollyPerspectiveCamera(double InDelta)
+{
+	ToggleOrbitCamera(true);
+	
+	FViewportCameraTransform& ViewTransform = GetViewTransform();
+	
+	const float CameraSpeed = GetCameraSpeed();
+	const float CameraSpeedDistanceScale = ShouldScaleCameraSpeedByDistance() ? FMath::Min(FVector::Dist( GetViewLocation(), GetLookAtLocation() ) / 1000.f, 1000.f) : 1.0f;
+	
+	FMatrix OrbitMatrix = ViewTransform.ComputeOrbitMatrix().InverseFast();
+
+	FVector DeltaLocation = FVector(0, InDelta * CameraSpeed * CameraSpeedDistanceScale, 0);
+
+	FVector LookAt = ViewTransform.GetLookAt();
+
+	// Orient the delta down the view direction towards the look at
+	FMatrix RotMat =
+		FTranslationMatrix( -LookAt ) *
+		FRotationMatrix( FRotator(0,GetViewRotation().Yaw,0) ) *
+		FRotationMatrix( FRotator(0, 0, GetViewRotation().Pitch));
+
+	FVector TransformedDelta = RotMat.InverseFast().TransformVector(DeltaLocation);
+
+	SetViewLocation( OrbitMatrix.GetOrigin() + TransformedDelta );
+
+	ToggleOrbitCamera(false);
+}
+
 /**
  * Forcibly disables lighting show flags if there are no lights in the scene, or restores lighting show
  * flags if lights are added to the scene.
@@ -2648,7 +2834,7 @@ void FEditorViewportClient::UpdateMouseDelta()
 	if( Viewport )
 	{
 		// If mouse delta tracker has delta OR there is transform delta coming from ITF Viewport Interaction, handle it
-		if (!DragDelta.IsNearlyZero() || ViewportNavigationHelper->HasTransformDelta())
+		if (!DragDelta.IsNearlyZero())
 		{
 			const bool LeftMouseButtonDown = Viewport->KeyState(EKeys::LeftMouseButton);
 			const bool MiddleMouseButtonDown = Viewport->KeyState(EKeys::MiddleMouseButton);
@@ -2666,13 +2852,7 @@ void FEditorViewportClient::UpdateMouseDelta()
 			FVector Scale = FVector::ZeroVector;
 			EAxisList::Type CurrentAxis = Widget->GetCurrentAxis();
 
-			// Try to handle ITF Viewport inputs before legacy (currently zoom and camera drag ones)
-			if (IsViewportInteractionActive(ModeTools.Get(), UE::Editor::ViewportInteractions::Zoom) ||
-				IsViewportInteractionActive(ModeTools.Get(), UE::Editor::ViewportInteractions::CameraDrag))
-			{
-				ViewportNavigationHelper->ConsumeTransformDelta(Drag, Rot);
-			}
-			else if ( IsOrtho() && ( LeftMouseButtonDown || bIsUsingTrackpad ) && RightMouseButtonDown )
+			if ( IsOrtho() && ( LeftMouseButtonDown || bIsUsingTrackpad ) && RightMouseButtonDown )
 			{
 				bWidgetAxisControlledByDrag = false;
 				Widget->SetCurrentAxis( EAxisList::None );
diff --git a/Engine/Source/Editor/UnrealEd/Private/ViewportClientNavigationHelper.cpp b/Engine/Source/Editor/UnrealEd/Private/ViewportClientNavigationHelper.cpp
--- a/Engine/Source/Editor/UnrealEd/Private/ViewportClientNavigationHelper.cpp
+++ b/Engine/Source/Editor/UnrealEd/Private/ViewportClientNavigationHelper.cpp
@@ -15,13 +15,15 @@ bool FViewportClientNavigationHelper::HasTransformDelta() const
 	return !RotationDelta.IsNearlyZero() || !LocationDelta.IsNearlyZero() || !OrbitDelta.IsNearlyZero();
 }
 
-void FViewportClientNavigationHelper::ConsumeTransformDelta(FVector& OutLocationDelta, FRotator& OutRotationDelta)
+void FViewportClientNavigationHelper::ConsumeTransformDelta(FVector& OutLocationDelta, FRotator& OutRotationDelta, FVector& OutOrbitDelta)
 {
 	OutLocationDelta = LocationDelta;
 	OutRotationDelta = RotationDelta;
+	OutOrbitDelta = OrbitDelta;
 
 	// Reset after use
 	ResetTransformDelta();
+	OrbitDelta = FVector::ZeroVector;
 }
 
 void FViewportClientNavigationHelper::ResetTransformDelta()
diff --git a/Engine/Source/Editor/UnrealEd/Public/CameraController.h b/Engine/Source/Editor/UnrealEd/Public/CameraController.h
--- a/Engine/Source/Editor/UnrealEd/Public/CameraController.h
+++ b/Engine/Source/Editor/UnrealEd/Public/CameraController.h
@@ -67,6 +67,26 @@ public:
 		  ZoomOutInImpulse( 0.0f )
 	{
 	}
+	
+	bool HasMovement() const
+	{
+		return RotateYawVelocityModifier != 0.0f ||
+			RotatePitchVelocityModifier != 0.0f ||
+			RotateRollVelocityModifier != 0.0f ||
+			MoveForwardBackwardImpulse != 0.0f ||
+			MoveRightLeftImpulse != 0.0f ||
+			MoveWorldUpDownImpulse != 0.0f ||
+			MoveLocalUpDownImpulse != 0.0f ||
+			ZoomOutInImpulse != 0.0f ||
+			HasRotationalImpulse();
+	}
+	
+	bool HasRotationalImpulse() const
+	{
+		return RotateYawImpulse != 0.0f ||
+			RotatePitchImpulse != 0.0f ||
+			RotateRollImpulse != 0.0f;
+	}
 };
 
 
diff --git a/Engine/Source/Editor/UnrealEd/Public/EditorViewportClient.h b/Engine/Source/Editor/UnrealEd/Public/EditorViewportClient.h
--- a/Engine/Source/Editor/UnrealEd/Public/EditorViewportClient.h
+++ b/Engine/Source/Editor/UnrealEd/Public/EditorViewportClient.h
@@ -1824,6 +1824,21 @@ private:
 	 */
 	UNREALED_API void UpdateCameraMovement( float DeltaTime );
 
+	/**
+	 * Updates the camera based on external movement sources.
+	 */
+	bool ConsumeExternalCameraMovement();
+
+	/**
+	 * Directs the camera to orbit based on a screen space movement delta
+	 */
+	void OrbitPerspectiveCamera(const FRotator& InOrbit);
+
+	/**
+	 * Dollies the camera towards or away from its orbit point
+	 */
+	void DollyPerspectiveCamera(double InDelta);
+
 	/**
 	 * Forcibly disables lighting show flags if there are no lights in the scene, or restores lighting show
 	 * flags if lights are added to the scene.
diff --git a/Engine/Source/Editor/UnrealEd/Public/ViewportClientNavigationHelper.h b/Engine/Source/Editor/UnrealEd/Public/ViewportClientNavigationHelper.h
--- a/Engine/Source/Editor/UnrealEd/Public/ViewportClientNavigationHelper.h
+++ b/Engine/Source/Editor/UnrealEd/Public/ViewportClientNavigationHelper.h
@@ -26,7 +26,9 @@ public:
 	FRotator RotationDelta;
 
 	/**
-	 * A delta which is used by Editor Viewport Client to compute orbit movement values
+	 * A delta which is used by Editor Viewport Client to compute orbit movement values.
+	 * In perspective view, Z indicates a Dolly action.
+	 * In orthographic view, Z indicates an orthographic zoom
 	 */
 	FVector OrbitDelta;
 
@@ -37,7 +39,7 @@ public:
 	FCameraControllerUserImpulseData ImpulseDataDelta;
 
 	/**
-	 * Check whether non-zero location or rotation deltas are available. Also checks for Orbit delta.
+	 * Check whether non-zero location, rotation, or orbit deltas are available.
 	 * Useful to skip certain logic in FEditorViewportClient::UpdateMouseDelta() or for other purposes
 	 */
 	bool HasTransformDelta() const;
@@ -47,7 +49,7 @@ public:
 	 * @param OutLocationDelta location delta will be copied to OutLocationDelta
 	 * @param OutRotationDelta rotation delta will be copied to OutRotationDelta
 	 */
-	void ConsumeTransformDelta(FVector& OutLocationDelta, FRotator& OutRotationDelta);
+	void ConsumeTransformDelta(FVector& OutLocationDelta, FRotator& OutRotationDelta, FVector& OutOrbitDelta);
 
 	/**
 	 * Retrieve current impulse deltas and zero them out
