[ITF Viewport Interactions] Fix orbit value initialization on ViewportClientNavigationHelper.

#rnx

diff --git a/Engine/Source/Editor/UnrealEd/Private/ViewportClientNavigationHelper.cpp b/Engine/Source/Editor/UnrealEd/Private/ViewportClientNavigationHelper.cpp
--- a/Engine/Source/Editor/UnrealEd/Private/ViewportClientNavigationHelper.cpp
+++ b/Engine/Source/Editor/UnrealEd/Private/ViewportClientNavigationHelper.cpp
@@ -30,6 +30,7 @@ void FViewportClientNavigationHelper::ResetTransformDelta()
 {
 	LocationDelta = FVector::ZeroVector;
 	RotationDelta = FRotator::ZeroRotator;
+	OrbitDelta = FVector::ZeroVector;
 }
 
 void FViewportClientNavigationHelper::ConsumeImpulseData(FCameraControllerUserImpulseData& OutImpulseData)
