[ITF Viewport Interactions] Hotfix to move a null check in EditorViewportClient to avoid a static analysis error.

#rnx

diff --git a/Engine/Source/Editor/UnrealEd/Private/EditorViewportClient.cpp b/Engine/Source/Editor/UnrealEd/Private/EditorViewportClient.cpp
--- a/Engine/Source/Editor/UnrealEd/Private/EditorViewportClient.cpp
+++ b/Engine/Source/Editor/UnrealEd/Private/EditorViewportClient.cpp
@@ -2106,7 +2106,7 @@ void FEditorViewportClient::UpdateCameraMovement( float DeltaTime )
 	const bool bIsITFInteractions = UsesViewportInteractions(ModeTools.Get());
 	
 	// We only want to move perspective cameras around like this
-	if( Viewport != nullptr && IsPerspective() && (bIsITFInteractions ? ViewportNavigationHelper->OrbitDelta.IsNearlyZero() : !ShouldOrbitCamera()))
+	if( Viewport != nullptr && IsPerspective() && (bIsITFInteractions ? (ViewportNavigationHelper && ViewportNavigationHelper->OrbitDelta.IsNearlyZero()) : !ShouldOrbitCamera()))
 	{
 		const bool bEnable = false;
 		ToggleOrbitCamera(bEnable);
@@ -2128,7 +2128,7 @@ void FEditorViewportClient::UpdateCameraMovement( float DeltaTime )
 
 		GestureMoveForwardBackwardImpulse = 0.0f;
 		
-		if (bIsITFInteractions && ViewportNavigationHelper)
+		if (bIsITFInteractions)
 		{
 			ViewportNavigationHelper->ConsumeImpulseData(*CameraUserImpulseData);
 		}
