UE: Bump IncrementalCookVersion

diff --git a/Engine/Source/Editor/UnrealEd/Private/Cooker/CookIncrementalVersion.cpp b/Engine/Source/Editor/UnrealEd/Private/Cooker/CookIncrementalVersion.cpp
--- a/Engine/Source/Editor/UnrealEd/Private/Cooker/CookIncrementalVersion.cpp
+++ b/Engine/Source/Editor/UnrealEd/Private/Cooker/CookIncrementalVersion.cpp
@@ -8,6 +8,6 @@ namespace UE::Cook
 {
 
 // Change CookIncrementalVersion to a new guid when all packages in an incremental cook need to be invalidated
-FGuid CookIncrementalVersion( 0xA47D92F1, 0x3C5E18AB, 0x7FB2046D, 0xE89C51D7 );
+FGuid CookIncrementalVersion( 0xBA201607, 0xD65B4B86, 0xAF26BDB6, 0xB8109F27 );
 
 }
