Harden the shader extraction code.

diff --git a/Engine/Source/Editor/UnrealEd/Private/Commandlets/ShaderPipelineCacheToolsCommandlet.cpp b/Engine/Source/Editor/UnrealEd/Private/Commandlets/ShaderPipelineCacheToolsCommandlet.cpp
--- a/Engine/Source/Editor/UnrealEd/Private/Commandlets/ShaderPipelineCacheToolsCommandlet.cpp
+++ b/Engine/Source/Editor/UnrealEd/Private/Commandlets/ShaderPipelineCacheToolsCommandlet.cpp
@@ -779,17 +779,24 @@ int32 ExtractShaderFromUshaderbytecode(const FString& UshaderbytecodeName, const
 			UE_LOG(LogShaderPipelineCacheTools, Display, TEXT("Extracting to %s"), *FinalSavePath);
 			TUniquePtr<FArchive> FileWriter(IFileManager::Get().CreateFileWriter(*FinalSavePath));
 
-			// small shaders might be stored without compression, handle them here
-			uint8* PtrToShader = Data.GetData() + OffsetToShaderCode + Entry.Offset;
-			if (Entry.Size == Entry.UncompressedSize)
+			if (FileWriter)
 			{
-				FileWriter->Serialize(PtrToShader, Entry.Size);
-			}
-			else
-			{
-				TUniquePtr<uint8[]> Buffer = MakeUnique<uint8[]>(Entry.UncompressedSize);
-				ShaderCodeArchive::DecompressShaderWithOodle(Buffer.Get(), Entry.UncompressedSize, PtrToShader, Entry.Size);
-				FileWriter->Serialize(Buffer.Get(), Entry.UncompressedSize);
+				uint8* PtrToShader = (static_cast<uint64>(OffsetToShaderCode) + Entry.Offset) < static_cast<uint64>(Data.Num()) ? (Data.GetData() + OffsetToShaderCode + Entry.Offset) : nullptr;
+
+				if (PtrToShader)
+				{
+					// small shaders might be stored without compression, handle them here
+					if (Entry.Size == Entry.UncompressedSize)
+					{
+						FileWriter->Serialize(PtrToShader, Entry.Size);
+					}
+					else
+					{
+						TUniquePtr<uint8[]> Buffer = MakeUnique<uint8[]>(Entry.UncompressedSize);
+						ShaderCodeArchive::DecompressShaderWithOodle(Buffer.Get(), Entry.UncompressedSize, PtrToShader, Entry.Size);
+						FileWriter->Serialize(Buffer.Get(), Entry.UncompressedSize);
+					}
+				}
 			}
 		}
 	}
