[UnrealEd] Restored occlusion settings update (accidental removal)

#FailedValidations Preflight

diff --git a/Engine/Source/Editor/UnrealEd/Private/EditorViewportClient.cpp b/Engine/Source/Editor/UnrealEd/Private/EditorViewportClient.cpp
--- a/Engine/Source/Editor/UnrealEd/Private/EditorViewportClient.cpp
+++ b/Engine/Source/Editor/UnrealEd/Private/EditorViewportClient.cpp
@@ -1612,6 +1612,10 @@ FSceneView* FEditorViewportClient::CalcSceneView(FSceneViewFamily* ViewFamily, c
 	{
 		View->AdditionalSelectionOutlineColors[OutlineColorIndex] = EditorStyle->AdditionalSelectionColors[OutlineColorIndex];
 	}
+	
+	View->CompositePrimitiveParams.OccludedDitheringPatternSize = EditorStyle->OccludedDitheringPatternSize;
+	View->CompositePrimitiveParams.OccludedDithering = EditorStyle->OccludedDithering;
+	View->CompositePrimitiveParams.OccludedBrightness = EditorStyle->OccludedBrightness;
 
 	int32 FamilyIndex = ViewFamily->Views.Add(View);
 	check(FamilyIndex == View->StereoViewIndex || View->StereoViewIndex == INDEX_NONE);
