[ITF Viewport inputs] Make enabling & disabling ITF Gizmos, Drag tools, and Interactions happen all together

- Set needed CVars value both when disabling and enabling new TRS gizmos
- Add delegates to `UModeManagerInteractiveToolsContext` to make sure editors register their own interactions at the right moment
- Update `SLevelEditor` and `SEditorViewport` accordingly

diff --git a/Engine/Source/Editor/LevelEditor/Private/SLevelEditor.cpp b/Engine/Source/Editor/LevelEditor/Private/SLevelEditor.cpp
--- a/Engine/Source/Editor/LevelEditor/Private/SLevelEditor.cpp
+++ b/Engine/Source/Editor/LevelEditor/Private/SLevelEditor.cpp
@@ -108,18 +108,11 @@ SLevelEditor::SLevelEditor()
 		ModeManagerInteractiveToolsContext->SetDragToolsEnabled(true);
 		ModeManagerInteractiveToolsContext->SetInputToolsEnabled(true);
 
-		RegisterViewportInteractions();
-	}
-
-	// This is used while working with both legacy and ITF interactions. Once ITF is shipped as default, it shouldn't be needed anymore
-	UE::Editor::ViewportInteractions::OnEditorViewportInteractionsActivated().AddRaw(
-		this, &SLevelEditor::RegisterViewportInteractions
+		ModeManagerInteractiveToolsContext->OnRegisterViewportInteractions().AddRaw(this, &SLevelEditor::RegisterViewportInteractions
 	);
 
-	UEditorInteractiveGizmoManager::OnUsesNewTRSGizmosChangedDelegate().AddLambda([this](bool bInUsesExperimentalGizmos)
-	{
 		RegisterViewportInteractions();
-	});
+	}
 }
 
 void SLevelEditor::BindCommands()
@@ -474,9 +467,11 @@ SLevelEditor::~SLevelEditor()
 		GUnrealEd->GetSelectedComponents()->SetElementSelectionSet(nullptr);
 	}
 
-	// This will be likely removed once ITF interactions become the default, and legacy ones are removed/deactivated
-	UE::Editor::ViewportInteractions::OnEditorViewportInteractionsActivated().RemoveAll(this);
-	UEditorInteractiveGizmoManager::OnUsesNewTRSGizmosChangedDelegate().RemoveAll(this);
+	if (UModeManagerInteractiveToolsContext* const ModeManagerInteractiveToolsContext =
+			SLevelEditor::GetEditorModeManager().GetInteractiveToolsContext())
+	{
+		ModeManagerInteractiveToolsContext->OnRegisterViewportInteractions().RemoveAll(this);
+	}
 }
 
 FText SLevelEditor::GetTabTitle() const
diff --git a/Engine/Source/Editor/UnrealEd/Private/SEditorViewport.cpp b/Engine/Source/Editor/UnrealEd/Private/SEditorViewport.cpp
--- a/Engine/Source/Editor/UnrealEd/Private/SEditorViewport.cpp
+++ b/Engine/Source/Editor/UnrealEd/Private/SEditorViewport.cpp
@@ -69,8 +69,11 @@ SEditorViewport::~SEditorViewport()
 
 	check( SceneViewport.IsUnique() );
 
-	// This might be removed once ITF viewport interactions become the default, and legacy ones are removed/deactivated - responds to dedicated CVar
-	UE::Editor::ViewportInteractions::OnEditorViewportInteractionsActivated().RemoveAll(this);
+	if (const TStrongObjectPtr<UModeManagerInteractiveToolsContext> ModeManagerInteractiveToolsContext =
+			ModeManagerInteractiveToolsContextWeak.Pin())
+	{
+		ModeManagerInteractiveToolsContext->OnRegisterViewportInteractions().RemoveAll(this);
+	}
 }
 
 void SEditorViewport::Construct( const FArguments& InArgs )
@@ -319,13 +322,19 @@ void SEditorViewport::Construct( const FArguments& InArgs )
 	// This ensures CachedDPIScale value can be properly retrieved at the right time, once the Widget is fully setup and its parent window can be accessed.
 	GetViewportClient()->RequestUpdateDPIScale();
 
+	if (UModeManagerInteractiveToolsContext* const ModeManagerInteractiveToolsContext =
+			GetModeManagerInteractiveToolsContext())
+	{
+		// GetModeManagerInteractiveToolsContext is virtual, so we store the context as a weak ptr.
+		// This allows us to access it in class destructor
+		ModeManagerInteractiveToolsContextWeak = ModeManagerInteractiveToolsContext;
+		ModeManagerInteractiveToolsContext->OnRegisterViewportInteractions().AddSP(
+			this, &SEditorViewport::RegisterViewportInteractions
+		);
+	}
+
 	// Register the Editor-specific viewport interactions
 	RegisterViewportInteractions();
-
-	// This is used while working with both legacy and ITF interactions. Once ITF is shipped as default AND CVar removed, it might not be needed anymore
-	UE::Editor::ViewportInteractions::OnEditorViewportInteractionsActivated().AddRaw(
-		this, &SEditorViewport::RegisterViewportInteractions
-	);
 }
 
 FReply SEditorViewport::OnKeyDown( const FGeometry& MyGeometry, const FKeyEvent& InKeyEvent )
diff --git a/Engine/Source/Editor/UnrealEd/Private/Tools/EdModeInteractiveToolsContext.cpp b/Engine/Source/Editor/UnrealEd/Private/Tools/EdModeInteractiveToolsContext.cpp
--- a/Engine/Source/Editor/UnrealEd/Private/Tools/EdModeInteractiveToolsContext.cpp
+++ b/Engine/Source/Editor/UnrealEd/Private/Tools/EdModeInteractiveToolsContext.cpp
@@ -1547,6 +1547,7 @@ void UModeManagerInteractiveToolsContext::RegisterDragTools()
 
 void UModeManagerInteractiveToolsContext::RegisterViewportInteractions()
 {
+	UE_LOG(LogTemp,Log,TEXT("RegisterViewportInteractions"));
 	FEditorModeTools* EditorModeTools = GetParentEditorModeManager();
 	if (!EditorModeTools)
 	{
@@ -1567,6 +1568,8 @@ void UModeManagerInteractiveToolsContext::RegisterViewportInteractions()
 		ViewportInteractionsBehaviorSource = NewObject<UViewportInteractionsBehaviorSource>();
 		ViewportInteractionsBehaviorSource->Initialize(this);
 		ViewportInteractionsBehaviorSource->RegisterBehaviorSources();
+
+		OnRegisterViewportInteractions().Broadcast();
 	}
 }
 
diff --git a/Engine/Source/Editor/UnrealEd/Public/SEditorViewport.h b/Engine/Source/Editor/UnrealEd/Public/SEditorViewport.h
--- a/Engine/Source/Editor/UnrealEd/Public/SEditorViewport.h
+++ b/Engine/Source/Editor/UnrealEd/Public/SEditorViewport.h
@@ -435,4 +435,10 @@ private:
 	TWeakPtr<SWidget> AutoUpgradeWidgetChild = nullptr;
 	/** A flag indicating that the automatic toolbar upgrade has been approved */
 	bool bLegacyToolbarIsAutomaticallyUpgradable = false;
+
+	/**
+	 * Storing Interactive Tools Context so that we can unbind delegates in the class constructor. Once the new ITF
+	 * based gizmos, drag tools and viewport interactions will be the default, this can be removed.
+	 */
+	TWeakObjectPtr<UModeManagerInteractiveToolsContext> ModeManagerInteractiveToolsContextWeak;
 };
diff --git a/Engine/Source/Editor/UnrealEd/Public/Tools/EdModeInteractiveToolsContext.h b/Engine/Source/Editor/UnrealEd/Public/Tools/EdModeInteractiveToolsContext.h
--- a/Engine/Source/Editor/UnrealEd/Public/Tools/EdModeInteractiveToolsContext.h
+++ b/Engine/Source/Editor/UnrealEd/Public/Tools/EdModeInteractiveToolsContext.h
@@ -408,6 +408,13 @@ public:
 	 */
 	UE_API UViewportInteractionsBehaviorSource* GetViewportInteractionsBehaviorSource();
 
+	DECLARE_MULTICAST_DELEGATE(FViewportInteractionsRegisterStateChanged)
+	FViewportInteractionsRegisterStateChanged& OnRegisterViewportInteractions()
+	{
+		return OnRegisterViewportInteractionsDelegate;
+	}
+	FViewportInteractionsRegisterStateChanged OnRegisterViewportInteractionsDelegate;
+
 protected:
 	/**
 	 * Activating Drag Tools if current Mode Manager requires them
