[TRS Gizmo] Make useage/visibility of TRS Gizmos be controlled from individual instances of FEditorModeTools. Actual useage requires the global setting to be on and individual mode tools instances to opt-in. Opts in the level editor mode tools.

#rnx

diff --git a/Engine/Plugins/Animation/ControlRig/Source/ControlRigEditor/Private/EditMode/ControlRigEditMode.cpp b/Engine/Plugins/Animation/ControlRig/Source/ControlRigEditor/Private/EditMode/ControlRigEditMode.cpp
--- a/Engine/Plugins/Animation/ControlRig/Source/ControlRigEditor/Private/EditMode/ControlRigEditMode.cpp
+++ b/Engine/Plugins/Animation/ControlRig/Source/ControlRigEditor/Private/EditMode/ControlRigEditMode.cpp
@@ -2059,7 +2059,7 @@ bool FControlRigEditMode::InputDelta(FEditorViewportClient* InViewportClient, FV
 	//button down if left and ctrl and right is down, needed for indirect posting
 
 	// enable MMB with the new TRS gizmos
-	const bool bEnableMMB = UEditorInteractiveGizmoManager::UsesNewTRSGizmos();
+	const bool bEnableMMB = InViewportClient->GetModeTools()->UsesViewportITF();
 	
 	const bool bMouseButtonDown =
 		InViewport->KeyState(EKeys::LeftMouseButton) ||
diff --git a/Engine/Source/Editor/Experimental/EditorInteractiveToolsFramework/Private/EditorGizmos/EditorTransformGizmoSource.cpp b/Engine/Source/Editor/Experimental/EditorInteractiveToolsFramework/Private/EditorGizmos/EditorTransformGizmoSource.cpp
--- a/Engine/Source/Editor/Experimental/EditorInteractiveToolsFramework/Private/EditorGizmos/EditorTransformGizmoSource.cpp
+++ b/Engine/Source/Editor/Experimental/EditorInteractiveToolsFramework/Private/EditorGizmos/EditorTransformGizmoSource.cpp
@@ -105,7 +105,7 @@ bool UEditorTransformGizmoSource::CanInteract(const EViewportContext InViewportC
 			bool bUseLegacyWidget = (WidgetMode == UE::Widget::WM_TranslateRotateZ || WidgetMode == UE::Widget::WM_2D);
 			if (!bUseLegacyWidget)
 			{
-				bUseLegacyWidget = !UEditorInteractiveGizmoManager::UsesNewTRSGizmos();
+				bUseLegacyWidget = !ModeTools.UsesViewportITF();
 			}
 			return !bUseLegacyWidget;
 		}
diff --git a/Engine/Source/Editor/Experimental/EditorInteractiveToolsFramework/Private/EditorGizmos/EditorTransformGizmoUtil.cpp b/Engine/Source/Editor/Experimental/EditorInteractiveToolsFramework/Private/EditorGizmos/EditorTransformGizmoUtil.cpp
--- a/Engine/Source/Editor/Experimental/EditorInteractiveToolsFramework/Private/EditorGizmos/EditorTransformGizmoUtil.cpp
+++ b/Engine/Source/Editor/Experimental/EditorInteractiveToolsFramework/Private/EditorGizmos/EditorTransformGizmoUtil.cpp
@@ -275,7 +275,7 @@ void UEditorTransformGizmoContextObject::UpdateGizmo(const TArray<FEditorViewpor
 	// create a new one
 	if (UTransformGizmo* TransformGizmo = UE::EditorTransformGizmoUtil::GetDefaultTransformGizmo(ToolManager))
 	{
-		TransformGizmo->SetVisibility(Swapped ? UEditorInteractiveGizmoManager::UsesNewTRSGizmos() : false);
+		TransformGizmo->SetVisibility(Swapped ? ModeTools->UsesViewportITF() : false);
 	}
 }
 
diff --git a/Engine/Source/Editor/UnrealEd/Private/EditorModeManager.cpp b/Engine/Source/Editor/UnrealEd/Private/EditorModeManager.cpp
--- a/Engine/Source/Editor/UnrealEd/Private/EditorModeManager.cpp
+++ b/Engine/Source/Editor/UnrealEd/Private/EditorModeManager.cpp
@@ -145,8 +145,7 @@ void FEditorModeTools::LoadConfig(void)
 	if (CoordinateSystem == COORD_Parent || CoordinateSystem == COORD_Explicit)
 	{
 		// parent mode is only supported with new trs gizmos for now
-		const bool bUseNewGizmo = UEditorInteractiveGizmoManager::UsesNewTRSGizmos();
-		if (!bUseNewGizmo)
+		if (!UsesViewportITF())
 		{
 			CoordSystemAsInt = static_cast<int32>(COORD_Local);
 		}
@@ -2074,6 +2073,16 @@ UModeManagerInteractiveToolsContext* FEditorModeTools::GetInteractiveToolsContex
 	return InteractiveToolsContext;
 }
 
+bool FEditorModeTools::UsesViewportITF() const
+{
+	return bViewportITF && UEditorInteractiveGizmoManager::UsesNewTRSGizmos();
+}
+
+void FEditorModeTools::SetSupportsViewportITF(bool bIsSupported)
+{
+	bViewportITF = bIsSupported;
+}
+
 IGizmoStateTarget* FEditorModeTools::GetGizmoStateTarget()
 {
 	if (!WeakGizmoStateTarget.IsValid())
diff --git a/Engine/Source/Editor/UnrealEd/Private/LevelEditorViewport.cpp b/Engine/Source/Editor/UnrealEd/Private/LevelEditorViewport.cpp
--- a/Engine/Source/Editor/UnrealEd/Private/LevelEditorViewport.cpp
+++ b/Engine/Source/Editor/UnrealEd/Private/LevelEditorViewport.cpp
@@ -3279,7 +3279,7 @@ bool FLevelEditorViewportClient::InputWidgetDelta(FViewport* InViewport, EAxisLi
 				const bool MiddleMouseButtonDown = InViewport->KeyState(EKeys::MiddleMouseButton);
 
 				// Some functionality can be skipped if new gizmos and ITF drag tools are active
-				const bool bUseLegacyGizmoDragTools = !(UEditorInteractiveGizmoManager::UsesNewTRSGizmos() && UE::Editor::DragTools::UseEditorDragTools());
+				const bool bUseLegacyGizmoDragTools = !(ModeTools->UsesViewportITF() && UE::Editor::DragTools::UseEditorDragTools());
 
 				// If duplicate dragging . . .
 				if (IsAltPressed() && (LeftMouseButtonDown || RightMouseButtonDown) && bUseLegacyGizmoDragTools)
@@ -3346,7 +3346,7 @@ bool FLevelEditorViewportClient::InputWidgetDelta(FViewport* InViewport, EAxisLi
 
 				// We do not want actors updated if we are holding down the middle mouse button.
 				// enable MMB for New TRS Gizmos
-				const bool bEnableMMB = UEditorInteractiveGizmoManager::UsesNewTRSGizmos() ? !bDraggingByHandle : false;
+				const bool bEnableMMB = ModeTools->UsesViewportITF() ? !bDraggingByHandle : false;
 				
 				if(!MiddleMouseButtonDown || bEnableMMB)
 				{
@@ -3391,7 +3391,7 @@ bool FLevelEditorViewportClient::InputWidgetDelta(FViewport* InViewport, EAxisLi
 				ModeTools->SnappedLocation += Drag;
 
 				const bool CameraFollowsGizmo = GetViewportNavigationHelper()->bCameraFollowsGizmo;
-				if ((IsShiftPressed() && !UEditorInteractiveGizmoManager::UsesNewTRSGizmos()) || CameraFollowsGizmo)
+				if ((IsShiftPressed() && !ModeTools->UsesViewportITF()) || CameraFollowsGizmo)
 				{
 					bApplyCameraSpeedScaleByDistance = false;
 					FVector CameraDelta( Drag );
diff --git a/Engine/Source/Editor/UnrealEd/Private/MouseDeltaTracker.cpp b/Engine/Source/Editor/UnrealEd/Private/MouseDeltaTracker.cpp
--- a/Engine/Source/Editor/UnrealEd/Private/MouseDeltaTracker.cpp
+++ b/Engine/Source/Editor/UnrealEd/Private/MouseDeltaTracker.cpp
@@ -13,6 +13,7 @@
 #include "EditorDragTools/EditorDragToolBehaviorTarget.h"
 #include "EditorInteractiveGizmoManager.h"
 #include "Misc/AxisDisplayInfo.h"
+#include "EditorModeManager.h"
 
 #define LOCTEXT_NAMESPACE "MouseDeltaTracker"
 
@@ -218,7 +219,7 @@ void FMouseDeltaTracker::StartTracking(FEditorViewportClient* InViewportClient,
 	const bool MiddleMouseButtonDown = InViewportClient->Viewport->KeyState(EKeys::MiddleMouseButton);
 
 	constexpr bool bAllowWidgetManipulationWithNewTRSGizmo = false;
-	const bool bUsesNewGizmo = UEditorInteractiveGizmoManager::UsesNewTRSGizmos();
+	const bool bUsesNewGizmo = InViewportClient->GetModeTools()->UsesViewportITF();
 
 	const bool bAllowWidgetManipulation = !bUsesNewGizmo || bAllowWidgetManipulationWithNewTRSGizmo;
 
diff --git a/Engine/Source/Editor/UnrealEd/Private/SEditorViewport.cpp b/Engine/Source/Editor/UnrealEd/Private/SEditorViewport.cpp
--- a/Engine/Source/Editor/UnrealEd/Private/SEditorViewport.cpp
+++ b/Engine/Source/Editor/UnrealEd/Private/SEditorViewport.cpp
@@ -1181,7 +1181,7 @@ void SEditorViewport::OnCycleCoordinateSystem()
 
 	// parent mode is only supported with new trs gizmos for now
 	int CoordMax = COORD_Parent;
-	if (UEditorInteractiveGizmoManager::UsesNewTRSGizmos())
+	if (Client->GetModeTools()->UsesViewportITF())
 	{
 		CoordMax = UEditorInteractiveGizmoManager::IsExplicitModeEnabled() ? COORD_Max : COORD_Explicit;
 	}
diff --git a/Engine/Source/Editor/UnrealEd/Private/UnrealEdGlobals.cpp b/Engine/Source/Editor/UnrealEd/Private/UnrealEdGlobals.cpp
--- a/Engine/Source/Editor/UnrealEd/Private/UnrealEdGlobals.cpp
+++ b/Engine/Source/Editor/UnrealEd/Private/UnrealEdGlobals.cpp
@@ -188,6 +188,7 @@ int32 EditorInit( IEngineLoop& EngineLoop )
 	if (!Internal::EditorModeToolsSingleton.IsValid())
 	{
 		Internal::EditorModeToolsSingleton = MakeShared<FEditorModeTools>();
+		Internal::EditorModeToolsSingleton->SetSupportsViewportITF(true);
 	}
 
 	// Initialize the misc editor
diff --git a/Engine/Source/Editor/UnrealEd/Private/UnrealWidgetRender.cpp b/Engine/Source/Editor/UnrealEd/Private/UnrealWidgetRender.cpp
--- a/Engine/Source/Editor/UnrealEd/Private/UnrealWidgetRender.cpp
+++ b/Engine/Source/Editor/UnrealEd/Private/UnrealWidgetRender.cpp
@@ -190,7 +190,7 @@ void FWidget::Render( const FSceneView* View,FPrimitiveDrawInterface* PDI, FEdit
 	bool bUseLegacyWidget = (WidgetMode == UE::Widget::EWidgetMode::WM_TranslateRotateZ || WidgetMode == UE::Widget::EWidgetMode::WM_2D);
 	if (!bUseLegacyWidget)
 	{
-		bUseLegacyWidget = !UEditorInteractiveGizmoManager::UsesNewTRSGizmos();
+		bUseLegacyWidget = !ViewportClient->GetModeTools()->UsesViewportITF();
 	}
 
 	bool bDrawWidget;
diff --git a/Engine/Source/Editor/UnrealEd/Private/ViewportToolbar/UnrealEdViewportToolbar.cpp b/Engine/Source/Editor/UnrealEd/Private/ViewportToolbar/UnrealEdViewportToolbar.cpp
--- a/Engine/Source/Editor/UnrealEd/Private/ViewportToolbar/UnrealEdViewportToolbar.cpp
+++ b/Engine/Source/Editor/UnrealEd/Private/ViewportToolbar/UnrealEdViewportToolbar.cpp
@@ -9,6 +9,7 @@
 #include "Editor/EditorPerformanceSettings.h"
 #include "EditorViewportClient.h"
 #include "EditorViewportCommands.h"
+#include "EditorModeManager.h"
 #include "Engine/GameViewportClient.h"
 #include "Framework/Commands/GenericCommands.h"
 #include "GPUSkinCache.h"
@@ -1062,19 +1063,22 @@ FToolMenuEntry CreateTransformsSubmenu()
 													);
 													UnnamedSection.AddEntry(LocalCoords);
 
-													if (UEditorInteractiveGizmoManager::UsesNewTRSGizmos())
+													if (TSharedPtr<FEditorViewportClient> Client = UUnrealEdViewportToolbarContext::GetEditorViewportClient(InSubmenu))
 													{
-														FToolMenuEntry ParentCoords = FToolMenuEntry::InitMenuEntry(
-															FEditorViewportCommands::Get().RelativeCoordinateSystem_Parent
-														);
-														UnnamedSection.AddEntry(ParentCoords);
-
-														if (UEditorInteractiveGizmoManager::IsExplicitModeEnabled())
+														if (FEditorModeTools* ModeTools = Client->GetModeTools(); ModeTools && ModeTools->UsesViewportITF())
 														{
-															FToolMenuEntry ExplicitCoords = FToolMenuEntry::InitMenuEntry(
-															   FEditorViewportCommands::Get().RelativeCoordinateSystem_Explicit
-														   );
-														   UnnamedSection.AddEntry(ExplicitCoords);
+															FToolMenuEntry ParentCoords = FToolMenuEntry::InitMenuEntry(
+																FEditorViewportCommands::Get().RelativeCoordinateSystem_Parent
+															);
+															UnnamedSection.AddEntry(ParentCoords);
+
+															if (UEditorInteractiveGizmoManager::IsExplicitModeEnabled())
+															{
+																FToolMenuEntry ExplicitCoords = FToolMenuEntry::InitMenuEntry(
+																   FEditorViewportCommands::Get().RelativeCoordinateSystem_Explicit
+															   );
+															   UnnamedSection.AddEntry(ExplicitCoords);
+															}
 														}
 													}
 												}
diff --git a/Engine/Source/Editor/UnrealEd/Private/ViewportToolbar/UnrealEdViewportToolbarContext.cpp b/Engine/Source/Editor/UnrealEd/Private/ViewportToolbar/UnrealEdViewportToolbarContext.cpp
--- a/Engine/Source/Editor/UnrealEd/Private/ViewportToolbar/UnrealEdViewportToolbarContext.cpp
+++ b/Engine/Source/Editor/UnrealEd/Private/ViewportToolbar/UnrealEdViewportToolbarContext.cpp
@@ -20,6 +20,15 @@ TSharedPtr<SEditorViewport> UUnrealEdViewportToolbarContext::GetEditorViewport(c
 	return nullptr;
 }
 
+TSharedPtr<FEditorViewportClient> UUnrealEdViewportToolbarContext::GetEditorViewportClient(const FToolMenuContext& Context)
+{
+	if (TSharedPtr<SEditorViewport> Viewport = GetEditorViewport(Context))
+	{
+		return Viewport->GetViewportClient();
+	}
+	return nullptr;
+}
+
 TSharedPtr<IPreviewProfileController> UUnrealEdViewportToolbarContext::GetPreviewProfileController() const
 {
 	if (const TSharedPtr<SEditorViewport> EditorViewport = Viewport.Pin())
diff --git a/Engine/Source/Editor/UnrealEd/Public/EditorModeManager.h b/Engine/Source/Editor/UnrealEd/Public/EditorModeManager.h
--- a/Engine/Source/Editor/UnrealEd/Public/EditorModeManager.h
+++ b/Engine/Source/Editor/UnrealEd/Public/EditorModeManager.h
@@ -557,6 +557,8 @@ public:
 	UNREALED_API UModeManagerInteractiveToolsContext* GetInteractiveToolsContext() const;
 
 	/** Begin New TRS Gizmo interface */
+	UNREALED_API bool UsesViewportITF() const;
+	UNREALED_API void SetSupportsViewportITF(bool bIsSupported);
 	UNREALED_API IGizmoStateTarget* GetGizmoStateTarget();
 	UNREALED_API bool BeginTransform(const FGizmoState& InState);
 	UNREALED_API bool EndTransform(const FGizmoState& InState) const;
@@ -664,6 +666,9 @@ private:
 	/** GizmoContext used to get context specific gizmo properties. */
 	mutable TWeakObjectPtr<UEditorTransformGizmoContextObject> WeakGizmoContext;
 
+	/** Determines whether this mode tool can ever support TRS gizmos. */
+	bool bViewportITF = false;
+
 	/** Flag to track if we started a new TRS gizmo transform.
 	 * NOTE: we have to use it for now as StartTracking / EndTracking iterates thru the modes even if the ITF context captured something.
 	 * This might not be needed anymore in the future but as some routing behavior had to be kept for legacy reasons, we use that extra flag.
diff --git a/Engine/Source/Editor/UnrealEd/Public/ViewportToolbar/UnrealEdViewportToolbarContext.h b/Engine/Source/Editor/UnrealEd/Public/ViewportToolbar/UnrealEdViewportToolbarContext.h
--- a/Engine/Source/Editor/UnrealEd/Public/ViewportToolbar/UnrealEdViewportToolbarContext.h
+++ b/Engine/Source/Editor/UnrealEd/Public/ViewportToolbar/UnrealEdViewportToolbarContext.h
@@ -27,6 +27,7 @@ DECLARE_DELEGATE_RetVal_OneParam(bool, DoesViewModeMenuShowSectionDelegate, EHid
 }
 
 class SEditorViewport;
+class FEditorViewportClient;
 struct FToolMenuContext;
 
 UCLASS(MinimalAPI)
@@ -60,6 +61,7 @@ public:
 	FTabId PreviewSettingsTabId;
 	
 	static UE_API TSharedPtr<SEditorViewport> GetEditorViewport(const FToolMenuContext& Context);
+	static UE_API TSharedPtr<FEditorViewportClient> GetEditorViewportClient(const FToolMenuContext& Context);
 	
 	UE_API virtual TSharedPtr<IPreviewProfileController> GetPreviewProfileController() const;
 	
