Cooker: Fix warning "InternalCookerError: Package %s was retracted from %s, but already has no platforms needing commit". This can occur when one CookRequestCluster finds a package needs incremental recook, but in a later cluster the dependencies have changed and it is found to be incrementally skippable. Avoid that possibility by using the stored value of IncrementallyUnModified instead of recomputing it in each cluster.
#rnx

diff --git a/Engine/Source/Editor/UnrealEd/Private/Cooker/CookRequestCluster.cpp b/Engine/Source/Editor/UnrealEd/Private/Cooker/CookRequestCluster.cpp
--- a/Engine/Source/Editor/UnrealEd/Private/Cooker/CookRequestCluster.cpp
+++ b/Engine/Source/Editor/UnrealEd/Private/Cooker/CookRequestCluster.cpp
@@ -1554,7 +1554,13 @@ void FRequestCluster::FGraphSearch::FExploreEdgesContext::CalculatePlatformsToPr
 		{
 			continue;
 		}
-		bool bIncrementallyUnmodifiedNeeded = !QueryPlatformData.bIncrementallyUnmodified.IsSet();
+		bool bIncrementallyUnmodifiedNeeded = false;
+		if (PlatformIndex != CookerLoadingPlatformIndex)
+		{
+			const ITargetPlatform* TargetPlatform = GraphSearch.FetchPlatforms[PlatformIndex].Platform;
+			const FPackagePlatformData& PackagePlatformData = PackageData->FindOrAddPlatformData(TargetPlatform);
+			bIncrementallyUnmodifiedNeeded = !PackagePlatformData.IsIncrementallyUnmodifiedSet();
+		}
 		bool bExploreNeeded = !QueryPlatformData.bExploreCompleted && QueryPlatformData.bExploreRequested;
 		if (!bIncrementallyUnmodifiedNeeded && !bExploreNeeded)
 		{
@@ -1615,15 +1621,14 @@ bool FRequestCluster::FGraphSearch::FExploreEdgesContext::TryCalculateIncrementa
 			continue;
 		}
 
-		FQueryPlatformData& QueryPlatformData = Vertex->GetPlatformData()[PlatformIndex];
-		if (QueryPlatformData.bIncrementallyUnmodified.IsSet())
-		{
-			continue;
-		}
-
 		FFetchPlatformData& FetchPlatformData = GraphSearch.FetchPlatforms[PlatformIndex];
 		const ITargetPlatform* TargetPlatform = FetchPlatformData.Platform;
 		FPackagePlatformData& PackagePlatformData = PackageData->FindOrAddPlatformData(TargetPlatform);
+		if (PackagePlatformData.IsIncrementallyUnmodifiedSet())
+		{
+			continue;
+		}
+		FQueryPlatformData& QueryPlatformData = Vertex->GetPlatformData()[PlatformIndex];
 
 		UE::Cook::FPackageArtifacts& Artifacts = QueryPlatformData.CookAttachments.Artifacts;
 		bool bPreviouslyCooked =
@@ -1649,13 +1654,10 @@ bool FRequestCluster::FGraphSearch::FExploreEdgesContext::TryCalculateIncrementa
 			// generated packages; they are all marked incrementally unmodified as well
 			const FPackagePlatformData* ParentPlatformData =
 				ParentPackageData->GetPlatformDatas().Find(TargetPlatform);
-			if (ParentPlatformData)
+			if (ParentPlatformData && ParentPlatformData->IsIncrementallyUnmodified())
 			{
-				if (ParentPlatformData->IsIncrementallyUnmodified())
-				{
-					SetIncrementallyUnmodified(PlatformIndex, PackagePlatformData);
-					continue;
-				}
+				SetIncrementallyUnmodified(PlatformIndex, PackagePlatformData);
+				continue;
 			}
 		}
 		if (!Artifacts.HasKeyMatch(TargetPlatform, GenerationHelper.GetReference()))
@@ -1707,13 +1709,13 @@ bool FRequestCluster::FGraphSearch::FExploreEdgesContext::TryCalculateIncrementa
 					break;
 				}
 
-				FQueryPlatformData& TransitivePlatformData = TransitiveBuildVertex.GetPlatformData()[PlatformIndex];
-				if (!TransitivePlatformData.bIncrementallyUnmodified.IsSet())
+				FPackagePlatformData& TransitivePlatformData = TransitiveBuildVertex.GetPackageData()->FindOrAddPlatformData(TargetPlatform);
+				if (!TransitivePlatformData.IsIncrementallyUnmodifiedSet())
 				{
 					UnreadyTransitiveBuildVertices.Add(&TransitiveBuildVertex);
 					continue;
 				}
-				if (!TransitivePlatformData.bIncrementallyUnmodified.GetValue())
+				if (!TransitivePlatformData.IsIncrementallyUnmodified())
 				{
 					ModifiedTransitiveBuildDependency = TransitiveBuildPackageName;
 					break;
@@ -1773,9 +1775,10 @@ bool FRequestCluster::FGraphSearch::FExploreEdgesContext::TryCalculateIncrementa
 			{
 				continue;
 			}
-			FQueryPlatformData& QueryPlatformData = Vertex->GetPlatformData()[PlatformIndex];
-			check(QueryPlatformData.bIncrementallyUnmodified.IsSet()); // Otherwise we early exited in !bAllPlatformsAreReady
-			bAllUnmodified &= *QueryPlatformData.bIncrementallyUnmodified;
+			const ITargetPlatform* TargetPlatform = GraphSearch.FetchPlatforms[PlatformIndex].Platform;
+			FPackagePlatformData& PackagePlatformData = PackageData->FindOrAddPlatformData(TargetPlatform);
+			check(PackagePlatformData.IsIncrementallyUnmodifiedSet()); // Otherwise we early exited in !bAllPlatformsAreReady
+			bAllUnmodified &= PackagePlatformData.IsIncrementallyUnmodified();
 		}
 		if (!bAllUnmodified)
 		{
@@ -1809,17 +1812,18 @@ bool FRequestCluster::FGraphSearch::FExploreEdgesContext::TryCalculateIncrementa
 	TArray<FVertexData*>& Listeners = Vertex->GetIncrementallyModifiedListeners();
 	if (!Listeners.IsEmpty())
 	{
-		bool bIncrementallyModified = false;
+		bool bAnyPlatformIncrementallyModified = false;
 		for (int32 PlatformIndex = 0; PlatformIndex < LocalNumFetchPlatforms; ++PlatformIndex)
 		{
 			if (PlatformIndex == PlatformAgnosticPlatformIndex)
 			{
 				continue;
 			}
-			FQueryPlatformData& QueryPlatformData = Vertex->GetPlatformData()[PlatformIndex];
-			if (QueryPlatformData.bIncrementallyUnmodified.IsSet() && !*QueryPlatformData.bIncrementallyUnmodified)
+			const ITargetPlatform* TargetPlatform = GraphSearch.FetchPlatforms[PlatformIndex].Platform;
+			FPackagePlatformData& PackagePlatformData = PackageData->FindOrAddPlatformData(TargetPlatform);
+			if (PackagePlatformData.IsIncrementallyUnmodifiedSet() && !PackagePlatformData.IsIncrementallyUnmodified())
 			{
-				bIncrementallyModified = true;
+				bAnyPlatformIncrementallyModified = true;
 				break;
 			}
 		}
@@ -1833,7 +1837,7 @@ bool FRequestCluster::FGraphSearch::FExploreEdgesContext::TryCalculateIncrementa
 			}
 
 			ListenerVertex->GetUnreadyDependencies().Remove(Vertex);
-			if (bIncrementallyModified || ListenerVertex->GetUnreadyDependencies().IsEmpty())
+			if (bAnyPlatformIncrementallyModified || ListenerVertex->GetUnreadyDependencies().IsEmpty())
 			{
 				ListenerVertex->SetWaitingOnUnreadyDependencies(false);
 				GraphSearch.KickVertex(ListenerVertex);
@@ -2211,8 +2215,8 @@ void FRequestCluster::FGraphSearch::FExploreEdgesContext::ProcessPlatformAttachm
 
 	if (Cluster.IsIncrementalCook())
 	{
-		check(QueryPlatformData.bIncrementallyUnmodified.IsSet());
-		bool bIncrementallyUnmodified = QueryPlatformData.bIncrementallyUnmodified.GetValue();
+		check(PackagePlatformData.IsIncrementallyUnmodifiedSet());
+		bool bIncrementallyUnmodified = PackagePlatformData.IsIncrementallyUnmodified();
 		if (bIncrementallyUnmodified)
 		{
 			// Queue runtime dependencies if we are exploring runtime dependencies, build definitions if
@@ -2389,8 +2393,6 @@ void FRequestCluster::FGraphSearch::FExploreEdgesContext::ProcessPlatformAttachm
 void FRequestCluster::FGraphSearch::FExploreEdgesContext::SetIncrementallyUnmodified(int32 PlatformIndex,
 	FPackagePlatformData& PackagePlatformData)
 {
-	FQueryPlatformData& ClusterPlatformData = Vertex->GetPlatformData()[PlatformIndex];
-	ClusterPlatformData.bIncrementallyUnmodified.Emplace(true);
 	PackagePlatformData.SetIncrementallyUnmodified(true);
 }
 
@@ -2399,7 +2401,6 @@ void FRequestCluster::FGraphSearch::FExploreEdgesContext::SetIncrementallyModifi
 	FIncrementallyModifiedContext& Context)
 {
 	FQueryPlatformData& ClusterPlatformData = Vertex->GetPlatformData()[PlatformIndex];
-	ClusterPlatformData.bIncrementallyUnmodified.Emplace(false);
 	if (!PackagePlatformData.IsIncrementallyUnmodifiedSet())
 	{
 		ClusterPlatformData.bIncrementallyModifiedInstigated = true;
diff --git a/Engine/Source/Editor/UnrealEd/Private/Cooker/CookRequestCluster.h b/Engine/Source/Editor/UnrealEd/Private/Cooker/CookRequestCluster.h
--- a/Engine/Source/Editor/UnrealEd/Private/Cooker/CookRequestCluster.h
+++ b/Engine/Source/Editor/UnrealEd/Private/Cooker/CookRequestCluster.h
@@ -144,7 +144,6 @@ private:
 		bool bIncrementallyUnmodifiedRequested = false;
 		bool bTransitiveBuildDependenciesResolvedAsNotModified = false;
 		bool bIncrementallyModifiedInstigated = false;
-		TOptional<bool> bIncrementallyUnmodified;
 	private:
 		std::atomic<EAsyncQueryStatus> AsyncQueryStatus;
 	};
