PR #13868: [ActorGrouping] fix crash when passing Array which contains nullptr

diff --git a/Engine/Source/Editor/UnrealEd/Private/ActorGroupingUtils.cpp b/Engine/Source/Editor/UnrealEd/Private/ActorGroupingUtils.cpp
--- a/Engine/Source/Editor/UnrealEd/Private/ActorGroupingUtils.cpp
+++ b/Engine/Source/Editor/UnrealEd/Private/ActorGroupingUtils.cpp
@@ -61,6 +61,11 @@ AGroupActor* UActorGroupingUtils::GroupActors(const TArray<AActor*>& ActorsToGro
 		bool bActorsInSameLevel = true;
 		for (AActor* Actor : ActorsToGroup)
 		{
+			if (!IsValid(Actor))
+			{
+				continue;
+			}
+			
 			if (!ActorLevel)
 			{
 				ActorLevel = Actor->GetLevel();
