Add delegate when an ensure failed
Display ensures in message log

diff --git a/Engine/Source/Editor/UnrealEd/Classes/Editor/EditorPerProjectUserSettings.h b/Engine/Source/Editor/UnrealEd/Classes/Editor/EditorPerProjectUserSettings.h
--- a/Engine/Source/Editor/UnrealEd/Classes/Editor/EditorPerProjectUserSettings.h
+++ b/Engine/Source/Editor/UnrealEd/Classes/Editor/EditorPerProjectUserSettings.h
@@ -104,6 +104,10 @@ class UEditorPerProjectUserSettings : public UObject
 	UPROPERTY(config)
 	uint32 bUseCurvesForDistributions:1; //(GDistributionType == 0)
 
+	/** If enabled, failed ensures will be displayed in the message log */
+	UPROPERTY(EditAnywhere, config, Category = Ensures)
+	uint32 bShowEnsuresInMessageLog : 1;
+
 	/** Controls the minimum value at which the property matrix editor will display a loading bar when pasting values */
 	UPROPERTY(config)
 	int32 PropertyMatrix_NumberOfPasteOperationsBeforeWarning;
diff --git a/Engine/Source/Editor/UnrealEd/Private/Settings/EditorPerProjectUserSettings.cpp b/Engine/Source/Editor/UnrealEd/Private/Settings/EditorPerProjectUserSettings.cpp
--- a/Engine/Source/Editor/UnrealEd/Private/Settings/EditorPerProjectUserSettings.cpp
+++ b/Engine/Source/Editor/UnrealEd/Private/Settings/EditorPerProjectUserSettings.cpp
@@ -44,6 +44,8 @@ UEditorPerProjectUserSettings::UEditorPerProjectUserSettings(const FObjectInitia
 	bPreviewFeatureLevelWasDefault = true;
 	PreviewDeviceProfileName = NAME_None;
 	bShowSelectionSubcomponents = true;
+
+	bShowEnsuresInMessageLog = true;
 }
 
 void UEditorPerProjectUserSettings::PostInitProperties()
diff --git a/Engine/Source/Editor/UnrealEd/Private/UnrealEdMisc.cpp b/Engine/Source/Editor/UnrealEd/Private/UnrealEdMisc.cpp
--- a/Engine/Source/Editor/UnrealEd/Private/UnrealEdMisc.cpp
+++ b/Engine/Source/Editor/UnrealEd/Private/UnrealEdMisc.cpp
@@ -105,7 +105,8 @@
 #include "Factories/FbxSkeletalMeshImportData.h"
 #include "Factories/FbxStaticMeshImportData.h"
 #include "OutputLogModule.h"
-
+#include "Editor/EditorPerProjectUserSettings.h"
+#include "Async/TaskGraphInterfaces.h"
 
 #define USE_UNIT_TESTS 0
 
@@ -318,6 +319,7 @@ void FUnrealEdMisc::OnInit()
 	FEditorDelegates::ChangeEditorMode.AddRaw(this, &FUnrealEdMisc::OnEditorChangeMode);
 	FCoreDelegates::PreModal.AddRaw(this, &FUnrealEdMisc::OnEditorPreModal);
 	FCoreDelegates::PostModal.AddRaw(this, &FUnrealEdMisc::OnEditorPostModal);
+	FCoreDelegates::OnEnsureFailed.AddRaw(this, &FUnrealEdMisc::OnEnsureFailed);
 
 	// Register the play world commands
 	FPlayWorldCommands::Register();
@@ -559,6 +561,12 @@ void FUnrealEdMisc::OnInit()
 		MessageLogModule.RegisterLogListing("PIE", LOCTEXT("PlayInEditor", "Play In Editor"), InitOptions);
 	}
 
+	{
+		FMessageLogInitializationOptions InitOptions;
+		InitOptions.bShowFilters = true;
+		MessageLogModule.RegisterLogListing("EnsureFailed", LOCTEXT("EnsureFailed", "Ensure Failed"), InitOptions);
+	}
+
 	// install message log delegates
 	FMessageLog::OnMessageSelectionChanged().BindRaw(this, &FUnrealEdMisc::OnMessageSelectionChanged);
 	FMessageLog::OnMultiSelectActorBegin().BindRaw(this, &FUnrealEdMisc::OnMultiSelectActorFromMessageTokenBegin);
@@ -1126,6 +1134,7 @@ void FUnrealEdMisc::OnExit()
 	FEditorDelegates::ChangeEditorMode.RemoveAll(this);
 	FCoreDelegates::PreModal.RemoveAll(this);
 	FCoreDelegates::PostModal.RemoveAll(this);
+	FCoreDelegates::OnEnsureFailed.RemoveAll(this);
 
 	FComponentAssetBrokerage::PRIVATE_ShutdownBrokerage();
 
@@ -1364,6 +1373,50 @@ void FUnrealEdMisc::OnEditorPostModal()
 	}
 }
 
+void FUnrealEdMisc::OnEnsureFailed(const ANSICHAR* Expr, const ANSICHAR* File, int32 Line, const TCHAR* Msg, const TCHAR* CombinedErrorMsg)
+{
+	if (CombinedErrorMsg == nullptr)
+	{
+		return;
+	}
+
+	const UEditorPerProjectUserSettings* const EditorSettings = GetDefault<UEditorPerProjectUserSettings>();
+	if (EditorSettings == nullptr || !EditorSettings->bShowEnsuresInMessageLog)
+	{
+		return;
+	}
+
+	auto EnsureNotify = [](const TCHAR* ErrorMsg, const TCHAR* NotifyMsg)
+		{
+			FMessageLog EnsureFailedLogger = FMessageLog("EnsureFailed");
+			EnsureFailedLogger.Error(FText::Format(NSLOCTEXT("Core", "EnsureMessage", "{0}"), FTextStringHelper::CreateFromBuffer(ErrorMsg)));
+			EnsureFailedLogger.Notify(FText::Format(NSLOCTEXT("Core", "EnsureFailed", "Ensure condition failed: {0}"), FTextStringHelper::CreateFromBuffer(NotifyMsg)));
+		};
+
+	FString NotifyMsg;
+	if (Msg != nullptr && FCString::Strlen(Msg) > 0)
+	{
+		NotifyMsg = Msg;
+	}
+	else
+	{
+		NotifyMsg = Expr;
+	}
+
+	if (IsInGameThread())
+	{
+		EnsureNotify(CombinedErrorMsg, *NotifyMsg);
+	}
+	else
+	{
+		FString ErrorMsg(CombinedErrorMsg);
+		FSimpleDelegateGraphTask::CreateAndDispatchWhenReady
+		(
+			FSimpleDelegateGraphTask::FDelegate::CreateLambda([EnsureNotify, ErrorMsg, NotifyMsg]() { EnsureNotify(*ErrorMsg, *NotifyMsg); }), TStatId(), nullptr, ENamedThreads::GameThread
+		);
+	}
+}
+
 void FUnrealEdMisc::OnActiveTabChanged(TSharedPtr<SDockTab> PreviouslyActive, TSharedPtr<SDockTab> NewlyActivated)
 {
 	OnUserActivityTabChanged(NewlyActivated);
diff --git a/Engine/Source/Editor/UnrealEd/Public/UnrealEdMisc.h b/Engine/Source/Editor/UnrealEd/Public/UnrealEdMisc.h
--- a/Engine/Source/Editor/UnrealEd/Public/UnrealEdMisc.h
+++ b/Engine/Source/Editor/UnrealEd/Public/UnrealEdMisc.h
@@ -213,6 +213,7 @@ public:
 	UNREALED_API void OnEditorChangeMode(FEditorModeID NewEditorMode);
 	UNREALED_API void OnEditorPreModal();
 	UNREALED_API void OnEditorPostModal();
+	UNREALED_API void OnEnsureFailed(const ANSICHAR* Expr, const ANSICHAR* File, int32 Line, const TCHAR* Msg, const TCHAR* CombinedErrorMsg);
 
 	/** BP nativization settings */
 	UE_DEPRECATED(5.0, "Blueprint Nativization has been removed as a supported feature. This delegate method will eventually be removed.")
diff --git a/Engine/Source/Runtime/Core/Private/Misc/AssertionMacros.cpp b/Engine/Source/Runtime/Core/Private/Misc/AssertionMacros.cpp
--- a/Engine/Source/Runtime/Core/Private/Misc/AssertionMacros.cpp
+++ b/Engine/Source/Runtime/Core/Private/Misc/AssertionMacros.cpp
@@ -489,14 +489,17 @@ FORCENOINLINE void FDebug::EnsureFailed(const ANSICHAR* Expr, const ANSICHAR* Fi
 	// Is there a debugger attached?  If not we'll submit an error report.
 	if (FPlatformMisc::IsDebuggerPresent() && !GAlwaysReportCrash)
 	{
+		const FString CombinedErrorMsg = FString::Printf(TEXT("Ensure condition failed: %hs" FILE_LINE_DESC_ANSI "\n%s\n"), Expr, File, Line, Msg);
+		FCoreDelegates::OnEnsureFailed.Broadcast(Expr, File, Line, Msg, *CombinedErrorMsg);
+
 #if !NO_LOGGING
 		if (GEnsuresAreErrors)
 		{
-			UE_LOG(LogOutputDevice, Error, TEXT("Ensure condition failed: %hs" FILE_LINE_DESC_ANSI "\n%s\n"), Expr, File, Line, Msg);
+			UE_LOG(LogOutputDevice, Error, TEXT("%s"), *CombinedErrorMsg);
 		}
 		else
 		{
-			UE_LOG(LogOutputDevice, Warning, TEXT("Ensure condition failed: %hs" FILE_LINE_DESC_ANSI "\n%s\n"), Expr, File, Line, Msg);
+			UE_LOG(LogOutputDevice, Warning, TEXT("%s"), *CombinedErrorMsg);
 		}
 #endif
 	}
@@ -544,6 +547,8 @@ FORCENOINLINE void FDebug::EnsureFailed(const ANSICHAR* Expr, const ANSICHAR* Fi
 #endif
 			GLog->Flush();
 
+			FCoreDelegates::OnEnsureFailed.Broadcast(Expr, File, Line, Msg, ErrorMsg);
+
 			// Trace the error
 #if !NO_LOGGING
 			if (GEnsuresAreErrors)
@@ -611,6 +616,9 @@ FORCENOINLINE void FDebug::EnsureFailed(const ANSICHAR* Expr, const ANSICHAR* Fi
 		}
 		else
 		{
+			const FString CombinedErrorMsg = FString::Printf(TEXT("Ensure condition failed: %hs" FILE_LINE_DESC_ANSI "\n%s\n"), Expr, File, Line, Msg);
+			FCoreDelegates::OnEnsureFailed.Broadcast(Expr, File, Line, Msg, *CombinedErrorMsg);
+
 			if (FParse::Param(FCommandLine::Get(), TEXT("nocrashreports")))
 			{
 				return;
diff --git a/Engine/Source/Runtime/Core/Private/Misc/CoreDelegates.cpp b/Engine/Source/Runtime/Core/Private/Misc/CoreDelegates.cpp
--- a/Engine/Source/Runtime/Core/Private/Misc/CoreDelegates.cpp
+++ b/Engine/Source/Runtime/Core/Private/Misc/CoreDelegates.cpp
@@ -83,6 +83,7 @@ FCurrentlyMountedPaksDelegate& FCoreInternalDelegates::GetCurrentlyMountedPaksDe
 
 TMulticastDelegate<void(bool, int32, int32)> FCoreDelegates::OnUserLoginChangedEvent;
 TMulticastDelegate<void()> FCoreDelegates::OnSafeFrameChangedEvent;
+TMulticastDelegate<void(const ANSICHAR*, const ANSICHAR*, int32, const TCHAR*, const TCHAR*)> FCoreDelegates::OnEnsureFailed;
 TMulticastDelegate<void()> FCoreDelegates::OnHandleSystemEnsure;
 TMulticastDelegate<void()> FCoreDelegates::OnHandleSystemError;
 
diff --git a/Engine/Source/Runtime/Core/Public/Misc/CoreDelegates.h b/Engine/Source/Runtime/Core/Public/Misc/CoreDelegates.h
--- a/Engine/Source/Runtime/Core/Public/Misc/CoreDelegates.h
+++ b/Engine/Source/Runtime/Core/Public/Misc/CoreDelegates.h
@@ -170,6 +170,9 @@ public:
     static CORE_API TMulticastDelegate<void(FMovieStreamerPtr)> UnRegisterMovieStreamerDelegate;
 
 	// Callback when an ensure has occurred
+	static CORE_API TMulticastDelegate<void(const ANSICHAR* Expr, const ANSICHAR* File, int32 Line, const TCHAR* Msg, const TCHAR* CombinedErrorMsg)> OnEnsureFailed;
+
+	// Callback when an ensure error report is being submitted
 	static CORE_API TMulticastDelegate<void()> OnHandleSystemEnsure;
 
 	// Callback when an error (crash) has occurred
