Fix issue when asset registry grows passed the 2Gb barrier, by using a 64 bit archive instead.

diff --git a/Engine/Source/Editor/UnrealEd/Private/Commandlets/AssetRegistryGenerator.cpp b/Engine/Source/Editor/UnrealEd/Private/Commandlets/AssetRegistryGenerator.cpp
--- a/Engine/Source/Editor/UnrealEd/Private/Commandlets/AssetRegistryGenerator.cpp
+++ b/Engine/Source/Editor/UnrealEd/Private/Commandlets/AssetRegistryGenerator.cpp
@@ -1965,7 +1965,7 @@ bool FAssetRegistryGenerator::SaveAssetRegistry(const FString& SandboxPath, bool
 			// No need to call FilterTags; it is called by InitializeFromExisting
 
 			// Create development registry data, used for DLC cooks, incremental cooks, and editor viewing
-			FArrayWriter SerializedAssetRegistry;
+			FBufferArchive64 SerializedAssetRegistry;
 
 			DevelopmentState.Save(SerializedAssetRegistry, DevelopmentSaveOptions);
 
