[Editor ITF] Restored legacy snapping support

diff --git a/Engine/Source/Editor/UnrealEd/Private/LevelEditorViewport.cpp b/Engine/Source/Editor/UnrealEd/Private/LevelEditorViewport.cpp
--- a/Engine/Source/Editor/UnrealEd/Private/LevelEditorViewport.cpp
+++ b/Engine/Source/Editor/UnrealEd/Private/LevelEditorViewport.cpp
@@ -3355,34 +3355,32 @@ bool FLevelEditorViewportClient::InputWidgetDelta(FViewport* InViewport, EAxisLi
 				
 				if(!MiddleMouseButtonDown || bEnableMMB)
 				{
+					bool bSnapped = false;
 					// Only perform snapping for legacy gizmos
 					if (!bUsesViewportITF)
 					{
-						bool bSnapped = FSnappingUtils::SnapActorsToNearestActor( Drag, this );
-						bSnapped = bSnapped || FSnappingUtils::SnapDraggedActorsToNearestVertex( Drag, this );
+						bSnapped = FSnappingUtils::SnapActorsToNearestActor( Drag, this );
+					}
+
+					bSnapped = bSnapped || FSnappingUtils::SnapDraggedActorsToNearestVertex( Drag, this );
 
-						// If we are only changing position, project the actors onto the world
-						const bool bOnlyTranslation = !Drag.IsZero() && Rot.IsZero() && Scale.IsZero();
+					// If we are only changing position, project the actors onto the world
+					const bool bOnlyTranslation = !Drag.IsZero() && Rot.IsZero() && Scale.IsZero();
 
-						const EAxisList::Type CurrentAxis = GetCurrentWidgetAxis();
-						const bool bSingleAxisDrag = CurrentAxis == EAxisList::X || CurrentAxis == EAxisList::Y || CurrentAxis == EAxisList::Z;
-						if (!bSnapped && !bSingleAxisDrag && GetDefault<ULevelEditorViewportSettings>()->SnapToSurface.bEnabled && bOnlyTranslation)
+					const EAxisList::Type CurrentAxis = GetCurrentWidgetAxis();
+					const bool bSingleAxisDrag = CurrentAxis == EAxisList::X || CurrentAxis == EAxisList::Y || CurrentAxis == EAxisList::Z;
+					if (!bSnapped && !bSingleAxisDrag && GetDefault<ULevelEditorViewportSettings>()->SnapToSurface.bEnabled && bOnlyTranslation)
+					{
+						TArray<AActor*> SelectedActors;
+						for (FSelectionIterator It(GEditor->GetSelectedActorIterator()); It; ++It)
 						{
-							TArray<AActor*> SelectedActors;
-							for (FSelectionIterator It(GEditor->GetSelectedActorIterator()); It; ++It)
+							if (AActor* Actor = Cast<AActor>(*It))
 							{
-								if (AActor* Actor = Cast<AActor>(*It))
-								{
-									SelectedActors.Add(Actor);
-								}
+								SelectedActors.Add(Actor);
 							}
-
-							ProjectActorsIntoWorld(SelectedActors, InViewport, Drag, Rot);
-						}
-						else
-						{
-							ApplyDeltaToSelectedElements(FTransform(Rot, Drag, Scale));
 						}
+
+						ProjectActorsIntoWorld(SelectedActors, InViewport, Drag, Rot);
 					}
 					else
 					{
