Added a console command to clear the FPreviewScene cache from UThumbnailRenderer as it currently encurs a rendering cost even if we're not rendering thumbnails.

diff --git a/Engine/Source/Editor/UnrealEd/Classes/ThumbnailRendering/BlueprintThumbnailRenderer.h b/Engine/Source/Editor/UnrealEd/Classes/ThumbnailRendering/BlueprintThumbnailRenderer.h
--- a/Engine/Source/Editor/UnrealEd/Classes/ThumbnailRendering/BlueprintThumbnailRenderer.h
+++ b/Engine/Source/Editor/UnrealEd/Classes/ThumbnailRendering/BlueprintThumbnailRenderer.h
@@ -34,6 +34,8 @@ class UBlueprintThumbnailRenderer : public UDefaultSizedThumbnailRenderer
 	/** Notifies the thumbnail scene to refresh components for the specified blueprint */
 	void BlueprintChanged(UBlueprint* Blueprint);
 
+	void ClearThumbnailScenes();
+
 private:
 	void OnBlueprintUnloaded(UBlueprint* Blueprint);
 
diff --git a/Engine/Source/Editor/UnrealEd/Private/ThumbnailRendering/BlueprintThumbnailRenderer.cpp b/Engine/Source/Editor/UnrealEd/Private/ThumbnailRendering/BlueprintThumbnailRenderer.cpp
--- a/Engine/Source/Editor/UnrealEd/Private/ThumbnailRendering/BlueprintThumbnailRenderer.cpp
+++ b/Engine/Source/Editor/UnrealEd/Private/ThumbnailRendering/BlueprintThumbnailRenderer.cpp
@@ -5,9 +5,24 @@
 #include "ShowFlags.h"
 #include "SceneView.h"
 #include "Misc/App.h"
+#include "UObject/UObjectIterator.h"
 
 #include UE_INLINE_GENERATED_CPP_BY_NAME(BlueprintThumbnailRenderer)
 
+static FAutoConsoleCommand CVarClearCachedBPThumbnailScenes(
+	TEXT("Editor.Thumbnails.ClearCachedBPThumbnailScenes"),
+	TEXT("Clears cached scenes used for Blueprint thumbnail rendering."),
+	FConsoleCommandDelegate::CreateLambda([]()
+	{
+		for (TObjectIterator<UBlueprintThumbnailRenderer> ThumbnailRendererIt; ThumbnailRendererIt; ++ThumbnailRendererIt)
+		{
+			if (UBlueprintThumbnailRenderer* ThumbnailRenderer = *ThumbnailRendererIt)
+			{
+				ThumbnailRenderer->ClearThumbnailScenes();
+			}
+		}
+	}));
+
 UBlueprintThumbnailRenderer::UBlueprintThumbnailRenderer(const FObjectInitializer& ObjectInitializer)
 	: Super(ObjectInitializer)
 {
@@ -87,3 +102,8 @@ void UBlueprintThumbnailRenderer::OnBlueprintUnloaded(UBlueprint* Blueprint)
 		ThumbnailScenes.RemoveThumbnailScene(Blueprint->GeneratedClass);
 	}
 }
+
+void UBlueprintThumbnailRenderer::ClearThumbnailScenes()
+{
+	ThumbnailScenes.Clear();
+}
