Avoid Category string allocation when there are no hide categories

#rnx

diff --git a/Engine/Source/Editor/UnrealEd/Private/EditorCategoryUtils.cpp b/Engine/Source/Editor/UnrealEd/Private/EditorCategoryUtils.cpp
--- a/Engine/Source/Editor/UnrealEd/Private/EditorCategoryUtils.cpp
+++ b/Engine/Source/Editor/UnrealEd/Private/EditorCategoryUtils.cpp
@@ -358,6 +358,10 @@ bool FEditorCategoryUtils::IsCategoryHiddenFromClass(const UStruct* Class, const
 bool FEditorCategoryUtils::IsCategoryHiddenFromClass(const TArray<FString>& ClassHideCategories, const UStruct* Class, const FString& Category)
 {
 	bool bIsHidden = false;
+	if(ClassHideCategories.Num() == 0)
+	{
+		return bIsHidden;
+	}
 
 	// run the category through sanitization so we can ensure compares will hit
 	const FString DisplayCategory = GetCategoryDisplayString(Category);
