Cooker: MultiProductBuild: Silently ignore the discrepency between .umap and .uasset in generated files that are not being cooked in the current cook session between the filename they were registered under in the basegame AssetRegistry and the filename they were found to have in the current cook session. The Extension might be different (and incorrect in the current cook session) because the list of object (a ULevel) is different in the current cook session and the authoritative one that saves the generated package.
#rnx

diff --git a/Engine/Source/Editor/UnrealEd/Private/Cooker/CookPackageData.cpp b/Engine/Source/Editor/UnrealEd/Private/Cooker/CookPackageData.cpp
--- a/Engine/Source/Editor/UnrealEd/Private/Cooker/CookPackageData.cpp
+++ b/Engine/Source/Editor/UnrealEd/Private/Cooker/CookPackageData.cpp
@@ -2362,17 +2362,67 @@ FPackageData& FPackageDatas::FindOrAddPackageData(const FName& PackageName, cons
 		FPackageData** PackageDataMapAddr = PackageNameToPackageData.Find(PackageName);
 		if (PackageDataMapAddr != nullptr)
 		{
+			FPackageData& ExistingPackageData = **PackageDataMapAddr;
 			FPackageData** FileNameMapAddr = FileNameToPackageData.Find(NormalizedFileName);
-			checkf(FileNameMapAddr,
-				TEXT("Package %s is being added with filename %s, but it already exists with filename %s, ")
-				TEXT("and it is not present in FileNameToPackageData map under the new name."),
-				*PackageName.ToString(), *NormalizedFileName.ToString(),
-				*(*PackageDataMapAddr)->GetFileName().ToString());
-			checkf(*FileNameMapAddr == *PackageDataMapAddr,
+			if (FileNameMapAddr == nullptr)
+			{
+				// Since we have an existing package, it should have registered its filename. Therefore the only
+				// way the current filename could be unregistered is if the existing package's filename is different.
+				FString ExistingFileNameStr = ExistingPackageData.GetFileName().ToString();
+				FString NewFileNameStr = NormalizedFileName.ToString();
+				FileNameMapAddr = FileNameToPackageData.Find(ExistingPackageData.GetFileName());
+				checkf(ExistingPackageData.GetFileName() != NormalizedFileName,
+					TEXT("InternalCookerError: FindOrAddPackageData found existing package at PackageName, and NormalizedFileName is not registered, and so we expected ExistingPackage.FileName != NormalizedFileName, but it is equal.")
+					TEXT("\n\tPackageName = %s, NormalizedFileName = %s"), *PackageName.ToString(), *NormalizedFileName.ToString());
+				checkf(FileNameMapAddr != nullptr,
+					TEXT("InternalCookerError: FindOrAddPackageData found existing package at PackageName with a different FileName, but that existing package's FileName is not registered.")
+					TEXT("\n\tPackageName = %s, NormalizedFileName = %s, ExistingFileName = %s"),
+					*PackageName.ToString(), *NewFileNameStr, *ExistingFileNameStr);
+
+				if (FPaths::GetBaseFilename(ExistingFileNameStr, false /* bRemovePath */)
+					!= FPaths::GetBaseFilename(NewFileNameStr, false /* bRemovePath */))
+				{
+					// If the filename's Path or BaseNameWithoutExtension are different, that's a discrepency we don't handle.
+					checkf(ExistingFileNameStr == NewFileNameStr,
+						TEXT("Package %s is being added with filename %s, but it already exists with filename with different path %s, ")
+						TEXT("and it is not present in FileNameToPackageData map under the new name."),
+						*PackageName.ToString(), *NewFileNameStr, *ExistingFileNameStr);
+				}
+				else
+				{
+					// We handle the extension being different, but only in the special case that the package does not need to be saved.
+					// If we are responsible for saving the package then we need to know which extension is correct, so we are broken.
+					bool bAllPreviouslyCooked = true;
+					for (const ITargetPlatform* TargetPlatform
+						: GetCookOnTheFlyServer().PlatformManager->GetSessionPlatforms())
+					{
+						if (ExistingPackageData.FindOrAddPlatformData(TargetPlatform).GetWhereCooked()
+							== EWhereCooked::ThisSession)
+						{
+							bAllPreviouslyCooked = false;
+						}
+					}
+					if (bAllPreviouslyCooked)
+					{
+						// Silently accept the existing extension; we will not need it since we are not cooking it.
+					}
+					else
+					{
+						// We don't handle different extensions when we need to cook the package, but it will probably
+						// not cause a crash elsewhere. Log it as an error and print the callstack.
+						UE_LOG(LogCook, Error,
+							TEXT("Package %s is being added with filename %s, but it already exists with filename with different extension %s. Keeping the old extension. This is unexpected and the output of the cook might be invalid."),
+							*PackageName.ToString(), *NewFileNameStr, *ExistingFileNameStr);
+						FDebug::DumpStackTraceToLog(ELogVerbosity::Display);
+					}
+					// MPCook TODO: Send a correction to all CookWorkers?
+				}
+			}
+			checkf(*FileNameMapAddr == &ExistingPackageData,
 				TEXT("Package %s is being added with filename %s, but that filename maps to a different package %s."),
 				*PackageName.ToString(), *NormalizedFileName.ToString(),
 				*(*FileNameMapAddr)->GetPackageName().ToString());
-			return **PackageDataMapAddr;
+			return ExistingPackageData;
 		}
 
 		checkf(FileNameToPackageData.Find(NormalizedFileName) == nullptr,
