Cooker: DiffOnly: Fix DiffWriterArchive failing unit test after the DiffWriterArchive bugfix that requires saving a callstack structure for every serialize call.
#rnx

diff --git a/Engine/Source/Editor/UnrealEd/Private/Tests/DiffWriterArchiveTests.cpp b/Engine/Source/Editor/UnrealEd/Private/Tests/DiffWriterArchiveTests.cpp
--- a/Engine/Source/Editor/UnrealEd/Private/Tests/DiffWriterArchiveTests.cpp
+++ b/Engine/Source/Editor/UnrealEd/Private/Tests/DiffWriterArchiveTests.cpp
@@ -67,17 +67,23 @@ bool FDiffWriterArchiveTestsCallstacks::RunTest(const FString& Parameters)
 
 	State.Serialize(Ar);
 
+	// Accumulator->bFirstSaveComplete should be is false, so bCollectingCallstacks should be false.
+	// We always collect a callstack struct, but in the !bCollectingCallstacks case the callstack struct
+	// we collected should have null cpp callstack data.
 	FCallstacks& Callstacks = Accumulator->LinkerCallstacks;
 	TestTrueExpr(Callstacks.Num() == 3);
 	TestTrueExpr(Callstacks.GetEndOffset() == sizeof(FBasicDiffState));
 	TestTrueExpr(Callstacks.GetCallstack(0).Offset == 0);
-	TestTrueExpr(Callstacks.GetCallstack(0).Callstack == 0);
+	TestTrueExpr(Callstacks.GetCallstack(0).Callstack != nullptr
+		&& Callstacks.GetCallstack(0).Callstack->Callstack == nullptr);
 	TestTrueExpr(Callstacks.GetCallstack(0).bSuppressLogging == false);
 	TestTrueExpr(Callstacks.GetCallstack(1).Offset == 4);
-	TestTrueExpr(Callstacks.GetCallstack(1).Callstack == 0);
+	TestTrueExpr(Callstacks.GetCallstack(1).Callstack != nullptr
+		&& Callstacks.GetCallstack(1).Callstack->Callstack == nullptr);
 	TestTrueExpr(Callstacks.GetCallstack(1).bSuppressLogging == false);
 	TestTrueExpr(Callstacks.GetCallstack(2).Offset == 8);
-	TestTrueExpr(Callstacks.GetCallstack(2).Callstack == 0);
+	TestTrueExpr(Callstacks.GetCallstack(2).Callstack != nullptr
+		&& Callstacks.GetCallstack(2).Callstack->Callstack == nullptr);
 	TestTrueExpr(Callstacks.GetCallstack(2).bSuppressLogging == false);
 	return true;
 }
