[AutoViz] Issue with using multiple Viewports and Geometry Inspection Modes
- Removed CVar for toggle wireframe, replaced with Editor Viewport Client variable and two member functions
- Moved handling of Geometry Inspection Wireframe from ShowFlags to EditorViewportClient
- Allowing Zebra mode stripe count up to 64
- Fixing the PIE mode displaying Geometry Inspection modes

diff --git a/Engine/Source/Editor/LevelEditor/Private/SLevelViewport.cpp b/Engine/Source/Editor/LevelEditor/Private/SLevelViewport.cpp
--- a/Engine/Source/Editor/LevelEditor/Private/SLevelViewport.cpp
+++ b/Engine/Source/Editor/LevelEditor/Private/SLevelViewport.cpp
@@ -2598,6 +2598,10 @@ TSharedRef<SWidget> SLevelViewport::BuildPIEViewportToolbar()
 				switch (ViewMode)
 				{
 				case VMI_Lit_Wireframe:
+				case VMI_Clay:
+				case VMI_Zebra:
+				case VMI_FrontBackFace:
+				case VMI_RandomColor:
 					return false;
 				default:
 					return true;
diff --git a/Engine/Source/Editor/UnrealEd/Private/EditorViewportClient.cpp b/Engine/Source/Editor/UnrealEd/Private/EditorViewportClient.cpp
--- a/Engine/Source/Editor/UnrealEd/Private/EditorViewportClient.cpp
+++ b/Engine/Source/Editor/UnrealEd/Private/EditorViewportClient.cpp
@@ -4829,10 +4829,16 @@ void FEditorViewportClient::Draw(FViewport* InViewport, FCanvas* Canvas)
 	const bool bRayTracingDebugEnabled = CurrentViewMode == VMI_RayTracingDebug && CurrentRayTracingDebugVisualizationMode != NAME_None;
 	const bool bVisualizeGPUSkinCache = CurrentViewMode == VMI_VisualizeGPUSkinCache && CurrentGPUSkinCacheVisualizationMode != NAME_None;
 	const bool bCanDisableTonemapper = bVisualizeBufferEnabled || bVisualizeGPUSkinCache || (bRayTracingDebugEnabled && !GetRayTracingVisualizationData().GetModeTonemapped(CurrentRayTracingDebugVisualizationMode));
-	
+	const bool bIsGeometryInspectionMode = CurrentViewMode == VMI_Clay || CurrentViewMode == VMI_Zebra || CurrentViewMode == VMI_FrontBackFace || CurrentViewMode == VMI_RandomColor;
+
 	EngineShowFlagOverride(ESFIM_Editor, ViewFamily.ViewMode, ViewFamily.EngineShowFlags, bCanDisableTonemapper);
 	EngineShowFlagOrthographicOverride(IsPerspective(), ViewFamily.EngineShowFlags);
 
+	if (bIsGeometryInspectionMode && bIsGeometryInspectionWireframeOn)
+	{
+		ViewFamily.EngineShowFlags.SetMeshEdges(true);
+	}
+
 	UpdateLightingShowFlags( ViewFamily.EngineShowFlags );
 
 	ViewFamily.ExposureSettings = ExposureSettings;
@@ -6353,6 +6359,16 @@ bool FEditorViewportClient::IsSetShowCollisionChecked() const
 	return EngineShowFlags.Collision;
 }
 
+void FEditorViewportClient::ToggleGeometryInspectionShowWireframe()
+{
+	bIsGeometryInspectionWireframeOn = !bIsGeometryInspectionWireframeOn;
+	Invalidate();
+}
+
+bool FEditorViewportClient::IsSetGeometryInspectionShowWireframe() const
+{
+	return bIsGeometryInspectionWireframeOn;
+}
 
 void FEditorViewportClient::SetViewMode(EViewModeIndex InViewModeIndex)
 {
@@ -6422,7 +6438,7 @@ bool FEditorViewportClient::IsZebraViewModeEnabled(bool bHorizontal) const
 
 void FEditorViewportClient::SetZebraViewModeStripeCount(uint8 InStripeCount)
 {
-	if (InStripeCount>0 && InStripeCount<32)
+	if (InStripeCount>0 && InStripeCount<64)
 	{
 		if (UMaterialInstanceDynamic* MaterialInstance = Cast<UMaterialInstanceDynamic>(GEngine->ZebraMaterialInstance))
 		{
diff --git a/Engine/Source/Editor/UnrealEd/Private/SEditorViewport.cpp b/Engine/Source/Editor/UnrealEd/Private/SEditorViewport.cpp
--- a/Engine/Source/Editor/UnrealEd/Private/SEditorViewport.cpp
+++ b/Engine/Source/Editor/UnrealEd/Private/SEditorViewport.cpp
@@ -741,16 +741,11 @@ void SEditorViewport::BindCommands()
 		FCanExecuteAction(),
 		FIsActionChecked::CreateSP( ClientRef, &FEditorViewportClient::IsZebraViewModeEnabled, false ) );
 	
-	IConsoleVariable* ShowWireframeCVar = IConsoleManager::Get().FindConsoleVariable(TEXT("r.GeometryInspection.DrawWireframe"));
-	check(ShowWireframeCVar);
 	CommandListRef.MapAction(
 		Commands.GeometryInspectionShowWireframe,
-		FExecuteAction::CreateLambda([=]()
-		{
-			ShowWireframeCVar->Set(!ShowWireframeCVar->GetBool());
-		}),
-		FCanExecuteAction::CreateLambda([]() { return true; }),
-		FIsActionChecked::CreateLambda([=]() { return ShowWireframeCVar->GetBool(); })
+		FExecuteAction::CreateSP( ClientRef, &FEditorViewportClient::ToggleGeometryInspectionShowWireframe ),
+		FCanExecuteAction(),
+		FIsActionChecked::CreateSP( ClientRef, &FEditorViewportClient::IsSetGeometryInspectionShowWireframe )
 	);
 
 	BindShowCommands( CommandListRef );
diff --git a/Engine/Source/Editor/UnrealEd/Private/ViewportToolbar/UnrealEdViewportToolbar.cpp b/Engine/Source/Editor/UnrealEd/Private/ViewportToolbar/UnrealEdViewportToolbar.cpp
--- a/Engine/Source/Editor/UnrealEd/Private/ViewportToolbar/UnrealEdViewportToolbar.cpp
+++ b/Engine/Source/Editor/UnrealEd/Private/ViewportToolbar/UnrealEdViewportToolbar.cpp
@@ -2157,7 +2157,7 @@ void PopulateViewModesMenu(UToolMenu* InMenu)
 				);
 			}
 		}
-		
+
 		{
 			struct Local
 			{
@@ -2249,7 +2249,8 @@ void PopulateViewModesMenu(UToolMenu* InMenu)
 								.Style(&FAppStyle::Get(), "Menu.SpinBox")
 								.Font( FAppStyle::GetFontStyle( TEXT( "MenuItem.Font" ) ) )
 								.MinValue(1.f)
-								.MaxValue(32.f)
+								.MaxValue(64.f)
+								.MaxSliderValue(32.f)
 								.SupportDynamicSliderMaxValue(false)
 								.Value_Static ( &FEditorViewportClient::GetZebraViewModeStripeCount )
 								.OnValueChanged_Lambda( [EditorViewportWeakPtr](uint16 NewValue)
diff --git a/Engine/Source/Editor/UnrealEd/Public/EditorViewportClient.h b/Engine/Source/Editor/UnrealEd/Public/EditorViewportClient.h
--- a/Engine/Source/Editor/UnrealEd/Public/EditorViewportClient.h
+++ b/Engine/Source/Editor/UnrealEd/Public/EditorViewportClient.h
@@ -485,6 +485,12 @@ public:
 
 	/** Callback for checking the collision geometry show flag. */
 	UNREALED_API bool IsSetShowCollisionChecked() const;
+	
+	/** Callback for toggling the geometry inspection show wireframe flag. */
+	UNREALED_API void ToggleGeometryInspectionShowWireframe();
+
+	/** Callback for checking the geometry inspection show wireframe flag. */
+	UNREALED_API bool IsSetGeometryInspectionShowWireframe() const;
 
 	/** Gets ViewportCameraTransform object for the current viewport type */
 	FViewportCameraTransform& GetViewTransform()
@@ -2300,6 +2306,9 @@ private:
 	 * true during FEditorViewportClient::StopTracking. Used to guard against reentry.
 	 */
 	bool bIsTrackingBeingStopped;
+
+	/** true if Geometry Inspection Mode Wireframe option is on. */
+	bool bIsGeometryInspectionWireframeOn = true;
 };
 
 
diff --git a/Engine/Source/Runtime/Core/Private/HAL/ConsoleManager.cpp b/Engine/Source/Runtime/Core/Private/HAL/ConsoleManager.cpp
--- a/Engine/Source/Runtime/Core/Private/HAL/ConsoleManager.cpp
+++ b/Engine/Source/Runtime/Core/Private/HAL/ConsoleManager.cpp
@@ -4582,12 +4582,6 @@ static TAutoConsoleVariable<int32> CVarLegacyLuminanceFactors(
 		"1: Legacy coefficients (inaccurate). Activate to retain the 5.4 and earlier versions behavior.\n"),
 	ECVF_ReadOnly | ECVF_RenderThreadSafe);
 
-static TAutoConsoleVariable<bool> CVarShowWireframe(
-		TEXT("r.GeometryInspection.DrawWireframe"),
-		true,
-		TEXT("Controls if wireframe is drawn in Geometry Inspection modes"),
-		ECVF_RenderThreadSafe);
-
 class FConsoleObjectShadowData
 {
 protected:
diff --git a/Engine/Source/Runtime/Engine/Private/ShowFlags.cpp b/Engine/Source/Runtime/Engine/Private/ShowFlags.cpp
--- a/Engine/Source/Runtime/Engine/Private/ShowFlags.cpp
+++ b/Engine/Source/Runtime/Engine/Private/ShowFlags.cpp
@@ -581,11 +581,6 @@ void EngineShowFlagOverride(EShowFlagInitMode ShowFlagInitMode, EViewModeIndex V
 			ViewModeIndex == VMI_FrontBackFace ||
 			ViewModeIndex == VMI_RandomColor)
 		{
-			static const auto ICVar = IConsoleManager::Get().FindTConsoleVariableDataBool(TEXT("r.GeometryInspection.DrawWireframe"));
-			const bool bDrawWireframe = (ICVar && ICVar->GetValueOnGameThread());
-
-			EngineShowFlags.SetMeshEdges(bDrawWireframe);
-
 			EngineShowFlags.SetLighting(false);
 			EngineShowFlags.SetCapsuleShadows(false);
 			EngineShowFlags.SetContactShadows(false);
