[AutoViz] Fixing the BugHawk issue with CL # 46621863
- Value of 64 is rejected to Zebra Stripe Count

diff --git a/Engine/Source/Editor/UnrealEd/Private/EditorViewportClient.cpp b/Engine/Source/Editor/UnrealEd/Private/EditorViewportClient.cpp
--- a/Engine/Source/Editor/UnrealEd/Private/EditorViewportClient.cpp
+++ b/Engine/Source/Editor/UnrealEd/Private/EditorViewportClient.cpp
@@ -6438,7 +6438,7 @@ bool FEditorViewportClient::IsZebraViewModeEnabled(bool bHorizontal) const
 
 void FEditorViewportClient::SetZebraViewModeStripeCount(uint8 InStripeCount)
 {
-	if (InStripeCount>0 && InStripeCount<64)
+	if (InStripeCount>0 && InStripeCount<=64)
 	{
 		if (UMaterialInstanceDynamic* MaterialInstance = Cast<UMaterialInstanceDynamic>(GEngine->ZebraMaterialInstance))
 		{
