[MaterialEditor] Allow drag-and-dropping a texture onto an existing texture sample node to replace the texture

diff --git a/Engine/Source/Editor/UnrealEd/Classes/MaterialGraph/MaterialGraphSchema.h b/Engine/Source/Editor/UnrealEd/Classes/MaterialGraph/MaterialGraphSchema.h
--- a/Engine/Source/Editor/UnrealEd/Classes/MaterialGraph/MaterialGraphSchema.h
+++ b/Engine/Source/Editor/UnrealEd/Classes/MaterialGraph/MaterialGraphSchema.h
@@ -272,6 +272,7 @@ class UMaterialGraphSchema : public UEdGraphSchema
 	virtual void BreakSinglePinLink(UEdGraphPin* SourcePin, UEdGraphPin* TargetPin) const override;
 	virtual bool CanEncapuslateNode(UEdGraphNode const& TestNode) const override;
 	virtual void DroppedAssetsOnGraph(const TArray<struct FAssetData>& Assets, const FVector2f& GraphPosition, UEdGraph* Graph) const override;
+	virtual void DroppedAssetsOnNode(const TArray<struct FAssetData>& Assets, const FVector2f& GraphPosition, UEdGraphNode* Node) const override;
 	virtual int32 GetNodeSelectionCount(const UEdGraph* Graph) const override;
 	virtual TSharedPtr<FEdGraphSchemaAction> GetCreateCommentAction() const override;
 	virtual bool IsCacheVisualizationOutOfDate(int32 InVisualizationCacheID) const override;
@@ -280,6 +281,7 @@ class UMaterialGraphSchema : public UEdGraphSchema
 	virtual void OnPinConnectionDoubleCicked(UEdGraphPin* PinA, UEdGraphPin* PinB, const FVector2f& GraphPosition) const override;
 	virtual bool SafeDeleteNodeFromGraph(UEdGraph* Graph, UEdGraphNode* NodeToDelete) const override;
 	virtual void GetAssetsGraphHoverMessage(const TArray<FAssetData>& Assets, const UEdGraph* HoverGraph, FString& OutTooltipText, bool& OutOkIcon) const;
+	virtual void GetAssetsNodeHoverMessage(const TArray<struct FAssetData>& Assets, const UEdGraphNode* HoverNode, FString& OutTooltipText, bool& OutOkIcon) const override;
 #if WITH_EDITORONLY_DATA
 	virtual float GetActionFilteredWeight(const FEdGraphSchemaAction& InCurrentAction, const TArray<FString>& InFilterTerms, const TArray<FString>& InSanitizedFilterTerms, const TArray<UEdGraphPin*>& DraggedFromPins) const override;
 	PRAGMA_DISABLE_DEPRECATION_WARNINGS
diff --git a/Engine/Source/Editor/UnrealEd/Private/MaterialGraphSchema.cpp b/Engine/Source/Editor/UnrealEd/Private/MaterialGraphSchema.cpp
--- a/Engine/Source/Editor/UnrealEd/Private/MaterialGraphSchema.cpp
+++ b/Engine/Source/Editor/UnrealEd/Private/MaterialGraphSchema.cpp
@@ -891,6 +891,33 @@ void UMaterialGraphSchema::DroppedAssetsOnGraph(const TArray<struct FAssetData>&
 	}
 }
 
+void UMaterialGraphSchema::DroppedAssetsOnNode(const TArray<struct FAssetData>& Assets, const FVector2f& GraphPosition, UEdGraphNode* Node) const
+{
+	const UMaterialGraphNode* MaterialGraphNode = Cast<UMaterialGraphNode>(Node);
+	if (!MaterialGraphNode || !MaterialGraphNode->MaterialExpression)
+	{
+		return;
+	}
+	UMaterialExpression* Expression = MaterialGraphNode->MaterialExpression;
+
+	// Drop a texture onto a TextureSample node -> replace texture
+	if (UMaterialExpressionTextureSample* TextureSampleNode = Cast<UMaterialExpressionTextureSample>(Expression))
+	{
+		if (Assets.Num() < 1)
+		{
+			return;
+		}
+
+		if (UTexture* Tex = Cast<UTexture>(Assets[0].GetAsset()))
+		{
+			TextureSampleNode->Texture = Tex;
+			TextureSampleNode->AutoSetSampleType();
+
+			FMaterialEditorUtilities::ForceRefreshExpressionPreviews(Node->GetGraph());
+		}
+	}
+}
+
 void UMaterialGraphSchema::UpdateMaterialOnDefaultValueChanged(const UEdGraph* Graph) const
 {
 	FMaterialEditorUtilities::UpdateMaterialAfterGraphChange(Graph);
@@ -1222,6 +1249,32 @@ void UMaterialGraphSchema::GetAssetsGraphHoverMessage(const TArray<FAssetData>&
 	}
 }
 
+void UMaterialGraphSchema::GetAssetsNodeHoverMessage(const TArray<struct FAssetData>& Assets, const UEdGraphNode* HoverNode, FString& OutTooltipText, bool& OutOkIcon) const
+{
+	OutTooltipText = FString();
+	OutOkIcon = false;
+
+	const UMaterialGraphNode* MaterialGraphNode = Cast<UMaterialGraphNode>(HoverNode);
+	if (!MaterialGraphNode || !MaterialGraphNode->MaterialExpression)
+	{
+		return;
+	}
+	UMaterialExpression* Expression = MaterialGraphNode->MaterialExpression;
+
+	if (UMaterialExpressionTextureSample* TextureSampleNode = Cast<UMaterialExpressionTextureSample>(Expression))
+	{
+		if (Assets.Num() < 1)
+		{
+			return;
+		}
+
+		if (UTexture* Tex = Cast<UTexture>(Assets[0].GetAsset()))
+		{
+			OutOkIcon = true;
+		}
+	}
+}
+
 #if WITH_EDITORONLY_DATA
 //////////////////////////////////////////////////////////////////////////
 
