[MaterialEditor] Follow-up to 46637463 to fix undo/redo and asset filter

diff --git a/Engine/Source/Editor/UnrealEd/Private/MaterialGraphSchema.cpp b/Engine/Source/Editor/UnrealEd/Private/MaterialGraphSchema.cpp
--- a/Engine/Source/Editor/UnrealEd/Private/MaterialGraphSchema.cpp
+++ b/Engine/Source/Editor/UnrealEd/Private/MaterialGraphSchema.cpp
@@ -900,16 +900,20 @@ void UMaterialGraphSchema::DroppedAssetsOnNode(const TArray<struct FAssetData>&
 	}
 	UMaterialExpression* Expression = MaterialGraphNode->MaterialExpression;
 
+	TSharedPtr<IAssetReferenceFilter> AssetReferenceFilter = FMaterialEditorUtilities::MakeMaterialAssetReferenceFilter(MaterialGraphNode->GetGraph());
+
 	// Drop a texture onto a TextureSample node -> replace texture
 	if (UMaterialExpressionTextureSample* TextureSampleNode = Cast<UMaterialExpressionTextureSample>(Expression))
 	{
-		if (Assets.Num() < 1)
+		if (Assets.Num() < 1 || (AssetReferenceFilter.IsValid() && !AssetReferenceFilter->PassesFilter(Assets[0])))
 		{
 			return;
 		}
 
 		if (UTexture* Tex = Cast<UTexture>(Assets[0].GetAsset()))
 		{
+			const FScopedTransaction Transaction(LOCTEXT("MaterialEditorReplaceTexture", "Material Editor: Replace texture"));
+			TextureSampleNode->Modify();
 			TextureSampleNode->Texture = Tex;
 			TextureSampleNode->AutoSetSampleType();
 
