[Backout] - CL46696001
#rnx
Original CL Desc
-----------------------------------------------------------------
Added Tagged Asset Browser implementation for new Material Wizard. Contained in default enabled plugin.

- Old factories removed from UI if browser is not setup in Material Editor Settings
- New factories with wizard added to UI if browser is setup in Material Editor Settings

diff --git a/Engine/Plugins/Editor/MaterialAssetWizard/MaterialAssetWizard.uplugin b/Engine/Plugins/Editor/MaterialAssetWizard/MaterialAssetWizard.uplugin
deleted file mode 100644
diff --git a/Engine/Plugins/Editor/MaterialAssetWizard/Source/MaterialAssetWizard/MaterialAssetWizard.Build.cs b/Engine/Plugins/Editor/MaterialAssetWizard/Source/MaterialAssetWizard/MaterialAssetWizard.Build.cs
deleted file mode 100644
diff --git a/Engine/Plugins/Editor/MaterialAssetWizard/Source/MaterialAssetWizard/Private/MaterialAssetWizardModule.cpp b/Engine/Plugins/Editor/MaterialAssetWizard/Source/MaterialAssetWizard/Private/MaterialAssetWizardModule.cpp
deleted file mode 100644
diff --git a/Engine/Plugins/Editor/MaterialAssetWizard/Source/MaterialAssetWizard/Private/MaterialFactoryWithAssetWizard.cpp b/Engine/Plugins/Editor/MaterialAssetWizard/Source/MaterialAssetWizard/Private/MaterialFactoryWithAssetWizard.cpp
deleted file mode 100644
diff --git a/Engine/Plugins/Editor/MaterialAssetWizard/Source/MaterialAssetWizard/Public/MaterialAssetWizardModule.h b/Engine/Plugins/Editor/MaterialAssetWizard/Source/MaterialAssetWizard/Public/MaterialAssetWizardModule.h
deleted file mode 100644
diff --git a/Engine/Plugins/Editor/MaterialAssetWizard/Source/MaterialAssetWizard/Public/MaterialFactoryWithAssetWizard.h b/Engine/Plugins/Editor/MaterialAssetWizard/Source/MaterialAssetWizard/Public/MaterialFactoryWithAssetWizard.h
deleted file mode 100644
diff --git a/Engine/Source/Editor/MaterialEditor/Public/MaterialEditorSettings.h b/Engine/Source/Editor/MaterialEditor/Public/MaterialEditorSettings.h
--- a/Engine/Source/Editor/MaterialEditor/Public/MaterialEditorSettings.h
+++ b/Engine/Source/Editor/MaterialEditor/Public/MaterialEditorSettings.h
@@ -105,14 +105,6 @@ class UMaterialEditorSettings : public UObject
 	GENERATED_BODY()
 
 public:
-
-	/** The configuration used in the Material Asset Wizard when creating a new Material. If not set, Content Browser uses the old factory without wizard. */
-	UPROPERTY(config, EditAnywhere, Category = Wizard, meta=(AllowedClasses="/Script/UserAssetTagsEditor.TaggedAssetBrowserConfiguration"))
-	FSoftObjectPath MaterialWizardConfiguration;
-
-	/** The configuration used in the Material Asset Wizard when creating a new Material Instance. If not set, Content Browser uses the old factory without wizard. */
-	UPROPERTY(config, EditAnywhere, Category = Wizard, meta=(AllowedClasses="/Script/UserAssetTagsEditor.TaggedAssetBrowserConfiguration"))
-	FSoftObjectPath MaterialInstanceWizardConfiguration;
 	
 	/**
 	Allow ignoring compilation errors of platform shaders and derived materials.
diff --git a/Engine/Source/Editor/UnrealEd/Classes/Factories/MaterialFactoryNew.h b/Engine/Source/Editor/UnrealEd/Classes/Factories/MaterialFactoryNew.h
--- a/Engine/Source/Editor/UnrealEd/Classes/Factories/MaterialFactoryNew.h
+++ b/Engine/Source/Editor/UnrealEd/Classes/Factories/MaterialFactoryNew.h
@@ -22,6 +22,5 @@ class UMaterialFactoryNew : public UFactory
 
 	//~ Begin UFactory Interface
 	virtual UObject* FactoryCreateNew(UClass* Class,UObject* InParent,FName Name,EObjectFlags Flags,UObject* Context,FFeedbackContext* Warn) override;
-	virtual bool ShouldShowInNewMenu() const override;
 	//~ Begin UFactory Interface	
 };
diff --git a/Engine/Source/Editor/UnrealEd/Classes/Factories/MaterialInstanceConstantFactoryNew.h b/Engine/Source/Editor/UnrealEd/Classes/Factories/MaterialInstanceConstantFactoryNew.h
--- a/Engine/Source/Editor/UnrealEd/Classes/Factories/MaterialInstanceConstantFactoryNew.h
+++ b/Engine/Source/Editor/UnrealEd/Classes/Factories/MaterialInstanceConstantFactoryNew.h
@@ -21,7 +21,6 @@ class UMaterialInstanceConstantFactoryNew : public UFactory
 
 	//~ Begin UFactory Interface
 	UNREALED_API virtual UObject* FactoryCreateNew(UClass* Class,UObject* InParent,FName Name,EObjectFlags Flags,UObject* Context,FFeedbackContext* Warn) override;
-	virtual bool ShouldShowInNewMenu() const override;
 	//~ Begin UFactory Interface	
 };
 
diff --git a/Engine/Source/Editor/UnrealEd/Private/Factories/EditorFactories.cpp b/Engine/Source/Editor/UnrealEd/Private/Factories/EditorFactories.cpp
--- a/Engine/Source/Editor/UnrealEd/Private/Factories/EditorFactories.cpp
+++ b/Engine/Source/Editor/UnrealEd/Private/Factories/EditorFactories.cpp
@@ -283,7 +283,6 @@
 #include "TextureReferenceResolver.h"
 #include "UDIMUtilities.h"
 #include "FileHelpers.h"
-#include "MaterialEditorSettings.h"
 #include "StaticMeshOperations.h"
 #include "UObject/OverridableManager.h"
 
@@ -490,11 +489,6 @@ UObject* UMaterialInstanceConstantFactoryNew::FactoryCreateNew(UClass* Class,UOb
 	return MIC;
 }
 
-bool UMaterialInstanceConstantFactoryNew::ShouldShowInNewMenu() const
-{
-	return Super::ShouldShowInNewMenu() && GetDefault<UMaterialEditorSettings>()->MaterialInstanceWizardConfiguration == nullptr;
-}
-
 /*------------------------------------------------------------------------------
 	UMaterialFactoryNew implementation.
 ------------------------------------------------------------------------------*/
@@ -546,11 +540,6 @@ UObject* UMaterialFactoryNew::FactoryCreateNew(UClass* Class,UObject* InParent,F
 	return NewMaterial;
 }
 
-bool UMaterialFactoryNew::ShouldShowInNewMenu() const
-{
-	return Super::ShouldShowInNewMenu() && GetDefault<UMaterialEditorSettings>()->MaterialWizardConfiguration == nullptr;
-}
-
 
 /*------------------------------------------------------------------------------
 	UMaterialFunctionFactoryNew implementation.
