Fixing the editor compilation errors when building with UE_WITH_PACKAGE_ACCESS_TRACKING=0

diff --git a/Engine/Source/Editor/UnrealEd/Private/Cooker/CookAssetRegistryAccessTracker.cpp b/Engine/Source/Editor/UnrealEd/Private/Cooker/CookAssetRegistryAccessTracker.cpp
--- a/Engine/Source/Editor/UnrealEd/Private/Cooker/CookAssetRegistryAccessTracker.cpp
+++ b/Engine/Source/Editor/UnrealEd/Private/Cooker/CookAssetRegistryAccessTracker.cpp
@@ -144,6 +144,7 @@ namespace UE::CookAssetRegistryAccessTracker
 
 	void FCookAssetRegistryAccessTracker::OnEnumerateAssets(const FARCompiledFilter& CompiledFilter, UE::AssetRegistry::EEnumerateAssetsFlags Flag)
 	{
+#if UE_WITH_PACKAGE_ACCESS_TRACKING
 		if (!bEnabled)
 		{
 			return;
@@ -185,5 +186,6 @@ namespace UE::CookAssetRegistryAccessTracker
 
 			Records.Add(NewRecord);
 		}
+#endif // UE_WITH_PACKAGE_ACCESS_TRACKING
 	}
 }
diff --git a/Engine/Source/Editor/UnrealEd/Private/Cooker/CookLog.cpp b/Engine/Source/Editor/UnrealEd/Private/Cooker/CookLog.cpp
--- a/Engine/Source/Editor/UnrealEd/Private/Cooker/CookLog.cpp
+++ b/Engine/Source/Editor/UnrealEd/Private/Cooker/CookLog.cpp
@@ -525,6 +525,7 @@ void FLogHandler::ReportActiveLog(FReplicatedLogData&& LogData, FStringView Form
 void FLogHandler::RecordLogForIncrementalCook(UE::Cook::FReplicatedLogData&& LogData,
 	ELogVerbosity::Type LogVerbosity)
 {
+#if UE_WITH_PACKAGE_ACCESS_TRACKING
 	// Note that this function can be called from any thread. Only threadsafe data only can be accessed.
 	if (LogVerbosity > ELogVerbosity::Warning)
 	{
@@ -553,6 +554,7 @@ void FLogHandler::RecordLogForIncrementalCook(UE::Cook::FReplicatedLogData&& Log
 	{
 		RecordLogForIncrementalCookGameThreadPortion(ActivePackage, MoveTemp(LogData));
 	}
+#endif // UE_WITH_PACKAGE_ACCESS_TRACKING
 }
 
 void FLogHandler::RecordLogForIncrementalCookGameThreadPortion(FName ActivePackage, FReplicatedLogData&& LogData)
