[ITF Interactions] Regression: Right-Mouse+WASD navigation in the viewport sets a keyframe when you use the 'S' key

diff --git a/Engine/Source/Editor/UnrealEd/Private/EditorViewportClient.cpp b/Engine/Source/Editor/UnrealEd/Private/EditorViewportClient.cpp
--- a/Engine/Source/Editor/UnrealEd/Private/EditorViewportClient.cpp
+++ b/Engine/Source/Editor/UnrealEd/Private/EditorViewportClient.cpp
@@ -5718,6 +5718,11 @@ bool FEditorViewportClient::IsFlightCameraInputModeActive() const
 	{
 		if( CameraController != NULL )
 		{
+			if (UsesViewportInteractions(ModeTools.Get()))
+			{
+				return ViewportNavigationHelper->bIsMouseLooking;
+			}
+
 			const bool bLeftMouseButtonDown = Viewport->KeyState(EKeys::LeftMouseButton) && !bLockFlightCamera;
 			const bool bMiddleMouseButtonDown = Viewport->KeyState( EKeys::MiddleMouseButton );
 			const bool bRightMouseButtonDown = Viewport->KeyState( EKeys::RightMouseButton );
