Add asset path for Material Functions Tooltips in the Material Editor Context Menu.

[RN] minor | Rendering Architecture

diff --git a/Engine/Source/Editor/UnrealEd/Private/MaterialGraphSchema.cpp b/Engine/Source/Editor/UnrealEd/Private/MaterialGraphSchema.cpp
--- a/Engine/Source/Editor/UnrealEd/Private/MaterialGraphSchema.cpp
+++ b/Engine/Source/Editor/UnrealEd/Private/MaterialGraphSchema.cpp
@@ -988,7 +988,6 @@ void UMaterialGraphSchema::GetMaterialFunctionActions(FGraphActionMenuBuilder& A
 			{
 				// Gather the relevant information from the asset data
 				const FString FunctionPathName = AssetData.GetObjectPathString();
-				const FText Description = AssetData.GetTagValueRef<FText>("Description");
 				TArray<FString> LibraryCategories;
 				{
 					const FString LibraryCategoriesString = AssetData.GetTagValueRef<FString>("LibraryCategories");
@@ -1064,14 +1063,17 @@ void UMaterialGraphSchema::GetMaterialFunctionActions(FGraphActionMenuBuilder& A
 						FunctionName = FunctionPathName.Right(FunctionPathName.Len() - PeriodIndex - 1);
 					}
 				}
-
+				// Remove leading slash
+				const FText AssetPath = FText::FromString(AssetData.PackagePath.ToString().RightChop(1));
+				const FText Description = AssetData.GetTagValueRef<FText>("Description");
+				const FText Tooltip = FText::Format(LOCTEXT("MaterialFunctionActionTooltip", "Location: {0}\n{1}"), AssetPath, Description);
 				// For each category the function should belong to...
 				for (const FText& CategoryName : LibraryCategoriesText)
 				{
 					TSharedPtr<FMaterialGraphSchemaAction_NewFunctionCall> NewFunctionAction(new FMaterialGraphSchemaAction_NewFunctionCall(
 						CategoryName,
 						FText::FromString(FunctionName),
-						Description, 0));
+						Tooltip, 0));
 					ActionMenuBuilder.AddAction(NewFunctionAction);
 					NewFunctionAction->FunctionPath = FunctionPathName;
 				}
