[World] ActorFolders refresh optimization, advantageous for deeply nested LevelInstances
- Listening to level actor list changes, and map changes, does not currently catch additional cases not handled by the more robust ULevel::FixupActorFolders loop, and results in costly redundant work
- ULevel::OnLevelLoaded already triggers a full UWorldFolders::RebuildList, and add/removal of folders in the outliner is handled at the FActorFolders::CreateFolder level

#tests Devtest by [REDACTED_PERSON], Load/unload levels, load/unload LevelInstances through partitions, PIE, add/remove folders and LI containing folders
[RN] [World] Optimized Actor Folder refresh on level modifications which helps loading deeply nested LevelInstances faster

diff --git a/Engine/Source/Editor/UnrealEd/Private/EditorActorFolders.cpp b/Engine/Source/Editor/UnrealEd/Private/EditorActorFolders.cpp
--- a/Engine/Source/Editor/UnrealEd/Private/EditorActorFolders.cpp
+++ b/Engine/Source/Editor/UnrealEd/Private/EditorActorFolders.cpp
@@ -48,10 +48,8 @@ FActorFolders::FActorFolders()
 {
 	check(GEngine);
 	GEngine->OnLevelActorFolderChanged().AddRaw(this, &FActorFolders::OnActorFolderChanged);
-	GEngine->OnLevelActorListChanged().AddRaw(this, &FActorFolders::OnLevelActorListChanged);
 	GEngine->OnActorFolderAdded().AddRaw(this, &FActorFolders::OnActorFolderAdded);
 
-	FEditorDelegates::MapChange.AddRaw(this, &FActorFolders::OnMapChange);
 	FEditorDelegates::PostSaveWorldWithContext.AddRaw(this, &FActorFolders::OnWorldSaved);
 	FEditorDelegates::PostSaveExternalActors.AddRaw(this, &FActorFolders::SaveWorldFoldersState);
 
@@ -65,11 +63,9 @@ FActorFolders::~FActorFolders()
 	if (GEngine)
 	{
 		GEngine->OnLevelActorFolderChanged().RemoveAll(this);
-		GEngine->OnLevelActorListChanged().RemoveAll(this);
 		GEngine->OnActorFolderAdded().RemoveAll(this);
 	}
 
-	FEditorDelegates::MapChange.RemoveAll(this);
 	FEditorDelegates::PostSaveWorldWithContext.RemoveAll(this);
 	FEditorDelegates::PostSaveExternalActors.RemoveAll(this);
 
@@ -142,36 +138,6 @@ static bool IsRunningGameOrPIE()
 	return (GIsEditor && GEditor && GEditor->GetPIEWorldContext()) || IsRunningGame();
 }
 
-void FActorFolders::OnLevelActorListChanged()
-{
-	if (!IsRunningGameOrPIE())
-	{
-		bAnyLevelsChanged = true;
-	}
-}
-
-void FActorFolders::OnAllLevelsChanged()
-{
-	TRACE_CPUPROFILER_EVENT_SCOPE(FActorFolders::OnAllLevelsChanged);
-
-	if (bAnyLevelsChanged && !IsRunningGameOrPIE())
-	{
-		QUICK_SCOPE_CYCLE_COUNTER(FActorFolders_OnAllLevelsChanged);
-		bAnyLevelsChanged = false;
-		Housekeeping();
-
-		if (UWorld* World = GEditor ? GEditor->GetEditorWorldContext().World() : nullptr)
-		{
-			RebuildFolderListForWorld(*World);
-		}
-	}
-}
-
-void FActorFolders::OnMapChange(uint32 MapChangeFlags)
-{
-	OnLevelActorListChanged();
-}
-
 void FActorFolders::OnWorldSaved(UWorld* World, FObjectPostSaveContext ObjectSaveContext)
 {
 	SaveWorldFoldersState(World);
@@ -215,24 +181,6 @@ void FActorFolders::OnActorFolderChanged(const AActor* InActor, FName OldPath)
 	}
 }
 
-void FActorFolders::RebuildFolderListForWorld(UWorld& InWorld)
-{
-	TRACE_CPUPROFILER_EVENT_SCOPE(FActorFolders::RebuildFolderListForWorld);
-
-	if (auto* Folders = WorldFolders.Find(&InWorld))
-	{
-		// For world folders, we don't empty the existing folders so that we keep empty ones.
-		// Explicitly deleted folders will already be removed from the list.
-		
-		(*Folders)->RebuildList();
-	}
-	else
-	{
-		// No folders exist for this world yet - creating them will ensure they're up to date
-		CreateWorldFolders(InWorld);
-	}
-}
-
 FActorFolderProps* FActorFolders::GetFolderProperties(UWorld& InWorld, const FFolder& InFolder)
 {
 	return GetOrCreateWorldFolders(InWorld).GetFolderProperties(InFolder);
@@ -243,9 +191,6 @@ UWorldFolders& FActorFolders::CreateWorldFolders(UWorld& InWorld)
 	// Clean up any stale worlds
 	Housekeeping();
 
-	InWorld.OnLevelsChanged().AddRaw(this, &FActorFolders::OnLevelActorListChanged);
-	InWorld.OnAllLevelsChanged().AddRaw(this, &FActorFolders::OnAllLevelsChanged);
-
 	// We intentionally don't pass RF_Transactional to ConstructObject so that we don't record the creation of the object into the undo buffer
 	// (to stop it getting deleted on undo as we manage its lifetime), but we still want it to be RF_Transactional so we can record any changes later
 	UWorldFolders* Folders = NewObject<UWorldFolders>(GetTransientPackage(), NAME_None, RF_NoFlags);
@@ -843,4 +788,4 @@ bool FActorFolders::RenameFolderInWorld(UWorld& InWorld, FName OldPath, FName Ne
 //~ End Deprecated
 ////////////////////////////////////////////
 
-#undef LOCTEXT_NAMESPACE
+#undef LOCTEXT_NAMESPACE
\ No newline at end of file
diff --git a/Engine/Source/Editor/UnrealEd/Public/EditorActorFolders.h b/Engine/Source/Editor/UnrealEd/Public/EditorActorFolders.h
--- a/Engine/Source/Editor/UnrealEd/Public/EditorActorFolders.h
+++ b/Engine/Source/Editor/UnrealEd/Public/EditorActorFolders.h
@@ -205,25 +205,12 @@ private:
 	/** Create and update a folder container for the specified world */
 	UNREALED_API UWorldFolders& CreateWorldFolders(UWorld& InWorld);
 
-	/** Rebuild the folder list for the specified world. This can be very slow as it
-		iterates all actors in memory to rebuild the array of actors for this world */
-	UNREALED_API void RebuildFolderListForWorld(UWorld& InWorld);
-
 	/** Called when an actor's folder has changed */
 	UNREALED_API void OnActorFolderChanged(const AActor* InActor, FName OldPath);
 
-	/** Called when the actor list of the current world has changed */
-	UNREALED_API void OnLevelActorListChanged();
-
-	/** Called when all the levels have been changed */
-	UNREALED_API void OnAllLevelsChanged();
-
 	/** Called when an actor folder is added */
 	UNREALED_API void OnActorFolderAdded(UActorFolder* InActorFolder);
 
-	/** Called when the global map in the editor has changed */
-	UNREALED_API void OnMapChange(uint32 MapChangeFlags);
-
 	/** Called after a world has been saved */
 	UNREALED_API void OnWorldSaved(UWorld* World, FObjectPostSaveContext ObjectSaveContext);
 
