Improved Customized UV tooltips in the material editor

#rnx

diff --git a/Engine/Source/Editor/UnrealEd/Private/MaterialGraph.cpp b/Engine/Source/Editor/UnrealEd/Private/MaterialGraph.cpp
--- a/Engine/Source/Editor/UnrealEd/Private/MaterialGraph.cpp
+++ b/Engine/Source/Editor/UnrealEd/Private/MaterialGraph.cpp
@@ -201,7 +201,9 @@ void UMaterialGraph::RebuildGraphInternal(const TMap<UMaterialExpression*, TArra
 		for (int32 UVIndex = 0; UVIndex < UE_ARRAY_COUNT(Material->GetEditorOnlyData()->CustomizedUVs); UVIndex++)
 		{
 			//@todo - localize
-			MaterialInputs.Add(FMaterialInputInfo(FText::FromString(FString::Printf(TEXT("Customized UV%u"), UVIndex)), (EMaterialProperty)(MP_CustomizedUVs0 + UVIndex), FText::FromString(FString::Printf(TEXT("CustomizedUV%uToolTip"), UVIndex))));
+			MaterialInputs.Add(FMaterialInputInfo(FText::FromString(FString::Printf(TEXT("Customized UV%u"), UVIndex)), (EMaterialProperty)(MP_CustomizedUVs0 + UVIndex), 
+				FText::Format(LOCTEXT("CustomizedUVTooltip", "Customized UV{0} allows running calculations in the vertex shader to increase performance over running them per-pixel.\n"
+				"Use the TexCoords{0} node (or ConstCoordinate {0} in the TextureSample node) to access the interpolated custom UV in the pixel shader."), UVIndex)));
 		}
 
 		MaterialInputs.Add(FMaterialInputInfo(FMaterialAttributeDefinitionMap::GetDisplayNameForMaterial(MP_PixelDepthOffset, Material), MP_PixelDepthOffset, LOCTEXT("PixelDepthOffsetToolTip", "Pixel Depth Offset")));
