Don't reset the cached world bounds for non-editor worlds.

diff --git a/Engine/Source/Editor/UnrealEd/Private/EditorViewportClient.cpp b/Engine/Source/Editor/UnrealEd/Private/EditorViewportClient.cpp
--- a/Engine/Source/Editor/UnrealEd/Private/EditorViewportClient.cpp
+++ b/Engine/Source/Editor/UnrealEd/Private/EditorViewportClient.cpp
@@ -7313,7 +7313,7 @@ void FEditorViewportClient::OnObjectTransacted(UObject* Object, const FTransacti
 		return;
 	}
 
-	if (Object->GetWorld() == GetWorld())
+	if (Object->GetWorld() == GetWorld() && !GetWorld()->IsGameWorld())
 	{
 		CachedWorldBounds.Init();
 	}
