Fix BugHawk comment: Null pointer dereference risk in OnObjectTransacted: the condition compares Object->GetWorld() to GetWorld() and then immediately dereferences GetWorld() via IsGameWorld() without ensuring GetWorld() is non-null. If both worlds are null, the equality evaluates true and GetWorld()->IsGameWorld() dereferences a null pointer.

diff --git a/Engine/Source/Editor/UnrealEd/Private/EditorViewportClient.cpp b/Engine/Source/Editor/UnrealEd/Private/EditorViewportClient.cpp
--- a/Engine/Source/Editor/UnrealEd/Private/EditorViewportClient.cpp
+++ b/Engine/Source/Editor/UnrealEd/Private/EditorViewportClient.cpp
@@ -7313,7 +7313,7 @@ void FEditorViewportClient::OnObjectTransacted(UObject* Object, const FTransacti
 		return;
 	}
 
-	if (Object->GetWorld() == GetWorld() && !GetWorld()->IsGameWorld())
+	if (GetWorld() && (Object->GetWorld() == GetWorld()) && !GetWorld()->IsGameWorld())
 	{
 		CachedWorldBounds.Init();
 	}
