Add -DlcPlugingOnly option, only cooks inside the specified plugin does not follow references outside the plugin. Removes the requirement for the development asset registry to be loaded.
#rnx

diff --git a/Engine/Source/Editor/UnrealEd/Classes/CookOnTheSide/CookOnTheFlyServer.h b/Engine/Source/Editor/UnrealEd/Classes/CookOnTheSide/CookOnTheFlyServer.h
--- a/Engine/Source/Editor/UnrealEd/Classes/CookOnTheSide/CookOnTheFlyServer.h
+++ b/Engine/Source/Editor/UnrealEd/Classes/CookOnTheSide/CookOnTheFlyServer.h
@@ -992,6 +992,9 @@ public:
 	/** Is the local CookOnTheFlyServer cooking a DLC plugin rather than a Project+Engine+EmbeddedPlugins? */
 	UNREALED_API bool IsCookingDLC() const;
 
+	/** Is the local CookOnTheFlyServer ONLY cooking a DLC plugin, ignoring all references outside the plugin path? */
+	UNREALED_API bool IsCookingDLCPluginOnly() const;
+
 	/** Return whether EDLCookInfo verification has NOT been rendered useless by settings such as CookFilter. */
 	UNREALED_API bool ShouldVerifyEDLCookInfo() const;
 
diff --git a/Engine/Source/Editor/UnrealEd/Private/CookOnTheFlyServer.cpp b/Engine/Source/Editor/UnrealEd/Private/CookOnTheFlyServer.cpp
--- a/Engine/Source/Editor/UnrealEd/Private/CookOnTheFlyServer.cpp
+++ b/Engine/Source/Editor/UnrealEd/Private/CookOnTheFlyServer.cpp
@@ -1274,6 +1274,11 @@ bool UCookOnTheFlyServer::IsCookingDLC() const
 	return !CookByTheBookOptions->DlcName.IsEmpty();
 }
 
+bool UCookOnTheFlyServer::IsCookingDLCPluginOnly() const
+{
+	return IsCookingDLC() && CookByTheBookOptions->bDlcPluginOnly;
+}
+
 bool UCookOnTheFlyServer::IsCookingAgainstFixedBase() const
 {
 	return IsCookingDLC() && CookByTheBookOptions->bCookAgainstFixedBase;
@@ -11092,6 +11097,7 @@ FBeginCookContext UCookOnTheFlyServer::CreateBeginCookByTheBookContext(const FCo
 	CookByTheBookOptions->bAllowUncookedAssetReferences = FParse::Param(FCommandLine::Get(), TEXT("AllowUncookedAssetReferences"));
 	CookByTheBookOptions->bCookList = StartupOptions.bCookList;
 	CookByTheBookOptions->DlcName = StartupOptions.DLCName;
+	CookByTheBookOptions->bDlcPluginOnly = FParse::Param(FCommandLine::Get(), TEXT("DlcPluginOnly"));
 	if (CookByTheBookOptions->bSkipHardReferences && !CookByTheBookOptions->bSkipSoftReferences)
 	{
 		UE_LOG(LogCook, Display, TEXT("Setting bSkipSoftReferences to true since bSkipHardReferences is true and skipping hard references requires skipping soft references."));
@@ -11670,6 +11676,32 @@ void UCookOnTheFlyServer::RecordDLCPackagesFromBaseGame(FBeginCookContext& Begin
 		return;
 	}
 
+	// If we're only cooking the DLC plugin we don't need Development Asset Registry.
+	// All references outside the plugin are never explored or cooked, just mark any
+	// existing PackageData's as already cooked.
+	if (IsCookingDLCPluginOnly())
+	{
+		FString DLCPath = GetMountedAssetPathForDLC();
+		UE_LOG(LogCook, Log, TEXT("Performing plugin only cook for: %s"), *DLCPath);
+
+		PackageDatas->LockAndEnumeratePackageDatas([&BeginContext,&DLCPath](FPackageData* PackageData)
+		{
+			FString PackageName = PackageData->GetPackageName().ToString();
+
+			if (!FStringView(PackageName).StartsWith(DLCPath))
+			{
+				for (const ITargetPlatform* TargetPlatform : BeginContext.TargetPlatforms)
+				{
+					PackageData->SetPlatformCooked(TargetPlatform, ECookResult::Succeeded, /*bWasCookedThisSession=*/false);
+					FPackagePlatformData& PlatformData = PackageData->FindOrAddPlatformData(TargetPlatform);
+					PlatformData.SetWhereCooked(EWhereCooked::BaseGame);
+				}
+			}
+		});
+
+		return;
+	}
+
 	const ECookByTheBookOptions& CookOptions = CookByTheBookOptions->StartupOptions;
 	const FString& BasedOnReleaseVersion = CookByTheBookOptions->BasedOnReleaseVersion;
 
diff --git a/Engine/Source/Editor/UnrealEd/Private/Cooker/CookRequestCluster.cpp b/Engine/Source/Editor/UnrealEd/Private/Cooker/CookRequestCluster.cpp
--- a/Engine/Source/Editor/UnrealEd/Private/Cooker/CookRequestCluster.cpp
+++ b/Engine/Source/Editor/UnrealEd/Private/Cooker/CookRequestCluster.cpp
@@ -2808,6 +2808,19 @@ void FRequestCluster::IsRequestCookable(const ITargetPlatform* Platform, const F
 	};
 
 	FName FileName = PackageData.GetFileName();
+
+	if (InCOTFS.CookByTheBookOptions->bDlcPluginOnly && !InDLCPath.IsEmpty())
+	{
+		FileName.ToString(NameBuffer);
+		if (!FStringView(NameBuffer).StartsWith(InDLCPath))
+		{
+			OutReason = ESuppressCookReason::NotInCurrentPlugin;
+			bOutCookable = false;
+			bOutExplorable = false;
+			return;
+		}
+	}
+
 	if (InCOTFS.PackageTracker->NeverCookPackageList.Contains(PackageName))
 	{
 		if (INDEX_NONE != UE::String::FindFirst(NameBuffer, ULevel::GetExternalActorsFolderName(), 
diff --git a/Engine/Source/Editor/UnrealEd/Private/Cooker/CookTypes.h b/Engine/Source/Editor/UnrealEd/Private/Cooker/CookTypes.h
--- a/Engine/Source/Editor/UnrealEd/Private/Cooker/CookTypes.h
+++ b/Engine/Source/Editor/UnrealEd/Private/Cooker/CookTypes.h
@@ -576,6 +576,8 @@ public:
 	bool							bCookSoftPackageReferences = false;
 	bool							bCookAgainstFixedBase = false;
 	bool							bDlcLoadMainAssetRegistry = false;
+	/** True when performing a DLC plugin only cook, always ignores references outside the DLC path. */
+	bool							bDlcPluginOnly = false;
 	/** True if CookByTheBook is being run in cooklist mode and will not be loading/saving packages. */
 	bool							bCookList = false;
 
