IAssetEditorInstance: Add a protected destructor.

IAssetEditorInstance is just a way for users to interact with asset editors in a generic manner (e.g focus) and is not really exposed to users in a way where they are meant to destroy derived class instances using a base class pointer. However it is always technically possible to do so and the IAssetEditorInstance did not have a virtual destructor which could cause issues. To explicitly disallow this a protected destructor is being added.

diff --git a/Engine/Source/Editor/UnrealEd/Public/Subsystems/AssetEditorSubsystem.h b/Engine/Source/Editor/UnrealEd/Public/Subsystems/AssetEditorSubsystem.h
--- a/Engine/Source/Editor/UnrealEd/Public/Subsystems/AssetEditorSubsystem.h
+++ b/Engine/Source/Editor/UnrealEd/Public/Subsystems/AssetEditorSubsystem.h
@@ -89,6 +89,10 @@ public:
 	virtual double GetLastActivationTime() = 0;
 	virtual void RemoveEditingAsset(UObject* Asset) = 0;
 	virtual EAssetOpenMethod GetOpenMethod() const { return EAssetOpenMethod::Edit; }
+
+protected:
+	// protected destructor since IAssetEditorInstance isn't meant to be used to delete derived class instances
+	~IAssetEditorInstance() =  default;
 };
 
 struct FRegisteredModeInfo
