Updated archive generation to process in parallel

#rnx

diff --git a/Engine/Source/Editor/UnrealEd/Private/Commandlets/GenerateGatherArchiveCommandlet.cpp b/Engine/Source/Editor/UnrealEd/Private/Commandlets/GenerateGatherArchiveCommandlet.cpp
--- a/Engine/Source/Editor/UnrealEd/Private/Commandlets/GenerateGatherArchiveCommandlet.cpp
+++ b/Engine/Source/Editor/UnrealEd/Private/Commandlets/GenerateGatherArchiveCommandlet.cpp
@@ -2,6 +2,7 @@
 
 #include "Commandlets/GenerateGatherArchiveCommandlet.h"
 
+#include "Async/ParallelFor.h"
 #include "Commandlets/Commandlet.h"
 #include "Containers/Array.h"
 #include "Containers/Map.h"
@@ -254,77 +255,115 @@ int32 UGenerateGatherArchiveCommandlet::Main( const FString& Params )
 		}
 	}
 
-	// We need to make sure the native culture is processed first
-	if (CulturesToGenerate.RemoveSingle(NativeCulture) > 0)
-	{
-		CulturesToGenerate.Insert(NativeCulture, 0);
-	}
-
-	for (const FString& CultureName : CulturesToGenerate)
-	{
-		// Add any missing manifest entries to the archive for this culture
-		const bool bIsNativeCulture = CultureName == NativeCulture;
-		LocTextHelper.EnumerateSourceTexts([&LocTextHelper, &CultureName, &NativeCulture, &bIsNativeCulture](TSharedRef<FManifestEntry> InManifestEntry) -> bool
+	auto UpdateArchiveForCulture = 
+		[&NativeCulture, &LocTextHelper](const FString& CultureName)
 		{
-			for (const FManifestContext& Context : InManifestEntry->Contexts)
-			{
-				if (!Context.bIsOptional)
+			// Add any missing manifest entries to the archive for this culture
+			const bool bIsNativeCulture = CultureName == NativeCulture;
+			LocTextHelper.EnumerateSourceTexts(
+				[&LocTextHelper, &CultureName, &NativeCulture, &bIsNativeCulture](TSharedRef<FManifestEntry> InManifestEntry) -> bool
 				{
-					TSharedPtr<FArchiveEntry> ArchiveEntry = LocTextHelper.FindTranslation(CultureName, InManifestEntry->Namespace, Context.Key, Context.KeyMetadataObj);
-					if (ArchiveEntry.IsValid())
-					{
-						// We only update existing entries for the native culture, as this lets us preserve stale translations in foreign archives so we can decide later whether we actually want to use them
-						if (bIsNativeCulture && !ArchiveEntry->Source.IsExactMatch(InManifestEntry->Source))
-						{
-							LocTextHelper.UpdateTranslation(CultureName, InManifestEntry->Namespace, Context.Key, Context.KeyMetadataObj, InManifestEntry->Source, InManifestEntry->Source);
-						}
-					}
-					else
+					for (const FManifestContext& Context : InManifestEntry->Contexts)
 					{
-						// Get the source we should use for the new entry
-						FLocItem ArchiveSource = InManifestEntry->Source;
-						if (!bIsNativeCulture)
+						if (!Context.bIsOptional)
 						{
-							TSharedPtr<FArchiveEntry> NativeArchiveEntry = LocTextHelper.FindTranslation(NativeCulture, InManifestEntry->Namespace, Context.Key, Context.KeyMetadataObj);
-							if (NativeArchiveEntry.IsValid() && !NativeArchiveEntry->Source.IsExactMatch(NativeArchiveEntry->Translation))
+							TSharedPtr<FArchiveEntry> ArchiveEntry = LocTextHelper.FindTranslation(CultureName, InManifestEntry->Namespace, Context.Key, Context.KeyMetadataObj);
+							if (ArchiveEntry.IsValid())
+							{
+								// We only update existing entries for the native culture, as this lets us preserve stale translations in foreign archives so we can decide later whether we actually want to use them
+								if (bIsNativeCulture && !ArchiveEntry->Source.IsExactMatch(InManifestEntry->Source))
+								{
+									LocTextHelper.UpdateTranslation(CultureName, InManifestEntry->Namespace, Context.Key, Context.KeyMetadataObj, InManifestEntry->Source, InManifestEntry->Source);
+								}
+							}
+							else
 							{
-								ArchiveSource = NativeArchiveEntry->Source;
+								// Get the source we should use for the new entry
+								FLocItem ArchiveSource = InManifestEntry->Source;
+								if (!bIsNativeCulture)
+								{
+									TSharedPtr<FArchiveEntry> NativeArchiveEntry = LocTextHelper.FindTranslation(NativeCulture, InManifestEntry->Namespace, Context.Key, Context.KeyMetadataObj);
+									if (NativeArchiveEntry.IsValid() && !NativeArchiveEntry->Source.IsExactMatch(NativeArchiveEntry->Translation))
+									{
+										ArchiveSource = NativeArchiveEntry->Source;
+									}
+								}
+
+								// We want to process the source before adding it to the archive
+								ConditionSource(ArchiveSource);
+
+								FLocItem ArchiveTranslation = ArchiveSource;
+								if (!bIsNativeCulture)
+								{
+									// We want to process the translation before adding it to the archive
+									// We skip this for native entries so that the source text is also added as the translation
+									ConditionTranslation(ArchiveTranslation);
+								}
+
+								LocTextHelper.AddTranslation(CultureName, InManifestEntry->Namespace, Context.Key, Context.KeyMetadataObj, ArchiveSource, ArchiveTranslation, Context.bIsOptional);
 							}
 						}
+					}
 
-						// We want to process the source before adding it to the archive
-						ConditionSource(ArchiveSource);
+					return true; // continue enumeration
+				}, true);
 
-						FLocItem ArchiveTranslation = ArchiveSource;
-						if (!bIsNativeCulture)
-						{
-							// We want to process the translation before adding it to the archive
-							// We skip this for native entries so that the source text is also added as the translation
-							ConditionTranslation(ArchiveTranslation);
-						}
+			// Trim any dead entries out of the archive
+			LocTextHelper.TrimArchive(CultureName);
 
-						LocTextHelper.AddTranslation(CultureName, InManifestEntry->Namespace, Context.Key, Context.KeyMetadataObj, ArchiveSource, ArchiveTranslation, Context.bIsOptional);
-					}
+			// Save the new archive
+			{
+				FText SaveError;
+				if (!LocTextHelper.SaveArchive(CultureName, &SaveError))
+				{
+					UE_LOGFMT(LogGenerateArchiveCommandlet, Error, "Save error: {error}",
+						("error", *SaveError.ToString()),
+						("id", GenerateGatherArchiveCommandlet::LocalizationLogIdentifier)
+					);
+					return false;
 				}
 			}
 
-			return true; // continue enumeration
-		}, true);
+			return true;
+		};
 
-		// Trim any dead entries out of the archive
-		LocTextHelper.TrimArchive(CultureName);
+	// We need to make sure the native culture is processed first, as other cultures query that when adding their sources
+	if (CulturesToGenerate.RemoveSingle(NativeCulture) > 0)
+	{
+		if (!UpdateArchiveForCulture(NativeCulture))
+		{
+			return -1;
+		}
+	}
+
+	// We can update the remaining cultures in parallel
+	if (CulturesToGenerate.Num() > 0)
+	{
+		TSharedPtr<ILocFileNotifies> LocFileNotifies = LocTextHelper.GetLocFileNotifies();
 
-		// Save the new archive
+		if (LocFileNotifies)
 		{
-			FText SaveError;
-			if (!LocTextHelper.SaveArchive(CultureName, &SaveError))
+			LocFileNotifies->BeginParallelTasks();
+		}
+
+		std::atomic<bool> bSuccess = true;
+		ParallelFor(TEXT("GenerateGatherArchiveCommandlet"), CulturesToGenerate.Num(), 1, [&CulturesToGenerate, &UpdateArchiveForCulture, &bSuccess](int32 Index)
+		{
+			const FString& CultureName = CulturesToGenerate[Index];
+			if (!UpdateArchiveForCulture(CultureName))
 			{
-				UE_LOGFMT(LogGenerateArchiveCommandlet, Error, "Save error: {error}",
-					("error", *SaveError.ToString()),
-					("id", GenerateGatherArchiveCommandlet::LocalizationLogIdentifier)
-				);
-				return -1;
+				bSuccess = false;
 			}
+		});
+
+		if (LocFileNotifies)
+		{
+			LocFileNotifies->EndParallelTasks();
+		}
+
+		if (!bSuccess)
+		{
+			return -1;
 		}
 	}
 	
