Fixed inverted logic

#rnx

diff --git a/Engine/Source/Editor/UnrealEd/Private/Commandlets/GenerateTextLocalizationResourceCommandlet.cpp b/Engine/Source/Editor/UnrealEd/Private/Commandlets/GenerateTextLocalizationResourceCommandlet.cpp
--- a/Engine/Source/Editor/UnrealEd/Private/Commandlets/GenerateTextLocalizationResourceCommandlet.cpp
+++ b/Engine/Source/Editor/UnrealEd/Private/Commandlets/GenerateTextLocalizationResourceCommandlet.cpp
@@ -111,11 +111,11 @@ int32 UGenerateTextLocalizationResourceCommandlet::Main(const FString& Params)
 		return -1;
 	}
 
-	for(int32 i = 0; i < CulturesToGenerate.Num(); ++i)
+	for (const FString& CultureToGenerate : CulturesToGenerate)
 	{
-		if( FInternationalization::Get().GetCulture( CulturesToGenerate[i] ).IsValid() )
+		if (!FInternationalization::Get().GetCulture(CultureToGenerate))
 		{
-			UE_LOG(LogGenerateTextLocalizationResourceCommandlet, Verbose, TEXT("Specified culture is not a valid runtime culture, but may be a valid base language: %s"), *(CulturesToGenerate[i]));
+			UE_LOG(LogGenerateTextLocalizationResourceCommandlet, Verbose, TEXT("Specified culture is not a valid runtime culture, but may be a valid base language: %s"), *CultureToGenerate);
 		}
 	}
 
