Cooker: Fix out of bounds array access due to off-by-one error, found by static analysis.
#rnx

diff --git a/Engine/Source/Editor/UnrealEd/Private/Cooker/CookTypes.cpp b/Engine/Source/Editor/UnrealEd/Private/Cooker/CookTypes.cpp
--- a/Engine/Source/Editor/UnrealEd/Private/Cooker/CookTypes.cpp
+++ b/Engine/Source/Editor/UnrealEd/Private/Cooker/CookTypes.cpp
@@ -1108,7 +1108,7 @@ FMultiPackageReaderResults GetSaveExportsAndImports(UPackage* Package, UObject*
 
 	for (int32 MultiOutputIndex = 0; MultiOutputIndex < PackageWriter->SavedRecord.Packages.Num(); ++MultiOutputIndex)
 	{
-		if (MultiOutputIndex > UE_ARRAY_COUNT(Result.Realms))
+		if (MultiOutputIndex >= UE_ARRAY_COUNT(Result.Realms))
 		{
 			break;
 		}
