[World] EDL Cook Filtering. When enabled, only the named EDLs will cook, and base maps will not be saved either. Used for multi product builds. Replaces the temporary bShouldSkipCookSaveGeneratorPackageIfBaseGamePartiallyCooked.

#tests Built local maps with EDLs filered out, also Hermes_Terrain, compared cook output
#rnx

diff --git a/Engine/Source/Editor/UnrealEd/Private/CookOnTheFlyServer.cpp b/Engine/Source/Editor/UnrealEd/Private/CookOnTheFlyServer.cpp
--- a/Engine/Source/Editor/UnrealEd/Private/CookOnTheFlyServer.cpp
+++ b/Engine/Source/Editor/UnrealEd/Private/CookOnTheFlyServer.cpp
@@ -8393,6 +8393,8 @@ const TCHAR* GCookRequestUsageMessage = TEXT(
 	"			Request all packages in the given directory.\n"
 	"		-mapinisection=<SectionNameInEditorIni>	\n"
 	"			Specify an ini section of packages to cook, in the style of AlwaysCookMaps.\n"
+	"		-EDLCookFilters=<EDLObjectPath1>+<EDLObjectPath2>+... \n"
+	"			Specify EDL object paths to cook exclusively, skipping all other EDLs used by the base map and the base map itself.\n"
 	"	Ini:\n"
 	"		Editor.ini\n"
 	"			[AlwaysCookMaps]\n"
@@ -8414,6 +8416,9 @@ const TCHAR* GCookRequestUsageMessage = TEXT(
 	"			[/Script/EngineSettings.GameMapsSettings]\n"
 	"				GameDefaultMap=<PackageName>\n"
 	"				; And other default types; see GameMapsSettings.\n"
+	"			[/Script/Engine.WorldPartitionSettings]\n"
+	"				; Specify EDL object paths to cook exclusively, skipping all other EDLs used by the base map and the base map itself.\n"
+	"				+EDLCookFilters=<EDLObjectPath>\n"
 	"	C++API\n"
 	"		FAssetManager::ModifyCook\n"
 	"			// Subclass FAssetManager (Engine.ini:[/Script/Engine.Engine]:AssetManagerClassName) and override this hook.\n"
diff --git a/Engine/Source/Editor/UnrealEd/Private/Cooker/CookGenerationHelper.cpp b/Engine/Source/Editor/UnrealEd/Private/Cooker/CookGenerationHelper.cpp
--- a/Engine/Source/Editor/UnrealEd/Private/Cooker/CookGenerationHelper.cpp
+++ b/Engine/Source/Editor/UnrealEd/Private/Cooker/CookGenerationHelper.cpp
@@ -2,7 +2,6 @@
 
 #include "Cooker/CookGenerationHelper.h"
 
-#include "Algo/AnyOf.h"
 #include "Algo/Unique.h"
 #include "AssetRegistry/AssetData.h"
 #include "Cooker/CookDirector.h"
@@ -1552,53 +1551,10 @@ void FGenerationHelper::OnRequestFencePassed(UCookOnTheFlyServer& COTFS)
 		// the request fence. This can happen in incremental cooks, or in race conditions if we sent all packages
 		// out for saving before receiving the EndQueueGeneratedPackagesOnDirector message.
 		OnNumSavedUpdated();
-
-		FName PackageName = GetOwner().GetPackageName();
-		constexpr bool bShouldSkipCookSaveGeneratorPackageIfBaseGamePartiallyCooked = true;
-		if (bShouldSkipCookSaveGeneratorPackageIfBaseGamePartiallyCooked)
-		{
-			auto IsCookedElsewhere = [](const FPackageData& PackageData)
-				{
-					for (const TPair<const ITargetPlatform*, FPackagePlatformData>& PlatformPair
-						: PackageData.GetPlatformDatas())
-					{
-						if (PlatformPair.Value.GetWhereCooked() != EWhereCooked::ThisSession)
-						{
-							return true;
-						}
-					}
-					return false;
-				};
-			if (!IsCookedElsewhere(GetOwner()))
-			{
-				// This is a temporary solution for ExperimentalFeature MultiProductBuilds that want to recook a
-				// Generator with extra GeneratedPackages inserted by plugins, but do not want to duplicate the cook
-				// of the generator package. MultiProductBuild implements that case by adding the generated packages
-				// into the basegame's DevelopmentAssetRegistry, but NOT adding the generator package, so that in
-				// RecordDLCPackagesFromBaseGame we marked the generated packages as
-				// SetWhereCooked(EWhereCooked::BaseGame), but did not mark the generator.
-				// Detect that case and suppress the save of the generator package via bSaveGenerator=false.
-				bool bAnyCookedElsewhere = Algo::AnyOf(PackagesToGenerate,
-					[&IsCookedElsewhere](const FCookGenerationInfo& Info)
-					{
-						return IsCookedElsewhere(*Info.PackageData);
-					});
-				if (bAnyCookedElsewhere)
-				{
-					// OnSkipSaveExistingGenerator, and the broadcast of it to CookWorkers, should be called before
-					// OnQueuedGeneratedPackagesFencePassed because Generators wait for
-					// OnQueuedGeneratedPackagesFencePassed and then execute code that uses this value if set.
-					OnSkipSaveExistingGenerator();
-					if (COTFS.CookDirector)
-					{
-						FGeneratorEventMessage Message(EGeneratorEvent::SkipSaveExistingGenerator, PackageName);
-						COTFS.CookDirector->BroadcastMessage(MoveTemp(Message), ECookBroadcastTiming::AfterAssignPackages);
-					}
-				}
-			}
-		}
+		
 		if (COTFS.CookDirector)
 		{
+			FName PackageName = GetOwner().GetPackageName();
 			FGeneratorEventMessage Message(EGeneratorEvent::QueuedGeneratedPackagesFencePassed, PackageName);
 			COTFS.CookDirector->BroadcastMessage(MoveTemp(Message), ECookBroadcastTiming::AfterAssignPackages);
 		}
diff --git a/Engine/Source/Runtime/Engine/Private/Level.cpp b/Engine/Source/Runtime/Engine/Private/Level.cpp
--- a/Engine/Source/Runtime/Engine/Private/Level.cpp
+++ b/Engine/Source/Runtime/Engine/Private/Level.cpp
@@ -3468,7 +3468,7 @@ void ULevel::OnLevelLoaded()
 		bIsMapBuildDataOwner = (OwningLevel == this);
 	}
 
-	// 1. Cook commandlet does it's own UWorldPartition::Initialize call in FWorldPartitionCookPackageSplitter::GetGenerateList
+	// 1. Cook commandlet does its own UWorldPartition::Initialize call in FWorldPartitionCookPackageSplitter::GetGenerateList
 	// 2. Do not Initialize if World doesn't have a UWorldPartitionSubsystem
 	if (!IsRunningCookCommandlet() && OwningWorld->HasSubsystem<UWorldPartitionSubsystem>())
 	{
diff --git a/Engine/Source/Runtime/Engine/Private/WorldCookPackageSplitter.cpp b/Engine/Source/Runtime/Engine/Private/WorldCookPackageSplitter.cpp
--- a/Engine/Source/Runtime/Engine/Private/WorldCookPackageSplitter.cpp
+++ b/Engine/Source/Runtime/Engine/Private/WorldCookPackageSplitter.cpp
@@ -5,6 +5,7 @@
 #if WITH_EDITOR
 #include "Engine/World.h"
 #include "Editor.h"
+#include "WorldPartition/WorldPartitionSettings.h"
 
 REGISTER_COOKPACKAGE_SPLITTER(FWorldCookPackageSplitter, UWorld);
 
@@ -119,6 +120,9 @@ ICookPackageSplitter::FGenerationManifest FWorldCookPackageSplitter::ReportGener
 
 	FWorldDelegates::OnWorldCleanup.AddRaw(this, &FWorldCookPackageSplitter::OnWorldCleanup);
 
+	// If we have EDL cook filters, we are trying to cook the EDLs in isolation. Do not save the map even though it was loaded.
+	Manifest.bCookSaveGeneratorPackage = UWorldPartitionSettings::GetEDLCookFilters().IsEmpty();
+
 	return Manifest;
 }
 
diff --git a/Engine/Source/Runtime/Engine/Private/WorldPartition/ContentBundle/ContentBundleEditor.cpp b/Engine/Source/Runtime/Engine/Private/WorldPartition/ContentBundle/ContentBundleEditor.cpp
--- a/Engine/Source/Runtime/Engine/Private/WorldPartition/ContentBundle/ContentBundleEditor.cpp
+++ b/Engine/Source/Runtime/Engine/Private/WorldPartition/ContentBundle/ContentBundleEditor.cpp
@@ -6,20 +6,21 @@
 #if WITH_EDITOR
 
 #include "Engine/Engine.h"
-#include "WorldPartition/WorldPartition.h"
 #include "Engine/Level.h"
-#include "WorldPartition/DataLayer/WorldDataLayers.h"
-#include "WorldPartition/ContentBundle/ContentBundleStatus.h"
+#include "WorldPartition/ActorDescContainerInstance.h"
+#include "WorldPartition/WorldPartition.h"
+#include "WorldPartition/WorldPartitionLevelHelper.h"
 #include "WorldPartition/WorldPartitionRuntimeHash.h"
+#include "WorldPartition/WorldPartitionSettings.h"
 #include "WorldPartition/ContentBundle/ContentBundleDescriptor.h"
 #include "WorldPartition/ContentBundle/ContentBundleEditorSubsystemInterface.h"
-#include "WorldPartition/ContentBundle/ContentBundleWorldSubsystem.h"
 #include "WorldPartition/ContentBundle/ContentBundleLog.h"
 #include "WorldPartition/ContentBundle/ContentBundlePaths.h"
-#include "WorldPartition/Cook/WorldPartitionCookPackageContextInterface.h"
+#include "WorldPartition/ContentBundle/ContentBundleStatus.h"
+#include "WorldPartition/ContentBundle/ContentBundleWorldSubsystem.h"
 #include "WorldPartition/Cook/WorldPartitionCookPackage.h"
-#include "WorldPartition/WorldPartitionLevelHelper.h"
-#include "WorldPartition/ActorDescContainerInstance.h"
+#include "WorldPartition/Cook/WorldPartitionCookPackageContextInterface.h"
+#include "WorldPartition/DataLayer/WorldDataLayers.h"
 
 FContentBundleEditor::FContentBundleEditor(TSharedPtr<FContentBundleClient>& InClient, UWorld* InWorld)
 	: FContentBundleBase(InClient, InWorld)
@@ -411,7 +412,8 @@ bool FContentBundleEditor::GatherPackagesToCook(class IWorldPartitionCookPackage
 
 	bool bIsSuccess = true;
 
-	if (HasCookedContent())
+	// Do not cook ContentBundles when building with (exclusive) EDL filters
+	if (HasCookedContent() && UWorldPartitionSettings::GetEDLCookFilters().IsEmpty())
 	{
 		ExternalStreamingObject->ForEachStreamingCells([this, &bIsSuccess, &CookContext](UWorldPartitionRuntimeCell& RuntimeCell)
 		{
diff --git a/Engine/Source/Runtime/Engine/Private/WorldPartition/Cook/WorldPartitionCook.cpp b/Engine/Source/Runtime/Engine/Private/WorldPartition/Cook/WorldPartitionCook.cpp
--- a/Engine/Source/Runtime/Engine/Private/WorldPartition/Cook/WorldPartitionCook.cpp
+++ b/Engine/Source/Runtime/Engine/Private/WorldPartition/Cook/WorldPartitionCook.cpp
@@ -4,6 +4,7 @@
 
 #include "WorldPartition/WorldPartition.h"
 #include "WorldPartition/WorldPartitionRuntimeHash.h"
+#include "WorldPartition/WorldPartitionSettings.h"
 #include "WorldPartition/DataLayer/ExternalDataLayerAsset.h"
 #include "WorldPartition/DataLayer/ExternalDataLayerHelper.h"
 #include "WorldPartition/DataLayer/ExternalDataLayerManager.h"
@@ -59,6 +60,17 @@ bool UWorldPartition::GatherPackagesToCook(IWorldPartitionCookPackageContext& Co
 		{
 			check(InCookPackageObject->GetPackageNameToCreate() == InPackageName);
 			const UExternalDataLayerAsset* ExternalDataLayerAsset = InCookPackageObject->GetExternalDataLayerAsset();
+
+			// The EDL filter list, when set, will exclude any package which is not part of the EDL list
+			const TSet<FTopLevelAssetPath>& EDLCookFilters = UWorldPartitionSettings::GetEDLCookFilters();
+			if (!EDLCookFilters.IsEmpty())
+			{
+				if (!ExternalDataLayerAsset || !EDLCookFilters.Contains(FTopLevelAssetPath(ExternalDataLayerAsset)))
+				{
+					return;
+				}
+			}
+
 			const FString ExternalContentRoot = ExternalDataLayerAsset ? ExternalDataLayerManager->GetExternalDataLayerLevelRootPath(ExternalDataLayerAsset) : FString();
 			const FString& Root = ExternalDataLayerAsset ? ExternalContentRoot : WorldPackageName;
 			if (CookContext.AddPackageToGenerate(this, InCookPackageObject, Root, InPackageName))
@@ -118,7 +130,6 @@ bool UWorldPartition::PopulateGeneratorPackageForCook(IWorldPartitionCookPackage
 	TArray<const FWorldPartitionCookPackage*> GenericPackagesToCook;
 	for (const FWorldPartitionCookPackage* PackageToCook : InPackagesToCook)
 	{
-		// Defer generic packages as they depend on level packages
 		if (PackageToCook->Type == FWorldPartitionCookPackage::EType::Level)
 		{
 			if (!OnPopulateGeneratorPackageForCook(PackageToCook))
@@ -126,13 +137,14 @@ bool UWorldPartition::PopulateGeneratorPackageForCook(IWorldPartitionCookPackage
 				return false;
 			}
 		}
-		else
+		else // Delay generic packages as they depend on level packages
 		{
 			check(PackageToCook->Type == FWorldPartitionCookPackage::EType::Generic);
 			GenericPackagesToCook.Add(PackageToCook);
 		}
 	}
 
+	// Move on to delayed packages
 	for (const FWorldPartitionCookPackage* GenericPackageToCook : GenericPackagesToCook)
 	{
 		if (!OnPopulateGeneratorPackageForCook(GenericPackageToCook))
diff --git a/Engine/Source/Runtime/Engine/Private/WorldPartition/WorldPartitionSettings.cpp b/Engine/Source/Runtime/Engine/Private/WorldPartition/WorldPartitionSettings.cpp
--- a/Engine/Source/Runtime/Engine/Private/WorldPartition/WorldPartitionSettings.cpp
+++ b/Engine/Source/Runtime/Engine/Private/WorldPartition/WorldPartitionSettings.cpp
@@ -1,9 +1,59 @@
 // Copyright Epic Games, Inc. All Rights Reserved.
 
 #include "WorldPartition/WorldPartitionSettings.h"
+#include "Misc/CommandLine.h"
+#include "Misc/Parse.h"
+#include "UObject/SoftObjectPath.h"
 
 #include UE_INLINE_GENERATED_CPP_BY_NAME(WorldPartitionSettings)
 
 UWorldPartitionSettings::UWorldPartitionSettings(const FObjectInitializer& ObjectInitializer)
 	: Super(ObjectInitializer)
 {}
+
+#if WITH_EDITOR
+
+const TSet<FTopLevelAssetPath>& UWorldPartitionSettings::GetEDLCookFilters()
+{
+	// Caches the current EDL cooking filters from the commandline and config
+	struct EDLCookFilterCommandLineLoader
+	{
+		EDLCookFilterCommandLineLoader()
+		{
+			// First parse EDL filters passed through commandline like -EDLCookFilters=Value1+Value2+Value3
+			FString SwitchValue;
+			if (FParse::Value(FCommandLine::Get(), TEXT("-EDLCookFilters="), SwitchValue))
+			{
+				TArray<FString> Tokens;
+				SwitchValue.ParseIntoArray(Tokens, TEXT("+"));
+				for (const FString& Filter : Tokens)
+				{
+					FTopLevelAssetPath Path;
+					if (Path.TrySetPath(Filter))
+					{
+						Filters.Add(Path);
+					}
+				}
+			}
+
+			// Then append filters found in config, e.g. in [Base/Default/etc.]Engine.ini
+			// [/Script/Engine.WorldPartitionSettings]
+			// +EDLCookFilters=/MyGFP/DataLayer/MyEDL.MyEDL
+			for (const TSoftObjectPtr<UExternalDataLayerAsset>& EDL : GetDefault<UWorldPartitionSettings>()->EDLCookFilters)
+			{
+				FSoftObjectPath Path = EDL.ToSoftObjectPath();
+				if (Path.IsValid())
+				{
+					Filters.Add(Path.GetAssetPath());
+				}
+			}
+		}
+
+		TSet<FTopLevelAssetPath> Filters;
+	};
+
+	static EDLCookFilterCommandLineLoader Loader;
+	return Loader.Filters;
+}
+
+#endif // #if WITH_EDITOR
diff --git a/Engine/Source/Runtime/Engine/Public/WorldPartition/WorldPartitionSettings.h b/Engine/Source/Runtime/Engine/Public/WorldPartition/WorldPartitionSettings.h
--- a/Engine/Source/Runtime/Engine/Public/WorldPartition/WorldPartitionSettings.h
+++ b/Engine/Source/Runtime/Engine/Public/WorldPartition/WorldPartitionSettings.h
@@ -35,7 +35,6 @@ public:
 	
 	TSubclassOf<UWorldPartitionEditorHash> GetEditorHashDefaultClass() const { return EditorHashDefaultClass; }
 	TSubclassOf<UWorldPartitionRuntimeHash> GetRuntimeHashDefaultClass() const { return RuntimeHashDefaultClass; }
-	bool ShouldEDLPackagesInheritWorldChunkAssignmentsDuringCook() const { return bShouldEDLPackagesInheritWorldChunkAssignmentsDuringCook;	}
 
 	UWorldPartitionPropertyOverridePolicy* GetPropertyOverridePolicy() const
 	{
@@ -47,6 +46,18 @@ public:
 		return UnsupportedActorClassesPerHLODLayerType;
 	}
 
+#if WITH_EDITOR
+	bool ShouldEDLPackagesInheritWorldChunkAssignmentsDuringCook() const { return bShouldEDLPackagesInheritWorldChunkAssignmentsDuringCook;	}
+
+	/**
+	 * EDL Filters
+	 * Multi product build optimization, where it is advantageous to cook EDLs on other machines for example, and where producing the map more than once is a waste and can lead to duplication in game packages.
+	 * If this list is empty, EDLs are not filtered and we have a full cook.
+	 * If there are paths in the list however, only these EDLs are cooked, excluding other EDLs not in the list, and most importantly excluding the base map
+	 */
+	static ENGINE_API const TSet<FTopLevelAssetPath>& GetEDLCookFilters();
+#endif
+
 protected:
 	friend class ULevelInstanceSettings;
 	void SetPropertyOverridePolicy(UWorldPartitionPropertyOverridePolicy* InPropertyOverridePolicy)
@@ -80,8 +91,13 @@ protected:
 	UPROPERTY(Config)
 	TMap<EHLODLayerType, FHLODLayerTypeUnsupportedActorClasses> UnsupportedActorClassesPerHLODLayerType;
 
+#if WITH_EDITORONLY_DATA
 	UPROPERTY(EditAnywhere, Config, Category = WorldPartition, meta=(DisplayName="Cook - External Data Layers packages inherit World chunk assignments"))
 	bool bShouldEDLPackagesInheritWorldChunkAssignmentsDuringCook = true;
+
+	UPROPERTY(EditAnywhere, Config, Category = WorldPartition, meta=(ToolTip="EDLs to cook in exclusion from base maps and other EDLs"))
+	TSet<TSoftObjectPtr<UExternalDataLayerAsset>> EDLCookFilters;
+#endif
 };
 
 #undef UE_API
