[UnrealEd][EntityFramework] Fixed surface alignment for entities. Introduced surface alignment for drag&drop in prefab editor.

diff --git a/Engine/Source/Editor/UnrealEd/Private/Editor/ActorPositioning.cpp b/Engine/Source/Editor/UnrealEd/Private/Editor/ActorPositioning.cpp
--- a/Engine/Source/Editor/UnrealEd/Private/Editor/ActorPositioning.cpp
+++ b/Engine/Source/Editor/UnrealEd/Private/Editor/ActorPositioning.cpp
@@ -20,6 +20,9 @@
 #include "LandscapeComponent.h"
 #include "Editor/EditorPerProjectUserSettings.h"
 #include "Editor/ObjectPositioning.h"
+#include "Elements/Framework/TypedElementRegistry.h"
+#include "Elements/Interfaces/TypedElementWorldInterface.h"
+#include "Elements/Framework/EngineElementsLibrary.h"
 
 namespace ActorPositioningLocals
 {
@@ -47,15 +50,46 @@ namespace ActorPositioningLocals
 
 		return ResultOut;
 	}
+
+	static void AddIgnoredActors(FCollisionQueryParams& OutParams, const TArray<AActor*>* IgnoreActors)
+	{
+		if (!IgnoreActors)
+		{
+			return;
+		}
+
+		UTypedElementRegistry* Registry = UTypedElementRegistry::GetInstance();
+		if (!Registry)
+		{
+			OutParams.AddIgnoredActors(*IgnoreActors);
+			return;
+		}
+
+		for (AActor* IgnoreActor : *IgnoreActors)
+		{
+			if (!IgnoreActor)
+			{
+				continue;
+			}
+
+			FTypedElementHandle Element = UEngineElementsLibrary::AcquireEditorActorElementHandle(IgnoreActor);
+			TTypedElement<ITypedElementWorldInterface> WorldElement = Registry->GetElement<ITypedElementWorldInterface>(Element);
+			if (WorldElement)
+			{
+				WorldElement.AddIgnoredElementToCollisionQueryParams(Element, OutParams, true);
+			}
+			else
+			{
+				OutParams.AddIgnoredActor(IgnoreActor);
+			}
+		}
+	}
 }
 
 FActorPositionTraceResult FActorPositioning::TraceWorldForPositionWithDefault(const FViewportCursorLocation& Cursor, const FSceneView& View, const TArray<AActor*>* IgnoreActors)
 {
 	FCollisionQueryParams Param(SCENE_QUERY_STAT(DragDropTrace), true);
-	if (IgnoreActors)
-	{
-		Param.AddIgnoredActors(*IgnoreActors);
-	}
+	ActorPositioningLocals::AddIgnoredActors(Param, IgnoreActors);
 
 	return ActorPositioningLocals::ConvertToActorPositionTraceResult(
 		UE::Positioning::TraceWorldForPositionWithDefault(Cursor, View, IgnoreActors ? &Param : nullptr));
@@ -64,10 +98,7 @@ FActorPositionTraceResult FActorPositioning::TraceWorldForPositionWithDefault(co
 FActorPositionTraceResult FActorPositioning::TraceWorldForPosition(const FViewportCursorLocation& Cursor, const FSceneView& View, const TArray<AActor*>* IgnoreActors)
 {
 	FCollisionQueryParams Param(SCENE_QUERY_STAT(DragDropTrace), true);
-	if (IgnoreActors)
-	{
-		Param.AddIgnoredActors(*IgnoreActors);
-	}
+	ActorPositioningLocals::AddIgnoredActors(Param, IgnoreActors);
 
 	return ActorPositioningLocals::ConvertToActorPositionTraceResult(
 		UE::Positioning::TraceWorldForPosition(Cursor, View, IgnoreActors ? &Param : nullptr));
@@ -76,10 +107,7 @@ FActorPositionTraceResult FActorPositioning::TraceWorldForPosition(const FViewpo
 FActorPositionTraceResult FActorPositioning::TraceWorldForPosition(const UWorld& InWorld, const FSceneView& InSceneView, const FVector& RayStart, const FVector& RayEnd, const TArray<AActor*>* IgnoreActors)
 {
 	FCollisionQueryParams Param(SCENE_QUERY_STAT(DragDropTrace), true);
-	if (IgnoreActors)
-	{
-		Param.AddIgnoredActors(*IgnoreActors);
-	}
+	ActorPositioningLocals::AddIgnoredActors(Param, IgnoreActors);
 
 	return ActorPositioningLocals::ConvertToActorPositionTraceResult(
 		UE::Positioning::TraceWorldForPosition(InWorld, InSceneView, RayStart, RayEnd, IgnoreActors ? &Param : nullptr));
@@ -157,6 +185,11 @@ FTransform FActorPositioning::GetSurfaceAlignedTransform(const FPositioningData&
 	{
 		RotatorQuat = Data.ActorFactory->AlignObjectToSurfaceNormal(Data.SurfaceNormal, RotatorQuat);
 	}
+	else
+	{
+		// Try to find a default alignment if we don't have a factory.
+		RotatorQuat = FindActorAlignmentRotation(RotatorQuat, FVector::UpVector, Data.SurfaceNormal);
+	}
 
 	// Choose the largest location offset of the various options (global viewport settings, collision, factory offset)
 	const ULevelEditorViewportSettings* ViewportSettings = GetDefault<ULevelEditorViewportSettings>();
diff --git a/Engine/Source/Editor/UnrealEd/Private/SnappingUtils.cpp b/Engine/Source/Editor/UnrealEd/Private/SnappingUtils.cpp
--- a/Engine/Source/Editor/UnrealEd/Private/SnappingUtils.cpp
+++ b/Engine/Source/Editor/UnrealEd/Private/SnappingUtils.cpp
@@ -21,6 +21,7 @@
 #include "ISnappingPolicy.h"
 #include "ViewportSnappingModule.h"
 #include "ActorGroupingUtils.h"
+#include "Elements/Actor/ActorElementData.h"
 
 namespace SnappingUtilsLocals
 {
@@ -92,7 +93,7 @@ public:
 	 * @param ViewportClient	The viewport client the user is dragging in
 	 * @return true if anything was snapped
 	 */
-	bool SnapDraggedActorsToNearestVertex( FVector& DragDelta, FLevelEditorViewportClient* ViewportClient );
+	bool SnapDraggedActorsToNearestVertex( FVector& DragDelta, TNotNull<FEditorViewportClient*> ViewportClient );
 
 	/**
 	 * Snaps a delta drag movement to the nearest vertex
@@ -102,7 +103,7 @@ public:
 	 * @param ViewportClient	The viewport client being dragged in.
 	 * @return true if anything was snapped
 	 */
-	bool SnapDragLocationToNearestVertex( const FVector& BaseLocation, FVector& DragDelta, FLevelEditorViewportClient* ViewportClient, bool bIsPivot = false );
+	bool SnapDragLocationToNearestVertex( const FVector& BaseLocation, FVector& DragDelta, TNotNull<FEditorViewportClient*> ViewportClient, bool bIsPivot = false );
 
 	/**
 	 * Snaps a location to the nearest vertex
@@ -113,7 +114,7 @@ public:
 	 * @param OutVertexNormal	The normal at the closest vertex
 	 * @return true if anything was snapped
 	 */
-	bool SnapLocationToNearestVertex( FVector& Location, const FVector2D& MouseLocation, FLevelEditorViewportClient* ViewportClient, FVector& OutVertexNormal, bool bDrawVertHelpers );
+	bool SnapLocationToNearestVertex( FVector& Location, const FVector2D& MouseLocation, TNotNull<FEditorViewportClient*> ViewportClient, FVector& OutVertexNormal, bool bDrawVertHelpers, const TSet<TWeakObjectPtr<AActor>>& ActorsToIgnore );
 
 	bool SnapToBSPVertex( FVector& Location, FVector GridBase, FRotator& Rotation );
 
@@ -391,14 +392,12 @@ void FEditorViewportSnapping::DrawSnappingHelpers(const FSceneView* View, FPrimi
 	VertexSnappingImpl.DrawSnappingHelpers( View, PDI );
 }
 
-
-
-bool FEditorViewportSnapping::SnapLocationToNearestVertex( FVector& Location, const FVector2D& MouseLocation, FLevelEditorViewportClient* ViewportClient, FVector& OutVertexNormal, bool bDrawVertHelpers )
+bool FEditorViewportSnapping::SnapLocationToNearestVertex( FVector& Location, const FVector2D& MouseLocation, TNotNull<FEditorViewportClient*> ViewportClient, FVector& OutVertexNormal, bool bDrawVertHelpers, const TSet<TWeakObjectPtr<AActor>>& ActorsToIgnore )
 {
 	bool bSnapped = false;
 	if( IsSnapToVertexEnabled() )
 	{
-		bSnapped = VertexSnappingImpl.SnapLocationToNearestVertex( Location, MouseLocation, ViewportClient, OutVertexNormal, bDrawVertHelpers );
+		bSnapped = VertexSnappingImpl.SnapLocationToNearestVertex( Location, MouseLocation, ViewportClient, OutVertexNormal, bDrawVertHelpers, ActorsToIgnore );
 	}
 	else
 	{
@@ -408,8 +407,7 @@ bool FEditorViewportSnapping::SnapLocationToNearestVertex( FVector& Location, co
 	return bSnapped;
 }
 
-
-bool FEditorViewportSnapping::SnapDraggedActorsToNearestVertex( FVector& DragDelta, FLevelEditorViewportClient* ViewportClient )
+bool FEditorViewportSnapping::SnapDraggedActorsToNearestVertex( FVector& DragDelta, TNotNull<FEditorViewportClient*> ViewportClient )
 {
 	bool bSnapped = false;
 	if( IsSnapToVertexEnabled() && !DragDelta.IsNearlyZero() )
@@ -419,7 +417,7 @@ bool FEditorViewportSnapping::SnapDraggedActorsToNearestVertex( FVector& DragDel
 	return bSnapped;
 }
 
-bool FEditorViewportSnapping::SnapDragLocationToNearestVertex( const FVector& BaseLocation, FVector& DragDelta, FLevelEditorViewportClient* ViewportClient, bool bIsPivot )
+bool FEditorViewportSnapping::SnapDragLocationToNearestVertex( const FVector& BaseLocation, FVector& DragDelta, TNotNull<FEditorViewportClient*> ViewportClient, bool bIsPivot )
 {
 	bool bSnapped = false;
 	if( IsSnapToVertexEnabled(bIsPivot) && !DragDelta.IsNearlyZero() )
@@ -485,19 +483,36 @@ TOptional<FVector> FSnappingUtils::GetDeltaToNearestActor(FVector& DragDelta, co
 	return EditorViewportSnapper.IsValid() ? EditorViewportSnapper->GetDeltaToNearestActor(DragDelta, InWorld) : TOptional<FVector>();
 }
 
-bool FSnappingUtils::SnapDraggedActorsToNearestVertex(FVector& DragDelta, FLevelEditorViewportClient* ViewportClient)
+bool FSnappingUtils::SnapDraggedActorsToNearestVertex(FVector& DragDelta, FEditorViewportClient* ViewportClient)
 {
 	return EditorViewportSnapper.IsValid() ? EditorViewportSnapper->SnapDraggedActorsToNearestVertex(DragDelta, ViewportClient) : false;
 }
 
-bool FSnappingUtils::SnapDragLocationToNearestVertex(const FVector& BaseLocation, FVector& DragDelta, FLevelEditorViewportClient* ViewportClient, bool bIsPivot)
+bool FSnappingUtils::SnapDragLocationToNearestVertex(const FVector& BaseLocation, FVector& DragDelta, FEditorViewportClient* ViewportClient, bool bIsPivot)
 {
 	return EditorViewportSnapper.IsValid() ? EditorViewportSnapper->SnapDragLocationToNearestVertex(BaseLocation, DragDelta, ViewportClient, bIsPivot) : false;
 }
 
+bool FSnappingUtils::SnapLocationToNearestVertex(FVector& Location, const FVector2D& MouseLocation, FEditorViewportClient* ViewportClient, FVector& OutVertexNormal, bool bDrawVertHelpers, const TSet<TWeakObjectPtr<AActor>>& ActorsToIgnore )
+{
+	return EditorViewportSnapper.IsValid() ? EditorViewportSnapper->SnapLocationToNearestVertex(Location, MouseLocation, ViewportClient, OutVertexNormal, bDrawVertHelpers, ActorsToIgnore) : false;
+}
+
 bool FSnappingUtils::SnapLocationToNearestVertex(FVector& Location, const FVector2D& MouseLocation, FLevelEditorViewportClient* ViewportClient, FVector& OutVertexNormal, bool bDrawVertHelpers )
 {
-	return EditorViewportSnapper.IsValid() ? EditorViewportSnapper->SnapLocationToNearestVertex(Location, MouseLocation, ViewportClient, OutVertexNormal, bDrawVertHelpers ) : false;
+	TSet<TWeakObjectPtr<AActor>> ActorsToIgnore;
+	if (const TSharedPtr<FTypedElementList>& Elements = ViewportClient->GetDropPreviewElements())
+	{
+		Elements->ForEachElementHandle([&ActorsToIgnore](const FTypedElementHandle& Element)
+			{
+				if (AActor* Actor = ActorElementDataUtil::GetActorFromHandle(Element, true))
+				{
+					ActorsToIgnore.Add(Actor);
+				}
+				return true;
+			});
+	}
+	return SnapLocationToNearestVertex(Location, MouseLocation, ViewportClient, OutVertexNormal, bDrawVertHelpers, MoveTemp(ActorsToIgnore));
 }
 
 void FSnappingUtils::SnapScale(FVector& Point, const FVector& GridBase)
diff --git a/Engine/Source/Editor/UnrealEd/Private/VertexSnapping.cpp b/Engine/Source/Editor/UnrealEd/Private/VertexSnapping.cpp
--- a/Engine/Source/Editor/UnrealEd/Private/VertexSnapping.cpp
+++ b/Engine/Source/Editor/UnrealEd/Private/VertexSnapping.cpp
@@ -326,7 +326,7 @@ void FVertexSnappingImpl::DrawSnappingHelpers(const FSceneView* View,FPrimitiveD
 
 struct FVertexSnappingArgs
 {
-	FVertexSnappingArgs( const FPlane& InActorPlane, const FVector& InCurrentLocation, FLevelEditorViewportClient* InViewportClient, const FSceneView* InView, const FVector2D& InMousePosition, EAxisList::Type InCurrentAxis, bool bInDrawVertexHelpers )
+	FVertexSnappingArgs( const FPlane& InActorPlane, const FVector& InCurrentLocation, TNotNull<FEditorViewportClient*> InViewportClient, const FSceneView* InView, const FVector2D& InMousePosition, EAxisList::Type InCurrentAxis, bool bInDrawVertexHelpers )
 		: ActorPlane( InActorPlane )
 		, CurrentLocation( InCurrentLocation )
 		, ViewportClient( InViewportClient )
@@ -341,7 +341,7 @@ struct FVertexSnappingArgs
 	/** Current pre-snap location that is being snapped */
 	const FVector& CurrentLocation;
 	/** Viewport client being used */
-	FLevelEditorViewportClient* ViewportClient;
+	TNotNull<FEditorViewportClient*> ViewportClient;
 	/** 2D position of the mouse */
 	const FVector2D MousePosition;
 	/** Our current view */
@@ -538,7 +538,7 @@ FSnappingVertex FVertexSnappingImpl::GetClosestVertex( const TArray<FSnapActor>&
 	return ClosestLocation;
 }
 
-void FVertexSnappingImpl::GetPossibleSnapActors( const FBox& AllowedBox, FIntPoint MouseLocation, FLevelEditorViewportClient* ViewportClient, const FSceneView* View, EAxisList::Type CurrentAxis, TSet< TWeakObjectPtr<AActor> >& ActorsToIgnore, TArray<FSnapActor>& OutActorsInBox )
+void FVertexSnappingImpl::GetPossibleSnapActors( const FBox& AllowedBox, FIntPoint MouseLocation, TNotNull<FEditorViewportClient*> ViewportClient, const FSceneView* View, EAxisList::Type CurrentAxis, const TSet< TWeakObjectPtr<AActor> >& ActorsToIgnore, TArray<FSnapActor>& OutActorsInBox )
 {
 	if( CurrentAxis == EAxisList::Screen && !ViewportClient->IsOrtho() )
 	{
@@ -548,7 +548,6 @@ void FVertexSnappingImpl::GetPossibleSnapActors( const FBox& AllowedBox, FIntPoi
 			AActor* HitProxyActor = static_cast<HActor*>(HitProxy)->Actor ;
 			if( HitProxyActor && !ActorsToIgnore.Contains(HitProxyActor)  )
 			{
-				ActorsToIgnore.Add( HitProxyActor );
 				OutActorsInBox.Add( FSnapActor( HitProxyActor, HitProxyActor->GetComponentsBoundingBox(true) ) );
 			}
 		}
@@ -579,8 +578,8 @@ void FVertexSnappingImpl::GetActorsInsideBox( const FBox& Box, UWorld* World, TA
 	}
 }
 
-
-bool FVertexSnappingImpl::SnapLocationToNearestVertex( FVector& Location, const FVector2D& MouseLocation, FLevelEditorViewportClient* ViewportClient, FVector& OutVertexNormal, bool bDrawVertexHelpers )
+bool FVertexSnappingImpl::SnapLocationToNearestVertex(FVector& Location, const FVector2D& MouseLocation, TNotNull<FEditorViewportClient*> ViewportClient,
+	FVector& OutVertexNormal, bool bDrawVertexHelpers, const TSet<TWeakObjectPtr<AActor>>& ActorsToIgnore)
 {
 	bool bSnapped = false;
 
@@ -596,10 +595,6 @@ bool FVertexSnappingImpl::SnapLocationToNearestVertex( FVector& Location, const
 	FSceneView* View = ViewportClient->CalcSceneView( &ViewFamily );
 
 	TArray<FSnapActor> ActorsInBox;
-	TSet<TWeakObjectPtr<AActor> > ActorsToIgnore;
-	// Ignore actors currently being moved
-	ActorsToIgnore.Append( ViewportClient->GetDropPreviewActors() );
-
 	GetPossibleSnapActors( AllowedSnappingBox, MouseLocation.IntPoint(), ViewportClient, View, EAxisList::Screen, ActorsToIgnore, ActorsInBox );
 
 	FViewportCursorLocation Cursor(View, ViewportClient, static_cast<int32>(MouseLocation.X), static_cast<int32>(MouseLocation.Y));
@@ -658,7 +653,7 @@ static void GetActorsToIgnore( AActor* Actor, TSet< TWeakObjectPtr<AActor> >& Ac
 	}
 }
 
-bool FVertexSnappingImpl::SnapDraggedActorsToNearestVertex( FVector& DragDelta, FLevelEditorViewportClient* ViewportClient )
+bool FVertexSnappingImpl::SnapDraggedActorsToNearestVertex( FVector& DragDelta, TNotNull<FEditorViewportClient*> ViewportClient )
 {
 	int32 MouseX = ViewportClient->Viewport->GetMouseX();
 	int32 MouseY = ViewportClient->Viewport->GetMouseY();
@@ -669,9 +664,6 @@ bool FVertexSnappingImpl::SnapDraggedActorsToNearestVertex( FVector& DragDelta,
 	bool bSnapped = false;
 	if( !DragDelta.IsNearlyZero() )
 	{
-		// Are there selected actors?
-		USelection* Selection = GEditor->GetSelectedActors();
-
 		FVector Direction = DragDelta.GetSafeNormal();
 
 		FSceneViewFamilyContext ViewFamily(FSceneViewFamily::ConstructionValues(
@@ -740,7 +732,7 @@ void FVertexSnappingImpl::SnapDragDelta( FVertexSnappingArgs& InArgs, const FVec
 	const FVector2D& MousePosition = InArgs.MousePosition;
 	const EAxisList::Type CurrentAxis = InArgs.CurrentAxis;
 	const FPlane ActorPlane = InArgs.ActorPlane;
-	FLevelEditorViewportClient* ViewportClient = InArgs.ViewportClient;
+	TNotNull<FEditorViewportClient*> ViewportClient = InArgs.ViewportClient;
 
 	TArray<FSnapActor> PossibleSnapPointActors;
 	GetPossibleSnapActors( AllowedSnappingBox, MousePosition.IntPoint(), ViewportClient, View, CurrentAxis, ActorsToIgnore, PossibleSnapPointActors );
@@ -813,7 +805,7 @@ void FVertexSnappingImpl::SnapDragDelta( FVertexSnappingArgs& InArgs, const FVec
 	}
 }
 
-bool FVertexSnappingImpl::SnapDragLocationToNearestVertex( const FVector& BaseLocation, FVector& DragDelta, FLevelEditorViewportClient* ViewportClient )
+bool FVertexSnappingImpl::SnapDragLocationToNearestVertex( const FVector& BaseLocation, FVector& DragDelta, TNotNull<FEditorViewportClient*> ViewportClient )
 {
 	int32 MouseX = ViewportClient->Viewport->GetMouseX();
 	int32 MouseY = ViewportClient->Viewport->GetMouseY();
diff --git a/Engine/Source/Editor/UnrealEd/Private/VertexSnapping.h b/Engine/Source/Editor/UnrealEd/Private/VertexSnapping.h
--- a/Engine/Source/Editor/UnrealEd/Private/VertexSnapping.h
+++ b/Engine/Source/Editor/UnrealEd/Private/VertexSnapping.h
@@ -6,6 +6,7 @@
 #include "UObject/WeakObjectPtr.h"
 
 class AActor;
+class FEditorViewportClient;
 class FLevelEditorViewportClient;
 class FPrimitiveDrawInterface;
 class FSceneView;
@@ -100,7 +101,7 @@ public:
 	 * @param ViewportClient	The viewport client the user is dragging in
 	 * @return true if anything was snapped
 	 */
-	bool SnapDraggedActorsToNearestVertex( FVector& DragDelta, FLevelEditorViewportClient* ViewportClient );
+	bool SnapDraggedActorsToNearestVertex( FVector& DragDelta, TNotNull<FEditorViewportClient*> ViewportClient );
 
 	/**
 	 * Snaps a delta drag movement to the nearest vertex
@@ -110,7 +111,7 @@ public:
 	 * @param ViewportClient	The viewport client being dragged in.
 	 * @return true if anything was snapped
 	 */
-	bool SnapDragLocationToNearestVertex( const FVector& BaseLocation, FVector& DragDelta, FLevelEditorViewportClient* ViewportClient );
+	bool SnapDragLocationToNearestVertex( const FVector& BaseLocation, FVector& DragDelta, TNotNull<FEditorViewportClient*> ViewportClient );
 
 	/**
 	 * Snaps a location to the nearest vertex
@@ -121,7 +122,7 @@ public:
 	 * @param OutVertexNormal	The normal at the closest vertex
 	 * @return true if anything was snapped
 	 */
-	bool SnapLocationToNearestVertex( FVector& Location, const FVector2D& MouseLocation, FLevelEditorViewportClient* ViewportClient, FVector& OutVertexNormal, bool bDrawVertexHelpers );
+	bool SnapLocationToNearestVertex( FVector& Location, const FVector2D& MouseLocation, TNotNull<FEditorViewportClient*> ViewportClient, FVector& OutVertexNormal, bool bDrawVertexHelpers, const TSet<TWeakObjectPtr<AActor>>& ActorsToIgnore );
 private:
 	/**
 	 * Gets the closest vertex on a primitive component
@@ -154,7 +155,7 @@ private:
 	 * @param ActorsToIgnore	Set of actors to ignore when checking for vertices
 	 * @param OutSnapActors		The populated list of possible snap actors
 	 */
-	void GetPossibleSnapActors( const FBox& AllowedBox, FIntPoint MouseLocation, FLevelEditorViewportClient* ViewportClient, const FSceneView* View, EAxisList::Type CurrentAxis, TSet< TWeakObjectPtr<AActor> >& ActorsToIgnore, TArray<FSnapActor>& OutSnapActors );
+	void GetPossibleSnapActors( const FBox& AllowedBox, FIntPoint MouseLocation, TNotNull<FEditorViewportClient*> ViewportClient, const FSceneView* View, EAxisList::Type CurrentAxis, const TSet< TWeakObjectPtr<AActor> >& ActorsToIgnore, TArray<FSnapActor>& OutSnapActors );
 
 	/**
 	 * Gets all actors inside a bounding box
diff --git a/Engine/Source/Editor/UnrealEd/Public/LevelEditorViewport.h b/Engine/Source/Editor/UnrealEd/Public/LevelEditorViewport.h
--- a/Engine/Source/Editor/UnrealEd/Public/LevelEditorViewport.h
+++ b/Engine/Source/Editor/UnrealEd/Public/LevelEditorViewport.h
@@ -177,6 +177,12 @@ public:
 		return DropPreviewActors;
 	}
 
+	/** Returns the current drop preview elements, or a null pointer if there are none. */
+	const TSharedPtr<FTypedElementList>& GetDropPreviewElements() const
+	{
+		return DropPreviewElements;
+	}
+
 	FVector2D GetDropPreviewLocation() const { return FVector2D( DropPreviewMouseX, DropPreviewMouseY ); }
 
 	/**
diff --git a/Engine/Source/Editor/UnrealEd/Public/SnappingUtils.h b/Engine/Source/Editor/UnrealEd/Public/SnappingUtils.h
--- a/Engine/Source/Editor/UnrealEd/Public/SnappingUtils.h
+++ b/Engine/Source/Editor/UnrealEd/Public/SnappingUtils.h
@@ -5,6 +5,7 @@
 #include "CoreMinimal.h"
 
 class FEditorViewportSnapping;
+class FEditorViewportClient;
 class FLevelEditorViewportClient;
 class FPrimitiveDrawInterface;
 class FSceneView;
@@ -65,7 +66,7 @@ public:
 	 * @param ViewportClient	The viewport client the user is dragging in
 	 * @return true if anything was snapped
 	 */
-	static UNREALED_API bool SnapDraggedActorsToNearestVertex( FVector& DragDelta, FLevelEditorViewportClient* ViewportClient );
+	static UNREALED_API bool SnapDraggedActorsToNearestVertex( FVector& DragDelta, FEditorViewportClient* ViewportClient );
 
 	/**
 	 * Snaps a delta drag movement to the nearest vertex
@@ -75,7 +76,7 @@ public:
 	 * @param ViewportClient	The viewport client being dragged in.
 	 * @return true if anything was snapped
 	 */
-	static UNREALED_API bool SnapDragLocationToNearestVertex( const FVector& BaseLocation, FVector& DragDelta, FLevelEditorViewportClient* ViewportClient, bool bIsPivot = false );
+	static UNREALED_API bool SnapDragLocationToNearestVertex( const FVector& BaseLocation, FVector& DragDelta, FEditorViewportClient* ViewportClient, bool bIsPivot = false );
 
 	/**
 	 * Snaps a location to the nearest vertex
@@ -86,6 +87,8 @@ public:
 	 * @param OutVertexNormal	The normal at the closest vertex
 	 * @return true if anything was snapped
 	 */
+	static UNREALED_API bool SnapLocationToNearestVertex( FVector& Location, const FVector2D& MouseLocation, FEditorViewportClient* ViewportClient,
+		FVector& OutVertexNormal, bool bDrawVertHelpers, const TSet<TWeakObjectPtr<class AActor>>& ActorsToIgnore );
 	static UNREALED_API bool SnapLocationToNearestVertex( FVector& Location, const FVector2D& MouseLocation, FLevelEditorViewportClient* ViewportClient, FVector& OutVertexNormal, bool bDrawVertHelpers );
 
 	/**
