Fix UE::Positioning::TraceWorldForPosition incorrectly filtering hits by evaluating against the hit object root component instead of the hit component

diff --git a/Engine/Source/Editor/UnrealEd/Private/Editor/ObjectPositioning.cpp b/Engine/Source/Editor/UnrealEd/Private/Editor/ObjectPositioning.cpp
--- a/Engine/Source/Editor/UnrealEd/Private/Editor/ObjectPositioning.cpp
+++ b/Engine/Source/Editor/UnrealEd/Private/Editor/ObjectPositioning.cpp
@@ -43,11 +43,11 @@ namespace ObjectPositioningLocals
 			const FActorInstanceHandle& HitObjHandle = Hit.HitObjectHandle;
 
 			// Try and find a primitive component for the hit
-			const UPrimitiveComponent* PrimitiveComponent = Cast<UPrimitiveComponent>(HitObjHandle.GetRootComponent());
+			const UPrimitiveComponent* PrimitiveComponent = Hit.GetComponent();
 
 			if (!PrimitiveComponent)
 			{
-				PrimitiveComponent = Hit.Component.Get();
+				PrimitiveComponent = Cast<UPrimitiveComponent>(HitObjHandle.GetRootComponent());		
 			}
 			if (PrimitiveComponent && PrimitiveComponent->IsA(ULandscapeHeightfieldCollisionComponent::StaticClass()))
 			{
