[ITF Gizmos] Fix issue where enabling the EditorTRSGizmo module but not enabling the preference could result in loading in some ITF interactions when none should be present.

#rnx
#code-analyzed

diff --git a/Engine/Source/Editor/UnrealEd/Private/Tools/EdModeInteractiveToolsContext.cpp b/Engine/Source/Editor/UnrealEd/Private/Tools/EdModeInteractiveToolsContext.cpp
--- a/Engine/Source/Editor/UnrealEd/Private/Tools/EdModeInteractiveToolsContext.cpp
+++ b/Engine/Source/Editor/UnrealEd/Private/Tools/EdModeInteractiveToolsContext.cpp
@@ -1451,7 +1451,7 @@ UViewportInteractionsBehaviorSource* UModeManagerInteractiveToolsContext::GetVie
 
 void UModeManagerInteractiveToolsContext::RegisterViewportInteractions()
 {
-	if (!ViewportInteractionsBehaviorSource)
+	if (!ViewportInteractionsBehaviorSource && UsesViewportInteractions())
 	{
 		SetBypassAltModifier(false);
 
