Added functions to the Collection Manager Scripting Subsystem for getting and setting a collections color.
Don't checkout/save .collection files if setting their color to their existing color.

diff --git a/Engine/Source/Developer/CollectionManager/Private/CollectionContainer.cpp b/Engine/Source/Developer/CollectionManager/Private/CollectionContainer.cpp
--- a/Engine/Source/Developer/CollectionManager/Private/CollectionContainer.cpp
+++ b/Engine/Source/Developer/CollectionManager/Private/CollectionContainer.cpp
@@ -1993,6 +1993,11 @@ bool FCollectionContainer::SetCollectionColor(FName CollectionName, ECollectionS
 			return false;
 		}
 
+		if ((*CollectionRefPtr)->GetCollectionColor() == NewColor)
+		{
+			return true;
+		}
+
 		(*CollectionRefPtr)->SetCollectionColor(NewColor);
 
 		constexpr bool bForceCommitToRevisionControl = false;
diff --git a/Engine/Source/Editor/UnrealEd/Private/Subsystems/CollectionManagerScriptingSubsystem.cpp b/Engine/Source/Editor/UnrealEd/Private/Subsystems/CollectionManagerScriptingSubsystem.cpp
--- a/Engine/Source/Editor/UnrealEd/Private/Subsystems/CollectionManagerScriptingSubsystem.cpp
+++ b/Engine/Source/Editor/UnrealEd/Private/Subsystems/CollectionManagerScriptingSubsystem.cpp
@@ -481,6 +481,77 @@ bool UCollectionManagerScriptingSubsystem::GetAssetsInCollection(const FCollecti
 	return true;
 }
 
+bool UCollectionManagerScriptingSubsystem::GetCollectionColor(const FCollectionScriptingRef& Collection, FLinearColor& OutColor)
+{
+	OutColor = FLinearColor::White;
+
+	const TSharedPtr<ICollectionContainer> CollectionContainer = FCollectionManagerModule::GetModule().Get().FindCollectionContainer(Collection.Container);
+	if (!CollectionContainer.IsValid())
+	{
+		CollectionManagerScriptingSubsystemUtil::LogLastCollectionManagerWarning(FText::Format(CollectionManagerScriptingSubsystemUtil::FConstants::Get().ContainerNotFound, { FText::FromName(Collection.Container) }));
+		return false;
+	}
+
+	if (!CollectionContainer->CollectionExists(Collection.Name, CollectionManagerScriptingSubsystemUtil::ECollectionScriptingShareType_To_ECollectionShareType(Collection.ShareType)))
+	{
+		CollectionManagerScriptingSubsystemUtil::LogLastCollectionManagerWarning(FText::Format(CollectionManagerScriptingSubsystemUtil::FConstants::Get().CollectionNotFound,
+			{ FText::FromName(Collection.Name), FText::FromString(CollectionManagerScriptingSubsystemUtil::ECollectionScriptingShareType_To_String(Collection.ShareType)), FText::FromName(Collection.Container) }));
+		return false;
+	}
+
+	TOptional<FLinearColor> Color;
+	FText Error;
+	if (!CollectionContainer->GetCollectionColor(Collection.Name, CollectionManagerScriptingSubsystemUtil::ECollectionScriptingShareType_To_ECollectionShareType(Collection.ShareType), Color, &Error))
+	{
+		CollectionManagerScriptingSubsystemUtil::LogLastCollectionManagerWarning(Error);
+		return false;
+	}
+
+	if (!Color.IsSet())
+	{
+		return false;
+	}
+
+	OutColor = Color.GetValue();
+	return true;
+}
+
+bool UCollectionManagerScriptingSubsystem::SetCollectionColor(const FCollectionScriptingRef& Collection, const FLinearColor& Color)
+{
+	return SetCollectionColor(Collection, TOptional<FLinearColor>(Color));
+}
+
+bool UCollectionManagerScriptingSubsystem::ClearCollectionColor(const FCollectionScriptingRef& Collection)
+{
+	return SetCollectionColor(Collection, TOptional<FLinearColor>());
+}
+
+bool UCollectionManagerScriptingSubsystem::SetCollectionColor(const FCollectionScriptingRef& Collection, const TOptional<FLinearColor>& Color)
+{
+	const TSharedPtr<ICollectionContainer> CollectionContainer = FCollectionManagerModule::GetModule().Get().FindCollectionContainer(Collection.Container);
+	if (!CollectionContainer.IsValid())
+	{
+		CollectionManagerScriptingSubsystemUtil::LogLastCollectionManagerWarning(FText::Format(CollectionManagerScriptingSubsystemUtil::FConstants::Get().ContainerNotFound, { FText::FromName(Collection.Container) }));
+		return false;
+	}
+
+	if (!CollectionContainer->CollectionExists(Collection.Name, CollectionManagerScriptingSubsystemUtil::ECollectionScriptingShareType_To_ECollectionShareType(Collection.ShareType)))
+	{
+		CollectionManagerScriptingSubsystemUtil::LogLastCollectionManagerWarning(FText::Format(CollectionManagerScriptingSubsystemUtil::FConstants::Get().CollectionNotFound,
+			{ FText::FromName(Collection.Name), FText::FromString(CollectionManagerScriptingSubsystemUtil::ECollectionScriptingShareType_To_String(Collection.ShareType)), FText::FromName(Collection.Container) }));
+		return false;
+	}
+
+	FText Error;
+	if (!CollectionContainer->SetCollectionColor(Collection.Name, CollectionManagerScriptingSubsystemUtil::ECollectionScriptingShareType_To_ECollectionShareType(Collection.ShareType), Color, &Error))
+	{
+		CollectionManagerScriptingSubsystemUtil::LogLastCollectionManagerWarning(Error);
+		return false;
+	}
+
+	return true;
+}
+
 bool UCollectionManagerScriptingSubsystem::GetCollectionsContainingAsset(const FCollectionScriptingContainerSource Container, const FSoftObjectPath& AssetPath, TArray<FCollectionScriptingRef>& OutCollections)
 {
 	const TSharedPtr<ICollectionContainer> CollectionContainer = FCollectionManagerModule::GetModule().Get().FindCollectionContainer(Container.Name);
diff --git a/Engine/Source/Editor/UnrealEd/Public/Subsystems/CollectionManagerScriptingSubsystem.h b/Engine/Source/Editor/UnrealEd/Public/Subsystems/CollectionManagerScriptingSubsystem.h
--- a/Engine/Source/Editor/UnrealEd/Public/Subsystems/CollectionManagerScriptingSubsystem.h
+++ b/Engine/Source/Editor/UnrealEd/Public/Subsystems/CollectionManagerScriptingSubsystem.h
@@ -303,6 +303,38 @@ public:
 	UFUNCTION(BlueprintCallable, Category="CollectionManager")
 	UE_API bool GetAssetsInCollection(const FCollectionScriptingRef& Collection, TArray<FAssetData>& OutAssets);
 
+	/**
+	 * Get the color of the given collection.
+	 *
+	 * @param Collection The collection to get the color of.
+	 * @param OutColor The collection's color.
+	 *
+	 * @return True if the collection exists and has a color set, false otherwise (see the output log for details on error).
+	 */
+	UFUNCTION(BlueprintCallable, Category = "CollectionManager")
+	UE_API bool GetCollectionColor(const FCollectionScriptingRef& Collection, FLinearColor& OutColor);
+
+	/**
+	 * Set the color of the given collection.
+	 *
+	 * @param Collection The collection to set the color of.
+	 * @param Color The new collection color.
+	 *
+	 * @return True if the collection exists and the color was set, false otherwise (see the output log for details on error).
+	 */
+	UFUNCTION(BlueprintCallable, Category = "CollectionManager")
+	UE_API bool SetCollectionColor(const FCollectionScriptingRef& Collection, const FLinearColor& Color);
+
+	/**
+	 * Clear the color of the given collection.
+	 *
+	 * @param Collection The collection to clear the color of.
+	 *
+	 * @return True if the collection exists and the color was cleared, false otherwise (see the output log for details on error).
+	 */
+	UFUNCTION(BlueprintCallable, Category = "CollectionManager")
+	UE_API bool ClearCollectionColor(const FCollectionScriptingRef& Collection);
+
 	/**
 	 * Get the collections in the specified container that contain the given asset.
 	 *
@@ -346,6 +378,9 @@ public:
 	 */
 	UFUNCTION(BlueprintCallable, Category="CollectionManager")
 	UE_API FCollectionScriptingContainerSource GetBaseGameCollectionContainer() const;
+
+private:
+	bool SetCollectionColor(const FCollectionScriptingRef& Collection, const TOptional<FLinearColor>& Color);
 };
 
 #undef UE_API
