Restore animation position and play state after undo/redo actions

diff --git a/Engine/Source/Editor/Persona/Private/SAnimMontagePanel.cpp b/Engine/Source/Editor/Persona/Private/SAnimMontagePanel.cpp
--- a/Engine/Source/Editor/Persona/Private/SAnimMontagePanel.cpp
+++ b/Engine/Source/Editor/Persona/Private/SAnimMontagePanel.cpp
@@ -620,7 +620,21 @@ void SAnimMontagePanel::PostRedo( bool bSuccess )
 
 void SAnimMontagePanel::PostRedoUndo()
 {
+	// When we rebuild the montage panel as part of a undo/redo it also restarts the montage from the start
+	// To facilitate a smoother workflow we try to restore the montage position afterwards
+	UAnimPreviewInstance* Preview = nullptr;
+	if (UDebugSkelMeshComponent* MeshComponent = WeakModel.Pin()->GetPreviewScene()->GetPreviewMeshComponent())
+	{
+		Preview = MeshComponent->PreviewInstance;
+	}
+	const float PrevTime = Preview ? Preview->GetCurrentTime() : 0.0f;
+
 	RebuildMontagePanel(); //Rebuild here, undoing adds can cause slate to crash later on if we don't (using dummy args since they aren't used by the method
+
+	if (Preview)
+	{
+		Preview->SetPosition(PrevTime, false);
+	}
 }
 
 void SAnimMontagePanel::Update()
diff --git a/Engine/Source/Editor/UnrealEd/Private/Animation/DebugSkelMeshComponent.cpp b/Engine/Source/Editor/UnrealEd/Private/Animation/DebugSkelMeshComponent.cpp
--- a/Engine/Source/Editor/UnrealEd/Private/Animation/DebugSkelMeshComponent.cpp
+++ b/Engine/Source/Editor/UnrealEd/Private/Animation/DebugSkelMeshComponent.cpp
@@ -693,6 +693,12 @@ void UDebugSkelMeshComponent::EnablePreview(bool bEnable, UAnimationAsset* Previ
 				SavedAnimScriptInstance = AnimScriptInstance;
 			}
 
+			// When this code path gets hit during undo/redo calls and we are previewing the same asset as before 
+			// we want to restore playback position and play state after SetAnimationAsset to avoid transacting actions reseting the state of anim editors
+			const bool bRestoreState = GIsTransacting && IsPreviewOn() && PreviewInstance->GetCurrentAsset() == PreviewAsset;
+			const bool bWasPlaying = PreviewInstance->IsPlaying();
+			const float PrevTime = PreviewInstance->GetCurrentTime();
+			
 			AnimScriptInstance = PreviewInstance;
 		    // restore previous state
 		    bDisableClothSimulation = bPrevDisableClothSimulation;
@@ -700,6 +706,12 @@ void UDebugSkelMeshComponent::EnablePreview(bool bEnable, UAnimationAsset* Previ
 			const bool bWasLooping = PreviewInstance->IsLooping();
 			PreviewInstance->SetAnimationAsset(PreviewAsset, bWasLooping);
 			
+			if (bRestoreState)
+			{
+				PreviewInstance->SetPosition(PrevTime, false);
+				PreviewInstance->SetPlaying(bWasPlaying);
+			}
+
 			// Reset to previous animation asset's root motion playback time to prevent this from influencing the new animation asset previewing during root motion consumption.
 			ConsumeRootMotionPreviousPlaybackTime = 0.0f;
 			
