Bump IncrementalCookVersion to work around Collision of two shader code chunks in horde. Collision is caused by stale files in shader library when assets are moved between chunk. The proper fix is to add incremental pruning of shader libraries in each incremental cook and is still pending.
#rnx

diff --git a/Engine/Source/Editor/UnrealEd/Private/Cooker/CookIncrementalVersion.cpp b/Engine/Source/Editor/UnrealEd/Private/Cooker/CookIncrementalVersion.cpp
--- a/Engine/Source/Editor/UnrealEd/Private/Cooker/CookIncrementalVersion.cpp
+++ b/Engine/Source/Editor/UnrealEd/Private/Cooker/CookIncrementalVersion.cpp
@@ -8,6 +8,6 @@ namespace UE::Cook
 {
 
 // Change CookIncrementalVersion to a new guid when all packages in an incremental cook need to be invalidated
-FGuid CookIncrementalVersion( 0xBA201607, 0xD65B4B86, 0xAF26BDB6, 0xB8109F27 );
+FGuid CookIncrementalVersion(0x7DCBB8C7, 0x72364DB2, 0xA97E2A13, 0x0CFA598D);
 
 }
