Cook Package Splitter: add a virtual function to ICookPackageSplitter to allow specifying if the splitter implementation requires or not empty packages before populate.
If true, this splitter forces the Generator package to be empty before calling the populate function.
Returning false will allow the splitter to precreate packages instead of relying on the cooker to create them.

diff --git a/Engine/Source/Editor/UnrealEd/Private/Cooker/CookGenerationHelper.cpp b/Engine/Source/Editor/UnrealEd/Private/Cooker/CookGenerationHelper.cpp
--- a/Engine/Source/Editor/UnrealEd/Private/Cooker/CookGenerationHelper.cpp
+++ b/Engine/Source/Editor/UnrealEd/Private/Cooker/CookGenerationHelper.cpp
@@ -1622,7 +1622,7 @@ UPackage* FGenerationHelper::TryCreateGeneratedPackage(FCookGenerationInfo& Gene
 		// If bResetToEmpty then we are starting the populate of the generated package and we need to remove all objects
 		// from the package. Generated packages are created empty and it is the job of the CookPackageSplitter to populate
 		// them during save.
-		if (bResetToEmpty)
+		if (bResetToEmpty && GetCookPackageSplitterInstance()->RequiresEmptyPackageBeforePopulate(GeneratedPackage))
 		{
 			TArray<UObject*> ExistingObjects;
 			GetObjectsWithPackage(GeneratedPackage, ExistingObjects, false /* bIncludeNestedObjects */);
diff --git a/Engine/Source/Editor/UnrealEd/Public/CookPackageSplitter.h b/Engine/Source/Editor/UnrealEd/Public/CookPackageSplitter.h
--- a/Engine/Source/Editor/UnrealEd/Public/CookPackageSplitter.h
+++ b/Engine/Source/Editor/UnrealEd/Public/CookPackageSplitter.h
@@ -90,6 +90,15 @@ public:
 		return false;
 	}
 
+	/**
+	 * If true, this splitter forces the Generator package to be empty before calling the populate function.
+	 * Returning false will allow the splitter to precreate packages instead of relying on the cooker to create them.
+	 */
+	virtual bool RequiresEmptyPackageBeforePopulate(const UPackage* Package)
+	{
+		return true;
+	}
+
 	/**
 	 * Return value for the DoesGeneratedRequireGenerator function. All levels behave correctly, but provide
 	 * different tradeoffs of guarantees to the splitter versus performance.
