Change the location of incremental cook diagnostic logs when we are running on build machines to go to the AutomationTool log directory.

diff --git a/Engine/Source/Editor/UnrealEd/Private/CookOnTheFlyServer.cpp b/Engine/Source/Editor/UnrealEd/Private/CookOnTheFlyServer.cpp
--- a/Engine/Source/Editor/UnrealEd/Private/CookOnTheFlyServer.cpp
+++ b/Engine/Source/Editor/UnrealEd/Private/CookOnTheFlyServer.cpp
@@ -9588,10 +9588,6 @@ void UCookOnTheFlyServer::CookByTheBookFinishedInternal()
 
 	if (!PlatformManager->GetSessionPlatforms().IsEmpty())
 	{
-		const ITargetPlatform* TargetPlatform = PlatformManager->GetSessionPlatforms()[0];
-		FString MetaDataDirectory = GetMetadataDirectory();
-		MetaDataDirectory = ConvertToFullSandboxPath(MetaDataDirectory, true, TargetPlatform->PlatformName());
-
 		if (IncrementallyModifiedDiagnostics)
 		{
 			if (bWriteIncrementallyModifiedDiagnosticsToLogs)
@@ -9600,7 +9596,19 @@ void UCookOnTheFlyServer::CookByTheBookFinishedInternal()
 			}
 			else
 			{
-				IncrementallyModifiedDiagnostics->WriteToFile(MetaDataDirectory);
+				const ITargetPlatform* TargetPlatform = PlatformManager->GetSessionPlatforms()[0];
+				FString DiagnosticsDir;
+				
+				if (GIsBuildMachine)
+				{
+					DiagnosticsDir = FPaths::Combine(*FPaths::EngineDir(), TEXT("Programs"), TEXT("AutomationTool"), TEXT("Saved"), TEXT("Logs"), TargetPlatform->PlatformName());
+				}
+				else
+				{
+					DiagnosticsDir = GetMetadataDirectory();
+					DiagnosticsDir = ConvertToFullSandboxPath(DiagnosticsDir, true, TargetPlatform->PlatformName());
+				}
+				IncrementallyModifiedDiagnostics->WriteToFile(DiagnosticsDir);
 			}
 		}
 	}
