New cvar Editor.CheckForWorldGCLeaksAreFatal to reduce the fatal error to a warning in CheckForWorldGCLeaks()

diff --git a/Engine/Source/Editor/UnrealEd/Private/EditorServer.cpp b/Engine/Source/Editor/UnrealEd/Private/EditorServer.cpp
--- a/Engine/Source/Editor/UnrealEd/Private/EditorServer.cpp
+++ b/Engine/Source/Editor/UnrealEd/Private/EditorServer.cpp
@@ -218,6 +218,13 @@ namespace
 		int32 Num;
 		int32 Size;
 	};
+
+	TAutoConsoleVariable<bool> CVarCheckForWorldGCLeaksAreFatal(
+		TEXT("Editor.CheckForWorldGCLeaksAreFatal"),
+		true,
+		TEXT("If true, CheckForWorldGCLeaks() will cause a Fatal error if leaks are found. If false, it is an Error."),
+		ECVF_Default
+	);
 }
 
 /**
@@ -1935,7 +1942,14 @@ void UEditorEngine::CheckForWorldGCLeaks( UWorld* NewWorld, UPackage* WorldPacka
 
 	if (NumFailedToCleanup > 0)
 	{
-		UE_LOG(LogEditorServer, Fatal, TEXT("World Memory Leaks: %d leaks objects and packages. See The output above."), NumFailedToCleanup);
+		if (CVarCheckForWorldGCLeaksAreFatal.GetValueOnGameThread())
+		{
+			UE_LOG(LogEditorServer, Fatal, TEXT("World Memory Leaks: %d leaks objects and packages. See The output above."), NumFailedToCleanup);
+		} 
+		else
+		{
+			UE_LOG(LogEditorServer, Error, TEXT("World Memory Leaks: %d leaks objects and packages. See The output above."), NumFailedToCleanup);
+		}
 	}
 }
 
