Move material instance EnumerationObjects to AdvancedDisplay UI.
Tidies up the UI in advance of adding some more items.

diff --git a/Engine/Source/Editor/UnrealEd/Classes/MaterialEditor/MaterialEditorInstanceConstant.h b/Engine/Source/Editor/UnrealEd/Classes/MaterialEditor/MaterialEditorInstanceConstant.h
--- a/Engine/Source/Editor/UnrealEd/Classes/MaterialEditor/MaterialEditorInstanceConstant.h
+++ b/Engine/Source/Editor/UnrealEd/Classes/MaterialEditor/MaterialEditorInstanceConstant.h
@@ -374,7 +374,7 @@ private:
 	TObjectPtr<UMaterialInterface> NaniteOverrideMaterial;
 
 	/** Array of enumeration objects for use by scalar parameter enumeration indices. */
-	UPROPERTY(EditAnywhere, Category = MaterialEditorInstanceConstant, meta = (DisplayThumbnail = "false", AllowedClasses = "/Script/CoreUObject.Enum, /Script/Engine.MaterialEnumerationProvider"))
+	UPROPERTY(EditAnywhere, Category = MaterialEditorInstanceConstant, AdvancedDisplay, meta = (DisplayThumbnail = "false", AllowedClasses = "/Script/CoreUObject.Enum, /Script/Engine.MaterialEnumerationProvider"))
 	TArray<TSoftObjectPtr<UObject>> EnumerationObjects;
 
 #if WITH_EDITOR
