Bump IncrementalCookVersion to work around another collision of two shader code chunks in horde.
#rnx

diff --git a/Engine/Source/Editor/UnrealEd/Private/Cooker/CookIncrementalVersion.cpp b/Engine/Source/Editor/UnrealEd/Private/Cooker/CookIncrementalVersion.cpp
--- a/Engine/Source/Editor/UnrealEd/Private/Cooker/CookIncrementalVersion.cpp
+++ b/Engine/Source/Editor/UnrealEd/Private/Cooker/CookIncrementalVersion.cpp
@@ -8,6 +8,6 @@ namespace UE::Cook
 {
 
 // Change CookIncrementalVersion to a new guid when all packages in an incremental cook need to be invalidated
-FGuid CookIncrementalVersion(0x7DCBB8C7, 0x72364DB2, 0xA97E2A13, 0x0CFA598D);
+FGuid CookIncrementalVersion(0xA42C52BE, 0x13E84181, 0xB025A1AF, 0xB33E6EC2);
 
 }
