- Dump package digest dependencies to the logs.
- New command line argument -DumpPackageDigest=<packagename>.

#rnx

diff --git a/Engine/Source/Editor/UnrealEd/Private/EditorDomain/EditorDomainUtils.cpp b/Engine/Source/Editor/UnrealEd/Private/EditorDomain/EditorDomainUtils.cpp
--- a/Engine/Source/Editor/UnrealEd/Private/EditorDomain/EditorDomainUtils.cpp
+++ b/Engine/Source/Editor/UnrealEd/Private/EditorDomain/EditorDomainUtils.cpp
@@ -42,6 +42,7 @@
 #include "Modules/ModuleManager.h"
 #include "Serialization/BulkDataRegistry.h"
 #include "Serialization/CompactBinaryWriter.h"
+#include "Serialization/JsonSerializer.h"
 #include "Serialization/MemoryReader.h"
 #include "Serialization/PackageWriterToSharedBuffer.h"
 #include "String/ParseTokens.h"
@@ -347,6 +348,18 @@ FPackageDigest CalculatePackageDigest(IAssetRegistry& AssetRegistry, FName Packa
 
 FPackageDigest CalculatePackageDigest(const FAssetPackageData& PackageData, FName PackageName)
 {
+	static FString PackageToDump;
+	static bool bDumpPackageDigest = FParse::Value(FCommandLine::Get(), TEXT("-DumpPackageDigest="), PackageToDump);
+	static FName PackageNameToDump(PackageToDump);
+
+	bool bDumpCurrentPackage = PackageName == PackageNameToDump;
+
+	TSharedPtr<FJsonObject> DigestDependencies = nullptr;
+	if (bDumpCurrentPackage)
+	{
+		DigestDependencies = MakeShareable(new FJsonObject());
+	}
+
 	FPackageDigest Result;
 	Result.DomainUse = EDomainUse::LoadEnabled | EDomainUse::SaveEnabled;;
 	EnumSetFlagsAnd(Result.DomainUse, EDomainUse::LoadEnabled | EDomainUse::SaveEnabled,
@@ -355,12 +368,42 @@ FPackageDigest CalculatePackageDigest(const FAssetPackageData& PackageData, FNam
 	FBlake3 Writer;
 	FStringView ProjectName(FApp::GetProjectName());
 	Writer.Update(ProjectName.GetData(), ProjectName.Len() * sizeof(ProjectName[0]));
+	if (bDumpCurrentPackage)
+	{
+		DigestDependencies->SetStringField("PackageDigest.ProjectName", FString(ProjectName));
+	}
+
 	Writer.Update(&EditorDomainVersion, sizeof(EditorDomainVersion));
+	if (bDumpCurrentPackage)
+	{
+		DigestDependencies->SetStringField("PackageDigest.EditorDomainVersion", EditorDomainVersion.ToString());
+	}
+
 	uint8 EditorDomainSaveUnversioned = GetEditorDomainSaveUnversioned() ? 1 : 0;
 	Writer.Update(&EditorDomainSaveUnversioned, sizeof(EditorDomainSaveUnversioned));
+	if (bDumpCurrentPackage)
+	{
+		DigestDependencies->SetStringField("PackageDigest.EditorDomainSaveUnversioned", FString::FromInt(EditorDomainSaveUnversioned));
+	}
+
 	Writer.Update(&PackageData.GetPackageSavedHash().GetBytes(), sizeof(PackageData.GetPackageSavedHash().GetBytes()));
+	if (bDumpCurrentPackage)
+	{
+		DigestDependencies->SetStringField("PackageDigest.PackageData", LexToString(PackageData.GetPackageSavedHash()));
+	}
+
 	Writer.Update(&GPackageFileUEVersion, sizeof(GPackageFileUEVersion));
+	if (bDumpCurrentPackage)
+	{
+		DigestDependencies->SetStringField("PackageDigest.GPackageFileUEVersion", FString::FromInt(GPackageFileUEVersion.ToValue()));
+	}
+
 	Writer.Update(&GPackageFileLicenseeUEVersion, sizeof(GPackageFileLicenseeUEVersion));
+	if (bDumpCurrentPackage)
+	{
+		DigestDependencies->SetStringField("PackageDigest.GPackageFileLicenseeUEVersion", FString::FromInt(GPackageFileLicenseeUEVersion));
+	}
+
 	TArray<int32> CustomVersionHandles;
 	// Reserve 10 custom versions per class times 100 classes per package times twice (once in package, once in class)
 	CustomVersionHandles.Reserve(10*100*2);
@@ -370,6 +413,10 @@ FPackageDigest CalculatePackageDigest(const FAssetPackageData& PackageData, FNam
 
 	InitializeGlobalConstructClasses();
 	Writer.Update(&GGlobalConstructClassesHash, sizeof(GGlobalConstructClassesHash));
+	if (bDumpCurrentPackage)
+	{
+		DigestDependencies->SetStringField("PackageDigest.GGlobalConstructClassesHash", LexToString(GGlobalConstructClassesHash));
+	}
 
 	// Add the InclusiveSchema for all classes imported by the package. The InclusiveSchema includes the schema of
 	// classes that can be constructed by the given class.
@@ -407,11 +454,20 @@ FPackageDigest CalculatePackageDigest(const FAssetPackageData& PackageData, FNam
 				FClassDigestData* NativeData = ExistingData;
 				if (ExistingData->ClosestNative != ClassPath)
 				{
+					if (bDumpCurrentPackage)
+					{
+						DigestDependencies->SetStringField(ClassPath.ToString(), TEXT("Non native class"));
+					}
 					NativeData = ClassDigests.Map.Find(ExistingData->ClosestNative);
 					checkf(NativeData, TEXT("Classes are only stored in a ClosestNative field if they exist"));
 				}
 				check(NativeData->bNative);
 				Writer.Update(&NativeData->InclusiveSchemaHash, sizeof(NativeData->InclusiveSchemaHash));
+				if (bDumpCurrentPackage)
+				{
+					DigestDependencies->SetStringField(NativeData->ClosestNative.ToString(), LexToString(NativeData->InclusiveSchemaHash));
+				}
+
 				CustomVersionHandles.Append(NativeData->CustomVersionHandles);
 				EnumSetFlagsAnd(Result.DomainUse, EDomainUse::LoadEnabled | EDomainUse::SaveEnabled,
 					Result.DomainUse, NativeData->EditorDomainUse);
@@ -442,21 +498,49 @@ FPackageDigest CalculatePackageDigest(const FAssetPackageData& PackageData, FNam
 	CustomVersionGuids.Sort();
 	CustomVersions.Reserve(CustomVersionGuids.Num());
 
+	int32 CustomVersionCounter = 0;
 	for (const FGuid& CustomVersionGuid : CustomVersionGuids)
 	{
 		Writer.Update(&CustomVersionGuid, sizeof(CustomVersionGuid));
+		if (bDumpCurrentPackage)
+		{
+			DigestDependencies->SetStringField("PackageDigest.CustomVersionGuid_" + FString::FromInt(CustomVersionCounter), LexToString(CustomVersionGuid));
+		}
+
 		TOptional<FCustomVersion> CurrentVersion = FCurrentCustomVersions::Get(CustomVersionGuid);
 		if (!CurrentVersion.IsSet())
 		{
 			return FPackageDigest(FPackageDigest::EStatus::MissingCustomVersion, FName(CustomVersionGuid.ToString()));
 		}
 		Writer.Update(&CurrentVersion->Version, sizeof(CurrentVersion->Version));
+		if (bDumpCurrentPackage)
+		{
+			DigestDependencies->SetStringField("PackageDigest.CurrentVersion_" + FString::FromInt(CustomVersionCounter), FString::FromInt(CurrentVersion->Version));
+			++CustomVersionCounter;
+		}
+
 		CustomVersions.Emplace(CustomVersionGuid, CurrentVersion->Version);
 	}
 
 	Result.Hash = Writer.Finalize();
 	Result.Status = FPackageDigest::EStatus::Successful;
 	Result.CustomVersions = UE::AssetRegistry::FPackageCustomVersionsHandle::FindOrAdd(CustomVersions);
+
+	if (bDumpCurrentPackage)
+	{
+		FString JsonStr;
+		TSharedRef<TJsonWriter<>> JsonWriter = TJsonWriterFactory<>::Create(&JsonStr);
+		FJsonSerializer::Serialize(DigestDependencies.ToSharedRef(), JsonWriter);
+
+		UE_LOG(LogEditorDomain, Display, TEXT("---Package digest dependencies---"));
+		UE_LOG(LogEditorDomain, Display, TEXT("%s"), *(PackageName.ToString()));
+		UE_LOG(LogEditorDomain, Display, TEXT("%s"), *JsonStr);
+		
+		FString PackageDigestHash = LexToString(Result.Hash);
+		UE_LOG(LogEditorDomain, Display, TEXT("Result Hash %s"), *PackageDigestHash);
+		UE_LOG(LogEditorDomain, Display, TEXT("---"));
+	}
+
 	return Result;
 }
 
