[Gizmos] Parent and Explicit gizmo space entries missing from menu

diff --git a/Engine/Source/Editor/UnrealEd/Private/ViewportToolbar/UnrealEdViewportToolbar.cpp b/Engine/Source/Editor/UnrealEd/Private/ViewportToolbar/UnrealEdViewportToolbar.cpp
--- a/Engine/Source/Editor/UnrealEd/Private/ViewportToolbar/UnrealEdViewportToolbar.cpp
+++ b/Engine/Source/Editor/UnrealEd/Private/ViewportToolbar/UnrealEdViewportToolbar.cpp
@@ -1063,7 +1063,7 @@ FToolMenuEntry CreateTransformsSubmenu()
 													);
 													UnnamedSection.AddEntry(LocalCoords);
 
-													if (TSharedPtr<FEditorViewportClient> Client = UUnrealEdViewportToolbarContext::GetEditorViewportClient(InSubmenu))
+													if (TSharedPtr<FEditorViewportClient> Client = UUnrealEdViewportToolbarContext::GetEditorViewportClient(InSubmenu->Context))
 													{
 														if (FEditorModeTools* ModeTools = Client->GetModeTools(); ModeTools && ModeTools->UsesViewportITF())
 														{
