Fixed TSet remove during iteration in UEditorEngine::BroadcastPostUndoRedo and UEditorEngine::UnregisterForUndo.

diff --git a/Engine/Source/Editor/UnrealEd/Classes/Editor/EditorEngine.h b/Engine/Source/Editor/UnrealEd/Classes/Editor/EditorEngine.h
--- a/Engine/Source/Editor/UnrealEd/Classes/Editor/EditorEngine.h
+++ b/Engine/Source/Editor/UnrealEd/Classes/Editor/EditorEngine.h
@@ -2884,7 +2884,7 @@ private:
 
 	/** A temporary copy used to add resilience during the undo/redo broadcast. Used to allow unregistering of clients while the broadcast is happening.
 	 * DO NOT add to this set while a broadcast is happening. */
-	TSet<class FEditorUndoClient*> InflightUndoClients;
+	TArray<class FEditorUndoClient*> InflightUndoClients;
 
 	/** The notification item to use for undo/redo */
 	TSharedPtr<class SNotificationItem> UndoRedoNotificationItem;
diff --git a/Engine/Source/Editor/UnrealEd/Private/EditorServer.cpp b/Engine/Source/Editor/UnrealEd/Private/EditorServer.cpp
--- a/Engine/Source/Editor/UnrealEd/Private/EditorServer.cpp
+++ b/Engine/Source/Editor/UnrealEd/Private/EditorServer.cpp
@@ -5359,10 +5359,9 @@ void UEditorEngine::BroadcastPostUndoRedo(const FTransactionContext& UndoContext
 
 	// Note that we use a copy here as clients can register/unregister with the undo system while in PostUndo()/PostRedo()
 	// which modifies UndoClients during the loop. This can cause an infinite loop where the iterator never finishes.
-	InflightUndoClients = UndoClients;
-	for (auto UndoIt = InflightUndoClients.CreateConstIterator(); UndoIt; ++UndoIt)
+	InflightUndoClients = UndoClients.Array();
+	for (FEditorUndoClient* Client : InflightUndoClients)
 	{
-		FEditorUndoClient* Client = *UndoIt;
 		if (Client && Client->MatchesContext(UndoContext, CurrentUndoRedoContext->TransactionObjects))
 		{
 			if (bWasUndo)
@@ -6858,7 +6857,13 @@ void UEditorEngine::UnregisterForUndo( FEditorUndoClient* Client)
 	if (Client)
 	{
 		UndoClients.Remove(Client);
-		InflightUndoClients.Remove(Client);
+		int32 Index = INDEX_NONE;
+		if (InflightUndoClients.Find(Client, Index))
+		{
+			// Can't remove from InflightUndoClients during iteration in UEditorEngine::BroadcastPostUndoRedo.
+			// Just null out the element.
+			InflightUndoClients[Index] = nullptr;
+		}
 	}
 }
 
