Editor: don't move a viewport-locked actor when input is disabled in that viewport

The "picture-in-picture" preview for camera actors could sometimes rotate that actor as if piloted. This was because the viewport's camera controller was clamping the pitch angle to -90;+90 degrees, so when that came into effect, the output pitch differed from the input pitch. This in turn was picked up by the viewport's camera movement handling as a difference that needed to be reproduced on the locked actor -- something that happens when piloting an actor using a movement mode that has inertia for smoother movements.

This change just disables that part of the code when input is itself disabled (which is the case in the PiP viewports), since it makes little sense to move a locked actor if the user doesn't have input to pilot it.

diff --git a/Engine/Source/Editor/UnrealEd/Private/EditorViewportClient.cpp b/Engine/Source/Editor/UnrealEd/Private/EditorViewportClient.cpp
--- a/Engine/Source/Editor/UnrealEd/Private/EditorViewportClient.cpp
+++ b/Engine/Source/Editor/UnrealEd/Private/EditorViewportClient.cpp
@@ -2376,7 +2376,6 @@ void FEditorViewportClient::UpdateCameraMovement( float DeltaTime )
 			CameraUserImpulseData->RotateRollVelocityModifier = 0.0f;
 		}
 
-
 		// Check for rotation difference within a small tolerance, ignoring winding
 		if( !GetViewRotation().GetDenormalized().Equals( FRotator::MakeFromEuler( NewViewEuler ).GetDenormalized(), SMALL_NUMBER ) )
 		{
@@ -2389,7 +2388,9 @@ void FEditorViewportClient::UpdateCameraMovement( float DeltaTime )
 		const bool bFOVDifferent = !FMath::IsNearlyEqual( NewViewFOV, ViewFOV, float(SMALL_NUMBER) );
 
 		// If something has changed, tell the actor
-		if(bTransformDifferent || bFOVDifferent)
+		// However, don't move the actor if input is disabled: the only reason for transform or FOV changes to be reproduced on the actor
+		// is if the user is piloting the actor directly or indirectly, which requires input.
+		if ((bTransformDifferent || bFOVDifferent) && !bDisableInput)
 		{
 			// Something has changed!
 			const bool bInvalidateChildViews=true;
