World Partition - HLOD: During build, fix issue with -ReuseParentBranchHLODs which would ignore important changes to the HLOD actors and only consider the HLOD hash
* For example, a change to the target runtime grid would be ignored... in which case the parent branch HLOD would be used.

diff --git a/Engine/Source/Editor/UnrealEd/Private/WorldPartition/WorldPartitionHLODsBuilder.cpp b/Engine/Source/Editor/UnrealEd/Private/WorldPartition/WorldPartitionHLODsBuilder.cpp
--- a/Engine/Source/Editor/UnrealEd/Private/WorldPartition/WorldPartitionHLODsBuilder.cpp
+++ b/Engine/Source/Editor/UnrealEd/Private/WorldPartition/WorldPartitionHLODsBuilder.cpp
@@ -1626,35 +1626,9 @@ bool UWorldPartitionHLODsBuilder::EvaluateHLODBuildConditions(AWorldPartitionHLO
 	// Retrieve the parent hash
 	TOptional<uint32> ParentHash;
 
-	// Option #1 (faster): Retrieve the HLOD hash from the asset registry tag
-	if (!ParentHash.IsSet())
 	{
-		// There's a possibility that the package we are about to process was saved with a newer version of the engine.
-		// This will be gracefully handled by the asset registry (package will fail to load).
-		// Suppress asset registry errors, as this is not a problem for our use case.
-		FScopedLogSuppression VerbosityScope(TEXT("LogAssetRegistry"));
-
-		IAssetRegistry::FLoadPackageRegistryData LoadedData;
-		IAssetRegistry& AssetRegistry = FModuleManager::LoadModuleChecked<FAssetRegistryModule>(TEXT("AssetRegistry")).Get();
-		AssetRegistry.LoadPackageRegistryData(LocalFilePathParent, LoadedData);
-		if (!LoadedData.Data.IsEmpty())
-		{
-			FString HLODActorHashString;
-			if (LoadedData.Data[0].GetTagValue(AWorldPartitionHLOD::NAME_HLODHash_AssetTag, HLODActorHashString))
-			{
-				uint32 HLODActorHash;
-				if (LexTryParseString(HLODActorHash, *HLODActorHashString))
-				{
-					ParentHash = HLODActorHash;
-				}
-			}
-		}
-	}
-
-	// Option #2 (slower): Load the package, find the actor and retrieve the hash
-	if (!ParentHash.IsSet())
-	{
-		// There's a possibility that the package we are about to load was saved with a newer version of the engine.
+		// Since we are copying the package from another (parent) branch, there's always the possibility that the package 
+		// we are about to load was saved with a newer version of the engine (a paused auto integration for example).
 		// This will be gracefully handled by the engine (package will fail to load).
 		// Suppress linker errors, as this is not a problem for our use case.
 		FScopedLogSuppression VerbosityScope(TEXT("LogLinker"));
@@ -1677,7 +1651,14 @@ bool UWorldPartitionHLODsBuilder::EvaluateHLODBuildConditions(AWorldPartitionHLO
 
 			if (FoundActor)
 			{
-				ParentHash = FoundActor->GetHLODHash();
+				TUniquePtr<FWorldPartitionActorDesc> CurrentActorDesc = HLODActor->CreateActorDesc();
+				TUniquePtr<FWorldPartitionActorDesc> ParentActorDesc = FoundActor->CreateActorDesc();
+				const FHLODActorDesc& CurrentHLODActorDesc = *(static_cast<FHLODActorDesc*>(CurrentActorDesc.Get()));
+				const FHLODActorDesc& ParentHLODActorDesc = *(static_cast<FHLODActorDesc*>(ParentActorDesc.Get()));
+				if (CurrentHLODActorDesc.EqualsUnbuilt(ParentHLODActorDesc))
+				{
+					ParentHash = FoundActor->GetHLODHash();
+				}
 			}
 
 			ResetLoaders(TempPackage);
diff --git a/Engine/Source/Runtime/Engine/Private/WorldPartition/HLOD/HLODActorDesc.cpp b/Engine/Source/Runtime/Engine/Private/WorldPartition/HLOD/HLODActorDesc.cpp
--- a/Engine/Source/Runtime/Engine/Private/WorldPartition/HLOD/HLODActorDesc.cpp
+++ b/Engine/Source/Runtime/Engine/Private/WorldPartition/HLOD/HLODActorDesc.cpp
@@ -196,6 +196,27 @@ bool FHLODActorDesc::Equals(const FWorldPartitionActorDesc* Other) const
 	return false;
 }
 
+bool FHLODActorDesc::EqualsUnbuilt(const FHLODActorDesc& Other) const
+{
+	return Guid == Other.Guid &&
+		BaseClass == Other.BaseClass &&
+		NativeClass == Other.NativeClass &&
+		ActorLabel == Other.ActorLabel &&
+		RuntimeGrid == Other.RuntimeGrid &&
+		bIsSpatiallyLoaded == Other.bIsSpatiallyLoaded &&
+		bActorIsRuntimeOnly == Other.bActorIsRuntimeOnly &&
+		bActorIsMainWorldOnly == Other.bActorIsMainWorldOnly &&
+		bIsUsingDataLayerAsset == Other.bIsUsingDataLayerAsset &&
+		HLODLayer == Other.HLODLayer &&
+		ContentBundleGuid == Other.ContentBundleGuid &&
+		ComponentDescs == Other.ComponentDescs &&
+		CompareUnsortedArrays(DataLayers, Other.DataLayers) &&
+		CompareUnsortedArrays(Tags, Other.Tags) &&
+		Properties == Other.Properties &&
+		ExternalDataLayerAsset == Other.ExternalDataLayerAsset &&
+		SourceHLODLayer == Other.SourceHLODLayer;
+}
+
 bool FHLODActorDesc::IsRuntimeRelevant(const FWorldPartitionActorDescInstance* InActorDescInstance) const
 {
 	return !InActorDescInstance->GetForceNonSpatiallyLoaded();
diff --git a/Engine/Source/Runtime/Engine/Public/WorldPartition/HLOD/HLODActorDesc.h b/Engine/Source/Runtime/Engine/Public/WorldPartition/HLOD/HLODActorDesc.h
--- a/Engine/Source/Runtime/Engine/Public/WorldPartition/HLOD/HLODActorDesc.h
+++ b/Engine/Source/Runtime/Engine/Public/WorldPartition/HLOD/HLODActorDesc.h
@@ -31,6 +31,9 @@ public:
 	inline const FStats& GetStats() const { return HLODStats; }
 	ENGINE_API int64 GetStat(FName InStatName) const;
 
+	// Test if HLOD actor descriptors are equals, but without considering properties that may change after the HLOD actor is built.
+	ENGINE_API bool EqualsUnbuilt(const FHLODActorDesc& Other) const;
+
 protected:
 	ENGINE_API FHLODActorDesc();
 
