[EditorTransaction]
- Avoid manipulating transactions from the loading thread

diff --git a/Engine/Source/Editor/UnrealEd/Private/EditorTransaction.cpp b/Engine/Source/Editor/UnrealEd/Private/EditorTransaction.cpp
--- a/Engine/Source/Editor/UnrealEd/Private/EditorTransaction.cpp
+++ b/Engine/Source/Editor/UnrealEd/Private/EditorTransaction.cpp
@@ -1224,6 +1224,13 @@ void UTransBuffer::OnObjectsReinstanced(const TMap<UObject*, UObject*>& OldToNew
 		return;
 	}
 
+	// Reinstanciation coming from the loading thread only concerns objects still being loaded and
+	// shouldn't be referenced in any way by an existing transaction.
+	if (IsInActualLoadingThread())
+	{
+		return;
+	}
+
 	class FReinstancingReferenceCollector : public FReferenceCollector
 	{
 	public:
