Bump IncrementalCookVersion to work around a deleted shader type that still exists in the incremental ShaderLibraries.
#rnx

diff --git a/Engine/Source/Editor/UnrealEd/Private/Cooker/CookIncrementalVersion.cpp b/Engine/Source/Editor/UnrealEd/Private/Cooker/CookIncrementalVersion.cpp
--- a/Engine/Source/Editor/UnrealEd/Private/Cooker/CookIncrementalVersion.cpp
+++ b/Engine/Source/Editor/UnrealEd/Private/Cooker/CookIncrementalVersion.cpp
@@ -8,6 +8,6 @@ namespace UE::Cook
 {
 
 // Change CookIncrementalVersion to a new guid when all packages in an incremental cook need to be invalidated
-FGuid CookIncrementalVersion(0x9E83A94D, 0x9E704115, 0xB8811C9A, 0x1DB6F152);
+FGuid CookIncrementalVersion(0x49FA64BB, 0xC79948DD, 0x9F97FD70, 0x7E0C3C1D);
 
 }
