UE-231160: Add ?Select Actors By Property? in Property Color viewmode

Problem
Users can highlight actors with matching properties in Property Color viewmode but cannot directly select all matches, slowing bulk edits and inspection.

Solution
Added a right-click context-menu action on properties, ?Select Actors By Property.? When invoked (only in Property Color viewmode), it selects all actors whose value matches the clicked property and tints the selection using the same property color mapping.

Notes
Scoped to Property Color viewmode to avoid accidental mass selection elsewhere. UX aligns with existing highlight behavior; no new keybind
Validation: enter Property Color viewmode, right-click a property such as ?Can Be Damaged,? choose ?Select Actors By Property,? verify all matches are selected and tinted; outside this viewmode the command is unavailable.

diff --git a/Engine/Source/Editor/PropertyEditor/Private/SDetailSingleItemRow.cpp b/Engine/Source/Editor/PropertyEditor/Private/SDetailSingleItemRow.cpp
--- a/Engine/Source/Editor/PropertyEditor/Private/SDetailSingleItemRow.cpp
+++ b/Engine/Source/Editor/PropertyEditor/Private/SDetailSingleItemRow.cpp
@@ -32,6 +32,7 @@
 #include "DetailsViewStyle.h"
 #include "SDetailsView.h"
 #include "ToolMenus.h"
+#include "LevelEditorViewport.h"
 
 namespace DetailWidgetConstants
 {
@@ -869,6 +870,30 @@ void SDetailSingleItemRow::Construct( const FArguments& InArgs, FDetailLayoutCus
 	);
 }
 
+void SDetailSingleItemRow::SetPropertyColorationTarget(bool bSelectActors /*=false*/)
+{
+	const TSharedPtr<FPropertyNode> PropertyNode = GetPropertyNode();
+	const TSharedPtr<IPropertyHandle> PropertyHandle = GetPropertyHandle();		
+
+	if (PropertyNode.IsValid() && PropertyHandle.IsValid())
+	{
+		TSharedPtr<IDetailsViewPrivate> DetailsView = OwnerTreeNode.Pin()->GetDetailsViewSharedPtr();
+		TSharedPtr<IPropertyHandle> Handle = PropertyEditorHelpers::GetPropertyHandle(PropertyNode.ToSharedRef(), DetailsView->GetNotifyHook(), DetailsView->GetPropertyUtilities());
+
+		FString Value;
+		if (Handle->GetValueAsFormattedString(Value, PPF_Copy) == FPropertyAccess::Success)
+		{
+			FProperty* Property = PropertyHandle->GetProperty();
+
+			TSharedRef<FEditPropertyChain> PropertyChain = PropertyNode->BuildPropertyChain(Property);
+			check(PropertyChain.IsUnique());
+
+			FProperty* TopProperty = PropertyChain->GetHead()->GetValue();
+			GEditor->SetPropertyColorationTarget(GWorld, Value, Property, TopProperty->GetOwnerClass(), &PropertyChain, bSelectActors);
+		}
+	}
+}
+
 FReply SDetailSingleItemRow::OnMouseButtonUp(const FGeometry& MyGeometry, const FPointerEvent& MouseEvent)
 {
 	if (MouseEvent.GetModifierKeys().IsShiftDown())
@@ -901,26 +926,7 @@ FReply SDetailSingleItemRow::OnMouseButtonUp(const FGeometry& MyGeometry, const
 	}
 	else if (MouseEvent.GetModifierKeys().IsControlDown() && GEditor && GWorld)
 	{
-		const TSharedPtr<FPropertyNode> PropertyNode = GetPropertyNode();
-		const TSharedPtr<IPropertyHandle> PropertyHandle = GetPropertyHandle();		
-
-		if (PropertyNode.IsValid() && PropertyHandle.IsValid())
-		{
-			TSharedPtr<IDetailsViewPrivate> DetailsView = OwnerTreeNode.Pin()->GetDetailsViewSharedPtr();
-			TSharedPtr<IPropertyHandle> Handle = PropertyEditorHelpers::GetPropertyHandle(PropertyNode.ToSharedRef(), DetailsView->GetNotifyHook(), DetailsView->GetPropertyUtilities());
-
-			FString Value;
-			if (Handle->GetValueAsFormattedString(Value, PPF_Copy) == FPropertyAccess::Success)
-			{
-				FProperty* Property = PropertyHandle->GetProperty();
-
-				TSharedRef<FEditPropertyChain> PropertyChain = PropertyNode->BuildPropertyChain(Property);
-				check(PropertyChain.IsUnique());
-
-				FProperty* TopProperty = PropertyChain->GetHead()->GetValue();
-				GEditor->SetPropertyColorationTarget(GWorld, Value, Property, TopProperty->GetOwnerClass(), &PropertyChain);
-			}
-		}
+		SetPropertyColorationTarget();
 	}
 
 	return SDetailTableRowBase::OnMouseButtonUp(MyGeometry, MouseEvent);
@@ -1293,6 +1299,21 @@ void SDetailSingleItemRow::PopulateContextMenu(UToolMenu* ToolMenu)
 				CopyInternalNameAction);
 		}
 
+		// Select By Property
+		{
+			FUIAction SelectByPropertyAction = FExecuteAction::CreateSP(this, &SDetailSingleItemRow::SetPropertyColorationTarget, true);
+			SelectByPropertyAction.IsActionVisibleDelegate = FIsActionButtonVisible::CreateSP(this, &SDetailSingleItemRow::CanSelectByProperty);
+
+			static const FTextFormat TooltipFormat = NSLOCTEXT("PropertyView_Single", "SelectActorsByProperty_ToolTip", "Select all actors with the same value for the property named:\n{0}");
+
+			EditSection.AddMenuEntry(
+				TEXT("SelectByProperty"),
+				NSLOCTEXT("PropertyView", "SelectActorsByProperty", "Select Actors by Property"),
+				FText::Format(TooltipFormat, GetPropertyDisplayName()),
+				FSlateIcon(FCoreStyle::Get().GetStyleSetName(), "GenericCommands.SelectAll"),
+				SelectByPropertyAction);
+		}
+
 		// Favorite
 		{
 			if (OwnerTreeNode.Pin()->GetDetailsViewSharedPtr()->IsFavoritingEnabled())
@@ -1504,6 +1525,12 @@ bool SDetailSingleItemRow::CanCopyPropertyInternalName()
 	return !GetPropertyInternalName().IsEmpty();
 }
 
+bool SDetailSingleItemRow::CanSelectByProperty()
+{
+	FEditorViewportClient* ViewportClient = static_cast<FEditorViewportClient*>(GEditor->GetActiveViewport()->GetClient());
+	return ViewportClient && ViewportClient->GetViewMode() == VMI_VisualizeActorColoration;
+}
+
 bool SDetailSingleItemRow::CanPasteProperty() const
 {
 	TSharedPtr<FPropertyNode> PropertyNode = GetPropertyNode();
diff --git a/Engine/Source/Editor/PropertyEditor/Private/SDetailSingleItemRow.h b/Engine/Source/Editor/PropertyEditor/Private/SDetailSingleItemRow.h
--- a/Engine/Source/Editor/PropertyEditor/Private/SDetailSingleItemRow.h
+++ b/Engine/Source/Editor/PropertyEditor/Private/SDetailSingleItemRow.h
@@ -79,6 +79,7 @@ private:
 	bool CanCopyPropertyDisplayName();
 	void OnCopyPropertyInternalName();
 	bool CanCopyPropertyInternalName();
+	bool CanSelectByProperty();
 	void OnPasteProperty();
 	bool CanPasteProperty() const;
 
@@ -154,6 +155,7 @@ private:
 	TSharedPtr<FPropertyNode> GetPropertyNode() const;
 	TSharedPtr<IPropertyHandle> GetPropertyHandle() const;
 
+	void SetPropertyColorationTarget(bool bSelectActors = false);
 private:
 	/** Customization for this widget */
 	FDetailLayoutCustomization* Customization;
diff --git a/Engine/Source/Editor/UnrealEd/Classes/Editor/EditorEngine.h b/Engine/Source/Editor/UnrealEd/Classes/Editor/EditorEngine.h
--- a/Engine/Source/Editor/UnrealEd/Classes/Editor/EditorEngine.h
+++ b/Engine/Source/Editor/UnrealEd/Classes/Editor/EditorEngine.h
@@ -2067,13 +2067,9 @@ public:
 	 * @param	Property			The property to color.
 	 * @param	CommonBaseClass		The class of object to color.
 	 * @param	PropertyChain		The chain of properties from member to lowest property.
+	 * @param	bSelectActors		Whether to also select the actors that match the PropertyValue or not
 	 */
-	UNREALED_API virtual void SetPropertyColorationTarget(UWorld* InWorld, const FString& PropertyValue, class FProperty* Property, class UClass* CommonBaseClass, TSharedRef<FEditPropertyChain>* PropertyChain);
-
-	/**
-	 * Selects actors that match the property coloration settings.
-	 */
-	UNREALED_API void SelectByPropertyColoration(UWorld* InWorld);
+	UNREALED_API virtual void SetPropertyColorationTarget(UWorld* InWorld, const FString& PropertyValue, class FProperty* Property, class UClass* CommonBaseClass, TSharedRef<FEditPropertyChain>* PropertyChain, bool bSelectActors = false);
 
 	/**
 	 * Warns the user of any hidden levels, and prompts them with a Yes/No dialog
diff --git a/Engine/Source/Editor/UnrealEd/Private/EditorServer.cpp b/Engine/Source/Editor/UnrealEd/Private/EditorServer.cpp
--- a/Engine/Source/Editor/UnrealEd/Private/EditorServer.cpp
+++ b/Engine/Source/Editor/UnrealEd/Private/EditorServer.cpp
@@ -3740,9 +3740,11 @@ namespace
 }
 
 
-void UEditorEngine::SetPropertyColorationTarget(UWorld* InWorld, const FString& PropertyValue, FProperty* Property, UClass* CommonBaseClass, TSharedRef<FEditPropertyChain>* PropertyChain)
+void UEditorEngine::SetPropertyColorationTarget(UWorld* InWorld, const FString& PropertyValue, FProperty* Property, UClass* CommonBaseClass, TSharedRef<FEditPropertyChain>* PropertyChain, bool bSelectActors /*=false*/)
 {
-	if ( GPropertyColorationProperty != Property ||
+
+	if (bSelectActors ||
+		GPropertyColorationProperty != Property ||
 		GPropertyColorationClass != CommonBaseClass ||
 		!!GPropertyColorationChain != !!PropertyChain ||
 		GPropertyColorationValue != PropertyValue )
@@ -3759,9 +3761,34 @@ void UEditorEngine::SetPropertyColorationTarget(UWorld* InWorld, const FString&
 		{
 			GPropertyColorationChain = new TSharedRef<FEditPropertyChain>(*PropertyChain);
 			GbColorationClassIsActor = GPropertyColorationClass && GPropertyColorationClass->IsChildOf(AActor::StaticClass());
-			GbColorationPropertyIsObjectProperty = CastField<FObjectPropertyBase>(GPropertyColorationProperty) != NULL;
+			GbColorationPropertyIsObjectProperty = CastField<FObjectPropertyBase>(GPropertyColorationProperty) != nullptr;
+
 
+			TArray<AActor*> Actors;
+			if (bSelectActors)
+			{
+				GPropertyColorationActorCollector = &Actors;
+			}
+			
 			FActorPrimitiveColorHandler::Get().RefreshPrimitiveColorHandler(TEXT("PropertyColor"), InWorld);
+
+			if (bSelectActors)
+			{
+				GPropertyColorationActorCollector = nullptr;
+
+				const FScopedTransaction Transaction(NSLOCTEXT("PropertyColor", "SelectByPropertyColor", "Select by Property Color"));
+				USelection* SelectedActors = GetSelectedActors();
+				SelectedActors->BeginBatchSelectOperation();
+				SelectedActors->Modify();
+				SelectNone(false, true);
+				for (int32 ActorIndex = 0; ActorIndex < Actors.Num(); ++ActorIndex)
+				{
+					AActor* Actor = Actors[ActorIndex];
+					SelectActor(Actor, true, false);
+				}
+				SelectedActors->EndBatchSelectOperation();
+				NoteSelectionChange();
+			}
 		}
 
 		RedrawLevelEditingViewports();
@@ -3774,14 +3801,13 @@ bool UEditorEngine::GetPropertyColorationMatch(UObject* Object)
 	if (GPropertyColorationChain)
 	{
 		UObject* MatchingBase = NULL;
-		AActor* Owner = NULL;
+		AActor* Actor = NULL;
 		if ( Object->IsA(GPropertyColorationClass) )
 		{
 			// The querying object matches the coloration class.
 			MatchingBase = Object;
 		}
-		else
-		{
+		else {
 			// If the coloration class is an actor, check if the querying object is a component.
 			// If so, compare the class of the component's owner against the coloration class.
 			if ( GbColorationClassIsActor )
@@ -3789,10 +3815,10 @@ bool UEditorEngine::GetPropertyColorationMatch(UObject* Object)
 				UActorComponent* ActorComponent = Cast<UActorComponent>( Object );
 				if ( ActorComponent )
 				{
-					Owner = ActorComponent->GetOwner();
-					if ( Owner && Owner->IsA( GPropertyColorationClass ) )
+					Actor = ActorComponent->GetOwner();
+					if (Actor && Actor->IsA( GPropertyColorationClass ) )
 					{
-						MatchingBase = Owner;
+						MatchingBase = Actor;
 					}
 				}
 			}
@@ -3839,9 +3865,13 @@ bool UEditorEngine::GetPropertyColorationMatch(UObject* Object)
 					bResult  = true;
 
 					// Collect actor references.
-					if ( GPropertyColorationActorCollector && Owner )
+					if (Actor == NULL)
 					{
-						GPropertyColorationActorCollector->AddUnique( Owner );
+						Actor = Cast<AActor>(Object);
+					}
+					if (GPropertyColorationActorCollector && Actor)
+					{
+						GPropertyColorationActorCollector->AddUnique(Actor);
 					}
 				}
 			}
@@ -3851,31 +3881,6 @@ bool UEditorEngine::GetPropertyColorationMatch(UObject* Object)
 }
 
 
-void UEditorEngine::SelectByPropertyColoration(UWorld* InWorld)
-{
-	TArray<AActor*> Actors;
-	GPropertyColorationActorCollector = &Actors;
-	InWorld->UpdateWorldComponents( false, false );
-	GPropertyColorationActorCollector = NULL;
-
-	if ( Actors.Num() > 0 )
-	{
-		const FScopedTransaction Transaction( NSLOCTEXT("UnrealEd", "SelectByProperty", "Select by Property") );
-		USelection* SelectedActors = GetSelectedActors();
-		SelectedActors->BeginBatchSelectOperation();
-		SelectedActors->Modify();
-		SelectNone( false, true );
-		for ( int32 ActorIndex = 0 ; ActorIndex < Actors.Num() ; ++ActorIndex )
-		{
-			AActor* Actor = Actors[ActorIndex];
-			SelectActor( Actor, true, false );
-		}
-		SelectedActors->EndBatchSelectOperation();
-		NoteSelectionChange();
-	}
-}
-
-
 bool UEditorEngine::Map_Check( UWorld* InWorld, const TCHAR* Str, FOutputDevice& Ar, bool bCheckDeprecatedOnly, EMapCheckNotification::Type Notification/*= EMapCheckNotification::DisplayResults*/, bool bClearLog/*= true*/)
 {
 #define LOCTEXT_NAMESPACE "EditorEngine"
diff --git a/Engine/Source/Editor/UnrealEd/Private/UnrealEdSrv.cpp b/Engine/Source/Editor/UnrealEd/Private/UnrealEdSrv.cpp
--- a/Engine/Source/Editor/UnrealEd/Private/UnrealEdSrv.cpp
+++ b/Engine/Source/Editor/UnrealEd/Private/UnrealEdSrv.cpp
@@ -2562,11 +2562,6 @@ bool UUnrealEdEngine::Exec_Actor( UWorld* InWorld, const TCHAR* Str, FOutputDevi
 			// @TODO no longer meaningful
 			return true;
 		}
-		else if( FParse::Command(&Str,TEXT("BYPROPERTY")) ) // ACTOR SELECT BYPROPERTY
-		{
-			GEditor->SelectByPropertyColoration(InWorld);
-			return true;
-		}
 		else if( FParse::Command(&Str,TEXT("DELETED")) ) // ACTOR SELECT DELETED
 		{
 			const FScopedTransaction Transaction( NSLOCTEXT("UnrealEd", "SelectDeleted", "Select Deleted") );
