[ITF] ITF AssetEditorContext properly recreated when needed #jira UE-336272

diff --git a/Engine/Source/Editor/UnrealEd/Private/EditorModeManager.cpp b/Engine/Source/Editor/UnrealEd/Private/EditorModeManager.cpp
--- a/Engine/Source/Editor/UnrealEd/Private/EditorModeManager.cpp
+++ b/Engine/Source/Editor/UnrealEd/Private/EditorModeManager.cpp
@@ -183,6 +183,9 @@ void FEditorModeTools::SetToolkitHost(TSharedRef<class IToolkitHost> InHost)
 	checkf(!ToolkitHost.IsValid(), TEXT("SetToolkitHost can only be called once"));
 	ToolkitHost = InHost;
 
+	// Remove the existing context, if it exists
+	InteractiveToolsContext->ContextObjectStore->RemoveContextObjectsOfType<UAssetEditorContextObject>();
+
 	if (HasToolkitHost())
 	{
 		UAssetEditorContextObject* AssetEditorContextObject = NewObject<UAssetEditorContextObject>(InteractiveToolsContext->ToolManager);
