Show Verse paths for tool tips of Verse classes that don't have any tooltip metadata instead of their mangled package name.

#rnx

diff --git a/Engine/Source/Editor/UnrealEd/Private/EditorClassUtils.cpp b/Engine/Source/Editor/UnrealEd/Private/EditorClassUtils.cpp
--- a/Engine/Source/Editor/UnrealEd/Private/EditorClassUtils.cpp
+++ b/Engine/Source/Editor/UnrealEd/Private/EditorClassUtils.cpp
@@ -10,9 +10,11 @@
 #include "AssetRegistry/AssetRegistryModule.h"
 #include "Subsystems/AssetEditorSubsystem.h"
 
+#include "Containers/VersePath.h"
 #include "IDocumentationPage.h"
 #include "IDocumentation.h"
 #include "SourceCodeNavigation.h"
+#include "VerseVM/VVMVerseClass.h"
 #include "Widgets/Input/SHyperlink.h"
 #include "BlueprintEditorSettings.h"
 
@@ -28,7 +30,30 @@ FString FEditorClassUtils::GetDocumentationExcerpt(const UClass* Class)
 
 TSharedRef<SToolTip> FEditorClassUtils::GetTooltip(const UClass* Class)
 {
-	return (Class ? GetTooltip(Class, Class->GetToolTipText(GetDefault<UBlueprintEditorSettings>()->bShowShortTooltips)) : SNew(SToolTip));
+	FText ToolTipText;
+	if (Class)
+	{
+		const bool bShowShortTooltips = GetDefault<UBlueprintEditorSettings>()->bShowShortTooltips;
+		if (const UVerseClass* VerseClass = Cast<const UVerseClass>(Class))
+		{
+			static const FName NAME_Tooltip("Tooltip");
+			static const FName NAME_ShortTooltip("ShortTooltip");
+			if ((bShowShortTooltips && !VerseClass->GetMetaData(NAME_ShortTooltip).IsEmpty()) || !VerseClass->GetMetaData(NAME_Tooltip).IsEmpty())
+			{
+				ToolTipText = Class->GetToolTipText(bShowShortTooltips);
+			}
+			else
+			{
+				// If the tooltip is empty, GetToolTipText will use the class FName, but that is the mangled package relative path, so instead use the full Verse path for Verse classes.
+				ToolTipText = VerseClass->GetVersePath().AsText();
+			}
+		}
+		else
+		{
+			ToolTipText = Class->GetToolTipText(bShowShortTooltips);
+		}
+	}
+	return GetTooltip(Class, ToolTipText);
 }
 
 TSharedRef<SToolTip> FEditorClassUtils::GetTooltip(const UClass* Class, const TAttribute<FText>& OverrideText)
