Fixed localization commandlets incorrectly creating any missing manifests/archives in-memory, rather than failing

diff --git a/Engine/Source/Developer/Localization/Private/TextLocalizationResourceGenerator.cpp b/Engine/Source/Developer/Localization/Private/TextLocalizationResourceGenerator.cpp
--- a/Engine/Source/Developer/Localization/Private/TextLocalizationResourceGenerator.cpp
+++ b/Engine/Source/Developer/Localization/Private/TextLocalizationResourceGenerator.cpp
@@ -342,9 +342,9 @@ bool FTextLocalizationResourceGenerator::GenerateLocResAndUpdateLiveEntriesFromC
 	FLocTextHelper LocTextHelper(SourcePath, ManifestName, ArchiveName, NativeCulture, CulturesToGenerate, nullptr);
 	{
 		FText LoadError;
-		if (!LocTextHelper.LoadAll(ELocTextHelperLoadFlags::LoadOrCreate, &LoadError))
+		if (!LocTextHelper.LoadAll(ELocTextHelperLoadFlags::Load, &LoadError))
 		{
-			UE_LOGFMT(LogTextLocalizationResourceGenerator, Error, "Load error: {error}",
+			UE_LOGFMT(LogTextLocalizationResourceGenerator, Warning, "Load error: {error}",
 				("error", *LoadError.ToString()),
 				("id", TextLocalizationResourceGenerator::LocalizationLogIdentifier)
 			);
diff --git a/Engine/Source/Editor/UnrealEd/Private/Commandlets/ExportDialogueScriptCommandlet.cpp b/Engine/Source/Editor/UnrealEd/Private/Commandlets/ExportDialogueScriptCommandlet.cpp
--- a/Engine/Source/Editor/UnrealEd/Private/Commandlets/ExportDialogueScriptCommandlet.cpp
+++ b/Engine/Source/Editor/UnrealEd/Private/Commandlets/ExportDialogueScriptCommandlet.cpp
@@ -159,7 +159,7 @@ int32 UExportDialogueScriptCommandlet::Main(const FString& Params)
 	if (!bUseCultureDirectory && CulturesToGenerate.Num() > 1)
 	{
 		UE_LOG(LogExportDialogueScriptCommandlet, Error, TEXT("bUseCultureDirectory may only be used with a single culture."));
-		return false;
+		return -1;
 	}
 
 	// Load the manifest and all archives
@@ -167,10 +167,10 @@ int32 UExportDialogueScriptCommandlet::Main(const FString& Params)
 	LocTextHelper.SetCopyrightNotice(GatherManifestHelper->GetCopyrightNotice());
 	{
 		FText LoadError;
-		if (!LocTextHelper.LoadAll(ELocTextHelperLoadFlags::LoadOrCreate, &LoadError))
+		if (!LocTextHelper.LoadAll(ELocTextHelperLoadFlags::Load, &LoadError))
 		{
 			UE_LOG(LogExportDialogueScriptCommandlet, Error, TEXT("%s"), *LoadError.ToString());
-			return false;
+			return -1;
 		}
 	}
 
diff --git a/Engine/Source/Editor/UnrealEd/Private/Commandlets/GenerateTextLocalizationReportCommandlet.cpp b/Engine/Source/Editor/UnrealEd/Private/Commandlets/GenerateTextLocalizationReportCommandlet.cpp
--- a/Engine/Source/Editor/UnrealEd/Private/Commandlets/GenerateTextLocalizationReportCommandlet.cpp
+++ b/Engine/Source/Editor/UnrealEd/Private/Commandlets/GenerateTextLocalizationReportCommandlet.cpp
@@ -186,7 +186,7 @@ bool UGenerateTextLocalizationReportCommandlet::ProcessWordCountReport(const FSt
 	LocTextHelper.SetCopyrightNotice(GatherManifestHelper->GetCopyrightNotice());
 	{
 		FText LoadError;
-		if (!LocTextHelper.LoadAll(ELocTextHelperLoadFlags::LoadOrCreate, &LoadError))
+		if (!LocTextHelper.LoadAll(ELocTextHelperLoadFlags::Load, &LoadError))
 		{
 			UE_LOGFMT(LogGenerateTextLocalizationReportCommandlet, Warning, "{error}",
 				("error", LoadError.ToString()),
diff --git a/Engine/Source/Editor/UnrealEd/Private/Commandlets/GenerateTextLocalizationResourceCommandlet.cpp b/Engine/Source/Editor/UnrealEd/Private/Commandlets/GenerateTextLocalizationResourceCommandlet.cpp
--- a/Engine/Source/Editor/UnrealEd/Private/Commandlets/GenerateTextLocalizationResourceCommandlet.cpp
+++ b/Engine/Source/Editor/UnrealEd/Private/Commandlets/GenerateTextLocalizationResourceCommandlet.cpp
@@ -170,9 +170,9 @@ int32 UGenerateTextLocalizationResourceCommandlet::Main(const FString& Params)
 	LocTextHelper.SetCopyrightNotice(GatherManifestHelper->GetCopyrightNotice());
 	{
 		FText LoadError;
-		if (!LocTextHelper.LoadAll(ELocTextHelperLoadFlags::LoadOrCreate, &LoadError))
+		if (!LocTextHelper.LoadAll(ELocTextHelperLoadFlags::Load, &LoadError))
 		{
-			UE_LOG(LogGenerateTextLocalizationResourceCommandlet, Error, TEXT("%s"), *LoadError.ToString());
+			UE_LOG(LogGenerateTextLocalizationResourceCommandlet, Warning, TEXT("%s"), *LoadError.ToString());
 			return -1;
 		}
 	}
diff --git a/Engine/Source/Editor/UnrealEd/Private/Commandlets/ImportDialogueScriptCommandlet.cpp b/Engine/Source/Editor/UnrealEd/Private/Commandlets/ImportDialogueScriptCommandlet.cpp
--- a/Engine/Source/Editor/UnrealEd/Private/Commandlets/ImportDialogueScriptCommandlet.cpp
+++ b/Engine/Source/Editor/UnrealEd/Private/Commandlets/ImportDialogueScriptCommandlet.cpp
@@ -132,7 +132,7 @@ int32 UImportDialogueScriptCommandlet::Main(const FString& Params)
 	if (!bUseCultureDirectory && CulturesToGenerate.Num() > 1)
 	{
 		UE_LOG(LogImportDialogueScriptCommandlet, Error, TEXT("bUseCultureDirectory may only be used with a single culture."));
-		return false;
+		return -1;
 	}
 
 	// Load the manifest and all archives
@@ -140,10 +140,10 @@ int32 UImportDialogueScriptCommandlet::Main(const FString& Params)
 	LocTextHelper.SetCopyrightNotice(GatherManifestHelper->GetCopyrightNotice());
 	{
 		FText LoadError;
-		if (!LocTextHelper.LoadAll(ELocTextHelperLoadFlags::LoadOrCreate, &LoadError))
+		if (!LocTextHelper.LoadAll(ELocTextHelperLoadFlags::Load, &LoadError))
 		{
 			UE_LOG(LogImportDialogueScriptCommandlet, Error, TEXT("%s"), *LoadError.ToString());
-			return false;
+			return -1;
 		}
 	}
 
@@ -151,7 +151,10 @@ int32 UImportDialogueScriptCommandlet::Main(const FString& Params)
 	{
 		const FString CultureSourcePath = SourcePath / (bUseCultureDirectory ? NativeCulture : TEXT(""));
 		const FString CultureDestinationPath = DestinationPath / NativeCulture;
-		ImportDialogueScriptForCulture(LocTextHelper, CultureSourcePath / DialogueScriptName, NativeCulture, true);
+		if (!ImportDialogueScriptForCulture(LocTextHelper, CultureSourcePath / DialogueScriptName, NativeCulture, true))
+		{
+			return -1;
+		}
 	}
 
 	// Import any remaining cultures
@@ -165,7 +168,10 @@ int32 UImportDialogueScriptCommandlet::Main(const FString& Params)
 
 		const FString CultureSourcePath = SourcePath / (bUseCultureDirectory ? CultureName : TEXT(""));
 		const FString CultureDestinationPath = DestinationPath / CultureName;
-		ImportDialogueScriptForCulture(LocTextHelper, CultureSourcePath / DialogueScriptName, CultureName, false);
+		if (!ImportDialogueScriptForCulture(LocTextHelper, CultureSourcePath / DialogueScriptName, CultureName, false))
+		{
+			return -1;
+		}
 	}
 
 	return 0;
diff --git a/Engine/Source/Editor/UnrealEd/Private/Commandlets/ImportLocalizedDialogueCommandlet.cpp b/Engine/Source/Editor/UnrealEd/Private/Commandlets/ImportLocalizedDialogueCommandlet.cpp
--- a/Engine/Source/Editor/UnrealEd/Private/Commandlets/ImportLocalizedDialogueCommandlet.cpp
+++ b/Engine/Source/Editor/UnrealEd/Private/Commandlets/ImportLocalizedDialogueCommandlet.cpp
@@ -144,10 +144,10 @@ int32 UImportLocalizedDialogueCommandlet::Main(const FString& Params)
 	LocTextHelper.SetCopyrightNotice(GatherManifestHelper->GetCopyrightNotice());
 	{
 		FText LoadError;
-		if (!LocTextHelper.LoadAll(ELocTextHelperLoadFlags::LoadOrCreate, &LoadError))
+		if (!LocTextHelper.LoadAll(ELocTextHelperLoadFlags::Load, &LoadError))
 		{
 			UE_LOG(LogImportLocalizedDialogueCommandlet, Error, TEXT("%s"), *LoadError.ToString());
-			return false;
+			return -1;
 		}
 	}
 	
@@ -382,7 +382,7 @@ bool UImportLocalizedDialogueCommandlet::ImportDialogueForCulture(FLocTextHelper
 	{
 		auto GetTranslatedTextForContext = [&](const FDialogueContextMapping& InContextMapping) -> FString
 		{
-			const FString ContextLocalizationKey = LocalizedDialogueWave->GetContextLocalizationKey(*ContextMappingPtr);
+			const FString ContextLocalizationKey = LocalizedDialogueWave->GetContextLocalizationKey(InContextMapping);
 
 			// Find the manifest entry for our context
 			TSharedPtr<FManifestEntry> ContextManifestEntry = InLocTextHelper.FindSourceText(FDialogueConstants::DialogueNamespace, ContextLocalizationKey, &LocalizedDialogueWave->SpokenText);
diff --git a/Engine/Source/Editor/UnrealEd/Private/Commandlets/InternationalizationExportCommandlet.cpp b/Engine/Source/Editor/UnrealEd/Private/Commandlets/InternationalizationExportCommandlet.cpp
--- a/Engine/Source/Editor/UnrealEd/Private/Commandlets/InternationalizationExportCommandlet.cpp
+++ b/Engine/Source/Editor/UnrealEd/Private/Commandlets/InternationalizationExportCommandlet.cpp
@@ -190,13 +190,13 @@ int32 UInternationalizationExportCommandlet::Main(const FString& Params)
 		LocTextHelper.SetCopyrightNotice(GatherManifestHelper->GetCopyrightNotice());
 		{
 			FText LoadError;
-			if (!LocTextHelper.LoadAll(ELocTextHelperLoadFlags::LoadOrCreate, &LoadError))
+			if (!LocTextHelper.LoadAll(ELocTextHelperLoadFlags::Load, &LoadError))
 			{
-				UE_LOGFMT(LogInternationalizationExportCommandlet, Error, "Load error: {error}",
+				UE_LOGFMT(LogInternationalizationExportCommandlet, Warning, "Load error: {error}",
 					("error", *LoadError.ToString()),
 					("id", InternationalizationExportCommandlet::LocalizationLogIdentifier)
 				);
-				return false;
+				return -1;
 			}
 		}
 
@@ -221,13 +221,13 @@ int32 UInternationalizationExportCommandlet::Main(const FString& Params)
 		LocTextHelper.SetCopyrightNotice(GatherManifestHelper->GetCopyrightNotice());
 		{
 			FText LoadError;
-			if (!LocTextHelper.LoadAll(ELocTextHelperLoadFlags::LoadOrCreate, &LoadError))
+			if (!LocTextHelper.LoadAll(ELocTextHelperLoadFlags::Load, &LoadError))
 			{
-				UE_LOGFMT(LogInternationalizationExportCommandlet, Error, "Load error: {error}",
+				UE_LOGFMT(LogInternationalizationExportCommandlet, Warning, "Load error: {error}",
 					("error", *LoadError.ToString()),
 					("id", InternationalizationExportCommandlet::LocalizationLogIdentifier)
 				);
-				return false;
+				return -1;
 			}
 		}
 
