Prevent propagating edits to actors that are in cleaned up worlds

[RN] Fixed a crash when editing an actor when instances of that actor exist in cleaned up UWorlds

diff --git a/Engine/Source/Editor/UnrealEd/Private/Kismet2/BlueprintEditorUtils.cpp b/Engine/Source/Editor/UnrealEd/Private/Kismet2/BlueprintEditorUtils.cpp
--- a/Engine/Source/Editor/UnrealEd/Private/Kismet2/BlueprintEditorUtils.cpp
+++ b/Engine/Source/Editor/UnrealEd/Private/Kismet2/BlueprintEditorUtils.cpp
@@ -8943,6 +8943,15 @@ void FBlueprintEditorUtils::PostEditChangeBlueprintActors(UBlueprint* Blueprint,
 		{
 			// We know the class was derived from AActor because we checked the Blueprint->GeneratedClass.
 			AActor* Actor = static_cast<AActor*>(MatchingObj);
+			
+			// If worlds leak or GC hasn't been performed we'll have actors in invalid worlds and reconstructing
+			// them in PostEditChange will crash
+			UWorld* ActorWorld = Actor->GetWorld();
+			if (ActorWorld && (!IsValid(ActorWorld) || ActorWorld->IsCleanedUp()))
+			{
+				continue;
+			}
+
 			Actor->PostEditChange();
 
 			// Broadcast edit notification if necessary so that the level editor's detail panel is refreshed
