Add new command line -DebugNullCook to track why certain core redirector is not found or applied in cook.
This code is not made to stay and should be removed once I found the issue.

#rnx

diff --git a/Engine/Source/Editor/UnrealEd/Private/EditorDomain/EditorDomainUtils.cpp b/Engine/Source/Editor/UnrealEd/Private/EditorDomain/EditorDomainUtils.cpp
--- a/Engine/Source/Editor/UnrealEd/Private/EditorDomain/EditorDomainUtils.cpp
+++ b/Engine/Source/Editor/UnrealEd/Private/EditorDomain/EditorDomainUtils.cpp
@@ -680,6 +680,17 @@ private:
 FClassDigestData* FPrecacheClassDigest::GetRecursive(FTopLevelAssetPath ClassPath, bool bAllowRedirects)
 {
 	check(!ClassPath.IsNull());
+
+	static FString ClassName;
+	static bool bDebugNullCook = FParse::Value(FCommandLine::Get(), TEXT("-DebugNullCook="), ClassName);
+	bool bPrintLogsDebugNullCook = bDebugNullCook && ClassPath.GetAssetName().ToString() == ClassName;
+
+	if (bPrintLogsDebugNullCook)
+	{
+		UE_LOG(LogEditorDomain, Display, TEXT("[DebugNullCook] Calculating class digest for class %s"), *ClassPath.ToString());
+		UE_LOG(LogEditorDomain, Display, TEXT("[DebugNullCook] bAllowRedirects=%d"), bAllowRedirects);
+	}
+
 	// Called within ClassDigestsMap.Lock.WriteLock()
 	FClassDigestData* DigestData = &ClassDigests.FindOrAdd(ClassPath);
 	if (DigestData->bConstructed)
@@ -704,6 +715,12 @@ FClassDigestData* FPrecacheClassDigest::GetRecursive(FTopLevelAssetPath ClassPat
 
 			return DigestData;
 		}
+
+		if (bPrintLogsDebugNullCook)
+		{
+			UE_LOG(LogEditorDomain, Display, TEXT("[DebugNullCook] Failed to find redirection for class %s"), *ClassPath.ToString());
+		}
+
 	}
 
 	// Fill in digest data config-driven flags
diff --git a/Engine/Source/Runtime/CoreUObject/Private/UObject/CoreRedirects.cpp b/Engine/Source/Runtime/CoreUObject/Private/UObject/CoreRedirects.cpp
--- a/Engine/Source/Runtime/CoreUObject/Private/UObject/CoreRedirects.cpp
+++ b/Engine/Source/Runtime/CoreUObject/Private/UObject/CoreRedirects.cpp
@@ -25,6 +25,7 @@
 
 #if WITH_EDITOR
 #include "Misc/RedirectCollector.h"
+#include "Misc/CommandLine.h"
 #endif
 
 DEFINE_LOG_CATEGORY(LogCoreRedirects);
@@ -2764,6 +2765,18 @@ bool FCoreRedirects::AddRedirectListUnderWriteLock(TArrayView<const FCoreRedirec
 bool FCoreRedirects::AddSingleRedirectUnderWriteLock(const FCoreRedirect& NewRedirect, const FString& SourceString, 
 	FScopeCoreRedirectsWriteLockedContext& LockedContext)
 {
+#if WITH_EDITOR
+	static FString ClassName;
+	static bool bDebugNullCook = FParse::Value(FCommandLine::Get(), TEXT("-DebugNullCook="), ClassName);
+	bool bPrintLogsDebugNullCook = bDebugNullCook && NewRedirect.OldName.ObjectName.ToString() == ClassName;
+
+	if (bPrintLogsDebugNullCook)
+	{
+		UE_LOG(LogCoreRedirects, Display, TEXT("[DebugNullCook] Adding redirector %s->%s from %s"), *NewRedirect.OldName.ToString(),
+			*NewRedirect.NewName.ToString(), *SourceString);
+	}
+#endif // #if WITH_EDITOR
+
 	const bool bIsWildcardMatch = NewRedirect.IsWildcardMatch();
 	FRedirectNameMap* ExistingNameMap;
 
