Fixed crash introduced in CL 47316511

#rnx

diff --git a/Engine/Source/Editor/UnrealEd/Private/Fbx/FbxSkeletalMeshImport.cpp b/Engine/Source/Editor/UnrealEd/Private/Fbx/FbxSkeletalMeshImport.cpp
--- a/Engine/Source/Editor/UnrealEd/Private/Fbx/FbxSkeletalMeshImport.cpp
+++ b/Engine/Source/Editor/UnrealEd/Private/Fbx/FbxSkeletalMeshImport.cpp
@@ -1536,7 +1536,10 @@ bool UnFbx::FFbxImporter::FillSkeletalMeshImportData(TArray<FbxNode*>& NodeArray
 
 	for (SkeletalMeshImportData::FMaterial& MaterialInfo : SkelMeshImportDataPtr->Materials)
 	{
-		MaterialInfo.Material->SetMaterialUsage(MATUSAGE_SkeletalMesh);
+		if (MaterialInfo.Material.IsValid())
+		{
+			MaterialInfo.Material->SetMaterialUsage(MATUSAGE_SkeletalMesh);
+		}
 	}
 
 	if (LastImportedMaterialNames.Num() > 0)
