Account for UsageFlags override in preview material.
#rnx

diff --git a/Engine/Source/Editor/UnrealEd/Private/PreviewMaterial.cpp b/Engine/Source/Editor/UnrealEd/Private/PreviewMaterial.cpp
--- a/Engine/Source/Editor/UnrealEd/Private/PreviewMaterial.cpp
+++ b/Engine/Source/Editor/UnrealEd/Private/PreviewMaterial.cpp
@@ -1162,6 +1162,10 @@ void UMaterialEditorInstanceConstant::CopyBasePropertiesFromParent()
 	{
 		BasePropertyOverrides.bCompatibleWithLumenCardSharing = SourceInstance->IsCompatibleWithLumenCardSharing();
 	}
+	if (!BasePropertyOverrides.bOverride_UsageFlags)
+	{
+		BasePropertyOverrides.UsageFlags = SourceInstance->GetUsageFlags();
+	}
 
 	// Copy the Lightmass settings...
 	// The lightmass functions (GetCastShadowAsMasked, etc.) check if the value is overridden and returns the current value if so, otherwise returns the parent value
