Recompile blueprints during presave if they have references to out of date struct properties
[RN] Blueprints that are in a dirty state because they are using a Struct that has been modified are now automatically recompiled on save

diff --git a/Engine/Source/Editor/KismetCompiler/Private/UserDefinedStructureCompilerUtils.cpp b/Engine/Source/Editor/KismetCompiler/Private/UserDefinedStructureCompilerUtils.cpp
--- a/Engine/Source/Editor/KismetCompiler/Private/UserDefinedStructureCompilerUtils.cpp
+++ b/Engine/Source/Editor/KismetCompiler/Private/UserDefinedStructureCompilerUtils.cpp
@@ -120,6 +120,11 @@ struct FUserDefinedStructureCompilerInner
 							ClearStructReferencesInBP(FoundBlueprint, BlueprintsToRecompile);
 							StructProperty->Struct = DuplicatedStruct;
 							StructsToRegenerateReferencesFor.Add(OwnerClass);
+							UStruct* OwnerStruct = StructProperty->GetOwnerStruct();
+							if (OwnerStruct != OwnerClass)
+							{
+								StructsToRegenerateReferencesFor.Add(OwnerStruct);
+							}
 						}
 					}
 					else if (UUserDefinedStruct* OwnerStruct = Cast<UUserDefinedStruct>(StructProperty->GetOwnerStruct()))
diff --git a/Engine/Source/Editor/UnrealEd/Private/Kismet2/BlueprintEditorUtils.cpp b/Engine/Source/Editor/UnrealEd/Private/Kismet2/BlueprintEditorUtils.cpp
--- a/Engine/Source/Editor/UnrealEd/Private/Kismet2/BlueprintEditorUtils.cpp
+++ b/Engine/Source/Editor/UnrealEd/Private/Kismet2/BlueprintEditorUtils.cpp
@@ -368,6 +368,41 @@ static void GetAllChildGraphVariables(UEdGraph* Graph, TArray<UK2Node_Variable*>
 	}
 }
 
+namespace UE::Private
+{
+	static void ForEachObjectReference(const UObject* RootObject, TFunctionRef<void(const UObject*)> Function);
+}
+
+void UE::Private::ForEachObjectReference(const UObject* RootObject, TFunctionRef<void(const UObject*)> Function)
+{
+	TSet<const UObject*> VisitedObjects;
+	TArray<const UObject*> Queue;
+	Queue.Add(RootObject);
+	VisitedObjects.Add(RootObject);
+	while(Queue.Num() != 0)
+	{
+		const UObject* CurrentObject = Queue.Pop();
+
+		TArray<UObject*> References;
+		FReferenceFinder Finder( References );
+		Finder.FindReferences( const_cast<UObject*>(CurrentObject) );
+
+		for(UObject* Obj: References)
+		{
+			if(VisitedObjects.Contains(Obj))
+			{
+				continue;
+			}
+			if(Obj->IsIn(RootObject))
+			{
+				Queue.Add(Obj);
+			}
+			VisitedObjects.Add(Obj);
+			Function(Obj);
+		}
+	}
+}
+
 FBlueprintEditorUtils::FOnRenameVariableReferences FBlueprintEditorUtils::OnRenameVariableReferencesEvent;
 
 void FBlueprintEditorUtils::RenameVariableReferences(UBlueprint* Blueprint, UClass* VariableClass, const FName& OldVarName, const FName& NewVarName)
@@ -1947,6 +1982,47 @@ void FBlueprintEditorUtils::MarkBlueprintAsModified(UBlueprint* Blueprint, FProp
 	}
 }
 
+void FBlueprintEditorUtils::RecompileBeforeSaveIfNeeded(UBlueprint* Blueprint)
+{
+	if (Blueprint->GeneratedClass == nullptr)
+	{
+		return; // no class, nothing could be out of date.. it could be that we need to regenerate it here
+	}
+
+	if (Blueprint->Status == EBlueprintStatus::BS_UpToDate)
+	{
+		return; // the UDS editor marks blueprint as out of date, we can rely on that flag
+	}
+
+	if (Blueprint->CachedUDSDependencies.Num() == 0)
+	{
+		return;
+	}
+
+	// Detect references to out of date UUserDefinedStructs (or other UScriptStructs):
+	TSet<const UScriptStruct*> OutOfDateStructs;
+	UE::Private::ForEachObjectReference(Blueprint->GeneratedClass, [&OutOfDateStructs](const UObject*Obj){
+		if (const UScriptStruct* Struct = Cast<UScriptStruct>(Obj))
+		{
+			const bool bStructOutOfDate = 
+				(Struct->StructFlags & STRUCT_NewerVersionExists) != 0
+				|| (!Struct->IsNative() && Struct->HasAnyFlags(RF_Transient));
+			if (bStructOutOfDate)
+			{
+				OutOfDateStructs.Add(Struct);
+			}
+		}
+	});
+	const bool bHasOutOfDateStructs = OutOfDateStructs.Num() != 0;
+	if (bHasOutOfDateStructs)
+	{
+		// full compile, really we just need to regenerate the class *and reinstance* if necessary,
+		// compile is just a convenient way to do that:
+		FKismetEditorUtilities::CompileBlueprint(
+			Blueprint, EBlueprintCompileOptions::SkipGarbageCollection);
+	}
+}
+
 bool FBlueprintEditorUtils::ShouldRegenerateBlueprint(UBlueprint* Blueprint)
 {
 	return !IsCompileOnLoadDisabled(Blueprint)
diff --git a/Engine/Source/Editor/UnrealEd/Public/Kismet2/BlueprintEditorUtils.h b/Engine/Source/Editor/UnrealEd/Public/Kismet2/BlueprintEditorUtils.h
--- a/Engine/Source/Editor/UnrealEd/Public/Kismet2/BlueprintEditorUtils.h
+++ b/Engine/Source/Editor/UnrealEd/Public/Kismet2/BlueprintEditorUtils.h
@@ -313,6 +313,13 @@ public:
 	 */
 	static UNREALED_API void MarkBlueprintAsModified(UBlueprint* Blueprint, FPropertyChangedEvent PropertyChangedEvent = FPropertyChangedEvent(nullptr));
 
+	/**
+	 * Handles regenerating a blueprint class such that it is in a good state for saving, currently this is 
+	 * required to get UserDefinedStruct types up to date before saving. It's likely that there may be
+	 * other situations that require fixup in this way.
+	 */
+	static UNREALED_API void RecompileBeforeSaveIfNeeded(UBlueprint* Blueprint);
+
 	/** See whether or not the specified graph name / entry point name is unique */
 	static UNREALED_API bool IsGraphNameUnique(UObject* InOuter, const FName& InName);
 
diff --git a/Engine/Source/Runtime/Engine/Private/Blueprint.cpp b/Engine/Source/Runtime/Engine/Private/Blueprint.cpp
--- a/Engine/Source/Runtime/Engine/Private/Blueprint.cpp
+++ b/Engine/Source/Runtime/Engine/Private/Blueprint.cpp
@@ -412,6 +412,8 @@ void UBlueprint::PreSave(FObjectPreSaveContext ObjectSaveContext)
 
 		// Cache the BP for use (immediate, since we're about to save)
 		FFindInBlueprintSearchManager::Get().AddOrUpdateBlueprintSearchMetadata(this, Flags, OverrideVersion);
+
+		FBlueprintEditorUtils::RecompileBeforeSaveIfNeeded(this);
 	}
 }
 
