Fix PVS false positive

#rnx

diff --git a/Engine/Source/Editor/UnrealEd/Private/Kismet2/BlueprintEditorUtils.cpp b/Engine/Source/Editor/UnrealEd/Private/Kismet2/BlueprintEditorUtils.cpp
--- a/Engine/Source/Editor/UnrealEd/Private/Kismet2/BlueprintEditorUtils.cpp
+++ b/Engine/Source/Editor/UnrealEd/Private/Kismet2/BlueprintEditorUtils.cpp
@@ -387,7 +387,7 @@ void UE::Private::ForEachObjectReference(const UObject* RootObject, TFunctionRef
 		FReferenceFinder Finder( References );
 		Finder.FindReferences( const_cast<UObject*>(CurrentObject) );
 
-		for(UObject* Obj: References)
+		for(UObject* Obj: References) //-V1078
 		{
 			if(VisitedObjects.Contains(Obj))
 			{
