[ITF Viewport Interactions] Rework the lifecycle of viewport interactions within the editor mode so that interactions are always cleared out and reinstated whenever an editor mode changes.

diff --git a/Engine/Plugins/Animation/ControlRig/Source/ControlRigEditor/Private/EditMode/ControlRigEditMode.cpp b/Engine/Plugins/Animation/ControlRig/Source/ControlRigEditor/Private/EditMode/ControlRigEditMode.cpp
--- a/Engine/Plugins/Animation/ControlRig/Source/ControlRigEditor/Private/EditMode/ControlRigEditMode.cpp
+++ b/Engine/Plugins/Animation/ControlRig/Source/ControlRigEditor/Private/EditMode/ControlRigEditMode.cpp
@@ -535,12 +535,14 @@ void FControlRigEditMode::Enter()
 			if (UModeManagerInteractiveToolsContext* const InteractiveToolsContext =
 					LevelEditor->GetEditorModeManager().GetInteractiveToolsContext())
 			{
-				if (UViewportInteractionsBehaviorSource* const InteractionsBehaviorSource =
-						InteractiveToolsContext->GetViewportInteractionsBehaviorSource())
+				InteractiveToolsContext->OnRegisterViewportInteractions().AddSPLambda(this, [InteractiveToolsContext]
 				{
-					// Add interactions specific to Control Rig
-					InteractionsBehaviorSource->AddInteraction<UControlRigFrustumSelectInteraction>(true);
-				}
+					if (UViewportInteractionsBehaviorSource* const InteractionsBehaviorSource = InteractiveToolsContext->GetViewportInteractionsBehaviorSource())
+					{
+						// Add interactions specific to Control Rig
+						InteractionsBehaviorSource->AddInteraction<UControlRigFrustumSelectInteraction>();
+					}
+				});
 			}
 		}
 	}
@@ -587,6 +589,14 @@ void FControlRigEditMode::Exit()
 	if (IsInLevelEditor())
 	{
 		UE::ControlRig::RemoveControlRigViewportToolbarExtensions();
+		
+		if (const ILevelEditor* const LevelEditor = Private::GetLevelEditor())
+		{
+			if (UModeManagerInteractiveToolsContext* const InteractiveToolsContext = LevelEditor->GetEditorModeManager().GetInteractiveToolsContext())
+			{
+				InteractiveToolsContext->OnRegisterViewportInteractions().RemoveAll(this);
+			}
+		}
 	}
 
 	if (GEngine)
diff --git a/Engine/Source/Editor/LevelEditor/Private/SLevelEditor.cpp b/Engine/Source/Editor/LevelEditor/Private/SLevelEditor.cpp
--- a/Engine/Source/Editor/LevelEditor/Private/SLevelEditor.cpp
+++ b/Engine/Source/Editor/LevelEditor/Private/SLevelEditor.cpp
@@ -771,8 +771,6 @@ void SLevelEditor::RegisterViewportInteractions() const
 		if (UViewportInteractionsBehaviorSource* InputToolsSource =
 				ModeManagerInteractiveToolsContext->GetViewportInteractionsBehaviorSource())
 		{
-			InputToolsSource->Reset();
-			
 			if (ModeManagerInteractiveToolsContext->UsesViewportInteractions())
 			{
 				UE::Editor::ViewportInteractions::AddDefaultCameraMovementInteractions(InputToolsSource);
@@ -787,9 +785,6 @@ void SLevelEditor::RegisterViewportInteractions() const
 
 			// Double click
 			InputToolsSource->AddInteraction<ULevelViewportDoubleClick>();
-
-			// make sure to register all BehaviorSources
-			InputToolsSource->RegisterBehaviorSources();
 		}
 	}
 }
diff --git a/Engine/Source/Editor/Sequencer/Private/SequencerEdMode.cpp b/Engine/Source/Editor/Sequencer/Private/SequencerEdMode.cpp
--- a/Engine/Source/Editor/Sequencer/Private/SequencerEdMode.cpp
+++ b/Engine/Source/Editor/Sequencer/Private/SequencerEdMode.cpp
@@ -209,10 +209,16 @@ void FSequencerEdMode::Enter()
 	NextPlayerStatus = EMovieScenePlayerStatus::Stopped;
 	StartXValue.Reset();
 	
-	if (UViewportInteractionsBehaviorSource* InteractionSource = GetModeManager()->GetInteractiveToolsContext()->GetViewportInteractionsBehaviorSource())
+	if (UModeManagerInteractiveToolsContext* Context = GetModeManager()->GetInteractiveToolsContext())
 	{
-		UMoveSequencerTimeSliderViewportInteraction* MoveInteraction = InteractionSource->AddInteraction<UMoveSequencerTimeSliderViewportInteraction>(true);
-		MoveInteraction->EdMode = this;
+		Context->OnRegisterViewportInteractions().AddSPLambda(this, [this, Context]
+		{
+			if (UViewportInteractionsBehaviorSource* Source = Context->GetViewportInteractionsBehaviorSource())
+			{
+				UMoveSequencerTimeSliderViewportInteraction* MoveInteraction = Source->AddInteraction<UMoveSequencerTimeSliderViewportInteraction>(true);
+				MoveInteraction->EdMode = this;	
+			}
+		});
 	}
 	
 	FEdMode::Enter();
@@ -222,6 +228,11 @@ void FSequencerEdMode::Exit()
 {
 	Sequencers.Reset();
 
+	if (UModeManagerInteractiveToolsContext* Context = GetModeManager()->GetInteractiveToolsContext())
+	{
+		Context->OnRegisterViewportInteractions().RemoveAll(this);
+	}
+	
 	FEdMode::Exit();
 }
 
diff --git a/Engine/Source/Editor/UnrealEd/Private/Tools/EdModeInteractiveToolsContext.cpp b/Engine/Source/Editor/UnrealEd/Private/Tools/EdModeInteractiveToolsContext.cpp
--- a/Engine/Source/Editor/UnrealEd/Private/Tools/EdModeInteractiveToolsContext.cpp
+++ b/Engine/Source/Editor/UnrealEd/Private/Tools/EdModeInteractiveToolsContext.cpp
@@ -1451,16 +1451,26 @@ UViewportInteractionsBehaviorSource* UModeManagerInteractiveToolsContext::GetVie
 
 void UModeManagerInteractiveToolsContext::RegisterViewportInteractions()
 {
-	if (!ViewportInteractionsBehaviorSource && UsesViewportInteractions())
+	if (!UsesViewportInteractions())
+	{
+		return;
+	}
+	
+	if (ViewportInteractionsBehaviorSource)
+	{
+		ViewportInteractionsBehaviorSource->Reset();
+	}
+	else
 	{
 		SetBypassAltModifier(false);
-
+		
 		ViewportInteractionsBehaviorSource = NewObject<UViewportInteractionsBehaviorSource>();
 		ViewportInteractionsBehaviorSource->Initialize(this);
-		ViewportInteractionsBehaviorSource->RegisterBehaviorSources();
-
-		OnRegisterViewportInteractions().Broadcast();
 	}
+	
+	OnRegisterViewportInteractions().Broadcast();
+	
+	ViewportInteractionsBehaviorSource->RegisterBehaviorSources();
 }
 
 void UModeManagerInteractiveToolsContext::UnregisterViewportInteractions()
@@ -1478,17 +1488,7 @@ void UModeManagerInteractiveToolsContext::UnregisterViewportInteractions()
 
 void UModeManagerInteractiveToolsContext::OnEditorModeChanged(const FEditorModeID& InModeID, bool bInIsEnteringMode)
 {
-	if (UE::Editor::ViewportInteractions::UseEditorViewportInteractions())
-	{
-		if (bInIsEnteringMode)
-		{
-			RegisterViewportInteractions();
-		}
-		else
-		{
-			UnregisterViewportInteractions();
-		}
-	}
+	RegisterViewportInteractions();
 }
 
 void UEdModeInteractiveToolsContext::InitializeContextFromModeManagerContext(UModeManagerInteractiveToolsContext* ModeManagerToolsContext)
