[AssetDefinition] Change Action retrieve method to avoid breaking validation

Currently the widget shared in the Action retrieval of AssetDefinition is used in 2 places and that breaks validation since 2 Widets will have the same child.
Instead just retrieve the function to determine if the action is playing or not and use that to create 2 different widget that based on that callback will decide the image to show (either Play or Stop)

diff --git a/Engine/Plugins/Editor/EngineAssetDefinitions/Source/Private/AssetDefinition_ForceFeedbackEffect.cpp b/Engine/Plugins/Editor/EngineAssetDefinitions/Source/Private/AssetDefinition_ForceFeedbackEffect.cpp
--- a/Engine/Plugins/Editor/EngineAssetDefinitions/Source/Private/AssetDefinition_ForceFeedbackEffect.cpp
+++ b/Engine/Plugins/Editor/EngineAssetDefinitions/Source/Private/AssetDefinition_ForceFeedbackEffect.cpp
@@ -287,16 +287,10 @@ TSharedPtr<SWidget> UAssetDefinition_ForceFeedbackEffect::GetThumbnailOverlay(co
 
 bool UAssetDefinition_ForceFeedbackEffect::GetThumbnailActionOverlay(const FAssetData& InAssetData, FAssetActionThumbnailOverlayInfo& OutActionOverlayInfo) const
 {
-	auto OnGetDisplayBrushLambda = [this, InAssetData]() -> const FSlateBrush*
-	{
-		if (MenuExtension_ForceFeedbackEffect::IsEffectPlaying(InAssetData))
+	OutActionOverlayInfo.IsActionPlayingDelegate = FIsActionPlaying::CreateLambda([InAssetData] ()
 		{
-			return FAppStyle::GetBrush("ContentBrowser.AssetAction.StopIcon");
-		}
-		return FAppStyle::GetBrush("ContentBrowser.AssetAction.PlayIcon");
-	};
-
-	OutActionOverlayInfo.ActionImageWidget = SNew(SImage).Image_Lambda(OnGetDisplayBrushLambda);
+			return MenuExtension_ForceFeedbackEffect::IsEffectPlaying(InAssetData);
+		});
 
 	auto OnClickedLambda = [InAssetData]() -> FReply
 	{
diff --git a/Engine/Plugins/Runtime/Metasound/Source/MetasoundEditor/Private/MetasoundAssetDefinitions.cpp b/Engine/Plugins/Runtime/Metasound/Source/MetasoundEditor/Private/MetasoundAssetDefinitions.cpp
--- a/Engine/Plugins/Runtime/Metasound/Source/MetasoundEditor/Private/MetasoundAssetDefinitions.cpp
+++ b/Engine/Plugins/Runtime/Metasound/Source/MetasoundEditor/Private/MetasoundAssetDefinitions.cpp
@@ -538,28 +538,31 @@ TSharedPtr<SWidget> UAssetDefinition_MetaSoundSource::GetThumbnailOverlay(const
 {
 	auto OnClickedLambdaOverride = [InAssetData]() -> FReply
 	{
-		TSharedPtr<Metasound::Editor::FEditor> Editor = Metasound::Editor::FGraphBuilder::GetEditorForMetasound(*InAssetData.GetAsset());
-		if (UE::AudioEditor::IsSoundPlaying(InAssetData))
-		{
-			if (Editor.IsValid())
-			{
-				Editor->Stop();
-			}
-			else
-			{
-				UE::AudioEditor::StopSound();
-			}
-		}
-		else
+		if (UObject* Object = InAssetData.GetAsset())
 		{
-			if (Editor.IsValid())
+			TSharedPtr<Metasound::Editor::FEditor> Editor = Metasound::Editor::FGraphBuilder::GetEditorForMetasound(*Object);
+			if (UE::AudioEditor::IsSoundPlaying(InAssetData))
 			{
-				Editor->Play();
+				if (Editor.IsValid())
+				{
+					Editor->Stop();
+				}
+				else
+				{
+					UE::AudioEditor::StopSound();
+				}
 			}
 			else
 			{
-				// Load and play sound
-				UE::AudioEditor::PlaySound(Cast<USoundBase>(InAssetData.GetAsset()));
+				if (Editor.IsValid())
+				{
+					Editor->Play();
+				}
+				else
+				{
+					// Load and play sound
+					UE::AudioEditor::PlaySound(Cast<USoundBase>(Object));
+				}
 			}
 		}
 		return FReply::Handled();
@@ -569,17 +572,10 @@ TSharedPtr<SWidget> UAssetDefinition_MetaSoundSource::GetThumbnailOverlay(const
 
 bool UAssetDefinition_MetaSoundSource::GetThumbnailActionOverlay(const FAssetData& InAssetData, FAssetActionThumbnailOverlayInfo& OutActionOverlayInfo) const
 {
-	auto OnGetDisplayBrushLambda = [InAssetData]() -> const FSlateBrush*
-	{
-		if (UE::AudioEditor::IsSoundPlaying(InAssetData))
+	OutActionOverlayInfo.IsActionPlayingDelegate = FIsActionPlaying::CreateLambda([InAssetData] ()
 		{
-			return FAppStyle::GetBrush("ContentBrowser.AssetAction.StopIcon");
-		}
-
-		return FAppStyle::GetBrush("ContentBrowser.AssetAction.PlayIcon");
-	};
-
-	OutActionOverlayInfo.ActionImageWidget = SNew(SImage).Image_Lambda(OnGetDisplayBrushLambda);
+			return UE::AudioEditor::IsSoundPlaying(InAssetData);
+		});
 
 	auto OnToolTipTextLambda = [InAssetData]() -> FText
 	{
@@ -593,28 +589,31 @@ bool UAssetDefinition_MetaSoundSource::GetThumbnailActionOverlay(const FAssetDat
 
 	auto OnClickedLambda = [InAssetData]() -> FReply
 	{
-		TSharedPtr<Metasound::Editor::FEditor> Editor = Metasound::Editor::FGraphBuilder::GetEditorForMetasound(*InAssetData.GetAsset());
-		if (UE::AudioEditor::IsSoundPlaying(InAssetData))
-		{
-			if (Editor.IsValid())
-			{
-				Editor->Stop();
-			}
-			else
-			{
-				UE::AudioEditor::StopSound();
-			}
-		}
-		else
+		if (UObject* Object = InAssetData.GetAsset())
 		{
-			if (Editor.IsValid())
+			TSharedPtr<Metasound::Editor::FEditor> Editor = Metasound::Editor::FGraphBuilder::GetEditorForMetasound(*Object);
+			if (UE::AudioEditor::IsSoundPlaying(InAssetData))
 			{
-				Editor->Play();
+				if (Editor.IsValid())
+				{
+					Editor->Stop();
+				}
+				else
+				{
+					UE::AudioEditor::StopSound();
+				}
 			}
 			else
 			{
-				// Load and play sound
-				UE::AudioEditor::PlaySound(Cast<USoundBase>(InAssetData.GetAsset()));
+				if (Editor.IsValid())
+				{
+					Editor->Play();
+				}
+				else
+				{
+					// Load and play sound
+					UE::AudioEditor::PlaySound(Cast<USoundBase>(Object));
+				}
 			}
 		}
 		return FReply::Handled();
diff --git a/Engine/Source/Editor/AssetDefinition/Public/AssetDefinition.h b/Engine/Source/Editor/AssetDefinition/Public/AssetDefinition.h
--- a/Engine/Source/Editor/AssetDefinition/Public/AssetDefinition.h
+++ b/Engine/Source/Editor/AssetDefinition/Public/AssetDefinition.h
@@ -331,16 +331,16 @@ public:
 	TArray<FAssetFilterData> Filters;
 };
 
-
+DECLARE_DELEGATE_RetVal(bool, FIsActionPlaying)
 struct FAssetActionThumbnailOverlayInfo
 {
 public:
+	UE_DEPRECATED(5.8, "Please use the IsActionPlayingDelegate instead.")
 	/** Image widget displayed in the top left border */
 	TSharedPtr<SWidget> ActionImageWidget;
 
-	UE_DEPRECATED(5.6, "Please use the ActionButtonArgs instead")
-	/** Button widget displayed in the center when hovering the thumbnail */
-	TSharedPtr<SWidget> ActionButtonWidget;
+	/** Attach your IsPlaying action to this delegate that will be used to decide on the action image to use */
+	FIsActionPlaying IsActionPlayingDelegate;
 
 	/** Button Args used to construct the Action Button */
 	SButton::FArguments ActionButtonArgs;
diff --git a/Engine/Source/Editor/AudioEditor/Private/AssetTypeActions/AssetDefinition_SoundBase.cpp b/Engine/Source/Editor/AudioEditor/Private/AssetTypeActions/AssetDefinition_SoundBase.cpp
--- a/Engine/Source/Editor/AudioEditor/Private/AssetTypeActions/AssetDefinition_SoundBase.cpp
+++ b/Engine/Source/Editor/AudioEditor/Private/AssetTypeActions/AssetDefinition_SoundBase.cpp
@@ -311,17 +311,10 @@ TSharedPtr<SWidget> UAssetDefinition_SoundBase::GetThumbnailOverlay(const FAsset
 
 bool UAssetDefinition_SoundBase::GetThumbnailActionOverlay(const FAssetData& InAssetData, FAssetActionThumbnailOverlayInfo& OutActionOverlayInfo) const
 {
-	auto OnGetDisplayBrushLambda = [InAssetData]() -> const FSlateBrush*
-	{
-		if (UE::AudioEditor::IsSoundPlaying(InAssetData))
+	OutActionOverlayInfo.IsActionPlayingDelegate = FIsActionPlaying::CreateLambda([InAssetData] ()
 		{
-			return FAppStyle::GetBrush("ContentBrowser.AssetAction.StopIcon");
-		}
-
-		return FAppStyle::GetBrush("ContentBrowser.AssetAction.PlayIcon");
-	};
-
-	OutActionOverlayInfo.ActionImageWidget = SNew(SImage).Image_Lambda(OnGetDisplayBrushLambda);
+			return UE::AudioEditor::IsSoundPlaying(InAssetData);
+		});
 
 	auto OnToolTipTextLambda = [InAssetData]() -> FText
 	{
diff --git a/Engine/Source/Editor/UnrealEd/Private/AssetThumbnail.cpp b/Engine/Source/Editor/UnrealEd/Private/AssetThumbnail.cpp
--- a/Engine/Source/Editor/UnrealEd/Private/AssetThumbnail.cpp
+++ b/Engine/Source/Editor/UnrealEd/Private/AssetThumbnail.cpp
@@ -482,7 +482,31 @@ PRAGMA_ENABLE_DEPRECATION_WARNINGS
 				if (AssetDefinition && AssetDefinition->GetThumbnailActionOverlay(AssetData, OutThumbnailInfo))
 				{
 					constexpr int32 OverlayZOrder = 1;
-					if (OutThumbnailInfo.ActionImageWidget.IsValid())
+					TSharedPtr<SWidget> TopLeftImage;
+					TSharedPtr<SWidget> CenterImage;
+					if (OutThumbnailInfo.IsActionPlayingDelegate.IsBound())
+					{
+						auto OnGetActionImage = [IsPlaying = OutThumbnailInfo.IsActionPlayingDelegate] ()
+							{
+								if (IsPlaying.IsBound() && IsPlaying.Execute())
+								{
+									return FAppStyle::GetBrush("ContentBrowser.AssetAction.StopIcon");
+								}
+								return FAppStyle::GetBrush("ContentBrowser.AssetAction.PlayIcon");
+							};
+						TopLeftImage = SNew(SImage).Image_Lambda(OnGetActionImage);
+						CenterImage = SNew(SImage).Image_Lambda(OnGetActionImage);
+					}
+					PRAGMA_DISABLE_DEPRECATION_WARNINGS
+					else if (OutThumbnailInfo.ActionImageWidget.IsValid())
+					{
+						// This approach use the same widget in 2 different place and that shouldn't happen, we have to create 2 separate widget, but keep it here for this release since we shouldn't break the API
+						TopLeftImage = OutThumbnailInfo.ActionImageWidget.ToSharedRef();
+						CenterImage = OutThumbnailInfo.ActionImageWidget.ToSharedRef();
+					}
+					PRAGMA_ENABLE_DEPRECATION_WARNINGS
+					
+					if (TopLeftImage.IsValid())
 					{
 						constexpr float PaddingFromTopLeftBorder = 4.f;
 
@@ -501,72 +525,43 @@ PRAGMA_ENABLE_DEPRECATION_WARNINGS
 								.Padding(this, &SAssetThumbnail::GetPlayIndicatorPadding)
 								.Visibility(this, &SAssetThumbnail::GetActionIconOverlayVisibility)
 								[
-									OutThumbnailInfo.ActionImageWidget.ToSharedRef()
+									TopLeftImage.ToSharedRef()
 								]
 							]
 						];
 					}
+					// Set the default style and padding for the Button
+					OutThumbnailInfo.ActionButtonArgs
+						.ButtonStyle(&FAppStyle::GetWidgetStyle<FButtonStyle>(Style, ".Action.Button"))
+						.ContentPadding(this, &SAssetThumbnail::GetPlayButtonContentPadding);
+
+					TAttribute<EVisibility> ActionVisibility = TAttribute<EVisibility>::CreateSP(this, &SAssetThumbnail::GetActionButtonVisibility);
+					TSharedRef<SButton> ActionButton = SArgumentNew(OutThumbnailInfo.ActionButtonArgs, SButton);
+					ActionButton->Construct(OutThumbnailInfo.ActionButtonArgs);
+					ActionButton->SetVisibility(ActionVisibility);
+					ActionButton->SetToolTipText(OutThumbnailInfo.ActionButtonArgs._ToolTipText);
+					ActionButton->SetToolTip(OutThumbnailInfo.ActionButtonArgs._ToolTip);
+					ActionButton->SetCursor(EMouseCursor::Default);
+					TSharedRef<SWidget> ActionImageWidget = CenterImage.IsValid()
+						? CenterImage.ToSharedRef()
+						: SNew(SImage).Image(FAppStyle::GetBrush("ContentBrowser.AssetAction.PlayIcon"));
+
+					ActionButton->SetContent(ActionImageWidget);
+
+					constexpr float CenterImageSize = 32.f;
 
-					PRAGMA_DISABLE_DEPRECATION_WARNINGS
-					if (OutThumbnailInfo.ActionButtonWidget.IsValid())
-					{
-						constexpr float CenterImageSize = 32.f;
-
-						OverlayWidget->AddSlot()
-						.ZOrder(OverlayZOrder)
-						.VAlign(VAlign_Center)
-						.HAlign(HAlign_Center)
-						[
-							SNew(SBorder)
-							.Padding(0.f)
-							.BorderImage(FStyleDefaults::GetNoBrush())
-							.Visibility(this, &SAssetThumbnail::GetActionButtonVisibility)
-							[
-								SNew(SBox)
-								.WidthOverride(CenterImageSize)
-								.HeightOverride(CenterImageSize)
-								[
-									OutThumbnailInfo.ActionButtonWidget.ToSharedRef()
-								]
-							]
-						];
-					}
-					PRAGMA_ENABLE_DEPRECATION_WARNINGS
-					else
-					{
-						// Set the default style and padding for the Button
-						OutThumbnailInfo.ActionButtonArgs
-							.ButtonStyle(&FAppStyle::GetWidgetStyle<FButtonStyle>(Style, ".Action.Button"))
-							.ContentPadding(this, &SAssetThumbnail::GetPlayButtonContentPadding);
-
-						TAttribute<EVisibility> ActionVisibility = TAttribute<EVisibility>::CreateSP(this, &SAssetThumbnail::GetActionButtonVisibility);
-						TSharedRef<SButton> ActionButton = SNew(SButton);
-						ActionButton->Construct(OutThumbnailInfo.ActionButtonArgs);
-						ActionButton->SetVisibility(ActionVisibility);
-						ActionButton->SetToolTipText(OutThumbnailInfo.ActionButtonArgs._ToolTipText);
-						ActionButton->SetToolTip(OutThumbnailInfo.ActionButtonArgs._ToolTip);
-						ActionButton->SetCursor(EMouseCursor::Default);
-						TSharedRef<SWidget> ActionImageWidget = OutThumbnailInfo.ActionImageWidget.IsValid()
-							? OutThumbnailInfo.ActionImageWidget.ToSharedRef()
-							: SNew(SImage).Image(FAppStyle::GetBrush("ContentBrowser.AssetAction.PlayIcon"));
-
-						ActionButton->SetContent(ActionImageWidget);
-
-						constexpr float CenterImageSize = 32.f;
-
-						OverlayWidget->AddSlot()
-						.ZOrder(OverlayZOrder)
-						.VAlign(VAlign_Center)
-						.HAlign(HAlign_Center)
+					OverlayWidget->AddSlot()
+					.ZOrder(OverlayZOrder)
+					.VAlign(VAlign_Center)
+					.HAlign(HAlign_Center)
+					[
+						SNew(SBox)
+						.WidthOverride(CenterImageSize)
+						.HeightOverride(CenterImageSize)
 						[
-							SNew(SBox)
-							.WidthOverride(CenterImageSize)
-							.HeightOverride(CenterImageSize)
-							[
-								ActionButton
-							]
-						];
-					}
+							ActionButton
+						]
+					];
 				}
 			}
 			else if (AssetTypeActions.IsValid())
@@ -730,7 +725,7 @@ PRAGMA_ENABLE_DEPRECATION_WARNINGS
 				PlayButtonContentPadding = PlayButtonContentDefaultPadding;
 			}
 
-			PlayButtonContentPadding = FMath::FloorToFloat(PlayButtonContentDefaultPadding);
+			PlayButtonContentPadding = FMath::FloorToFloat(PlayButtonContentPadding);
 			PlayIndicatorPadding = FMath::FloorToFloat(PlayIndicatorPadding);
 			PlayIndicatorSize = FMath::FloorToFloat(PlayIndicatorSize);
 		}
