DebugNullCook : Print the package currently being cooked and the call stack.

#rnx

diff --git a/Engine/Source/Editor/UnrealEd/Private/EditorDomain/EditorDomainUtils.cpp b/Engine/Source/Editor/UnrealEd/Private/EditorDomain/EditorDomainUtils.cpp
--- a/Engine/Source/Editor/UnrealEd/Private/EditorDomain/EditorDomainUtils.cpp
+++ b/Engine/Source/Editor/UnrealEd/Private/EditorDomain/EditorDomainUtils.cpp
@@ -689,6 +689,14 @@ FClassDigestData* FPrecacheClassDigest::GetRecursive(FTopLevelAssetPath ClassPat
 	{
 		UE_LOG(LogEditorDomain, Display, TEXT("[DebugNullCook] Calculating class digest for class %s"), *ClassPath.ToString());
 		UE_LOG(LogEditorDomain, Display, TEXT("[DebugNullCook] bAllowRedirects=%d"), bAllowRedirects);
+
+		const PackageAccessTracking_Private::FTrackedData* TrackedData = PackageAccessTracking_Private::FPackageAccessRefScope::GetCurrentThreadAccumulatedData();
+		if (TrackedData)
+		{
+			UE_LOG(LogEditorDomain, Display, TEXT("[DebugNullCook] Package %s"), *TrackedData->PackageName.ToString());
+		}
+
+		FDebug::DumpStackTraceToLog(ELogVerbosity::Display);
 	}
 
 	// Called within ClassDigestsMap.Lock.WriteLock()
