[ITF] Add utility function for getting button state via FKey. Includes optional delegate that can provide a fallback value.

#rnx

diff --git a/Engine/Source/Editor/UnrealEd/Private/Tools/EdModeInteractiveToolsContext.cpp b/Engine/Source/Editor/UnrealEd/Private/Tools/EdModeInteractiveToolsContext.cpp
--- a/Engine/Source/Editor/UnrealEd/Private/Tools/EdModeInteractiveToolsContext.cpp
+++ b/Engine/Source/Editor/UnrealEd/Private/Tools/EdModeInteractiveToolsContext.cpp
@@ -1099,6 +1099,7 @@ bool UModeManagerInteractiveToolsContext::InputKey(FEditorViewportClient* Viewpo
 			InputState.Keyboard.ActiveKey.Button = Key;
 			bool bPressed = (Event == IE_Pressed);
 			InputState.Keyboard.ActiveKey.SetStates(bPressed, bPressed, !bPressed);
+			InputState.OnGetButtonStateFallback.BindUObject(this, &UModeManagerInteractiveToolsContext::OnGetButtonStateFallback);
 			return InputRouter->PostInputEvent(InputState);
 		}
 	}
@@ -1355,6 +1356,8 @@ void UModeManagerInteractiveToolsContext::Initialize(IToolsContextQueriesAPI* Qu
 	// Delegates used to handle InputTools.EnableITFTools CVAR changes
 	UE::Editor::ViewportInteractions::OnEditorViewportInteractionsActivated().AddUObject(this, &UModeManagerInteractiveToolsContext::RegisterViewportInteractions);
 	UE::Editor::ViewportInteractions::OnEditorViewportInteractionsDeactivated().AddUObject(this, &UModeManagerInteractiveToolsContext::UnregisterViewportInteractions);
+	
+	CurrentMouseState.OnGetButtonStateFallback.BindUObject(this, &UModeManagerInteractiveToolsContext::OnGetButtonStateFallback);
 }
 
 void UModeManagerInteractiveToolsContext::Shutdown()
@@ -1384,6 +1387,19 @@ void UModeManagerInteractiveToolsContext::Shutdown()
 	Super::Shutdown();
 }
 
+FDeviceButtonState UModeManagerInteractiveToolsContext::OnGetButtonStateFallback(const FKey& Key) const
+{
+	if (FEditorViewportClient* Client = GetFocusedViewportClient())
+	{
+		if (Client->Viewport->KeyState(Key))
+		{
+			FDeviceButtonState State(Key);
+			State.bDown = true;
+			return State;
+		}
+	}
+	return FDeviceButtonState();
+}
 
 UEdModeInteractiveToolsContext* UModeManagerInteractiveToolsContext::CreateNewChildEdModeToolsContext()
 {
diff --git a/Engine/Source/Editor/UnrealEd/Public/Tools/EdModeInteractiveToolsContext.h b/Engine/Source/Editor/UnrealEd/Public/Tools/EdModeInteractiveToolsContext.h
--- a/Engine/Source/Editor/UnrealEd/Public/Tools/EdModeInteractiveToolsContext.h
+++ b/Engine/Source/Editor/UnrealEd/Public/Tools/EdModeInteractiveToolsContext.h
@@ -324,6 +324,8 @@ protected:
 	UE_API virtual void DeactivateAllActiveTools(EToolShutdownType ShutdownType) override;
 	UE_API virtual void Initialize(IToolsContextQueriesAPI* QueriesAPIIn, IToolsContextTransactionsAPI* TransactionsAPIIn) override;
 	UE_API virtual void Shutdown() override;
+		
+	FDeviceButtonState OnGetButtonStateFallback(const FKey& Key) const;
 
 	/** Input event instance used to keep track of various button states, etc, that we cannot directly query on-demand */
 	FInputDeviceState CurrentMouseState;
diff --git a/Engine/Source/Runtime/InteractiveToolsFramework/Private/InputState.cpp b/Engine/Source/Runtime/InteractiveToolsFramework/Private/InputState.cpp
--- a/Engine/Source/Runtime/InteractiveToolsFramework/Private/InputState.cpp
+++ b/Engine/Source/Runtime/InteractiveToolsFramework/Private/InputState.cpp
@@ -4,4 +4,65 @@
 
 #include UE_INLINE_GENERATED_CPP_BY_NAME(InputState)
 
+namespace UE::Editor::Private
+{
+	FDeviceButtonState DownState(const FKey& InKey, bool bInDown)
+	{
+		FDeviceButtonState State(InKey);
+		State.bDown = bInDown;
+		return State;
+	}
+}
 
+FDeviceButtonState FInputDeviceState::GetButtonState(const FKey& Key) const
+{
+	using namespace UE::Editor::Private;
+	
+	if (IsFromDevice(EInputDevices::Mouse))
+	{
+		if (Key == EKeys::LeftMouseButton)
+		{
+			return Mouse.Left;
+		}
+		if (Key == EKeys::MiddleMouseButton)
+		{
+			return Mouse.Middle;
+		}
+		if (Key == EKeys::RightMouseButton)
+		{
+			return Mouse.Right;
+		}	
+	}
+	
+	if (IsFromDevice(EInputDevices::Keyboard) && Key == Keyboard.ActiveKey.Button)
+	{
+		return Keyboard.ActiveKey;
+	}
+	
+	// Treat left/right modifier keys as the same
+	if (Key == EKeys::LeftControl || Key == EKeys::RightControl)
+	{
+		return DownState(Key, bCtrlKeyDown);
+	}
+	if (Key == EKeys::LeftShift || Key == EKeys::RightShift)
+	{
+		return DownState(Key, bShiftKeyDown);
+	}
+	if (Key == EKeys::LeftAlt || Key == EKeys::RightAlt)
+	{
+		return DownState(Key, bAltKeyDown);
+	}
+	if (Key == EKeys::LeftCommand || Key == EKeys::RightCommand)
+	{
+		return DownState(Key, bCmdKeyDown);
+	}
+	
+PRAGMA_DISABLE_EXPERIMENTAL_WARNINGS
+	if (OnGetButtonStateFallback.IsBound())
+	{
+		return OnGetButtonStateFallback.Execute(Key);
+	}
+PRAGMA_ENABLE_EXPERIMENTAL_WARNINGS
+	
+	return FDeviceButtonState();
+}
diff --git a/Engine/Source/Runtime/InteractiveToolsFramework/Public/InputState.h b/Engine/Source/Runtime/InteractiveToolsFramework/Public/InputState.h
--- a/Engine/Source/Runtime/InteractiveToolsFramework/Public/InputState.h
+++ b/Engine/Source/Runtime/InteractiveToolsFramework/Public/InputState.h
@@ -184,6 +184,11 @@ struct FDeviceButtonState
 		bReleased = bReleasedIn;
 		bDoubleClicked = bDoubleClickedIn;
 	}
+	
+	bool IsValid() const
+	{
+		return Button.IsValid();
+	}
 };
 
 
@@ -259,11 +264,14 @@ struct FMouseInputDeviceState
  * Assumption is that the state refers to a single physical input device,
  * ie InputDevice field is a single value of EInputDevices and not a combination.
  */
-USTRUCT(BlueprintType)
+USTRUCT(BlueprintType, MinimalAPI)
 struct FInputDeviceState
 {
 	GENERATED_BODY()
 
+	UE_EXPERIMENTAL(5.8, "Experimental extension that allows for generic access to state via FKey.")
+	DECLARE_DELEGATE_RetVal_OneParam(FDeviceButtonState, FGetDeviceButtonState, const FKey&);
+
 	/** Which InputDevice member is valid in this state */
 	UPROPERTY(transient, BlueprintReadWrite, Category = MouseInputDeviceState)
 	EInputDevices InputDevice;
@@ -321,9 +329,15 @@ struct FInputDeviceState
 	{
 		return ((InputDevice & DeviceType) != EInputDevices::None);
 	}
-
-
-
+	
+	/** @return The state for the current button. */
+	UE_EXPERIMENTAL(5.8, "Experimental extension that allows for generic access to state via FKey.")
+	INTERACTIVETOOLSFRAMEWORK_API FDeviceButtonState GetButtonState(const FKey& Key) const;
+
+PRAGMA_DISABLE_EXPERIMENTAL_WARNINGS
+	UE_EXPERIMENTAL(5.8, "Experimental extension that allows for generic access to state via FKey.")
+	FGetDeviceButtonState OnGetButtonStateFallback;
+PRAGMA_ENABLE_EXPERIMENTAL_WARNINGS
 
 	//
 	// utility functions to pass as lambdas
