[UE] Fix for missing native/instanced struct value data on runtime-spawned BP actor types with "Generate optimized Blueprint component data" enabled.

Notes:
- This change makes sure we also consider the struct property in cases where its value only differs from the default value with a native comparison. Previously the comparison was inclusive of only UHT-reflected properties.

diff --git a/Engine/Source/Editor/UnrealEd/Private/Kismet2/BlueprintEditorUtils.cpp b/Engine/Source/Editor/UnrealEd/Private/Kismet2/BlueprintEditorUtils.cpp
--- a/Engine/Source/Editor/UnrealEd/Private/Kismet2/BlueprintEditorUtils.cpp
+++ b/Engine/Source/Editor/UnrealEd/Private/Kismet2/BlueprintEditorUtils.cpp
@@ -10279,6 +10279,22 @@ struct FComponentInstancingDataUtils
 	// Recursively gathers properties that differ from class/struct defaults, and fills out the cooked property list structure.
 	static void RecursivePropertyGather(UStruct* InStruct, const uint8* DataPtr, const uint8* DefaultDataPtr, FBlueprintCookedComponentInstancingData& OutData)
 	{
+		// Determines if a struct value matches the default value using native comparison only (if implemented).
+		auto IsNativeIdentical = [](const UScriptStruct* InStruct, const uint8* InValuePtr, const uint8* InDefaultValuePtr, uint32 PortFlags)
+		{
+			bool bIsIdentical = false;
+			UScriptStruct::ICppStructOps* CppStructOps = InStruct->GetCppStructOps();
+			if (CppStructOps
+				&& CppStructOps->HasIdentical()
+				&& CppStructOps->Identical(InValuePtr, InDefaultValuePtr, PortFlags, bIsIdentical))
+			{
+				return bIsIdentical;
+			}
+
+			// Not bound or no actual comparison was performed at the native level.
+			return true;
+		};
+
 		for (FProperty* Property = InStruct->PropertyLink; Property; Property = Property->PropertyLinkNext)
 		{
 			// Skip editor-only properties since they won't be compiled in a non-editor configuration. Also skip transient and deprecated properties since they won't be serialized on save/duplicate. 
@@ -10301,8 +10317,9 @@ struct FComponentInstancingDataUtils
 
 						RecursivePropertyGather(StructProperty->Struct, PropertyValue, DefaultPropertyValue, OutData);
 
-						// Prepend the struct property only if there is at least one changed sub-property.
-						if (NumChangedProperties < OutData.ChangedPropertyList.Num())
+						// Prepend the struct property only if there is at least one changed sub-property, or otherwise if there are any native differences.
+						if (NumChangedProperties < OutData.ChangedPropertyList.Num()
+							|| !IsNativeIdentical(StructProperty->Struct, PropertyValue, DefaultPropertyValue, PPF_None))
 						{
 							OutData.ChangedPropertyList.Insert(ChangedPropertyInfo, NumChangedProperties);
 						}
@@ -10330,8 +10347,9 @@ struct FComponentInstancingDataUtils
 
 									RecursivePropertyGather(InnerStructProperty->Struct, ArrayPropertyValue, DefaultArrayPropertyValue, OutData);
 
-									// Prepend the struct property only if there is at least one changed sub-property.
-									if (NumChangedArrayProperties < OutData.ChangedPropertyList.Num())
+									// Prepend the struct property only if there is at least one changed sub-property, or otherwise if there are any native differences.
+									if (NumChangedArrayProperties < OutData.ChangedPropertyList.Num()
+										|| !IsNativeIdentical(InnerStructProperty->Struct, ArrayPropertyValue, DefaultArrayPropertyValue, PPF_None))
 									{
 										OutData.ChangedPropertyList.Insert(ChangedArrayPropertyInfo, NumChangedArrayProperties);
 									}
