AssetEditorToolkit - Allow child standalone toolkit implementations to override the default tab ID (StandaloneToolkit)

This change addresses a layout conflict in LiveLinkHub, which functions as an editor-like application hosting a WorkflowCentricApplication/AssetEditor.
When opening another asset editor from within LiveLinkHub, both editors previously shared the same default tab ID (StandaloneToolkit), causing layout clashes in the application.

diff --git a/Engine/Source/Editor/UnrealEd/Private/Toolkits/AssetEditorToolkit.cpp b/Engine/Source/Editor/UnrealEd/Private/Toolkits/AssetEditorToolkit.cpp
--- a/Engine/Source/Editor/UnrealEd/Private/Toolkits/AssetEditorToolkit.cpp
+++ b/Engine/Source/Editor/UnrealEd/Private/Toolkits/AssetEditorToolkit.cpp
@@ -210,6 +210,11 @@ void FAssetEditorToolkit::InitAssetEditor( const EToolkitMode::Type Mode, const
 				check(false);
 			}
 
+			if (FName OverrideTabId = GetOverrideTabId(); !OverrideTabId.IsNone())
+			{
+				PlaceholderId = OverrideTabId;
+			}
+
 			FGlobalTabmanager::Get()->InsertNewDocumentTab(PlaceholderId, *SearchPreference, NewMajorTab.ToSharedRef());
 
 			// Bring the window to front.  The tab manager will not do this for us to avoid intrusive stealing focus behavior
diff --git a/Engine/Source/Editor/UnrealEd/Public/Toolkits/AssetEditorToolkit.h b/Engine/Source/Editor/UnrealEd/Public/Toolkits/AssetEditorToolkit.h
--- a/Engine/Source/Editor/UnrealEd/Public/Toolkits/AssetEditorToolkit.h
+++ b/Engine/Source/Editor/UnrealEd/Public/Toolkits/AssetEditorToolkit.h
@@ -478,6 +478,12 @@ protected:
 	 */
 	UNREALED_API void RestoreFromLayout(const TSharedRef<FTabManager::FLayout>& NewLayout, bool bLoadUserLayout);
 
+	/** 
+	 * Enables standalone toolkit implementations to override the tab layout identifier so it differs from "StandaloneToolkit" or one of its variations. 
+	 * If it's not overridden or NAME_None is returned, the default id will be used.
+	 */
+	virtual FName GetOverrideTabId() const { return NAME_None; }
+
 private:
 	// Callback for persisting the Asset Editor's layout.
 	void HandleTabManagerPersistLayout( const TSharedRef<FTabManager::FLayout>& LayoutToSave )
