Bump IncrementalCookVersion to work around a deleted shader type that still exists in the incremental ShaderLibraries.
#rnx

diff --git a/Engine/Source/Editor/UnrealEd/Private/Cooker/CookIncrementalVersion.cpp b/Engine/Source/Editor/UnrealEd/Private/Cooker/CookIncrementalVersion.cpp
--- a/Engine/Source/Editor/UnrealEd/Private/Cooker/CookIncrementalVersion.cpp
+++ b/Engine/Source/Editor/UnrealEd/Private/Cooker/CookIncrementalVersion.cpp
@@ -8,6 +8,6 @@ namespace UE::Cook
 {
 
 // Change CookIncrementalVersion to a new guid when all packages in an incremental cook need to be invalidated
-FGuid CookIncrementalVersion(0xAC1FE10D, 0xE076413F, 0x9CC1D237, 0x52903FB9);
+FGuid CookIncrementalVersion(0x02C0F385, 0x02DB4B69, 0x8483457F, 0x41CD76FF);
 
 }
