World Partition - HLOD: Fixed issue in HLOD workload computation, uncovered by the recent changes to TSet<> internal behavior.

diff --git a/Engine/Source/Editor/UnrealEd/Private/WorldPartition/WorldPartitionHLODsBuilder.cpp b/Engine/Source/Editor/UnrealEd/Private/WorldPartition/WorldPartitionHLODsBuilder.cpp
--- a/Engine/Source/Editor/UnrealEd/Private/WorldPartition/WorldPartitionHLODsBuilder.cpp
+++ b/Engine/Source/Editor/UnrealEd/Private/WorldPartition/WorldPartitionHLODsBuilder.cpp
@@ -1014,24 +1014,24 @@ TArray<UWorldPartitionHLODsBuilder::FHLODWorkload> UWorldPartitionHLODsBuilder::
 	TMap<FGuid, TArray<FGuid>> HLODGroups;
 	TSet<FGuid>				   TriagedHLODs;
 
-	TFunction<void(TArray<FGuid>&, const FGuid&)> RecursiveAdd = [&TriagedHLODs, &HLODParenting, &HLODGroups, &RecursiveAdd](TArray<FGuid>& HLODGroup, const FGuid& HLODGuid)
+	TFunction<void(const FGuid&, const FGuid&)> RecursiveAdd = [&TriagedHLODs, &HLODParenting, &HLODGroups, &RecursiveAdd](const FGuid& GroupGuid, const FGuid& HLODGuid)
 	{
 		if (!TriagedHLODs.Contains(HLODGuid))
 		{
 			TriagedHLODs.Add(HLODGuid);
-			HLODGroup.Insert(HLODGuid, 0); // Child will come first in the list, as they need to be built first...
+			HLODGroups.FindChecked(GroupGuid).Insert(HLODGuid, 0); // Child will come first in the list, as they need to be built first...
 			TArray<FGuid>* ChildHLODs = HLODParenting.Find(HLODGuid);
 			if (ChildHLODs)
 			{
 				for (const auto& ChildGuid : *ChildHLODs)
 				{
-					RecursiveAdd(HLODGroup, ChildGuid);
+					RecursiveAdd(GroupGuid, ChildGuid);
 				}
 			}
 		}
 		else
 		{
-			HLODGroup.Insert(MoveTemp(HLODGroups.FindChecked(HLODGuid)), 0);
+			HLODGroups.FindChecked(GroupGuid).Insert(MoveTemp(HLODGroups.FindChecked(HLODGuid)), 0);
 			HLODGroups.Remove(HLODGuid);
 		}
 	};
@@ -1040,8 +1040,8 @@ TArray<UWorldPartitionHLODsBuilder::FHLODWorkload> UWorldPartitionHLODsBuilder::
 	{
 		if (!TriagedHLODs.Contains(Pair.Key))
 		{
-			TArray<FGuid>& HLODGroup = HLODGroups.Add(Pair.Key);
-			RecursiveAdd(HLODGroup, Pair.Key);
+			HLODGroups.Add(Pair.Key);
+			RecursiveAdd(Pair.Key, Pair.Key);
 		}
 	}
 
