UE: Bump incremental version due to stale error involving native changes to BP API

diff --git a/Engine/Source/Editor/UnrealEd/Private/Cooker/CookIncrementalVersion.cpp b/Engine/Source/Editor/UnrealEd/Private/Cooker/CookIncrementalVersion.cpp
--- a/Engine/Source/Editor/UnrealEd/Private/Cooker/CookIncrementalVersion.cpp
+++ b/Engine/Source/Editor/UnrealEd/Private/Cooker/CookIncrementalVersion.cpp
@@ -8,6 +8,6 @@ namespace UE::Cook
 {
 
 // Change CookIncrementalVersion to a new guid when all packages in an incremental cook need to be invalidated
-FGuid CookIncrementalVersion(0x02C0F385, 0x02DB4B69, 0x8483457F, 0x41CD76FF);
+FGuid CookIncrementalVersion(0x87F61775, 0xB53244C4, 0xB6DBFFC1, 0xD8A5AC03);
 
 }
