IncrementalCook: ShaderLibraries: Fix bug in intermittent cook when shaders move between chunks that causes the spurious error during staging "Collision of two shader code chunks, same Id (%s), different code".  Delete the chunk-specific ShaderLibraries present from previous incrementalcooks before writing them all again at the end of cook. These libraries are rewritten on every incremental cook and we need to handle the deletion of ones that are no longer referenced.
#rnx

diff --git a/Engine/Source/Editor/UnrealEd/Private/Commandlets/AssetRegistryGenerator.cpp b/Engine/Source/Editor/UnrealEd/Private/Commandlets/AssetRegistryGenerator.cpp
--- a/Engine/Source/Editor/UnrealEd/Private/Commandlets/AssetRegistryGenerator.cpp
+++ b/Engine/Source/Editor/UnrealEd/Private/Commandlets/AssetRegistryGenerator.cpp
@@ -645,14 +645,68 @@ bool FAssetRegistryGenerator::GenerateStreamingInstallManifest(int64 InOverrideC
 			return !DLCNameReturn.IsEmpty();
 		}();
 
-	// generate per-chunk pak list files
 	FDefaultPakFileRules DefaultPakFileRules;
+	TArray<int32> ReferencedChunkIds;
+	TArray<bool> FinalChunkIsReferenced;
+	FinalChunkIsReferenced.SetNum(FinalChunkManifests.Num());
+	ReferencedChunkIds.Reserve(FinalChunkManifests.Num());
+	// Get the list of ChunkIds to report to ChunkDataGenerators' BeginGenerateChunkDataFiles.
+	for (int32 PakchunkIndex = 0; PakchunkIndex < FinalChunkManifests.Num(); ++PakchunkIndex)
+	{
+		const FChunkPackageSet* Manifest = FinalChunkManifests[PakchunkIndex].Get();
+
+		// Is this index a null placeholder that we added in the loop over EncryptedNonUFSFileGroups and then never filled in?  If so, 
+		// fill it in with an empty FChunkPackageSet
+		if (!Manifest)
+		{
+			FinalChunkManifests[PakchunkIndex].Reset(new FChunkPackageSet());
+			Manifest = FinalChunkManifests[PakchunkIndex].Get();
+		}
+
+		FinalChunkIsReferenced[PakchunkIndex] = false;
+
+		// Do not create any files, and do not tell ChunkDataGenerators that the chunk is referenced, if the chunk is
+		// empty and is not referenced by rules applied during staging
+		if (Manifest->IsEmpty())
+		{
+			// During a chunk'd DLC cook ignore DefaultPakFile Rules. Those are handled during main cooks.
+			if (bIsDLCCook)
+			{
+				continue;
+			}
+
+			DefaultPakFileRules.InitializeFromConfig(TargetPlatform);
+			if (!DefaultPakFileRules.IsChunkReferenced(PakchunkIndex))
+			{
+				continue;
+			}
+		}
+
+		FinalChunkIsReferenced[PakchunkIndex] = true;
+		ReferencedChunkIds.Add(PakchunkIndex);
+	}
+	// TOOD: Need to make a public interface for FCookSandbox and pass it into GenerateChunkDataFiles instead of the
+	// internal FSandboxPlatformFile, in case any of the generators need to correctly map files in remapped plugins
+	FSandboxPlatformFile& SandboxPlatformFile = InSandboxFile.GetSandboxPlatformFile();
+	if (ChunkDataGenerators.Num() > 0)
+	{
+		IChunkDataGenerator::FBeginGenerateContext BeginContext;
+		BeginContext.CurrentChunkIds = ReferencedChunkIds;
+		BeginContext.TargetPlatform = TargetPlatform;
+		BeginContext.SandboxFile = &SandboxPlatformFile;
+		for (const TSharedRef<IChunkDataGenerator>& ChunkDataGenerator : ChunkDataGenerators)
+		{
+			ChunkDataGenerator->BeginGenerateChunkDataFiles(BeginContext);
+		}
+	}
+
 	bool bSucceeded = true;
 	TMap<FString, int64> PackageFileSizes;
+	// generate per-chunk pak list files
 	for (int32 PakchunkIndex = 0; PakchunkIndex < FinalChunkManifests.Num() && bSucceeded; ++PakchunkIndex)
 	{
 		const FChunkPackageSet* Manifest = FinalChunkManifests[PakchunkIndex].Get();
-
+		check(Manifest); // This was created in the loop above if not already existing
 		// Serialize chunk layers whether chunk is empty or not
 		int32 TargetLayer = 0;
 		FGameDelegates::Get().GetAssignLayerChunkDelegate().ExecuteIfBound(Manifest, Platform, PakchunkIndex, TargetLayer);
@@ -660,18 +714,10 @@ bool FAssetRegistryGenerator::GenerateStreamingInstallManifest(int64 InOverrideC
 		FString LayerString = FString::Printf(TEXT("%d\r\n"), TargetLayer);
 		ChunkLayerFile->Serialize(TCHAR_TO_ANSI(*LayerString), LayerString.Len());
 
-		// Is this index a null placeholder that we added in the loop over EncryptedNonUFSFileGroups and then never filled in?  If so, 
-		// fill it in with an empty FChunkPackageSet
-		if (!Manifest)
+		if (!FinalChunkIsReferenced[PakchunkIndex])
 		{
-			FinalChunkManifests[PakchunkIndex].Reset(new FChunkPackageSet());
-			Manifest = FinalChunkManifests[PakchunkIndex].Get();
+			continue;
 		}
-		
-		// Split the chunk into subchunks as necessary and create and register a PakListFile for each subchunk
-		int32 FilenameIndex = 0;
-		TArray<FString> ChunkFilenames;
-		Manifest->GenerateValueArray(ChunkFilenames);
 
 		int64* MaxChunkSizeForPakIndex = MaxChunkSizeOverrideFromIndex.Find(PakchunkIndex);
 		int64 PakChunkMaxChunkSize = MaxChunkSizeForPakIndex ? *MaxChunkSizeForPakIndex : MaxChunkSize;
@@ -691,21 +737,10 @@ bool FAssetRegistryGenerator::GenerateStreamingInstallManifest(int64 InOverrideC
 			}
 		}
 
-		// Do not create any files if the chunk is empty and is not referenced by rules applied during staging
-		if (ChunkFilenames.IsEmpty())
-		{
-			// During a chunk'd DLC cook ignore DefaultPakFile Rules. Those are handled during main cooks.
-			if (bIsDLCCook)
-			{
-				continue;
-			}
-
-			DefaultPakFileRules.InitializeFromConfig(TargetPlatform);
-			if (!DefaultPakFileRules.IsChunkReferenced(PakchunkIndex))
-			{
-				continue;
-			}
-		}
+		// Split the chunk into subchunks as necessary and create and register a PakListFile for each subchunk
+		int32 FilenameIndex = 0;
+		TArray<FString> ChunkFilenames;
+		Manifest->GenerateValueArray(ChunkFilenames);
 
 		const FString PakChunkName = GetPakChunkNameForId(PakChunkIdToStringOverride, PakchunkIndex);
 		bool bFinishedAllFiles = false;
@@ -777,9 +812,6 @@ bool FAssetRegistryGenerator::GenerateStreamingInstallManifest(int64 InOverrideC
 
 				for (const TSharedRef<IChunkDataGenerator>& ChunkDataGenerator : ChunkDataGenerators)
 				{
-					// TOOD: Need to make a public interface for FCookSandbox and pass it into GenerateChunkDataFiles instead of the
-					// internal FSandboxPlatformFile, in case any of the generators need to correctly map files in remapped plugins
-					FSandboxPlatformFile& SandboxPlatformFile = InSandboxFile.GetSandboxPlatformFile();
 					ChunkDataGenerator->GenerateChunkDataFiles(PakchunkIndex, PackagesInChunk, TargetPlatform, &SandboxPlatformFile, ChunkFilenames);
 				}
 			}
diff --git a/Engine/Source/Editor/UnrealEd/Private/ShaderLibraryChunkDataGenerator.cpp b/Engine/Source/Editor/UnrealEd/Private/ShaderLibraryChunkDataGenerator.cpp
--- a/Engine/Source/Editor/UnrealEd/Private/ShaderLibraryChunkDataGenerator.cpp
+++ b/Engine/Source/Editor/UnrealEd/Private/ShaderLibraryChunkDataGenerator.cpp
@@ -47,6 +47,15 @@ FShaderLibraryChunkDataGenerator::FShaderLibraryChunkDataGenerator(UCookOnTheFly
 	}
 }
 
+void FShaderLibraryChunkDataGenerator::BeginGenerateChunkDataFiles(FBeginGenerateContext& Context)
+{
+	const ITargetPlatform* TargetPlatform = Context.GetTargetPlatform();
+	FString ShaderCodeDir;
+	FString MetaDataPath;
+	COTFS.GetShaderLibraryPaths(TargetPlatform, ShaderCodeDir, MetaDataPath, true /* bUseProjectDirForDLC */);
+	FShaderLibraryCooker::DeleteShaderLibraryChunks(TargetPlatform, ShaderCodeDir, MetaDataPath);
+}
+
 void FShaderLibraryChunkDataGenerator::GenerateChunkDataFiles(const int32 InChunkId, const TSet<FName>& InPackagesInChunk,
 	const ITargetPlatform* TargetPlatform, FSandboxPlatformFile* InSandboxFile, TArray<FString>& OutChunkFilenames)
 {
diff --git a/Engine/Source/Editor/UnrealEd/Private/ShaderLibraryChunkDataGenerator.h b/Engine/Source/Editor/UnrealEd/Private/ShaderLibraryChunkDataGenerator.h
--- a/Engine/Source/Editor/UnrealEd/Private/ShaderLibraryChunkDataGenerator.h
+++ b/Engine/Source/Editor/UnrealEd/Private/ShaderLibraryChunkDataGenerator.h
@@ -37,6 +37,7 @@ public:
 	virtual ~FShaderLibraryChunkDataGenerator() = default;
 
 	//~ IChunkDataGenerator
+	virtual void BeginGenerateChunkDataFiles(FBeginGenerateContext& Context) override;
 	virtual void GenerateChunkDataFiles(const int32 InChunkId, const TSet<FName>& InPackagesInChunk, const ITargetPlatform* TargetPlatform,
 		FSandboxPlatformFile* InSandboxFile, TArray<FString>& OutChunkFilenames) override;
 };
diff --git a/Engine/Source/Editor/UnrealEd/Public/Commandlets/IChunkDataGenerator.h b/Engine/Source/Editor/UnrealEd/Public/Commandlets/IChunkDataGenerator.h
--- a/Engine/Source/Editor/UnrealEd/Public/Commandlets/IChunkDataGenerator.h
+++ b/Engine/Source/Editor/UnrealEd/Public/Commandlets/IChunkDataGenerator.h
@@ -2,11 +2,23 @@
 
 #pragma once
 
+#include "CoreTypes.h"
+#include "Containers/Array.h"
+#include "Containers/Set.h"
+#include "Containers/UnrealString.h"
+#include "Templates/Function.h"
+#include "Templates/SharedPointer.h"
+#include "UObject/NameTypes.h"
+
+#if UE_ENABLE_INCLUDE_ORDER_DEPRECATED_IN_5_8
 #include "CoreMinimal.h"
-#include "UObject/ICookInfo.h"
+#endif
+
 
+class FAssetRegistryGenerator;
 class FSandboxPlatformFile;
 class ITargetPlatform;
+namespace UE::Cook { class ICookInfo; }
 
 /**
  * Interface for generating extra data when creating streaming install manifests.
@@ -17,12 +29,41 @@ class IChunkDataGenerator
 public:
 	virtual ~IChunkDataGenerator() = default;
 
+	struct FBeginGenerateContext
+	{
+	public:
+		/**
+		 * Get the list of ChunkIds that are referenced from the current cook. Files for other ChunkIds might
+		 * exist on disk from previous cooks, and should be deleted.
+		 */
+		TConstArrayView<int32> GetCurrentChunkIds();
+		/** The TargetPlatform for which chunks are being generated. */
+		const ITargetPlatform* GetTargetPlatform();
+		/** The sandbox location that staging data for this chunks should be written to. */
+		FSandboxPlatformFile* GetSandboxFile();
+
+	private:
+		TConstArrayView<int32> CurrentChunkIds;
+		const ITargetPlatform* TargetPlatform = nullptr;
+		FSandboxPlatformFile* SandboxFile = nullptr;
+
+		friend FAssetRegistryGenerator;
+	};
+
+	/**
+	 * Called to clear or udpate chunk data files from a previous incremental cook. This will be called
+	 * before the call to GenerateChunkDataFiles for each of the listed chunks.
+	 */
+	virtual void BeginGenerateChunkDataFiles(FBeginGenerateContext& Context)
+	{
+	}
+
 	/**
 	 * Called to generate any additional files that should be added to the given chunk.
 	 *
 	 * @param InChunkId The ID of the chunk to generate the files for.
 	 * @param InPackagesInChunk The set of packages that are in the chunk.
-	 * @param InPlatformName The name of the platform this chunk is for.
+	 * @param TargetPlatform The platform this chunk is for.
 	 * @param InSandboxFile The sandbox location that staging data for this chunk should be written to.
 	 * @param OutChunkFilenames Array of filenames that belong to this chunk (to be appended to).
 	 * @return true if successful. If this function returns false, the deprecated one will be called
@@ -37,3 +78,22 @@ public:
 private:
 	static TArray<TFunction<TSharedRef<IChunkDataGenerator>(const UE::Cook::ICookInfo&)>> GeneratorsDataGeneratorFactories;
 };
+
+
+///////////////////////////////////////////////////////
+// Inline implementations
+///////////////////////////////////////////////////////
+
+
+inline TConstArrayView<int32> IChunkDataGenerator::FBeginGenerateContext::GetCurrentChunkIds()
+{
+	return CurrentChunkIds;
+}
+inline const ITargetPlatform* IChunkDataGenerator::FBeginGenerateContext::GetTargetPlatform()
+{
+	return TargetPlatform;
+}
+inline FSandboxPlatformFile* IChunkDataGenerator::FBeginGenerateContext::GetSandboxFile()
+{
+	return SandboxFile;
+}
diff --git a/Engine/Source/Runtime/RenderCore/Private/ShaderCodeLibrary.cpp b/Engine/Source/Runtime/RenderCore/Private/ShaderCodeLibrary.cpp
--- a/Engine/Source/Runtime/RenderCore/Private/ShaderCodeLibrary.cpp
+++ b/Engine/Source/Runtime/RenderCore/Private/ShaderCodeLibrary.cpp
@@ -80,6 +80,7 @@ static FString PipelineExtension = ".ushaderpipelines";
 static FString ShaderCodePatternStr = "^ShaderCode\\-(\\w*)\\-([\\w\\-]*)\\.ushaderbytecode$";
 static FString ShaderArchivePatternStr = "^ShaderArchive\\-(\\w*)\\-([\\w\\-]*)\\.ushaderbytecode$";
 static FString StableInfoPatternStr = "^ShaderStableInfo\\-(\\w*)\\-([\\w\\-]*)\\.shk$";
+static FString ShaderFileForChunkPatternStr = "Shader.*_Chunk[0-9]+-.*\\.(ushaderbytecode|assetinfo.json|stinfo|ushaderpipelines)";
 
 int32 GShaderCodeLibrarySeparateLoadingCache = 0;
 static FAutoConsoleVariableRef CVarShaderCodeLibrarySeparateLoadingCache(
@@ -3854,6 +3855,46 @@ public:
 		return bOk;
 	}
 
+	void DeleteShaderCodeChunks(const FString& SandboxDestinationPath, const FString& SandboxMetadataPath)
+	{
+		TArray<FString> ExistingFiles;
+		IFileManager::Get().FindFiles(ExistingFiles, *SandboxDestinationPath);
+		for (FString& FileName : ExistingFiles)
+		{
+			FileName = FPaths::Combine(SandboxDestinationPath, FileName);
+		}
+		if (SandboxMetadataPath.Len() > 0)
+		{
+			FString MetadataShaderPath = SandboxMetadataPath / TEXT("../ShaderLibrarySource");
+			TArray<FString> ExistingFilesMetadata;
+			IFileManager::Get().FindFiles(ExistingFilesMetadata, *MetadataShaderPath);
+			for (FString& FileName : ExistingFilesMetadata)
+			{
+				FileName = FPaths::Combine(MetadataShaderPath, FileName);
+			}
+			ExistingFiles.Append(ExistingFilesMetadata);
+		}
+		if (ExistingFiles.IsEmpty())
+		{
+			return;
+		}
+
+		FRegexPattern ShaderFileForChunkPattern(ShaderFileForChunkPatternStr);
+		for (FString& ExistingFileName : ExistingFiles)
+		{
+			// Delete all the per-chunk files, we haven't written any of those yet in this cook, but do not
+			// delete .ushaderbytecode files for other libraries such as the GlobalShaderLibrary which is
+			// written at beginning of cook.
+			FString LeafName(FPathViews::GetPathLeaf(ExistingFileName));
+			FRegexMatcher ChunkFileMatcher(ShaderFileForChunkPattern, LeafName);
+			if (ChunkFileMatcher.FindNext())
+			{
+				IFileManager::Get().Delete(*ExistingFileName,
+					true /* RequireExists */, false /* EvenReadOnly */, true /* Quiet */);
+			}
+		}
+	}
+
 	bool PackageNativeShaderLibrary(const FString& ShaderCodeDir, const TArray<FName>& ShaderFormats)
 	{
 		bool bOK = true;
@@ -4690,6 +4731,14 @@ bool FShaderLibraryCooker::SaveShaderLibraryChunk(int32 ChunkId, const TSet<FNam
 		SandboxDestinationPath, SandboxMetadataPath, OutChunkFilenames);
 }
 
+void FShaderLibraryCooker::DeleteShaderLibraryChunks(const ITargetPlatform* TargetPlatform,
+	const FString& SandboxDestinationPath, const FString& SandboxMetadataPath)
+{
+	// The destinationpaths are specific to the TargetPlatform, so we are free to delete files in
+	// them without worrying about deleting files for other TargetPlatforms
+	FShaderLibrariesCollection::Impl->DeleteShaderCodeChunks(SandboxDestinationPath, SandboxMetadataPath);
+}
+
 #endif// WITH_EDITOR
 
 void FShaderCodeLibrary::SafeAssignHash(FRHIShader* InShader, const FSHAHash& Hash)
diff --git a/Engine/Source/Runtime/RenderCore/Public/ShaderCodeLibrary.h b/Engine/Source/Runtime/RenderCore/Public/ShaderCodeLibrary.h
--- a/Engine/Source/Runtime/RenderCore/Public/ShaderCodeLibrary.h
+++ b/Engine/Source/Runtime/RenderCore/Public/ShaderCodeLibrary.h
@@ -394,6 +394,13 @@ struct FShaderLibraryCooker
 	static RENDERCORE_API bool SaveShaderLibraryChunk(int32 ChunkId, const TSet<FName>& InPackagesInChunk, const ITargetPlatform* TargetPlatform,
 		const FString& SandboxDestinationPath, const FString& SandboxMetadataPath, TArray<FString>& OutChunkFilenames, bool& bOutHasData);
 
+	/**
+	 * Delete the chunk-specific shaderlibrary files that are still present on disk from a previous incremental cook.
+	 * They will be rewritten by SaveShaderLibraryChunk calls.
+	 */
+	static RENDERCORE_API void DeleteShaderLibraryChunks(const ITargetPlatform* TargetPlatform,
+		const FString& SandboxDestinationPath, const FString& SandboxMetadataPath);
+
 	// Dump collected stats for each shader platform
 	UE_DEPRECATED(5.6, "DumpShaderCodeStats is deprecated due to being incorrect w.r.t shader duplication due to chunking.")
 	static RENDERCORE_API void DumpShaderCodeStats();
