[ITF Drag Tools] Alt + Duplicate - with Surface snapping incorrect visual feedback.

- Adds API to `FLevelEditorViewportClient` for duplicating selection, which differentiates between selecting duplicates or not.
- Adds a reset of `PreDragActorTransforms` to duplication logic, which fixes issues seen with "stamped" duplicates.

diff --git a/Engine/Source/Editor/Experimental/EditorInteractiveToolsFramework/Private/EditorDragTools/DuplicateDragSelectionInteraction.cpp b/Engine/Source/Editor/Experimental/EditorInteractiveToolsFramework/Private/EditorDragTools/DuplicateDragSelectionInteraction.cpp
--- a/Engine/Source/Editor/Experimental/EditorInteractiveToolsFramework/Private/EditorDragTools/DuplicateDragSelectionInteraction.cpp
+++ b/Engine/Source/Editor/Experimental/EditorInteractiveToolsFramework/Private/EditorDragTools/DuplicateDragSelectionInteraction.cpp
@@ -163,7 +163,7 @@ void UDuplicateDragSelectionInteraction::OnGizmoMovementEnd(UTransformProxy* InT
 	DuplicateTransaction.Reset();
 }
 
-void UDuplicateDragSelectionInteraction::DuplicateSelection(bool bInSelectNewElements) const
+void UDuplicateDragSelectionInteraction::DuplicateSelection(bool bInSelectNewElements)
 {
 	// "Active" Component visualizers process inputs, so they might want to handle Alt + Drag differently (e.g. spline points editing)
 	if (GUnrealEd->ComponentVisManager.IsActive())
@@ -173,46 +173,7 @@ void UDuplicateDragSelectionInteraction::DuplicateSelection(bool bInSelectNewEle
 
 	if (GCurrentLevelEditingViewportClient)
 	{
-		if (TSharedPtr<ILevelEditor> LevelEditor = GCurrentLevelEditingViewportClient->ParentLevelEditor.Pin())
-		{
-			if (UTypedElementCommonActions* CommonActions = LevelEditor->GetCommonActions())
-			{
-				FToolBuilderState State;
-				if (IAssetEditorContextInterface* AssetEditorContext =
-						ToolManager->GetContextObjectStore()->FindContext<IAssetEditorContextInterface>())
-				{
-					if (UTypedElementSelectionSet* MutableSelectionSet = AssetEditorContext->GetMutableSelectionSet())
-					{
-						const TArray<FTypedElementHandle> DuplicatedElements = CommonActions->DuplicateSelectedElements(
-							MutableSelectionSet, LevelEditor->GetWorld(), FVector::ZeroVector
-						);
-						if (DuplicatedElements.Num() > 0)
-						{
-							// (assuming chosen modifier key is Alt)
-							// Alt Duplicate has 2 modes:
-							// 1 - with Alt pressed, dragging will create and select a copy, then drag it
-							// 2 - while dragging, hitting Alt will crate a copy, NOT selecting it
-
-							// Select newly created elements: Alt + Drag only
-							if (bInSelectNewElements)
-							{
-								// Select the newly created elements
-								MutableSelectionSet->SetSelection(DuplicatedElements, FTypedElementSelectionOptions());
-								MutableSelectionSet->NotifyPendingChanges();
-							}
-
-							// Notify the global mode tools, the selection set should be identical to the new actors at this point
-							TArray<AActor*> SelectedActors = MutableSelectionSet->GetSelectedObjects<AActor>();
-							constexpr bool bDidOffsetDuplicate = false;
-							GLevelEditorModeTools().ActorsDuplicatedNotify(SelectedActors, SelectedActors, bDidOffsetDuplicate);
-						}
-
-						// Invalidate all viewports, so the new gizmo is rendered in each one
-						GEditor->RedrawLevelEditingViewports();
-					}
-				}
-			}
-		}
+		GCurrentLevelEditingViewportClient->DuplicateSelection(bInSelectNewElements);
 	}
 }
 
diff --git a/Engine/Source/Editor/Experimental/EditorInteractiveToolsFramework/Public/EditorDragTools/DuplicateDragSelectionInteraction.h b/Engine/Source/Editor/Experimental/EditorInteractiveToolsFramework/Public/EditorDragTools/DuplicateDragSelectionInteraction.h
--- a/Engine/Source/Editor/Experimental/EditorInteractiveToolsFramework/Public/EditorDragTools/DuplicateDragSelectionInteraction.h
+++ b/Engine/Source/Editor/Experimental/EditorInteractiveToolsFramework/Public/EditorDragTools/DuplicateDragSelectionInteraction.h
@@ -31,6 +31,12 @@ public:
 	virtual ~UDuplicateDragSelectionInteraction() override;
 
 protected:
+	/**
+	 * Duplicate the current selection
+	 * @param bInSelectNewElements if newly created elements should be selected right after duplication
+	 */
+	static void DuplicateSelection(bool bInSelectNewElements);
+
 	void OnGizmoMovementBegin(UTransformProxy* InTransformProxy);
 	void OnGizmoMovementEnd(UTransformProxy* InTransformProxy);
 
@@ -42,12 +48,6 @@ protected:
 
 	virtual void Initialize(UViewportInteractionsBehaviorSource* const InViewportInteractionsBehaviorSource) override;
 
-	/**
-	 * Duplicate the current selection
-	 * @param bInSelectNewElements if newly created elements should be selected right after duplication
-	 */
-	void DuplicateSelection(bool bInSelectNewElements) const;
-
 private:
 	void OnModifierKeyChanged(bool bInIsDown);
 
diff --git a/Engine/Source/Editor/UnrealEd/Private/LevelEditorViewport.cpp b/Engine/Source/Editor/UnrealEd/Private/LevelEditorViewport.cpp
--- a/Engine/Source/Editor/UnrealEd/Private/LevelEditorViewport.cpp
+++ b/Engine/Source/Editor/UnrealEd/Private/LevelEditorViewport.cpp
@@ -3295,55 +3295,9 @@ bool FLevelEditorViewportClient::InputWidgetDelta(FViewport* InViewport, EAxisLi
 						if ( !Drag.IsNearlyZero() || !Rot.IsZero() )
 						{
 							bDuplicateOnNextDrag = false;
-							
-							TSharedPtr<ILevelEditor> LevelEditor = ParentLevelEditor.Pin();
-							UTypedElementCommonActions* CommonActions = LevelEditor ? LevelEditor->GetCommonActions() : nullptr;
-							if (CommonActions)
-							{
-								FTypedElementListConstRef ElementsToManipulate = GetElementsToManipulate();
-								TArray<FTypedElementHandle> DuplicatedElements;
-								{
-									// Do not used the cached manipulation list here here, as it will have removed attachments, and we do want to duplicate those
-									DuplicatedElements = CommonActions->DuplicateSelectedElements(GetSelectionSet(), GetWorld(), FVector::ZeroVector);
-								}
 
-								// Exclusively select the new elements, so that future gizmo interaction manipulates those items instead
-								if (DuplicatedElements.Num() > 0)
-								{
-									// We need to end the gizmo manipulation on the current selection, as they won't be moved again this drag...
-									if (bHasBegunGizmoManipulation)
-									{
-										ViewportInteraction->EndGizmoManipulation(ElementsToManipulate, GetWidgetMode(), ETypedElementViewportInteractionGizmoManipulationType::Click);
-										// Note: We don't reset bHasBegunGizmoManipulation here, as we test it again below to begin manipulating the duplicated selection
-									}
-
-									UTypedElementSelectionSet* SelectionSet = GetMutableSelectionSet();
-
-									const FTypedElementSelectionOptions SelectionOptions = FTypedElementSelectionOptions()
-										.SetAllowLegacyNotifications(false); // Old drag duplicate code didn't used to notify about this selection change
-
-									SelectionSet->SetSelection(DuplicatedElements, SelectionOptions);
-
-									// Selected actors are the same as the selection set at this point, send the legacy mode tools update
-									{
-										TArray<AActor*> SelectedActors = SelectionSet->GetSelectedObjects<AActor>();
-										constexpr bool bDidOffsetDuplicate = false;
-										ModeTools->ActorsDuplicatedNotify(SelectedActors, SelectedActors, bDidOffsetDuplicate);
-									}
-
-									// Force the cached manipulation list to update, as we don't notify about the change above
-									ResetElementsToManipulate();
-
-									// We need to start gizmo manipulation on the duplicated selection, as they may continue to move as part of this drag duplicate
-									if (bHasBegunGizmoManipulation)
-									{
-										ElementsToManipulate = GetElementsToManipulate();
-										ViewportInteraction->BeginGizmoManipulation(ElementsToManipulate, GetWidgetMode());
-									}
-								}
-
-								RedrawAllViewportsIntoThisScene();
-							}
+							constexpr bool bSelectNewElements = true;
+							DuplicateSelection(bSelectNewElements);
 						}
 					}
 				}
@@ -3652,28 +3606,8 @@ bool FLevelEditorViewportClient::InputKey(const FInputKeyEventArgs& InEventArgs)
 	// Stamp a duplicate at the current position during a drag when ALT is pressed and we're dragging with a transaction open, and there is no component visualizer
 	else if ( bAltIsBeingPressed && bDraggingByHandle && TrackingTransaction.IsActive() && GUnrealEd->ComponentVisManager.GetActiveComponentVis() == nullptr )
 	{
-		TSharedPtr<ILevelEditor> LevelEditor = ParentLevelEditor.Pin();
-		UTypedElementCommonActions* CommonActions = LevelEditor ? LevelEditor->GetCommonActions() : nullptr;
-		if (CommonActions)
-		{
-			FTypedElementListConstRef ElementsToManipulate = GetElementsToManipulate();
-			TArray<FTypedElementHandle> DuplicatedElements;
-			{
-				// Do not used the cached manipulation list here here, as it will have removed attachments, and we do want to duplicate those
-				DuplicatedElements = CommonActions->DuplicateSelectedElements(GetSelectionSet(), GetWorld(), FVector::ZeroVector);
-			}
-
-			// Do not select the duplicate actors so that the drag operation can continue freely
-			if (DuplicatedElements.Num() > 0)
-			{
-				// Although we did not change the selection, we still need to notify legacy mode tools of the duplication
-				TArray<AActor*> SelectedActors = GetSelectionSet()->GetSelectedObjects<AActor>();
-				constexpr bool bDidOffsetDuplicate = false;
-				ModeTools->ActorsDuplicatedNotify(SelectedActors, SelectedActors, bDidOffsetDuplicate);
-			}
-
-			RedrawAllViewportsIntoThisScene();
-		}
+		constexpr bool bSelectNewElements = false;
+		DuplicateSelection(bSelectNewElements);
 	}
 
 	
@@ -5951,5 +5885,71 @@ void FLevelEditorViewportClient::SetEditingThroughMovementWidget()
 	bCurrentlyEditingThroughMovementWidget = true;
 }
 
+void FLevelEditorViewportClient::DuplicateSelection(bool bInSelectNewElements)
+{
+	TSharedPtr<ILevelEditor> LevelEditor = ParentLevelEditor.Pin();
+	if (!LevelEditor)
+	{
+		return;
+	}
+
+	if (UTypedElementCommonActions* CommonActions = LevelEditor->GetCommonActions())
+	{
+		FToolBuilderState State;
+		if (UTypedElementSelectionSet* MutableSelectionSet = GetMutableSelectionSet())
+		{
+			// Do not use the cached manipulation list here, as it will have removed attachments, and we do want to duplicate those
+			const TArray<FTypedElementHandle> DuplicatedElements = CommonActions->DuplicateSelectedElements(
+				GetSelectionSet(), LevelEditor->GetWorld(), FVector::ZeroVector
+			);
+
+			if (DuplicatedElements.Num() > 0)
+			{
+				// (assuming chosen modifier key is Alt)
+				// Alt Duplicate has 2 modes:
+				// 1 - with Alt pressed, dragging will create and select a copy, then drag it
+				// 2 - while dragging, hitting Alt will crate a copy, NOT selecting it
+
+				// Select newly created elements: Alt + Drag only
+				if (bInSelectNewElements)
+				{
+					// We need to end the gizmo manipulation on the current selection, as they won't be moved again this drag...
+					if (bHasBegunGizmoManipulation)
+					{
+						ViewportInteraction->EndGizmoManipulation(GetElementsToManipulate(), GetWidgetMode(), ETypedElementViewportInteractionGizmoManipulationType::Click);
+						// Note: We don't reset bHasBegunGizmoManipulation here, as we test it again below to begin manipulating the duplicated selection
+					}
+
+					// Select the newly created elements
+					const FTypedElementSelectionOptions SelectionOptions = FTypedElementSelectionOptions()
+										.SetAllowLegacyNotifications(false); // Old drag duplicate code didn't used to notify about this selection change
+
+					MutableSelectionSet->SetSelection(DuplicatedElements, SelectionOptions);
+					MutableSelectionSet->NotifyPendingChanges();
+
+					// Force the cached manipulation list to update, as we don't notify about the change above
+					ResetElementsToManipulate();
+
+					// We need to start gizmo manipulation on the duplicated selection, as they may continue to move as part of this drag duplicate
+					if (bHasBegunGizmoManipulation)
+					{
+						ViewportInteraction->BeginGizmoManipulation(GetElementsToManipulate(), GetWidgetMode());
+					}
+				}
+
+				// Notify the global mode tools, the selection set should be identical to the new actors at this point
+				TArray<AActor*> SelectedActors = MutableSelectionSet->GetSelectedObjects<AActor>();
+				constexpr bool bDidOffsetDuplicate = false;
+				GLevelEditorModeTools().ActorsDuplicatedNotify(SelectedActors, SelectedActors, bDidOffsetDuplicate);
+			}
+
+			// Invalidate all viewports, so the new gizmo is rendered in each one
+			RedrawAllViewportsIntoThisScene();
+
+			// Prevent issues with dragged elements sometimes being set back to their starting position after stamping the duplicate
+			PreDragActorTransforms.Empty();
+		}
+	}
+}
 
 #undef LOCTEXT_NAMESPACE
diff --git a/Engine/Source/Editor/UnrealEd/Public/LevelEditorViewport.h b/Engine/Source/Editor/UnrealEd/Public/LevelEditorViewport.h
--- a/Engine/Source/Editor/UnrealEd/Public/LevelEditorViewport.h
+++ b/Engine/Source/Editor/UnrealEd/Public/LevelEditorViewport.h
@@ -690,6 +690,15 @@ public:
 	 */
 	UNREALED_API void SetEditingThroughMovementWidget();
 
+	/**
+	 * Duplicates the current Level Editor selection. Newly created elements can be selected after creation or not.
+	 * E.g. Alt + Gizmo Drag creates a copy and then selects it, while pressing Alt while Gizmo Drag is being performed
+	 * "stamps" a duplicate of the dragged selection, without selecting it.
+	 *
+	 * @param bInSelectNewElements whether the duplicates should be selected or not
+	 */
+	UNREALED_API void DuplicateSelection(bool bInSelectNewElements);
+
 protected:
 	/**
 	* Sets the state of creating a preview actor in the viewport.
