[ITF Viewport Interactions] Defer registering viewport interaction behaviors until the next tick after mode tools are changed. This fixes an issue where viewport interactions were registered before mode tools had their `Exit()` function called. It also ensures there is only one registration pass even if multiple modes were activated or deactivated within a single frame.

#rnx

diff --git a/Engine/Source/Editor/UnrealEd/Private/Tools/EdModeInteractiveToolsContext.cpp b/Engine/Source/Editor/UnrealEd/Private/Tools/EdModeInteractiveToolsContext.cpp
--- a/Engine/Source/Editor/UnrealEd/Private/Tools/EdModeInteractiveToolsContext.cpp
+++ b/Engine/Source/Editor/UnrealEd/Private/Tools/EdModeInteractiveToolsContext.cpp
@@ -909,6 +909,12 @@ void UModeManagerInteractiveToolsContext::Tick(FEditorViewportClient* ViewportCl
 		EdModeContext->Tick(ViewportClient, DeltaTime);
 	}
 
+	if (bViewportInteractionsDirty)
+	{
+		bViewportInteractionsDirty = false;
+		RegisterViewportInteractions();
+	}
+
 	if (ViewportInteractionsBehaviorSource)
 	{
 		ViewportInteractionsBehaviorSource->Tick(DeltaTime);
@@ -1508,7 +1514,7 @@ void UModeManagerInteractiveToolsContext::UnregisterViewportInteractions()
 
 void UModeManagerInteractiveToolsContext::OnEditorModeChanged(const FEditorModeID& InModeID, bool bInIsEnteringMode)
 {
-	RegisterViewportInteractions();
+	bViewportInteractionsDirty = true;
 }
 
 void UEdModeInteractiveToolsContext::InitializeContextFromModeManagerContext(UModeManagerInteractiveToolsContext* ModeManagerToolsContext)
diff --git a/Engine/Source/Editor/UnrealEd/Public/Tools/EdModeInteractiveToolsContext.h b/Engine/Source/Editor/UnrealEd/Public/Tools/EdModeInteractiveToolsContext.h
--- a/Engine/Source/Editor/UnrealEd/Public/Tools/EdModeInteractiveToolsContext.h
+++ b/Engine/Source/Editor/UnrealEd/Public/Tools/EdModeInteractiveToolsContext.h
@@ -481,6 +481,7 @@ protected:
 	PRAGMA_ENABLE_EXPERIMENTAL_WARNINGS
 
 	bool bUsesViewportInteractions = false;
+	bool bViewportInteractionsDirty = false;
 };
 
 
