[HLOD] Add ability to run HLOD setup and build steps directly from the editor for parts of the world within a selected filter volume
* The Setup HLOD Actors step is executed for all HLOD actors, but only those within the filter volume are kept.
* The Build HLOD Actors step is executed only for actors within the filter volume.
* New and modified actors are not automatically saved, but marked as dirty.
* Added a new option in the World Partition minimap context menu that allows building HLODs from a selection.
* Added a new option to the viewport and outliner context menus, allowing users to build HLODs for selected actor(s); the actor(s)' bounds are used as filter volume.

diff --git a/Engine/Plugins/Editor/WorldPartitionHLODUtilities/Source/Private/WorldPartition/HLOD/Utilities/WorldPartitionHLODUtilities.cpp b/Engine/Plugins/Editor/WorldPartitionHLODUtilities/Source/Private/WorldPartition/HLOD/Utilities/WorldPartitionHLODUtilities.cpp
--- a/Engine/Plugins/Editor/WorldPartitionHLODUtilities/Source/Private/WorldPartition/HLOD/Utilities/WorldPartitionHLODUtilities.cpp
+++ b/Engine/Plugins/Editor/WorldPartitionHLODUtilities/Source/Private/WorldPartition/HLOD/Utilities/WorldPartitionHLODUtilities.cpp
@@ -11,6 +11,7 @@
 #include "WorldPartition/DataLayer/DataLayerInstance.h"
 #include "WorldPartition/HLOD/CustomHLODActor.h"
 #include "WorldPartition/HLOD/HLODActor.h"
+#include "WorldPartition/HLOD/HLODCreationFilter.h"
 #include "WorldPartition/HLOD/HLODHashBuilder.h"
 #include "WorldPartition/HLOD/HLODLayer.h"
 #include "WorldPartition/HLOD/HLODModifier.h"
@@ -271,6 +272,26 @@ TArray<AWorldPartitionHLOD*> FWorldPartitionHLODUtilities::CreateHLODActors(FHLO
 			HLODActor = CastChecked<AWorldPartitionHLOD>(HLODActorHandle.GetActor());
 		}
 
+		// If there was no ActorDesc found, check if we can find this actor in UnsavedHLODActors
+		TWeakObjectPtr<AWorldPartitionHLOD> HLODActorPtr;
+		if (HLODActor == nullptr && InCreationContext.UnsavedHLODActors.RemoveAndCopyValue(HLODActorName, HLODActorPtr))
+		{
+			HLODActor = HLODActorPtr.Get();
+		}
+
+		FHLODCreationFilterContext FilterContext;
+		FilterContext.Bounds = Bounds;
+		const bool bPassesAllFilters = UE::HLOD::CreationFilter::PassesFilters(InCreationParams.Filters, FilterContext);
+		if (!bPassesAllFilters)
+		{
+			if (HLODActor)
+			{
+				// If HLOD actor doesn't pass any of the filters, don't create it or update it. But if it exists, add it to HLODActors, so that it can be used when processing next layer
+				HLODActors.Add(HLODActor);
+			}
+			continue;
+		}
+
 		bool bNewActor = HLODActor == nullptr;
 		if (bNewActor)
 		{
diff --git a/Engine/Source/Editor/LevelEditor/Private/LevelEditor.cpp b/Engine/Source/Editor/LevelEditor/Private/LevelEditor.cpp
--- a/Engine/Source/Editor/LevelEditor/Private/LevelEditor.cpp
+++ b/Engine/Source/Editor/LevelEditor/Private/LevelEditor.cpp
@@ -1622,6 +1622,10 @@ void FLevelEditorModule::BindGlobalLevelEditorCommands()
 	ActionList.MapAction(Commands.BuildHLODs,
 		FExecuteAction::CreateStatic(&FLevelEditorActionCallbacks::BuildHLODs_Execute),
 		FCanExecuteAction::CreateStatic(&FLevelEditorActionCallbacks::IsWorldPartitionStreamingEnabled));
+
+	ActionList.MapAction(Commands.BuildHLODs_SelectedActors,
+		FExecuteAction::CreateStatic(&FLevelEditorActionCallbacks::BuildHLODs_SelectedActors_Execute),
+		FCanExecuteAction::CreateStatic(&FLevelEditorActionCallbacks::IsWorldPartitionStreamingEnabled));
 	
 	ActionList.MapAction(Commands.BuildMinimap,
 		FExecuteAction::CreateStatic(&FLevelEditorActionCallbacks::BuildMinimap_Execute),
diff --git a/Engine/Source/Editor/LevelEditor/Private/LevelEditorActions.cpp b/Engine/Source/Editor/LevelEditor/Private/LevelEditorActions.cpp
--- a/Engine/Source/Editor/LevelEditor/Private/LevelEditorActions.cpp
+++ b/Engine/Source/Editor/LevelEditor/Private/LevelEditorActions.cpp
@@ -43,8 +43,10 @@
 #include "WorldPartition/WorldPartitionSubsystem.h"
 #include "WorldPartition/WorldPartition.h"
 #include "WorldPartition/WorldPartitionEditorHash.h"
+#include "WorldPartition/WorldPartitionHLODsBuilder.h"
 #include "WorldPartition/WorldPartitionRuntimeHash.h"
 #include "WorldPartition/IWorldPartitionEditorModule.h"
+#include "WorldPartition/HLOD/HLODCreationFilter.h"
 #include "WorldBrowserModule.h"
 #include "ExternalPackageHelper.h"
 
@@ -1031,6 +1033,39 @@ void FLevelEditorActionCallbacks::BuildHLODs_Execute()
 	FEditorBuildUtils::EditorBuild(GetWorld(), FBuildOptions::BuildHierarchicalLOD);
 }
 
+void FLevelEditorActionCallbacks::BuildHLODs_SelectedActors_Execute()
+{
+	TArray<FBox> FilterVolumes;
+	for (FSelectionIterator It = GEditor->GetSelectedActorIterator(); It; ++It)
+	{
+		if (AActor* Actor = Cast<AActor>(*It))
+		{
+			FBox RuntimeBounds;
+			FBox EditorBounds;
+			Actor->GetStreamingBounds(RuntimeBounds, EditorBounds);
+			FilterVolumes.Add(RuntimeBounds);
+		}
+	}
+
+	UWorldPartitionSubsystem* WorldPartitionSubsystem = UWorld::GetSubsystem<UWorldPartitionSubsystem>(GetWorld());
+	if (WorldPartitionSubsystem)
+	{
+		for (const FWorldPartitionHandle& SelectedActorHandle : WorldPartitionSubsystem->SelectedActorHandles)
+		{
+			FilterVolumes.Add(SelectedActorHandle->GetRuntimeBounds());
+		}
+	}
+
+	if (FilterVolumes.Num() > 0)
+	{
+		UWorldPartitionHLODsBuilder::BuildWithFilters(GetWorld(), { MakeShared<FHLODCreationVolumeFilter>(MoveTemp(FilterVolumes)) });
+	}
+	else
+	{
+		UE_LOG(LevelEditorActions, Warning, TEXT("Build HLODs - No actors selected"));
+	}
+}
+
 void FLevelEditorActionCallbacks::BuildMinimap_Execute()
 {
 	// Build Minimap
@@ -4208,6 +4243,7 @@ void FLevelEditorCommands::RegisterCommands()
 	UI_COMMAND( BuildGeometryOnly_OnlyCurrentLevel, "Build Geometry (Current Level)", "Builds geometry, only for the current level", EUserInterfaceActionType::Button, FInputChord() );
 	UI_COMMAND( BuildPathsOnly, "Build Paths", "Only builds paths (all levels.)", EUserInterfaceActionType::Button, FInputChord() );
 	UI_COMMAND( BuildHLODs, "Build HLODs", "Builds all HLODs for the current world", EUserInterfaceActionType::Button, FInputChord() );
+	UI_COMMAND( BuildHLODs_SelectedActors, "Build HLODs (Selected Actors)", "Builds HLODs that are affected by selected actors", EUserInterfaceActionType::Button, FInputChord() );
 	UI_COMMAND( BuildMinimap, "Build World Partition Editor Minimap", "Builds the minimap for the current world", EUserInterfaceActionType::Button, FInputChord() );
 	UI_COMMAND( BuildLandscapeSplineMeshes, "Build Landscape Spline Meshes", "Builds landscape spline meshes for the current world", EUserInterfaceActionType::Button, FInputChord());
 	UI_COMMAND( BuildTextureStreamingOnly, "Build Texture Streaming", "Build texture streaming data", EUserInterfaceActionType::Button, FInputChord() );
diff --git a/Engine/Source/Editor/LevelEditor/Private/LevelEditorMenu.cpp b/Engine/Source/Editor/LevelEditor/Private/LevelEditorMenu.cpp
--- a/Engine/Source/Editor/LevelEditor/Private/LevelEditorMenu.cpp
+++ b/Engine/Source/Editor/LevelEditor/Private/LevelEditorMenu.cpp
@@ -750,6 +750,7 @@ void FLevelEditorMenu::RegisterBuildMenu()
 			{
 				FToolMenuSection& Section = InMenu->AddSection("LevelEditorWorldPartition", LOCTEXT("WorldPartitionHeading", "World Partition"));
 				Section.AddMenuEntry(FLevelEditorCommands::Get().BuildHLODs);
+				Section.AddMenuEntry(FLevelEditorCommands::Get().BuildHLODs_SelectedActors);
 				Section.AddMenuEntry(FLevelEditorCommands::Get().BuildMinimap);
 				Section.AddMenuEntry(FLevelEditorCommands::Get().BuildLandscapeSplineMeshes);
 			}
diff --git a/Engine/Source/Editor/LevelEditor/Public/LevelEditorActions.h b/Engine/Source/Editor/LevelEditor/Public/LevelEditorActions.h
--- a/Engine/Source/Editor/LevelEditor/Public/LevelEditorActions.h
+++ b/Engine/Source/Editor/LevelEditor/Public/LevelEditorActions.h
@@ -89,6 +89,7 @@ public:
 	TSharedPtr< FUICommandInfo > BuildGeometryOnly_OnlyCurrentLevel;
 	TSharedPtr< FUICommandInfo > BuildPathsOnly;
 	TSharedPtr< FUICommandInfo > BuildHLODs;
+	TSharedPtr< FUICommandInfo > BuildHLODs_SelectedActors;
 	TSharedPtr< FUICommandInfo > BuildMinimap;
 	TSharedPtr< FUICommandInfo > BuildLandscapeSplineMeshes;
 	TSharedPtr< FUICommandInfo > BuildTextureStreamingOnly;
@@ -853,6 +854,7 @@ public:
 	static UE_API bool IsWorldPartitionEnabled();
 	static UE_API bool IsWorldPartitionStreamingEnabled();
 	static UE_API void BuildHLODs_Execute();
+	static UE_API void BuildHLODs_SelectedActors_Execute();
 	static UE_API void BuildMinimap_Execute();
 	static UE_API void BuildLandscapeSplineMeshes_Execute();
 	static UE_API void BuildTextureStreamingOnly_Execute();
diff --git a/Engine/Source/Editor/UnrealEd/Private/EditorBuildUtils.cpp b/Engine/Source/Editor/UnrealEd/Private/EditorBuildUtils.cpp
--- a/Engine/Source/Editor/UnrealEd/Private/EditorBuildUtils.cpp
+++ b/Engine/Source/Editor/UnrealEd/Private/EditorBuildUtils.cpp
@@ -81,6 +81,7 @@ const FName FBuildOptions::BuildAll(TEXT("BuildAll"));
 const FName FBuildOptions::BuildAllSubmit(TEXT("BuildAllSubmit"));
 const FName FBuildOptions::BuildAllOnlySelectedPaths(TEXT("BuildAllOnlySelectedPaths"));
 const FName FBuildOptions::BuildHierarchicalLOD(TEXT("BuildHierarchicalLOD"));
+const FName FBuildOptions::BuildHierarchicalLOD_SelectedActors(TEXT("BuildHierarchicalLOD_SelectedActors"));
 const FName FBuildOptions::BuildMinimap(TEXT("BuildMinimap"));
 const FName FBuildOptions::BuildLandscapeSplineMeshes(TEXT("BuildLandscapeSplineMeshes"));
 const FName FBuildOptions::BuildTextureStreaming(TEXT("BuildTextureStreaming"));
diff --git a/Engine/Source/Editor/UnrealEd/Private/WorldPartition/WorldPartitionHLODsBuilder.cpp b/Engine/Source/Editor/UnrealEd/Private/WorldPartition/WorldPartitionHLODsBuilder.cpp
--- a/Engine/Source/Editor/UnrealEd/Private/WorldPartition/WorldPartitionHLODsBuilder.cpp
+++ b/Engine/Source/Editor/UnrealEd/Private/WorldPartition/WorldPartitionHLODsBuilder.cpp
@@ -10,10 +10,13 @@
 #include "Misc/ConfigCacheIni.h"
 #include "Misc/EngineVersion.h"
 #include "Misc/FileHelper.h"
+#include "Misc/MessageDialog.h"
 #include "Misc/OutputDeviceNull.h"
 #include "Misc/BufferedOutputDevice.h"
+#include "Misc/ScopedSlowTask.h"
 #include "Algo/ForEach.h"
 #include "UObject/Linker.h"
+#include "UObject/GCObjectScopeGuard.h"
 #include "UObject/SavePackage.h"
 
 #include "ActorFolder.h"
@@ -39,9 +42,11 @@
 #include "WorldPartition/WorldPartitionHelpers.h"
 #include "WorldPartition/HLOD/HLODActor.h"
 #include "WorldPartition/HLOD/HLODActorDesc.h"
+#include "WorldPartition/HLOD/HLODCreationFilter.h"
 #include "WorldPartition/HLOD/HLODSourceActors.h"
 #include "WorldPartition/HLOD/HLODLayer.h"
 #include "WorldPartition/HLOD/HLODProviderInterface.h"
+#include "WorldPartition/HLOD/HLODSourceActorsFromCell.h"
 #include "WorldPartition/HLOD/IWorldPartitionHLODUtilities.h"
 #include "WorldPartition/HLOD/IWorldPartitionHLODUtilitiesModule.h"
 #include "WorldPartition/HLOD/StandaloneHLODSubsystem.h"
@@ -63,6 +68,23 @@ static const FString BuildProductsFileName = TEXT("BuildProducts.txt");
 FString GetHLODBuilderFolderName(uint32 BuilderIndex) { return FString::Printf(TEXT("HLODBuilder%d"), BuilderIndex); }
 FString GetToSubmitFolderName() { return TEXT("ToSubmit"); }
 
+void CacheUnsavedHLODActors(UWorld* InWorld, TMap<FGuid, TWeakObjectPtr<AWorldPartitionHLOD> >& OutUnsavedHLODActors)
+{
+	if (InWorld)
+	{
+		for (TActorIterator<AWorldPartitionHLOD> It(InWorld); It; ++It)
+		{
+			if (UPackage* Package = It->GetPackage())
+			{
+				if (Package->IsDirty())
+				{
+					OutUnsavedHLODActors.Add(It->GetActorGuid(), *It);
+				}
+			}
+		}
+	}
+}
+
 namespace BuildEvaluationStats
 {
 	struct FStats
@@ -128,6 +150,9 @@ UWorldPartitionHLODsBuilder::UWorldPartitionHLODsBuilder(const FObjectInitialize
 	, BuilderIdx(INDEX_NONE)
 	, BuilderCount(INDEX_NONE)
 	, bBuildingStandaloneHLOD(false)
+	, bUseSlowTask(false)
+	, bSaveActors(true)
+	, bConsiderUnsavedHLODActors(false)
 {
 	if (!IsTemplate())
 	{
@@ -167,6 +192,57 @@ UWorldPartitionHLODsBuilder::UWorldPartitionHLODsBuilder(const FObjectInitialize
 	}
 }
 
+bool UWorldPartitionHLODsBuilder::BuildWithFilters(UWorld* InWorld, const TArray<TSharedPtr<IHLODCreationFilter>>& InFilters)
+{
+	bool bResult = false;
+
+	bool bUnsavedActorsPassingFilters = false;
+	for (TActorIterator<AActor> It(InWorld); It; ++It)
+	{
+		AActor* Actor = *It;
+		if (Actor && !Actor->IsA<AWorldPartitionHLOD>() && Actor->GetPackage()->IsDirty())
+		{
+			FBox RuntimeBounds;
+			FBox EditorBounds;
+			Actor->GetStreamingBounds(RuntimeBounds, EditorBounds);
+
+			FHLODCreationFilterContext FilterContext;
+			FilterContext.Bounds = RuntimeBounds;
+			if (UE::HLOD::CreationFilter::PassesFilters(InFilters, FilterContext))
+			{
+				bUnsavedActorsPassingFilters = true;
+				break;
+			}
+		}
+	}
+
+	if (bUnsavedActorsPassingFilters)
+	{
+		EAppReturnType::Type Response = FMessageDialog::Open(EAppMsgType::YesNo,
+			LOCTEXT("HLODBuildUnsavedActorsMessage", "There are unsaved actors that won't be included in this HLOD build. Do you want to continue?"));
+
+		if (Response == EAppReturnType::No)
+		{
+			return false;
+		}
+	}
+
+	UWorldPartitionHLODsBuilder* Builder = NewObject<UWorldPartitionHLODsBuilder>(GetTransientPackage(), UWorldPartitionHLODsBuilder::StaticClass());
+	Builder->SetUseSlowTask(true);
+	Builder->SetSaveActors(false);
+	Builder->SetConsiderUnsavedHLODActors(true);
+	Builder->SetFilters(InFilters);
+
+	{
+		FGCObjectScopeGuard BuilderGuard(Builder);
+		bResult = Builder->RunBuilder(InWorld);
+	}
+
+	GEditor->RedrawLevelEditingViewports();
+
+	return bResult;
+}
+
 void UWorldPartitionHLODsBuilder::AllowExternalDataLayerInjection(const UWorld* InWorld, const UExternalDataLayerAsset* InExternalDataLayerAsset, bool& bOutAllowInjection)
 {
 	// Always allow EDL injections during HLOD builds
@@ -326,13 +402,19 @@ bool UWorldPartitionHLODsBuilder::RunInternal(UWorld* InWorld, const FCellInfo&
 	World = InWorld;
 	WorldPartition = World->GetWorldPartition();
 
+	FScopedSlowTask SlowTask(2.f, LOCTEXT("BuildingHLODs", "Building HLODs..."), bUseSlowTask);
+	SlowTask.MakeDialog(true);
+
 	if (WorldPartition != nullptr)
 	{
 		bBuildingStandaloneHLOD = WorldPartition->HasStandaloneHLOD();
 	}
 	
-	// Allows HLOD Streaming levels to be GCed properly
-	FLevelStreamingGCHelper::EnableForCommandlet();
+	if (IsRunningCommandlet())
+	{
+		// Allows HLOD Streaming levels to be GCed properly
+		FLevelStreamingGCHelper::EnableForCommandlet();
+	}
 
 	SourceControlHelper = new FSourceControlHelper(PackageHelper, ModifiedFiles);
 
@@ -340,6 +422,7 @@ bool UWorldPartitionHLODsBuilder::RunInternal(UWorld* InWorld, const FCellInfo&
 
 	if (bRet && ShouldRunStep(EHLODBuildStep::HLOD_Setup))
 	{
+		SlowTask.EnterProgressFrame(1.f, LOCTEXT("SetupHLODActors", "Setup HLOD Actors"));
 		bRet = SetupHLODActors();
 	}
 
@@ -347,6 +430,7 @@ bool UWorldPartitionHLODsBuilder::RunInternal(UWorld* InWorld, const FCellInfo&
 	{
 		if (bRet && ShouldRunStep(EHLODBuildStep::HLOD_Build))
 		{
+			SlowTask.EnterProgressFrame(1.f, LOCTEXT("BuildHLODActors", "Build HLOD Actors"));
 			bRet = BuildHLODActors();
 		}
 
@@ -396,7 +480,10 @@ bool UWorldPartitionHLODsBuilder::SetupHLODActors()
 
 		UWorldPartition::FSetupHLODActorsParams SetupHLODActorsParams = UWorldPartition::FSetupHLODActorsParams()
 			.SetSourceControlHelper(SourceControlHelper)
-			.SetReportOnly(bReportOnly);
+			.SetReportOnly(bReportOnly)
+			.SetConsiderUnsavedHLODActors(bConsiderUnsavedHLODActors)
+			.SetSaveActors(bSaveActors)
+			.SetFilters(Filters);
 
 		WorldPartition->SetupHLODActors(SetupHLODActorsParams);
 
@@ -440,7 +527,7 @@ bool UWorldPartitionHLODsBuilder::SetupHLODActors()
 		UE_LOG(LogWorldPartitionHLODsBuilder, Display, TEXT("#### World HLOD actors ####"));
 
 		int32 NumActors = 0;
-		TFunction<void(UWorldPartition*)> ListHLODActors = [&NumActors](UWorldPartition* WorldPartitionToProcess)
+		TFunction<void(UWorldPartition*)> ListHLODActors = [&NumActors, bConsiderUnsavedHLODActors = bConsiderUnsavedHLODActors](UWorldPartition* WorldPartitionToProcess)
 		{
 			for (FActorDescContainerInstanceCollection::TIterator<AWorldPartitionHLOD> HLODIterator(WorldPartitionToProcess); HLODIterator; ++HLODIterator)
 			{
@@ -451,6 +538,21 @@ bool UWorldPartitionHLODsBuilder::SetupHLODActors()
 
 				NumActors++;
 			}
+
+			if (bConsiderUnsavedHLODActors)
+			{
+				if (UWorld* World = WorldPartitionToProcess->GetWorld())
+				{
+					for (TActorIterator<AWorldPartitionHLOD> It(World); It; ++It)
+					{
+						if (!WorldPartitionToProcess->GetActorDescInstance(It->GetActorGuid()))
+						{
+							UE_LOG(LogWorldPartitionHLODsBuilder, Display, TEXT("    [%d] %s"), NumActors, *It->GetPackage()->GetName());
+							NumActors++;
+						}
+					}
+				}
+			}
 		};
 
 		ListHLODActors(WorldPartition);
@@ -615,16 +717,19 @@ bool UWorldPartitionHLODsBuilder::BuildHLODActors()
 
 	auto SaveHLODActor = [this](AWorldPartitionHLOD* HLODActor)
 	{
-		UPackage* ActorPackage = HLODActor->GetPackage();
-		if (ActorPackage->IsDirty())
+		if (bSaveActors)
 		{
-			UE_LOG(LogWorldPartitionHLODsBuilder, Display, TEXT("HLOD actor %s was modified, saving..."), *HLODActor->GetActorLabel());
-
-			bool bSaved = SourceControlHelper->Save(ActorPackage);
-			if (!bSaved)
+			UPackage* ActorPackage = HLODActor->GetPackage();
+			if (ActorPackage->IsDirty())
 			{
-				UE_LOG(LogWorldPartitionHLODsBuilder, Error, TEXT("Failed to save %s, exiting..."), *USourceControlHelpers::PackageFilename(ActorPackage));
-				return false;
+				UE_LOG(LogWorldPartitionHLODsBuilder, Display, TEXT("HLOD actor %s was modified, saving..."), *HLODActor->GetActorLabel());
+
+				bool bSaved = SourceControlHelper->Save(ActorPackage);
+				if (!bSaved)
+				{
+					UE_LOG(LogWorldPartitionHLODsBuilder, Error, TEXT("Failed to save %s, exiting..."), *USourceControlHelpers::PackageFilename(ActorPackage));
+					return false;
+				}
 			}
 		}
 
@@ -646,14 +751,30 @@ bool UWorldPartitionHLODsBuilder::BuildHLODActors()
 			return false;
 		}
 
+		FScopedSlowTask SlowTask(static_cast<float>(HLODActorsToBuild.Num() - ResumeBuildIndex), LOCTEXT("BuildingHLODActors", "Building HLOD Actors"), bUseSlowTask);
+
 		UE_LOG(LogWorldPartitionHLODsBuilder, Display, TEXT("#### Building %d HLOD actors ####"), HLODActorsToBuild.Num());
 		if (bResumeBuild)
 		{
 			UE_LOG(LogWorldPartitionHLODsBuilder, Display, TEXT("#### Resuming build at %d ####"), ResumeBuildIndex);
 		}
 
+		// Cache all unsaved HLOD actors
+		TMap<FGuid, TWeakObjectPtr<AWorldPartitionHLOD>> UnsavedHLODActors;
+		if (bConsiderUnsavedHLODActors)
+		{
+			CacheUnsavedHLODActors(WorldPartition->GetWorld(), UnsavedHLODActors);
+		}
+
+		int32 BuiltHLODActors = 0;
 		for (int32 CurrentActor = ResumeBuildIndex; CurrentActor < HLODActorsToBuild.Num(); ++CurrentActor)
 		{
+			if (SlowTask.ShouldCancel())
+			{
+				UE_LOG(LogWorldPartitionHLODsBuilder, Display, TEXT("#### HLOD build cancelled ####"));
+				break;
+			}
+
 			TRACE_BOOKMARK(TEXT("BuildHLOD Start - %d"), CurrentActor);
 
 			UPackage* HLODActorPackage = nullptr;
@@ -661,14 +782,32 @@ bool UWorldPartitionHLODsBuilder::BuildHLODActors()
 				const FGuid& HLODActorGuid = HLODActorsToBuild[CurrentActor];
 
 				FWorldPartitionReference ActorRef(WorldPartition, HLODActorGuid);
+				AWorldPartitionHLOD* HLODActor = nullptr;
+				if (ActorRef.IsValid())
+				{ 
+					HLODActor = CastChecked<AWorldPartitionHLOD>(ActorRef.GetActor());
+				}
+				else if (bConsiderUnsavedHLODActors)
+				{
+					if (TWeakObjectPtr<AWorldPartitionHLOD>* UnsavedHLODActorPtr = UnsavedHLODActors.Find(HLODActorGuid))
+					{
+						HLODActor = UnsavedHLODActorPtr->Get();
+					}
+
+				}
+				check(HLODActor);
+
+				SlowTask.EnterProgressFrame(1.f, FText::Format(LOCTEXT("BuildingHLODActor", "Building HLOD actor: {0}"), FText::FromString(HLODActor->GetActorLabel())));
 
-				AWorldPartitionHLOD* HLODActor = CastChecked<AWorldPartitionHLOD>(ActorRef.GetActor());
 				HLODActorPackage = HLODActor->GetPackage();
 
 				UE_LOG(LogWorldPartitionHLODsBuilder, Display, TEXT("[%d / %d] %s %s..."), CurrentActor + 1, HLODActorsToBuild.Num(), *LOCTEXT("BuildingHLODActor", "Building HLOD actor").ToString(), *HLODActor->GetActorLabel());
 
-				// Simulate an engine tick to make sure engine & render resources that are queued for deletion are processed.
-				FWorldPartitionHelpers::FakeEngineTick(World);
+				if (IsRunningCommandlet())
+				{
+					// Simulate an engine tick to make sure engine & render resources that are queued for deletion are processed.
+					FWorldPartitionHelpers::FakeEngineTick(World);
+				}
 
 				HLODActor->BuildHLOD(bForceBuild);
 
@@ -693,9 +832,11 @@ bool UWorldPartitionHLODsBuilder::BuildHLODActors()
 			{
 				FWorldPartitionHelpers::DoCollectGarbage();
 			}
+
+			BuiltHLODActors++;
 		}
 
-		UE_LOG(LogWorldPartitionHLODsBuilder, Display, TEXT("#### Built %d HLOD actors ####"), HLODActorsToBuild.Num());
+		UE_LOG(LogWorldPartitionHLODsBuilder, Display, TEXT("#### Built %d HLOD actors ####"), BuiltHLODActors);
 	}
 	else
 	{
@@ -955,42 +1096,101 @@ TArray<UWorldPartitionHLODsBuilder::FHLODWorkload> UWorldPartitionHLODsBuilder::
 
 	// Build a mapping of 1 HLOD[Level] -> N HLOD[Level - 1]
 	TMap<FGuid, TArray<FGuid>>	HLODParenting;
-	TFunction<void(UWorldPartition*, uint32)> ProcessWorldPartition = [this, &HLODParenting, &HLODActorToWorldPartitionIndex, bShouldConsiderExternalHLODActors](UWorldPartition* WorldPartitionToProcess, uint32 WorldPartitionIndex)
+
+	TMap<FGuid, TWeakObjectPtr<AWorldPartitionHLOD> > UnsavedHLODActors;
+	if (bConsiderUnsavedHLODActors)
 	{
-		for (FActorDescContainerInstanceCollection::TIterator<AWorldPartitionHLOD> HLODIterator(WorldPartitionToProcess); HLODIterator; ++HLODIterator)
+		CacheUnsavedHLODActors(WorldPartition->GetWorld(), UnsavedHLODActors);
+	}
+
+	TFunction<void(const FHLODActorDesc&, const FGuid&, uint32)> ProcessHLODActor = [this, &HLODParenting, &HLODActorToWorldPartitionIndex, bShouldConsiderExternalHLODActors](const FHLODActorDesc& HLODActorDesc, const FGuid& ActorGuid, uint32 WorldPartitionIndex)
+	{
+		// Filter by HLOD actor
+		if (!HLODActorToBuild.IsNone() && HLODActorDesc.GetActorLabel() != HLODActorToBuild)
 		{
-			const FHLODActorDesc& HLODActorDesc = *(FHLODActorDesc*)HLODIterator->GetActorDesc();
-			// Filter by HLOD actor
-			if (!HLODActorToBuild.IsNone() && HLODActorDesc.GetActorLabel() != HLODActorToBuild)
-			{
-				continue;
-			}
+			return;
+		}
 
-			// Filter by HLOD layer
-			if (!HLODLayerToBuild.IsNone() && HLODActorDesc.GetSourceHLODLayer().GetAssetName() != HLODLayerToBuild)
-			{
-				continue;
-			}
+		// Filter by HLOD layer
+		if (!HLODLayerToBuild.IsNone() && HLODActorDesc.GetSourceHLODLayer().GetAssetName() != HLODLayerToBuild)
+		{
+			return;
+		}
+		
+		FHLODCreationFilterContext FilterContext;
+		FilterContext.Bounds = HLODActorDesc.GetRuntimeBounds();
+		const bool bPassesAllFilters = UE::HLOD::CreationFilter::PassesFilters(Filters, FilterContext);
 
-			if (bBuildingStandaloneHLOD && IsDistributedBuild())
+		if (!bPassesAllFilters)
+		{
+			return;
+		}
+
+		if (bBuildingStandaloneHLOD && IsDistributedBuild())
+		{
+			HLODActorToWorldPartitionIndex.Add(ActorGuid, WorldPartitionIndex);
+		}
+
+		// When requested to build a single HLOD Layer, skip the child actors
+		if (HLODLayerToBuild.IsNone())
+		{
+			TArray<FGuid>& ChildActors = HLODParenting.Add(ActorGuid, HLODActorDesc.GetChildHLODActors());
+
+			if (bShouldConsiderExternalHLODActors)
 			{
-				HLODActorToWorldPartitionIndex.Add(HLODIterator->GetGuid(), WorldPartitionIndex);
+				ChildActors.Append(HLODActorDesc.GetExternalChildHLODActors());
 			}
+		}
+		else
+		{
+			HLODParenting.Add(ActorGuid);
+		}
+	};
 
-			// When requested to build a single HLOD Layer, skip the child actors
-			if (HLODLayerToBuild.IsNone())
+	TFunction<void(UWorldPartition*, uint32)> ProcessWorldPartition = [this, &HLODParenting, &HLODActorToWorldPartitionIndex, &UnsavedHLODActors, bShouldConsiderExternalHLODActors, &ProcessHLODActor](UWorldPartition* WorldPartitionToProcess, uint32 WorldPartitionIndex)
+	{
+		// If considering unsaved HLOD Actors, process them first to use current state for filtering
+		// For remaining actors, use their saved ActorDesc
+		TSet<FGuid> ProcessedLoadedActors;
+		if (bConsiderUnsavedHLODActors)
+		{
+			for (const TPair<FGuid, TWeakObjectPtr<AWorldPartitionHLOD>>& Pair : UnsavedHLODActors)
 			{
-				TArray<FGuid>& ChildActors = HLODParenting.Add(HLODIterator->GetGuid(), HLODActorDesc.GetChildHLODActors());
-
-				if (bShouldConsiderExternalHLODActors)
+				if (AWorldPartitionHLOD* HLODActor = Cast<AWorldPartitionHLOD>(Pair.Value))
 				{
-					ChildActors.Append(HLODActorDesc.GetExternalChildHLODActors());
+					ProcessedLoadedActors.Add(HLODActor->GetActorGuid());
+					TUniquePtr<FWorldPartitionActorDesc> HLODActorDescPtr = HLODActor->CreateActorDesc();
+					const FHLODActorDesc& HLODActorDesc = *(FHLODActorDesc*)HLODActorDescPtr.Get();
+					ProcessHLODActor(HLODActorDesc, HLODActor->GetActorGuid(), WorldPartitionIndex);
+
+					// Unsaved HLOD actors might use other unsaved HLOD actors as ChildActors. If a ChildActor was never saved, it won't appear in ChildActors list generated by CreateActorDesc. Add them now.
+					if (UWorldPartitionHLODSourceActorsFromCell* SourceActors = Cast<UWorldPartitionHLODSourceActorsFromCell>(HLODActor->GetSourceActors()))
+					{
+						if (TArray<FGuid>* ChildActorsPtr = HLODParenting.Find(HLODActor->GetActorGuid()))
+						{
+							TArray<FGuid>& ChildActors = *ChildActorsPtr;
+							for (const FWorldPartitionRuntimeCellObjectMapping& SourceActor : SourceActors->GetActors())
+							{
+								if (UnsavedHLODActors.Contains(SourceActor.ActorInstanceGuid) && !ChildActors.Contains(SourceActor.ActorInstanceGuid))
+								{
+									ChildActors.Add(SourceActor.ActorInstanceGuid);
+								}
+							}
+						}
+					}
 				}
 			}
-			else
+		}
+
+		for (FActorDescContainerInstanceCollection::TIterator<AWorldPartitionHLOD> HLODIterator(WorldPartitionToProcess); HLODIterator; ++HLODIterator)
+		{
+			if (bConsiderUnsavedHLODActors && ProcessedLoadedActors.Contains(HLODIterator->GetGuid()))
 			{
-				HLODParenting.Add(HLODIterator->GetGuid());
+				continue;
 			}
+
+			const FHLODActorDesc& HLODActorDesc = *(FHLODActorDesc*)HLODIterator->GetActorDesc();
+			ProcessHLODActor(HLODActorDesc, HLODIterator->GetGuid(), WorldPartitionIndex);
 		}
 	};
 
@@ -1114,25 +1314,48 @@ bool UWorldPartitionHLODsBuilder::ValidateWorkload(const FHLODWorkload& Workload
 	for (int32 WorldIndex = 0; WorldIndex < Workload.PerWorldHLODWorkloads.Num(); WorldIndex++)
 	{
 		UWorldPartition* CurrentWorldPartition = (bBuildingStandaloneHLOD && IsDistributedBuild()) ? AdditionalWorldPartitionsForStandaloneHLOD[WorldIndex].Get() : WorldPartition;
+
+		TMap<FGuid, TWeakObjectPtr<AWorldPartitionHLOD>> UnsavedHLODActors;
+		if (bConsiderUnsavedHLODActors)
+		{
+			CacheUnsavedHLODActors(CurrentWorldPartition->GetWorld(), UnsavedHLODActors);
+		}
+
 		for (const FGuid& HLODActorGuid : Workload.PerWorldHLODWorkloads[WorldIndex])
 		{
 			const FWorldPartitionActorDescInstance* ActorDescInstance = CurrentWorldPartition->GetActorDescInstance(HLODActorGuid);
-			if(!ActorDescInstance)
+			bool bUsingUnsavedHLODActor = false;
+			TUniquePtr<FWorldPartitionActorDesc> UnsavedHLODActorDesc;
+			if (bConsiderUnsavedHLODActors && !ActorDescInstance)
 			{
-				UE_LOG(LogWorldPartitionHLODsBuilder, Error, TEXT("Unknown actor guid found (\"%s\"), your HLOD actors are probably out of date. Run with -SetupHLODs to fix this. Exiting..."), *HLODActorGuid.ToString());
-				return false;
+				if (TWeakObjectPtr<AWorldPartitionHLOD>* UnsavedHLODActorPtr = UnsavedHLODActors.Find(HLODActorGuid))
+				{
+					UnsavedHLODActorDesc = (*UnsavedHLODActorPtr)->CreateActorDesc();
+					bUsingUnsavedHLODActor = true;
+				}
 			}
 
-			if (!ActorDescInstance->GetActorNativeClass()->IsChildOf<AWorldPartitionHLOD>())
+			if (!bUsingUnsavedHLODActor)
 			{
-				UE_LOG(LogWorldPartitionHLODsBuilder, Error, TEXT("Unexpected actor guid found in HLOD workload (\"%s\"), exiting..."), *HLODActorGuid.ToString());
-				return false;
+				if (!ActorDescInstance)
+				{
+					UE_LOG(LogWorldPartitionHLODsBuilder, Error, TEXT("Unknown actor guid found (\"%s\"), your HLOD actors are probably out of date. Run with -SetupHLODs to fix this. Exiting..."), *HLODActorGuid.ToString());
+					return false;
+				}
+
+				if (!ActorDescInstance->GetActorNativeClass()->IsChildOf<AWorldPartitionHLOD>())
+				{
+					UE_LOG(LogWorldPartitionHLODsBuilder, Error, TEXT("Unexpected actor guid found in HLOD workload (\"%s\"), exiting..."), *HLODActorGuid.ToString());
+					return false;
+				}
 			}
 
 			// When requested to build a single HLOD Layer, do not validate that child actors are included
-			if (HLODLayerToBuild.IsNone())
+			// Don't validate when using Volume Filter, as not all child actors have to pass the volume Volume Filter
+			const bool bHasAnyActiveFilters = UE::HLOD::CreationFilter::HasAnyActiveFilters(Filters);
+			if (HLODLayerToBuild.IsNone() && !bHasAnyActiveFilters)
 			{
-				const FHLODActorDesc* HLODActorDesc = static_cast<const FHLODActorDesc*>(ActorDescInstance->GetActorDesc());
+				const FHLODActorDesc* HLODActorDesc = static_cast<const FHLODActorDesc*>(bUsingUnsavedHLODActor ? UnsavedHLODActorDesc.Get() : ActorDescInstance->GetActorDesc());
 
 				for (const FGuid& ChildHLODActorGuid : HLODActorDesc->GetChildHLODActors())
 				{
diff --git a/Engine/Source/Editor/UnrealEd/Public/EditorBuildUtils.h b/Engine/Source/Editor/UnrealEd/Public/EditorBuildUtils.h
--- a/Engine/Source/Editor/UnrealEd/Public/EditorBuildUtils.h
+++ b/Engine/Source/Editor/UnrealEd/Public/EditorBuildUtils.h
@@ -34,6 +34,8 @@ struct FBuildOptions
 	UNREALED_API static const FName BuildAllOnlySelectedPaths;
 	/** Build HLODs */
 	UNREALED_API static const FName BuildHierarchicalLOD;
+	/** Build HLODs for selected actors */
+	UNREALED_API static const FName BuildHierarchicalLOD_SelectedActors;
 	/** Build minimap */
 	UNREALED_API static const FName BuildMinimap;
 	/** Build landscape spline meshes */
diff --git a/Engine/Source/Editor/UnrealEd/Public/WorldPartition/WorldPartitionHLODsBuilder.h b/Engine/Source/Editor/UnrealEd/Public/WorldPartition/WorldPartitionHLODsBuilder.h
--- a/Engine/Source/Editor/UnrealEd/Public/WorldPartition/WorldPartitionHLODsBuilder.h
+++ b/Engine/Source/Editor/UnrealEd/Public/WorldPartition/WorldPartitionHLODsBuilder.h
@@ -25,6 +25,13 @@ class UWorldPartitionHLODsBuilder : public UWorldPartitionBuilder
 {
 	GENERATED_UCLASS_BODY()
 public:
+	virtual void SetUseSlowTask(bool bInUseSlowTask) { bUseSlowTask = bInUseSlowTask; }
+	virtual void SetSaveActors(bool bInSaveActors) { bSaveActors = bInSaveActors; }
+	virtual void SetConsiderUnsavedHLODActors(bool bInConsiderUnsavedHLODActors) { bConsiderUnsavedHLODActors = bInConsiderUnsavedHLODActors; }
+	virtual void SetFilters(const TArray<TSharedPtr<IHLODCreationFilter>> InFilters) { Filters = InFilters; }
+
+	UNREALED_API static bool BuildWithFilters(UWorld* InWorld, const TArray<TSharedPtr<IHLODCreationFilter>>& InFilters);
+
 	// UWorldPartitionBuilder interface begin
 	UNREALED_API virtual bool RequiresCommandletRendering() const override;
 	virtual ELoadingMode GetLoadingMode() const override { return ELoadingMode::Custom; }
@@ -101,4 +108,8 @@ private:
 	bool bBuildingStandaloneHLOD;
 	TArray<TObjectPtr<UWorldPartition>> AdditionalWorldPartitionsForStandaloneHLOD;
 	TArray<FString> StandaloneHLODWorkingDirs;
+	bool bUseSlowTask;
+	bool bSaveActors;
+	bool bConsiderUnsavedHLODActors;
+	TArray<TSharedPtr<IHLODCreationFilter>> Filters;
 };
\ No newline at end of file
diff --git a/Engine/Source/Editor/WorldPartitionEditor/Private/WorldPartition/SWorldPartitionEditorGrid2D.cpp b/Engine/Source/Editor/WorldPartitionEditor/Private/WorldPartition/SWorldPartitionEditorGrid2D.cpp
--- a/Engine/Source/Editor/WorldPartitionEditor/Private/WorldPartition/SWorldPartitionEditorGrid2D.cpp
+++ b/Engine/Source/Editor/WorldPartitionEditor/Private/WorldPartition/SWorldPartitionEditorGrid2D.cpp
@@ -31,6 +31,7 @@
 #include "Framework/MultiBox/MultiBoxBuilder.h"
 #include "WorldBrowserModule.h"
 #include "WorldPartition/HLOD/HLODActor.h"
+#include "WorldPartition/HLOD/HLODCreationFilter.h"
 #include "WorldPartition/LoaderAdapter/LoaderAdapterShape.h"
 #include "WorldPartition/WorldPartition.h"
 #include "WorldPartition/WorldPartitionLog.h"
@@ -42,6 +43,7 @@
 #include "WorldPartition/WorldPartitionEditorPerProjectUserSettings.h"
 #include "WorldPartition/WorldPartitionEditorSettings.h"
 #include "WorldPartition/WorldPartitionHelpers.h"
+#include "WorldPartition/WorldPartitionHLODsBuilder.h"
 #include "WorldPartition/WorldPartitionMiniMap.h"
 #include "WorldPartition/WorldPartitionMiniMapHelper.h"
 #include "WorldPartition/WorldPartitionMiniMapVolume.h"
@@ -229,6 +231,7 @@ void SWorldPartitionEditorGrid2D::FEditorCommands::RegisterCommands()
 {
 	// Context Menu
 	UI_COMMAND(CreateRegionFromSelection, "Load Region From Selection", "Load region from selection.", EUserInterfaceActionType::Button, FInputChord());
+	UI_COMMAND(BuildHLODsForSelection, "Build HLODs for selection", "Build HLODs in selected area.", EUserInterfaceActionType::Button, FInputChord());
 	UI_COMMAND(ConvertSelectedRegionsToActors, "Convert Selected Region(s) To Actor(s)", "Convert the selected region(s) to actor(s).", EUserInterfaceActionType::Button, FInputChord());
 	UI_COMMAND(LoadSelectedRegions, "Load Selected Region(s)", "Load the selected region(s).", EUserInterfaceActionType::Button, FInputChord());
 	UI_COMMAND(UnloadSelectedRegions, "Unload Selected Region(s)", "Unload the selected region(s).", EUserInterfaceActionType::Button, FInputChord());
@@ -424,6 +427,7 @@ TSharedRef<SWidget> SWorldPartitionEditorGrid2D::SToolBar::GenerateBuildMenu() c
 	FLevelEditorModule& LevelEditorModule = FModuleManager::GetModuleChecked<FLevelEditorModule>(TEXT("LevelEditor"));
 	const FLevelEditorCommands& LevelEditorCommands = LevelEditorModule.GetLevelEditorCommands();
 	Section.AddMenuEntry(LevelEditorCommands.BuildHLODs);
+	Section.AddMenuEntry(LevelEditorCommands.BuildHLODs_SelectedActors);
 	Section.AddMenuEntry(LevelEditorCommands.BuildMinimap);
 	Section.AddMenuEntry(LevelEditorCommands.BuildLandscapeSplineMeshes);
 
@@ -604,6 +608,7 @@ void SWorldPartitionEditorGrid2D::BindCommands()
 	const FEditorCommands& Commands = FEditorCommands::Get();
 
 	CommandList->MapAction(Commands.CreateRegionFromSelection, FExecuteAction::CreateSP(this, &SWorldPartitionEditorGrid2D::CreateRegionFromSelection), FCanExecuteAction::CreateLambda(CanCreateRegionFromSelection));
+	CommandList->MapAction(Commands.BuildHLODsForSelection, FExecuteAction::CreateSP(this, &SWorldPartitionEditorGrid2D::BuildHLODsForSelection), FCanExecuteAction::CreateLambda(CanCreateRegionFromSelection));
 	CommandList->MapAction(Commands.LoadSelectedRegions, FExecuteAction::CreateSP(this, &SWorldPartitionEditorGrid2D::LoadSelectedRegions), FCanExecuteAction::CreateLambda(CanLoadSelectedRegions));
 	CommandList->MapAction(Commands.UnloadSelectedRegions, FExecuteAction::CreateSP(this, &SWorldPartitionEditorGrid2D::UnloadSelectedRegions), FCanExecuteAction::CreateLambda(CanUnloadSelectedRegions));
 	CommandList->MapAction(Commands.ConvertSelectedRegionsToActors, FExecuteAction::CreateSP(this, &SWorldPartitionEditorGrid2D::ConvertSelectedRegionsToActors), FCanExecuteAction::CreateLambda(CanConvertSelectedRegionsToActors));
@@ -690,6 +695,16 @@ void SWorldPartitionEditorGrid2D::CreateRegionFromSelection()
 	Refresh();
 }
 
+void SWorldPartitionEditorGrid2D::BuildHLODsForSelection()
+{
+	TArray<FBox> FilterVolumes = { FBox(FVector(SelectBoxGridSnapped.Min.X, SelectBoxGridSnapped.Min.Y, -HALF_WORLD_MAX), FVector(SelectBoxGridSnapped.Max.X, SelectBoxGridSnapped.Max.Y, HALF_WORLD_MAX)) };
+
+	UWorldPartitionHLODsBuilder::BuildWithFilters(GetWorld(), { MakeShared<FHLODCreationVolumeFilter>(MoveTemp(FilterVolumes)) });
+
+	ClearSelection();
+	Refresh();
+}
+
 void SWorldPartitionEditorGrid2D::LoadSelectedRegions()
 {
 	for (const TWeakInterfacePtr<IWorldPartitionActorLoaderInterface>& SelectedLoaderAdapter : SelectedLoaderInterfaces)
@@ -926,6 +941,7 @@ TSharedRef<SWidget> SWorldPartitionEditorGrid2D::GenerateContextualMenu() const
 		static const FName SectionSelectionName(TEXT("ContextMenu.Selection"));
 		FToolMenuSection& SectionSelection = ConxtextualMenu->AddSection(SectionSelectionName, LOCTEXT("WorldPartitionSelectionHeader", "Selection"));
 		SectionSelection.AddMenuEntry(Commands.CreateRegionFromSelection);
+		SectionSelection.AddMenuEntry(Commands.BuildHLODsForSelection);
 		
 		static const FName SectionRegionsName(TEXT("ContextMenu.Regions"));
 		FToolMenuSection& SectionRegions = ConxtextualMenu->AddSection(SectionRegionsName, LOCTEXT("WorldPartitionRegionsHeader", "Region(s)"));
diff --git a/Engine/Source/Editor/WorldPartitionEditor/Private/WorldPartitionEditorModule.cpp b/Engine/Source/Editor/WorldPartitionEditor/Private/WorldPartitionEditorModule.cpp
--- a/Engine/Source/Editor/WorldPartitionEditor/Private/WorldPartitionEditorModule.cpp
+++ b/Engine/Source/Editor/WorldPartitionEditor/Private/WorldPartitionEditorModule.cpp
@@ -21,6 +21,7 @@
 #include "WorldPartition/WorldPartitionConvertOptions.h"
 #include "WorldPartition/WorldPartitionEditorSettings.h"
 #include "WorldPartition/HLOD/HLODActor.h"
+#include "WorldPartition/HLOD/HLODCreationFilter.h"
 #include "WorldPartition/HLOD/SWorldPartitionBuildHLODsDialog.h"
 #include "WorldPartition/HLOD/HLODEditorSubsystem.h"
 #include "WorldPartition/WorldPartitionClassDescRegistry.h"
@@ -122,7 +123,7 @@ static bool CanLoadUnloadSelectedVolumes(TArray<TWeakObjectPtr<AActor>> Volumes,
 	return false;
 }
 
-static void CreateLevelViewportContextMenuEntries(FMenuBuilder& MenuBuilder, TArray<TWeakObjectPtr<AActor>> Volumes, FBox SelectionBox)
+static void CreateLevelViewportContextMenuEntries(FMenuBuilder& MenuBuilder, TArray<TWeakObjectPtr<AActor>> Volumes, FBox SelectionBox, TSharedPtr<FHLODCreationVolumeFilter> VolumeFilter)
 {
 	MenuBuilder.BeginSection("WorldPartition", LOCTEXT("WorldPartition", "World Partition"));
 	
@@ -179,8 +180,31 @@ static void CreateLevelViewportContextMenuEntries(FMenuBuilder& MenuBuilder, TAr
 			})
 		);
 
+
+
 		MenuBuilder.AddMenuEntry(LOCTEXT("LoadRegionFromSelection", "Load Region From Selection"), LOCTEXT("LoadRegionFromSelection_Tooltip", "Load region from selected actor(s) bounds"), FSlateIcon(), LoadRegion);
 	}
+
+	if (GCurrentLevelEditingViewportClient && VolumeFilter->IsActive())
+	{
+		TWeakObjectPtr<UWorld> World = GCurrentLevelEditingViewportClient->GetWorld();
+
+		FUIAction BuildHLODsFromSelection(
+			FExecuteAction::CreateLambda([World, VolumeFilter]()
+			{
+				if (World.IsValid())
+				{
+					UWorldPartitionHLODsBuilder::BuildWithFilters(World.Get(), { VolumeFilter });
+				}
+			}),
+			FCanExecuteAction::CreateLambda([World]()
+			{
+				return (World.IsValid());
+			})
+		);
+
+		MenuBuilder.AddMenuEntry(LOCTEXT("BuildHLODsFromSelection", "Build HLODs From Selection"), LOCTEXT("BuildHLODSFromSelection_Tooltip", "Build HLODs from selected actor(s) bounds"), FSlateIcon(), BuildHLODsFromSelection);
+	}
 	
 	MenuBuilder.EndSection();
 }
@@ -191,6 +215,7 @@ static TSharedRef<FExtender> OnExtendLevelEditorMenu(const TSharedRef<FUICommand
 
 	TArray<TWeakObjectPtr<AActor>> Volumes;
 	FBoxSphereBounds::Builder BoundsBuilder;
+	TArray<FBox> FilterVolumes;
 
 	for (AActor* Actor : SelectedActors)
 	{
@@ -202,15 +227,22 @@ static TSharedRef<FExtender> OnExtendLevelEditorMenu(const TSharedRef<FUICommand
 		FBoxSphereBounds ActorBounds;
 		Actor->GetActorBounds(false, ActorBounds.Origin, ActorBounds.BoxExtent);
 		BoundsBuilder += ActorBounds;
+
+		FBox RuntimeBounds;
+		FBox EditorBounds;
+		Actor->GetStreamingBounds(RuntimeBounds, EditorBounds);
+		FilterVolumes.Add(RuntimeBounds);
 	}
 
-	if (!Volumes.IsEmpty() || BoundsBuilder.IsValid())
+	TSharedPtr<FHLODCreationVolumeFilter> VolumeFilter = MakeShared<FHLODCreationVolumeFilter>(MoveTemp(FilterVolumes));
+
+	if (!Volumes.IsEmpty() || BoundsBuilder.IsValid() || VolumeFilter->IsActive())
 	{
 		Extender->AddMenuExtension(
 			"ActorTypeTools",
 			EExtensionHook::After,
 			nullptr,
-			FMenuExtensionDelegate::CreateStatic(&CreateLevelViewportContextMenuEntries, Volumes, FBoxSphereBounds(BoundsBuilder).GetBox()));
+			FMenuExtensionDelegate::CreateStatic(&CreateLevelViewportContextMenuEntries, Volumes, FBoxSphereBounds(BoundsBuilder).GetBox(), VolumeFilter));
 	}
 
 	return Extender;
diff --git a/Engine/Source/Editor/WorldPartitionEditor/Public/WorldPartition/SWorldPartitionEditorGrid2D.h b/Engine/Source/Editor/WorldPartitionEditor/Public/WorldPartition/SWorldPartitionEditorGrid2D.h
--- a/Engine/Source/Editor/WorldPartitionEditor/Public/WorldPartition/SWorldPartitionEditorGrid2D.h
+++ b/Engine/Source/Editor/WorldPartitionEditor/Public/WorldPartition/SWorldPartitionEditorGrid2D.h
@@ -26,6 +26,7 @@ protected:
 	
 		// Context Menu
 		TSharedPtr<FUICommandInfo> CreateRegionFromSelection;
+		TSharedPtr<FUICommandInfo> BuildHLODsForSelection;
 		TSharedPtr<FUICommandInfo> LoadSelectedRegions;
 		TSharedPtr<FUICommandInfo> UnloadSelectedRegions;
 		TSharedPtr<FUICommandInfo> UnloadHoveredRegion;
@@ -93,6 +94,7 @@ protected:
 	void BindCommands();
 
 	void CreateRegionFromSelection();
+	void BuildHLODsForSelection();
 	void LoadSelectedRegions();
 	void UnloadSelectedRegions();
 	void ConvertSelectedRegionsToActors();
diff --git a/Engine/Source/Runtime/Engine/Private/WorldPartition/HLOD/HLODCreationFilter.cpp b/Engine/Source/Runtime/Engine/Private/WorldPartition/HLOD/HLODCreationFilter.cpp
new file mode 100644
--- /dev/null
+++ b/Engine/Source/Runtime/Engine/Private/WorldPartition/HLOD/HLODCreationFilter.cpp
@@ -0,0 +1,49 @@
+// Copyright Epic Games, Inc. All Rights Reserved.
+
+#include "WorldPartition/HLOD/HLODCreationFilter.h"
+
+#include "Algo/AllOf.h"
+#include "Algo/AnyOf.h"
+
+namespace UE::HLOD::CreationFilter
+{
+	bool HasAnyActiveFilters(const TArray<TSharedPtr<IHLODCreationFilter>>& InFilters)
+	{
+		return Algo::AnyOf(InFilters, [](const TSharedPtr<IHLODCreationFilter>& Filter)
+		{
+			return (Filter.IsValid() && Filter->IsActive());
+		});
+	}
+
+	bool PassesFilters(const TArray<TSharedPtr<IHLODCreationFilter>>& InFilters, const FHLODCreationFilterContext& InFilterContext)
+	{
+		return Algo::AllOf(InFilters, [&InFilterContext](const TSharedPtr<IHLODCreationFilter>& Filter)
+		{
+			return (Filter.IsValid() == false || Filter->PassesFilter(InFilterContext));
+		});
+	}
+}
+
+FHLODCreationVolumeFilter::FHLODCreationVolumeFilter(TArray<FBox> && InFilterVolumes)
+	: FilterVolumes(MoveTemp(InFilterVolumes))
+{
+}
+
+bool FHLODCreationVolumeFilter::IsActive() const
+{
+	return !FilterVolumes.IsEmpty();
+}
+
+bool FHLODCreationVolumeFilter::PassesFilter(const FHLODCreationFilterContext& InFilterContext) const
+{
+	if (InFilterContext.Bounds.IsSet() == false)
+	{
+		return false;
+	}
+
+	const FBox& Bounds = InFilterContext.Bounds.GetValue();
+	return (IsActive() == false || Algo::AnyOf(FilterVolumes, [&Bounds](const FBox& FilterVolume)
+	{
+		return FilterVolume.Intersect(Bounds);
+	}));
+}
\ No newline at end of file
diff --git a/Engine/Source/Runtime/Engine/Private/WorldPartition/HLOD/HLODSourceActorsFromCell.cpp b/Engine/Source/Runtime/Engine/Private/WorldPartition/HLOD/HLODSourceActorsFromCell.cpp
--- a/Engine/Source/Runtime/Engine/Private/WorldPartition/HLOD/HLODSourceActorsFromCell.cpp
+++ b/Engine/Source/Runtime/Engine/Private/WorldPartition/HLOD/HLODSourceActorsFromCell.cpp
@@ -9,6 +9,7 @@
 #include "WorldPartition/WorldPartitionLevelStreamingDynamic.h"
 
 #include "Editor/UnrealEd/Public/EditorLevelUtils.h" // MoveActorsToLevel
+#include "EngineUtils.h"
 #include "UObject/Linker.h"
 #include "WorldPartition/ContentBundle/ContentBundlePaths.h"
 #include "WorldPartition/DataLayer/ExternalDataLayerHelper.h"
@@ -60,6 +61,63 @@ bool UWorldPartitionHLODSourceActorsFromCell::LoadSourceActors(bool& bOutDirty,
 
 	TArray<FWorldPartitionRuntimeCellObjectMapping> ActorsCopy = Actors;
 
+	// Find dirty source HLOD actors and duplicate them instead of loading them from disk
+	TArray<AActor*> DirtySourceActors;
+	for (TActorIterator<AWorldPartitionHLOD> It(SourceWorld); It; ++It)
+	{
+		AActor* SourceActor = *It;
+		if (UPackage* SourceActorPackage = SourceActor->GetPackage())
+		{
+			if (SourceActorPackage->IsDirty())
+			{
+				int32 Index = ActorsCopy.IndexOfByPredicate([UnsavedActorGuid = SourceActor->GetActorInstanceGuid()](FWorldPartitionRuntimeCellObjectMapping Actor)
+				{
+					return UnsavedActorGuid == Actor.ActorInstanceGuid;
+				});
+
+				if (Index != INDEX_NONE)
+				{
+					DirtySourceActors.Add(SourceActor);
+					ActorsCopy.RemoveAtSwap(Index);
+				}
+			}
+		}
+	}
+
+	if (DirtySourceActors.Num())
+	{
+		UActorContainer* ActorContainer = NewObject<UActorContainer>(SourceWorld->PersistentLevel);
+
+		for (AActor* DirtySourceActor : DirtySourceActors)
+		{
+			DirtySourceActor->UObject::Rename(nullptr, ActorContainer, REN_DoNotDirty);
+			ActorContainer->Actors.Add(DirtySourceActor->GetFName(), DirtySourceActor);
+		}
+
+		FObjectDuplicationParameters Parameters(ActorContainer, TargetWorld->PersistentLevel);
+		Parameters.DestClass = ActorContainer->GetClass();
+		Parameters.FlagMask = RF_AllFlags & ~(RF_MarkAsRootSet | RF_MarkAsNative | RF_HasExternalPackage);
+		Parameters.InternalFlagMask = EInternalObjectFlags_AllFlags;
+		Parameters.DuplicationSeed.Add(SourceWorld->PersistentLevel, TargetWorld->PersistentLevel);
+
+		UActorContainer* ActorContainerDup = (UActorContainer*)StaticDuplicateObjectEx(Parameters);
+		check(ActorContainerDup);
+
+		for (TPair<FName, AActor*> Pair : ActorContainerDup->Actors)
+		{
+			AActor* DuplicatedActor = Pair.Value;
+			DuplicatedActor->Rename(nullptr, TargetWorld->PersistentLevel, REN_DoNotDirty);
+		}
+
+		for (AActor* DirtySourceActor : DirtySourceActors)
+		{
+			DirtySourceActor->UObject::Rename(nullptr, SourceWorld->PersistentLevel, REN_DoNotDirty);
+		}
+
+		ActorContainer->MarkAsGarbage();
+		ActorContainerDup->MarkAsGarbage();
+	}
+
 	FWorldPartitionLevelHelper::FPackageReferencer PackageReferencer;
 	FWorldPartitionLevelHelper::FLoadActorsParams Params = FWorldPartitionLevelHelper::FLoadActorsParams()
 		.SetOuterWorld(TargetWorld)
diff --git a/Engine/Source/Runtime/Engine/Private/WorldPartition/RuntimeHashSet/WorldPartitionRuntimeHashSetHLODGeneration.cpp b/Engine/Source/Runtime/Engine/Private/WorldPartition/RuntimeHashSet/WorldPartitionRuntimeHashSetHLODGeneration.cpp
--- a/Engine/Source/Runtime/Engine/Private/WorldPartition/RuntimeHashSet/WorldPartitionRuntimeHashSetHLODGeneration.cpp
+++ b/Engine/Source/Runtime/Engine/Private/WorldPartition/RuntimeHashSet/WorldPartitionRuntimeHashSetHLODGeneration.cpp
@@ -12,6 +12,7 @@
 #include "WorldPartition/HLOD/HLODLayer.h"
 #include "WorldPartition/HLOD/HLODActor.h"
 #include "WorldPartition/HLOD/HLODActorDesc.h"
+#include "WorldPartition/HLOD/HLODCreationFilter.h"
 #include "WorldPartition/HLOD/HLODModifier.h"
 #include "WorldPartition/HLOD/IWorldPartitionHLODUtilities.h"
 #include "WorldPartition/HLOD/IWorldPartitionHLODUtilitiesModule.h"
@@ -23,6 +24,7 @@
 #if WITH_EDITOR
 #include "EditorLevelUtils.h"
 #include "EditorWorldUtils.h"
+#include "EngineUtils.h"
 #endif
 
 #include "UObject/SavePackage.h"
@@ -196,6 +198,20 @@ bool UWorldPartitionRuntimeHashSet::SetupHLODActors(const IStreamingGenerationCo
 		}
 	});
 
+	TArray<TUniquePtr<FWorldPartitionActorDesc> > UnsavedHLODActorDescs;
+	TArray<TUniquePtr<FWorldPartitionActorDescInstance> > UnsavedHLODActorDescInstances;
+	
+	if (Params.bConsiderUnsavedHLODActors)
+	{
+		for (TActorIterator<AWorldPartitionHLOD> It(WorldPartition->GetWorld()); It; ++It)
+		{
+			if (!HLODCreationContext.HLODActorDescs.Contains(It->GetFName()))
+			{
+				HLODCreationContext.UnsavedHLODActors.Emplace(It->GetFName(), *It);
+			}
+		}
+	}
+
 	TUniquePtr<IStreamingGenerationContext> CurrentHLODStreamingGenerationContext = MakeUnique<FStreamingGenerationContextProxy>(StreamingGenerationContext);
 
 	int32 HLODLevel = 0;
@@ -387,6 +403,7 @@ bool UWorldPartitionRuntimeHashSet::SetupHLODActors(const IStreamingGenerationCo
 				HLODCreationParams.ContentBundleGuid = CellDescInstance.ContentBundleID;
 				HLODCreationParams.DataLayerInstances = CellDescInstance.DataLayerInstances;
 				HLODCreationParams.bIsStandalone = false;
+				HLODCreationParams.Filters = Params.Filters;
 
 				if (WorldPartition->HasStandaloneHLOD())
 				{
@@ -433,7 +450,7 @@ bool UWorldPartitionRuntimeHashSet::SetupHLODActors(const IStreamingGenerationCo
 						NumNextLayerHLODActors += CellHLODActor->GetHLODLayer() ? 1 : 0;
 					}
 
-					if (!Params.bReportOnly)
+					if (!Params.bReportOnly && Params.bSaveActors)
 					{
 						for (AWorldPartitionHLOD* CellHLODActor : CellHLODActors)
 						{
@@ -455,6 +472,18 @@ bool UWorldPartitionRuntimeHashSet::SetupHLODActors(const IStreamingGenerationCo
 		// Build the next HLOD generation context
 		if (NumNextLayerHLODActors)
 		{
+			TMap<FGuid, TWeakObjectPtr<AWorldPartitionHLOD> > UnsavedHLODActors;
+			if (Params.bConsiderUnsavedHLODActors)
+			{
+				for (TActorIterator<AWorldPartitionHLOD> It(WorldPartition->GetWorld()); It; ++It)
+				{	
+					if (It->GetPackage() && It->GetPackage()->IsDirty())
+					{
+						UnsavedHLODActors.Add(It->GetActorGuid(), *It);
+					}
+				}
+			}
+
 			CurrentHLODStreamingGenerationContext = MakeUnique<FHLODStreamingGenerationContext>();
 			FHLODStreamingGenerationContext* HLODStreamingGenerationContext = (FHLODStreamingGenerationContext*)CurrentHLODStreamingGenerationContext.Get();
 
@@ -475,7 +504,29 @@ bool UWorldPartitionRuntimeHashSet::SetupHLODActors(const IStreamingGenerationCo
 			UE_LOG(LogWorldPartition, Log, TEXT("%s"), *LOCTEXT("CreatingHLODContext:", "Creating HLOD context:").ToString());
 			for (const FGuid& HLODActorGuid : HLODActorGuids)
 			{
-				FWorldPartitionActorDescInstance* HLODActorDescInstance = CurrentWorldPartition->GetActorDescInstance(HLODActorGuid);
+				FWorldPartitionActorDescInstance* HLODActorDescInstance = nullptr;
+				if (Params.bConsiderUnsavedHLODActors)
+				{
+					if (TWeakObjectPtr<AWorldPartitionHLOD>* UnsavedHLODActorPtr = UnsavedHLODActors.Find(HLODActorGuid))
+					{
+						if (AWorldPartitionHLOD* UnsavedHLODActor = UnsavedHLODActorPtr->Get())
+						{
+							if (UActorDescContainerInstance* HandlingContainer = const_cast<UActorDescContainerInstance*>(BaseContainerInstanceCollection->FindHandlingContainerInstance(UnsavedHLODActor).Get()))
+							{
+								TUniquePtr<FWorldPartitionActorDesc>& UnsavedHLODActorDesc = UnsavedHLODActorDescs.Add_GetRef(UnsavedHLODActor->CreateActorDesc());
+								TUniquePtr<FWorldPartitionActorDescInstance>& UnsavedHLODActorDescInstance = UnsavedHLODActorDescInstances.Add_GetRef(MakeUnique<FWorldPartitionActorDescInstance>(HandlingContainer, UnsavedHLODActorDesc.Get()));
+								DataLayerManager->ResolveActorDescInstanceDataLayers(UnsavedHLODActorDescInstance.Get());
+
+								HLODActorDescInstance = UnsavedHLODActorDescInstance.Get();
+							}
+						}
+					}
+				}
+
+				if (!HLODActorDescInstance)
+				{
+					HLODActorDescInstance = CurrentWorldPartition->GetActorDescInstance(HLODActorGuid);
+				}
 				check(HLODActorDescInstance);
 
 				FStreamingGenerationActorDescView* HLODActorDescView = HLODStreamingGenerationContext->ActorDescViewMap.Emplace(HLODActorDescInstance);
@@ -514,7 +565,28 @@ bool UWorldPartitionRuntimeHashSet::SetupHLODActors(const IStreamingGenerationCo
 		for (const auto& HLODActorPair : HLODCreationContext.HLODActorDescs)
 		{
 			check(HLODActorPair.Value.IsValid());
-			PrivateUtils::DeletePackage(WorldPartition, HLODActorPair.Value, Params.SourceControlHelper);
+
+			FHLODCreationFilterContext FilterContext;
+			FilterContext.Bounds = HLODActorPair.Value->GetRuntimeBounds();
+			const bool bPassesAllFilters = UE::HLOD::CreationFilter::PassesFilters(Params.Filters, FilterContext);
+
+			if (bPassesAllFilters)
+			{
+				if (Params.bSaveActors)
+				{
+					PrivateUtils::DeletePackage(WorldPartition, HLODActorPair.Value, Params.SourceControlHelper);
+				}
+				else
+				{
+					FWorldPartitionReference HLODActorRef = HLODActorPair.Value.ToReference();
+					AWorldPartitionHLOD* HLODActor = CastChecked<AWorldPartitionHLOD>(HLODActorRef.GetActor());
+					
+					ResetLoaders(HLODActor);
+					WorldPartition->GetWorld()->EditorDestroyActor(HLODActor, /*bShouldModifyLevel=*/false);
+					WorldPartition->OnPackageDeleted(HLODActor->GetPackage());
+					HLODActor->UObject::Rename(nullptr, GetTransientPackage(), REN_DontCreateRedirectors);
+				}
+			}
 		}
 	}
 
diff --git a/Engine/Source/Runtime/Engine/Private/WorldPartition/RuntimeSpatialHash/RuntimeSpatialHashHLOD.cpp b/Engine/Source/Runtime/Engine/Private/WorldPartition/RuntimeSpatialHash/RuntimeSpatialHashHLOD.cpp
--- a/Engine/Source/Runtime/Engine/Private/WorldPartition/RuntimeSpatialHash/RuntimeSpatialHashHLOD.cpp
+++ b/Engine/Source/Runtime/Engine/Private/WorldPartition/RuntimeSpatialHash/RuntimeSpatialHashHLOD.cpp
@@ -10,6 +10,7 @@
 #include "WorldPartition/HLOD/HLODLayer.h"
 #include "WorldPartition/HLOD/HLODActor.h"
 #include "WorldPartition/HLOD/HLODActorDesc.h"
+#include "WorldPartition/HLOD/HLODCreationFilter.h"
 #include "WorldPartition/HLOD/HLODModifier.h"
 #include "WorldPartition/HLOD/IWorldPartitionHLODUtilities.h"
 #include "WorldPartition/HLOD/IWorldPartitionHLODUtilitiesModule.h"
@@ -229,6 +230,7 @@ static TArray<FGuid> GenerateHLODActorsForGrid(UWorldPartition* WorldPartition,
 				CreationParams.ContentBundleGuid = GridCellDataChunk.GetContentBundleID();
 				CreationParams.DataLayerInstances = GridCellDataChunk.GetDataLayers();
 				CreationParams.bIsStandalone = false;
+				CreationParams.Filters = Params.Filters;
 
 				TArray<AWorldPartitionHLOD*> CellHLODActors = WPHLODUtilities->CreateHLODActors(Context, CreationParams, ActorInstances);
 				if (!CellHLODActors.IsEmpty())
@@ -615,11 +617,31 @@ bool UWorldPartitionRuntimeSpatialHash::SetupHLODActors(const IStreamingGenerati
 	// Destroy all unreferenced HLOD actors
 	for (const auto& HLODActorPair : Context.HLODActorDescs)
 	{
-		DeletedActors.Add(HLODActorPair.Key);
+		FHLODCreationFilterContext FilterContext;
+		FilterContext.Bounds = HLODActorPair.Value->GetRuntimeBounds();
+		const bool bPassesAllFilters = UE::HLOD::CreationFilter::PassesFilters(Params.Filters, FilterContext);
 
-		if (!Params.bReportOnly)
+		if (bPassesAllFilters)
 		{
-			DeleteHLODActor(HLODActorPair.Value);
+			DeletedActors.Add(HLODActorPair.Key);
+
+			if (!Params.bReportOnly)
+			{
+				if (Params.bSaveActors)
+				{
+					DeleteHLODActor(HLODActorPair.Value);
+				}
+				else
+				{
+					FWorldPartitionReference HLODActorRef = HLODActorPair.Value.ToReference();
+					AWorldPartitionHLOD* HLODActor = CastChecked<AWorldPartitionHLOD>(HLODActorRef.GetActor());
+
+					ResetLoaders(HLODActor);
+					WorldPartition->GetWorld()->EditorDestroyActor(HLODActor, /*bShouldModifyLevel=*/false);
+					WorldPartition->OnPackageDeleted(HLODActor->GetPackage());
+					HLODActor->UObject::Rename(nullptr, GetTransientPackage(), REN_DontCreateRedirectors);
+				}
+			}
 		}
 	}
 
diff --git a/Engine/Source/Runtime/Engine/Public/WorldPartition/HLOD/HLODCreationFilter.h b/Engine/Source/Runtime/Engine/Public/WorldPartition/HLOD/HLODCreationFilter.h
new file mode 100644
--- /dev/null
+++ b/Engine/Source/Runtime/Engine/Public/WorldPartition/HLOD/HLODCreationFilter.h
@@ -0,0 +1,39 @@
+// Copyright Epic Games, Inc. All Rights Reserved.
+
+#pragma once
+
+#include "Math/Box.h"
+#include "Templates/SharedPointer.h"
+
+#define UE_API ENGINE_API
+
+struct FHLODCreationFilterContext
+{
+	TOptional<FBox> Bounds;
+};
+
+struct IHLODCreationFilter
+{
+	virtual ~IHLODCreationFilter() = default;
+	virtual bool IsActive() const = 0;
+	virtual bool PassesFilter(const FHLODCreationFilterContext& InFilterContext) const = 0;
+};
+
+namespace UE::HLOD::CreationFilter
+{
+	UE_API bool HasAnyActiveFilters(const TArray<TSharedPtr<IHLODCreationFilter>>& InFilters);
+	UE_API bool PassesFilters(const TArray<TSharedPtr<IHLODCreationFilter>>& InFilters, const FHLODCreationFilterContext& InFilterContext);
+}
+
+struct FHLODCreationVolumeFilter : IHLODCreationFilter
+{
+	FHLODCreationVolumeFilter() = default;
+	UE_API FHLODCreationVolumeFilter(TArray<FBox>&& InFilterVolumes);
+
+	UE_API virtual bool IsActive() const override;
+	UE_API virtual bool PassesFilter(const FHLODCreationFilterContext& InFilterContext) const override;
+
+	TArray<FBox> FilterVolumes;
+};
+
+#undef UE_API
\ No newline at end of file
diff --git a/Engine/Source/Runtime/Engine/Public/WorldPartition/HLOD/IWorldPartitionHLODUtilities.h b/Engine/Source/Runtime/Engine/Public/WorldPartition/HLOD/IWorldPartitionHLODUtilities.h
--- a/Engine/Source/Runtime/Engine/Public/WorldPartition/HLOD/IWorldPartitionHLODUtilities.h
+++ b/Engine/Source/Runtime/Engine/Public/WorldPartition/HLOD/IWorldPartitionHLODUtilities.h
@@ -20,10 +20,13 @@ class UHLODBuilder;
 class UHLODBuilderSettings;
 class AWorldPartitionHLOD;
 
+struct IHLODCreationFilter;
+
 struct FHLODCreationContext
 {
 	TMap<FName, FWorldPartitionHandle> HLODActorDescs;
 	TArray<FWorldPartitionReference> ActorReferences;
+	TMap<FName, TWeakObjectPtr<AWorldPartitionHLOD>> UnsavedHLODActors;
 };
 
 struct FHLODCreationParams
@@ -38,6 +41,7 @@ struct FHLODCreationParams
 	FGuid ContentBundleGuid;
 	TArray<const UDataLayerInstance*> DataLayerInstances;
 	bool bIsStandalone;
+	TArray<TSharedPtr<IHLODCreationFilter>> Filters;
 
 	const UExternalDataLayerAsset* GetExternalDataLayerAsset() const
 	{
diff --git a/Engine/Source/Runtime/Engine/Public/WorldPartition/WorldPartition.h b/Engine/Source/Runtime/Engine/Public/WorldPartition/WorldPartition.h
--- a/Engine/Source/Runtime/Engine/Public/WorldPartition/WorldPartition.h
+++ b/Engine/Source/Runtime/Engine/Public/WorldPartition/WorldPartition.h
@@ -55,6 +55,7 @@ class UExternalDataLayerManager;
 class IWorldPartitionCookPackageObject;
 
 struct IWorldPartitionStreamingSourceProvider;
+struct IHLODCreationFilter;
 
 enum class EWorldPartitionRuntimeCellState : uint8;
 enum class EWorldPartitionStreamingPerformance : uint8;
@@ -367,14 +368,22 @@ public:
 		FSetupHLODActorsParams()
 			: SourceControlHelper(nullptr)
 			, bReportOnly(false)
+			, bConsiderUnsavedHLODActors(false)
+			, bSaveActors(true)
 		{}
 
 		ISourceControlHelper* SourceControlHelper;
 		bool bReportOnly;
+		bool bConsiderUnsavedHLODActors;
+		bool bSaveActors;
 		mutable TArray<TObjectPtr<UWorldPartition>> OutAdditionalWorldPartitionsForStandaloneHLOD;
+		TArray<TSharedPtr<IHLODCreationFilter> > Filters;
 
 		FSetupHLODActorsParams& SetSourceControlHelper(ISourceControlHelper* InSourceControlHelper) { SourceControlHelper = InSourceControlHelper; return *this; }
 		FSetupHLODActorsParams& SetReportOnly(bool bInReportOnly) { bReportOnly = bInReportOnly; return *this; }
+		FSetupHLODActorsParams& SetConsiderUnsavedHLODActors(bool bInConsiderUnsavedHLODActors) { bConsiderUnsavedHLODActors = bInConsiderUnsavedHLODActors; return *this; }
+		FSetupHLODActorsParams& SetSaveActors(bool bInSaveActors) { bSaveActors = bInSaveActors; return *this; }
+		FSetupHLODActorsParams& SetFilters(const TArray<TSharedPtr<IHLODCreationFilter> >& InFilters) { Filters = InFilters; return *this; }
 	};
 
 	ENGINE_API void SetupHLODActors(const FSetupHLODActorsParams& Params);
