Fixed the material editor providing positive feedback when dragging a texture asset that cannot be referenced by the current material asset over a texture sample node.

diff --git a/Engine/Source/Editor/UnrealEd/Private/MaterialGraphSchema.cpp b/Engine/Source/Editor/UnrealEd/Private/MaterialGraphSchema.cpp
--- a/Engine/Source/Editor/UnrealEd/Private/MaterialGraphSchema.cpp
+++ b/Engine/Source/Editor/UnrealEd/Private/MaterialGraphSchema.cpp
@@ -1274,9 +1274,20 @@ void UMaterialGraphSchema::GetAssetsNodeHoverMessage(const TArray<struct FAssetD
 			return;
 		}
 
-		if (UTexture* Tex = Cast<UTexture>(Assets[0].GetAsset()))
+		const FAssetData& AssetData = Assets[0];
+		if (UTexture* Tex = Cast<UTexture>(AssetData.GetAsset()))
 		{
-			OutOkIcon = true;
+			TSharedPtr<IAssetReferenceFilter> AssetReferenceFilter = FMaterialEditorUtilities::MakeMaterialAssetReferenceFilter(MaterialGraphNode->GetGraph());
+
+			FText FailureReason;
+			if (AssetReferenceFilter.IsValid() && !AssetReferenceFilter->PassesFilter(AssetData, &FailureReason))
+			{
+				OutTooltipText = FailureReason.ToString();
+			}
+			else
+			{
+				OutOkIcon = true;
+			}
 		}
 	}
 }
