SourceControl: UnrealEdEngine should not trigger `FUpdateStatus` operations if `bPromptForCheckoutOnAssetModification` is disabled.

After the `FUpdateStatus` completes, the `UUnrealEdEngine::OnSourceControlStateUpdated` callback is called. This -potentially- updates the `PackageToNotifyState` and -potentially- sets `bShowPackageNotification` to `true` - both depending on the SCC status.
If `bShowPackageNotification` is `true`, the `UUnrealEdEngine::AttemptModifiedPackageNotification()` method calls `UUnrealEdEngine::ShowPackageNotification()` on the next tick.

However, in that `ShowPackageNotification()` method, there is a check on `bPromptForCheckoutOnAssetModification` and the prompt is never shown if that's disabled.
It was thus never useful to query the source control status in the first place.

Note, this code appears to be ancient (2015, #2403502). Not sure if it's in any way superceded already.

#rnx

diff --git a/Engine/Source/Editor/UnrealEd/Private/UnrealEdEngine.cpp b/Engine/Source/Editor/UnrealEd/Private/UnrealEdEngine.cpp
--- a/Engine/Source/Editor/UnrealEd/Private/UnrealEdEngine.cpp
+++ b/Engine/Source/Editor/UnrealEd/Private/UnrealEdEngine.cpp
@@ -626,9 +626,6 @@ void UUnrealEdEngine::AttemptModifiedPackageNotification()
 	FMessageLog EditorLog("EditorErrors");
 	if (PackagesDirtiedThisTick.Num() > 0 && !bIsCooking)
 	{
-		// Force source control state to be updated
-		ISourceControlProvider& SourceControlProvider = ISourceControlModule::Get().GetProvider();
-
 		FString PackageName;
 		TArray<FString> Files;
 		TArray<TWeakObjectPtr<UPackage>> Packages;
@@ -649,9 +646,11 @@ void UUnrealEdEngine::AttemptModifiedPackageNotification()
 			}
 		}
 
-		if (Files.Num() > 0)
+		const bool bIsCheckoutPromptEnabled = GetDefault<UEditorLoadingSavingSettings>()->bPromptForCheckoutOnAssetModification;
+		if (bIsCheckoutPromptEnabled && Files.Num() > 0)
 		{
-			SourceControlProvider.Execute(ISourceControlOperation::Create<FUpdateStatus>(), SourceControlHelpers::AbsoluteFilenames(Files), EConcurrency::Asynchronous, FSourceControlOperationComplete::CreateUObject(this, &UUnrealEdEngine::OnSourceControlStateUpdated, Packages));
+			// Force source control state to be updated
+			ISourceControlModule::Get().GetProvider().Execute(ISourceControlOperation::Create<FUpdateStatus>(), SourceControlHelpers::AbsoluteFilenames(Files), EConcurrency::Asynchronous, FSourceControlOperationComplete::CreateUObject(this, &UUnrealEdEngine::OnSourceControlStateUpdated, Packages));
 		}
 	}
 	
