IncrementalCook: IncrementalReplay: When replaying warnings and errors from incrementally skipped packages, modify their text to include (IncrementalReplay) to allow QA to know that the package did not cook in the current session.
#rnx

diff --git a/Engine/Source/Editor/UnrealEd/Private/Cooker/CookLog.cpp b/Engine/Source/Editor/UnrealEd/Private/Cooker/CookLog.cpp
--- a/Engine/Source/Editor/UnrealEd/Private/Cooker/CookLog.cpp
+++ b/Engine/Source/Editor/UnrealEd/Private/Cooker/CookLog.cpp
@@ -110,7 +110,8 @@ public:
 	explicit FLogHandler(UCookOnTheFlyServer& InCOTFS);
 	virtual ~FLogHandler();
 
-	virtual void ReplayLogsFromIncrementallySkipped(TConstArrayView<FReplicatedLogData> LogMessages) override;
+	virtual void ReplayLogsFromIncrementallySkipped(FName PackageName,
+		TConstArrayView<FReplicatedLogData> LogMessages) override;
 	virtual void ReplayLogFromCookWorker(FReplicatedLogData&& LogData, int32 CookWorkerProfileId) override;
 	virtual void ConditionalPruneReplay() override;
 	virtual void FlushIncrementalCookLogs() override;
@@ -197,7 +198,8 @@ void FLogHandler::FlushIncrementalCookLogs()
 	}
 }
 
-void FLogHandler::ReplayLogsFromIncrementallySkipped(TConstArrayView<FReplicatedLogData> LogMessages)
+void FLogHandler::ReplayLogsFromIncrementallySkipped(FName PackageName,
+	TConstArrayView<FReplicatedLogData> LogMessages)
 {
 	// Replays only come from MarkPackageIncrementallySkipped, which happens only on the CookDirector, during
 	// CookRequestCluster traversal. We rely on that, and do not report whether messages from CookWorkers
@@ -205,14 +207,21 @@ void FLogHandler::ReplayLogsFromIncrementallySkipped(TConstArrayView<FReplicated
 	// replay on CookWorkers.
 	check(!COTFS.CookWorkerClient);
 
+	if (LogMessages.IsEmpty())
+	{
+		return;
+	}
+
 	auto MessagePassesFilter = [](FName Category, const FString& Message)
 		{
 			return true;
 		};
 	auto TryTransformMessage = [](const FString& In, FString& Out)
 		{
-			return false;
+			Out = FString::Printf(TEXT("[IncrementalReplay]: %s"), *In);
+			return true;
 		};
+	UE_LOG(LogCook, Display, TEXT("IncrementalReplay of messages for %s"), *WriteToString<256>(PackageName))
 	for (const FReplicatedLogData& LogMessage : LogMessages)
 	{
 		UnMarshalAndLog(LogMessage, MessagePassesFilter, TryTransformMessage);
diff --git a/Engine/Source/Editor/UnrealEd/Private/Cooker/CookLogPrivate.h b/Engine/Source/Editor/UnrealEd/Private/Cooker/CookLogPrivate.h
--- a/Engine/Source/Editor/UnrealEd/Private/Cooker/CookLogPrivate.h
+++ b/Engine/Source/Editor/UnrealEd/Private/Cooker/CookLogPrivate.h
@@ -43,7 +43,8 @@ class ILogHandler
 public:
 	virtual ~ILogHandler() = default;
 
-	virtual void ReplayLogsFromIncrementallySkipped(TConstArrayView<FReplicatedLogData> LogMessages) = 0;
+	virtual void ReplayLogsFromIncrementallySkipped(FName PackageName,
+		TConstArrayView<FReplicatedLogData> LogMessages) = 0;
 	virtual void ReplayLogFromCookWorker(FReplicatedLogData&& LogData, int32 CookWorkerProfileId) = 0;
 	virtual void ConditionalPruneReplay() = 0;
 	virtual void FlushIncrementalCookLogs() = 0;
diff --git a/Engine/Source/Editor/UnrealEd/Private/Cooker/CookRequestCluster.cpp b/Engine/Source/Editor/UnrealEd/Private/Cooker/CookRequestCluster.cpp
--- a/Engine/Source/Editor/UnrealEd/Private/Cooker/CookRequestCluster.cpp
+++ b/Engine/Source/Editor/UnrealEd/Private/Cooker/CookRequestCluster.cpp
@@ -2337,7 +2337,8 @@ void FRequestCluster::FGraphSearch::FExploreEdgesContext::ProcessPlatformAttachm
 				TConstArrayView<FReplicatedLogData> LogMessages = QueryPlatformData.CookAttachments.LogMessages;
 				if (!LogMessages.IsEmpty() && !PackageData->HasReplayedLogMessages())
 				{
-					Cluster.COTFS.LogHandler->ReplayLogsFromIncrementallySkipped(LogMessages);
+					Cluster.COTFS.LogHandler->ReplayLogsFromIncrementallySkipped(PackageData->GetPackageName(),
+						LogMessages);
 					PackageData->SetHasReplayedLogMessages(true);
 				}
 			}
