Merging CL#47816498

[UnrealEd] ObjectPositioning traces now consider IPrimitiveComponent interface instead of UPrimitiveComponent.

diff --git a/Engine/Source/Editor/UnrealEd/Private/Editor/ObjectPositioning.cpp b/Engine/Source/Editor/UnrealEd/Private/Editor/ObjectPositioning.cpp
--- a/Engine/Source/Editor/UnrealEd/Private/Editor/ObjectPositioning.cpp
+++ b/Engine/Source/Editor/UnrealEd/Private/Editor/ObjectPositioning.cpp
@@ -21,19 +21,25 @@ namespace ObjectPositioningLocals
 		bCVarAllowNonPrimitiveComponentHits,
 		TEXT("When raycasting the world in placement mode, allow hits of physics objects that are not tied to a UPrimitiveComponent (to work with non-actor workflows)."));
 
+	struct FPrimitiveOwner
+	{
+		TWeakObjectPtr<UObject> Object = nullptr;
+		bool bConsiderForPlacementWhenHidden = false;
+	};
+	
 	/**
 	 * Prunes list of hit results for object positioning calculations based on conditions that could be tested
 	 * on the game thread and returns a list of primitives for the remaining this.
 	 * @note If a non-primitive based hit is found and bCVarAllowNonPrimitiveComponentHits is true then
      * an empty weak obj ptr will be added to the result to represent the hit.
 	 */
-	TArray<TWeakObjectPtr<const UPrimitiveComponent>> FilterHitsGameThread(TArray<FHitResult>& InOutHits)
+	TArray<FPrimitiveOwner> FilterHitsGameThread(TArray<FHitResult>& InOutHits)
 	{
 		check(IsInGameThread() || IsInParallelGameThread());
-		TArray<TWeakObjectPtr<const UPrimitiveComponent>> WeakPrimitives;
-		WeakPrimitives.Reserve(InOutHits.Num());
+		TArray<FPrimitiveOwner> PrimitiveOwners;
+		PrimitiveOwners.Reserve(InOutHits.Num());
 
-		InOutHits.RemoveAll([&WeakPrimitives](const FHitResult& Hit)
+		InOutHits.RemoveAll([&PrimitiveOwners](const FHitResult& Hit)
 		{
 			if (Hit.bStartPenetrating)
 			{
@@ -43,7 +49,7 @@ namespace ObjectPositioningLocals
 			const FActorInstanceHandle& HitObjHandle = Hit.HitObjectHandle;
 
 			// Try and find a primitive component for the hit
-			const UPrimitiveComponent* PrimitiveComponent = Hit.GetComponent();
+			UPrimitiveComponent* PrimitiveComponent = Hit.GetComponent();
 
 			if (!PrimitiveComponent)
 			{
@@ -54,58 +60,68 @@ namespace ObjectPositioningLocals
 				PrimitiveComponent = CastChecked<ULandscapeHeightfieldCollisionComponent>(PrimitiveComponent)->GetRenderComponent();
 			}
 
+			FPrimitiveOwner Owner;
 			if (!PrimitiveComponent)
 			{
 				// If we don't have a primitive component, either ignore the hit, or pass it through if the CVar is set appropriately.
 				// If we pass the hit through, we still need to add an entry to the WeakPrimitives list to make sure that
 				// we have an entry for each hit index.
-				if (bCVarAllowNonPrimitiveComponentHits)
+				if (!bCVarAllowNonPrimitiveComponentHits)
 				{
-					WeakPrimitives.Add(nullptr);
-					return false;
+					// Filter out the hit if the CVar didn't allow it.
+					return true;
 				}
-				// Filter out the hit if the CVar didn't allow it.
-				return true;
-			}
 
-			// Ignore volumes and shapes
-			if (HitObjHandle.DoesRepresentClass(AVolume::StaticClass()))
+				Owner = { Hit.PhysicsObjectOwner.Get(), false };
+			}
+			else
 			{
-				return true;
+				// Ignore volumes and shapes
+				if (HitObjHandle.DoesRepresentClass(AVolume::StaticClass()))
+				{
+					return true;
+				}
+
+				if (PrimitiveComponent->IsA(UShapeComponent::StaticClass()))
+				{
+					return true;
+				}
+
+				Owner = { PrimitiveComponent, (bool)PrimitiveComponent->bConsiderForActorPlacementWhenHidden };
 			}
 
-			if (PrimitiveComponent->IsA(UShapeComponent::StaticClass()))
+			if (!Owner.Object.IsValid())
 			{
 				return true;
 			}
-
-			WeakPrimitives.Add(PrimitiveComponent);
+				
+			PrimitiveOwners.Add(MoveTemp(Owner));
 			return false;
 		});
 
-		return MoveTemp(WeakPrimitives);
+		return MoveTemp(PrimitiveOwners);
 	}
 
 	
 	/** Check to see if the specified hit result should be ignored from object positioning calculations for the specified scene view */
-	bool IsHitIgnoredRenderingThread(const TWeakObjectPtr<const UPrimitiveComponent>& InWeakPrimitiveComponent, const FSceneView& InSceneView)
+	bool IsHitIgnoredRenderingThread(const FPrimitiveOwner& InPrimitiveOwner, const FSceneView& InSceneView)
 	{
 		// We're using the SceneProxy and ViewRelevance here, we should execute from the render thread
 		check(IsInParallelRenderingThread());
 
-		const UPrimitiveComponent* PrimitiveComponent = InWeakPrimitiveComponent.Get();
+		const IPrimitiveComponent* PrimitiveComponent = TComponentInterfaceProvider<IPrimitiveComponent>::Provides(InPrimitiveOwner.Object.Get());
 
 		// Only use this component if it is visible in the specified scene views
 		bool bIsRenderedOnScreen = false;
 		bool bIgnoreTranslucentPrimitive = false;
-		const bool bConsiderInvisibleComponentForPlacement = PrimitiveComponent && PrimitiveComponent->bConsiderForActorPlacementWhenHidden;
+		const bool bConsiderInvisibleComponentForPlacement = PrimitiveComponent && InPrimitiveOwner.bConsiderForPlacementWhenHidden;
 
-		if (PrimitiveComponent && PrimitiveComponent->SceneProxy)
+		if (PrimitiveComponent && PrimitiveComponent->GetSceneProxy())
 		{
 			// Only use this component if it is visible in the specified scene views
-			const FPrimitiveViewRelevance ViewRelevance = PrimitiveComponent->SceneProxy->GetViewRelevance(&InSceneView);
+			const FPrimitiveViewRelevance ViewRelevance = PrimitiveComponent->GetSceneProxy()->GetViewRelevance(&InSceneView);
 			// BSP is a bit special in that its bDrawRelevance is false even when drawn as wireframe because InSceneView.Family->EngineShowFlags.BSPTriangles is off
-			bIsRenderedOnScreen = ViewRelevance.bDrawRelevance || (PrimitiveComponent->IsA(UModelComponent::StaticClass()) && InSceneView.Family->EngineShowFlags.BSP);
+			bIsRenderedOnScreen = ViewRelevance.bDrawRelevance || (InPrimitiveOwner.Object->IsA(UModelComponent::StaticClass()) && InSceneView.Family->EngineShowFlags.BSP);
 			bIgnoreTranslucentPrimitive = ViewRelevance.HasTranslucency() && !GetDefault<UEditorPerProjectUserSettings>()->bAllowSelectTranslucent;
 		}
 
@@ -165,18 +181,18 @@ UE::Positioning::FObjectPositioningTraceResult UE::Positioning::TraceWorldForPos
 		{
 			// Filter out anything that should be ignored based on information accessible on the game thread
 			// and build list of remaining weak primitive components that need to be filtered on the rendering thread 
-			TArray<TWeakObjectPtr<const UPrimitiveComponent>> WeakPrimitives = FilterHitsGameThread(Hits);
-			ensure(Hits.Num() == WeakPrimitives.Num());
+			TArray<FPrimitiveOwner> PrimitiveOwners = FilterHitsGameThread(Hits);
+			ensure(Hits.Num() == PrimitiveOwners.Num());
 
 			// Send IsHitIgnoredRenderingThread on the render thread since we're accessing view relevance
 			ENQUEUE_RENDER_COMMAND(TraceWorldForPosition_FilterHitsByViewRelevance)(
-				[&Hits, &WeakPrimitives, &InSceneView](FRHICommandListImmediate& RHICmdList)
+				[&Hits, &PrimitiveOwners, &InSceneView](FRHICommandListImmediate& RHICmdList)
 				{
 					// Filter out anything that should be ignored
 					int32 Index = 0;
-					Hits.RemoveAll([&Index, &InSceneView, &WeakPrimitives](const FHitResult&)
+					Hits.RemoveAll([&Index, &InSceneView, &PrimitiveOwners](const FHitResult&)
 					{
-						return IsHitIgnoredRenderingThread(WeakPrimitives[Index++], InSceneView);
+						return IsHitIgnoredRenderingThread(PrimitiveOwners[Index++], InSceneView);
 					});
 				}
 			);
