Thumbnails optims
- Removed most of the preview scene caches since they can be heavy on the render thread even when not rendering thumbnails and choke the GC when we swap them out for other preview scenes. Kept just 1 since we can have 1 real-time thumbnail at a time (its currently based on mouse hover).
- Added a object thumbnail cache to the thumbnail pools to persist the result of a thumbnail render even when thumbnail textures are moved out of the pool.
- Prioritized cached thumbnails when doing a thumbnail load for a loaded asset that just entered the pool. Subsequent loads are triggered by calls to refresh and will thus prioritize a fresh render.
- Made SAssetView create the initial FAssetThumbnail with the correct resolution. Otherwise, we were getting a cache miss when requesting a thumbnail of the correct size.
- Added a version of RenderThumbnail that doesn't flush the render thread to fill the FObjectThumbnail. Created a FObjectThumbnail locally so that we can poke at it from the render thread, before giving it back asynchronously to the game thread.
- Removed cvar to flush the BlueprintThumbnailRenderer preview scene cache since it's not useful anymore.

diff --git a/Engine/Source/Editor/ContentBrowser/Private/SAssetView.cpp b/Engine/Source/Editor/ContentBrowser/Private/SAssetView.cpp
--- a/Engine/Source/Editor/ContentBrowser/Private/SAssetView.cpp
+++ b/Engine/Source/Editor/ContentBrowser/Private/SAssetView.cpp
@@ -5400,7 +5400,7 @@ TSharedRef<ITableRow> SAssetView::MakeTileViewWidget(TSharedPtr<FAssetViewItem>
 		TSharedPtr<FAssetThumbnail>& AssetThumbnail = RelevantThumbnails.FindOrAdd(AssetItem);
 		if (!AssetThumbnail)
 		{
-			AssetThumbnail = MakeShared<FAssetThumbnail>(FAssetData(), TileViewThumbnailResolution, TileViewThumbnailResolution, AssetThumbnailPool);
+			AssetThumbnail = MakeShared<FAssetThumbnail>(FAssetData(), CurrentThumbnailSize, CurrentThumbnailSize, AssetThumbnailPool);
 			AssetItem->GetItem().UpdateThumbnail(*AssetThumbnail);
 			AssetThumbnail->GetViewportRenderTargetTexture(); // Access the texture once to trigger it to render
 		}
diff --git a/Engine/Source/Editor/UnrealEd/Classes/ThumbnailRendering/AnimBlueprintThumbnailRenderer.h b/Engine/Source/Editor/UnrealEd/Classes/ThumbnailRendering/AnimBlueprintThumbnailRenderer.h
--- a/Engine/Source/Editor/UnrealEd/Classes/ThumbnailRendering/AnimBlueprintThumbnailRenderer.h
+++ b/Engine/Source/Editor/UnrealEd/Classes/ThumbnailRendering/AnimBlueprintThumbnailRenderer.h
@@ -31,6 +31,6 @@ class UAnimBlueprintThumbnailRenderer : public UDefaultSizedThumbnailRenderer
 	// End UObject Implementation
 
 private:
-	TClassInstanceThumbnailScene<FAnimBlueprintThumbnailScene, 100> ThumbnailScenes;
+	TOptional<FAnimBlueprintThumbnailScene> ThumbnailScene;
 };
 
diff --git a/Engine/Source/Editor/UnrealEd/Classes/ThumbnailRendering/BlueprintThumbnailRenderer.h b/Engine/Source/Editor/UnrealEd/Classes/ThumbnailRendering/BlueprintThumbnailRenderer.h
--- a/Engine/Source/Editor/UnrealEd/Classes/ThumbnailRendering/BlueprintThumbnailRenderer.h
+++ b/Engine/Source/Editor/UnrealEd/Classes/ThumbnailRendering/BlueprintThumbnailRenderer.h
@@ -34,11 +34,9 @@ class UBlueprintThumbnailRenderer : public UDefaultSizedThumbnailRenderer
 	/** Notifies the thumbnail scene to refresh components for the specified blueprint */
 	void BlueprintChanged(UBlueprint* Blueprint);
 
-	void ClearThumbnailScenes();
-
 private:
 	void OnBlueprintUnloaded(UBlueprint* Blueprint);
 
 private:
-	TClassInstanceThumbnailScene<FBlueprintThumbnailScene, 100> ThumbnailScenes;
+	TOptional<FBlueprintThumbnailScene> ThumbnailScene;
 };
diff --git a/Engine/Source/Editor/UnrealEd/Classes/ThumbnailRendering/ClassThumbnailRenderer.h b/Engine/Source/Editor/UnrealEd/Classes/ThumbnailRendering/ClassThumbnailRenderer.h
--- a/Engine/Source/Editor/UnrealEd/Classes/ThumbnailRendering/ClassThumbnailRenderer.h
+++ b/Engine/Source/Editor/UnrealEd/Classes/ThumbnailRendering/ClassThumbnailRenderer.h
@@ -30,5 +30,5 @@ private:
 	void OnBlueprintGeneratedClassUnloaded(UBlueprintGeneratedClass* BPGC);
 
 private:
-	TClassInstanceThumbnailScene<FClassThumbnailScene, 100> ThumbnailScenes;
+	TOptional<FClassThumbnailScene> ThumbnailScene;
 };
diff --git a/Engine/Source/Editor/UnrealEd/Classes/ThumbnailRendering/ThumbnailRenderer.h b/Engine/Source/Editor/UnrealEd/Classes/ThumbnailRendering/ThumbnailRenderer.h
--- a/Engine/Source/Editor/UnrealEd/Classes/ThumbnailRendering/ThumbnailRenderer.h
+++ b/Engine/Source/Editor/UnrealEd/Classes/ThumbnailRendering/ThumbnailRenderer.h
@@ -52,20 +52,6 @@ public:
 	 */
 	virtual void GetThumbnailSize(UObject* Object, float Zoom, uint32& OutWidth, uint32& OutHeight) const PURE_VIRTUAL(UThumbnailRenderer::GetThumbnailSize,);
 
-	/**
-	 * Draws a thumbnail for the object that was specified.
-	 *
-	 * @param Object the object to draw the thumbnail for
-	 * @param X the X coordinate to start drawing at
-	 * @param Y the Y coordinate to start drawing at
-	 * @param Width the width of the thumbnail to draw
-	 * @param Height the height of the thumbnail to draw
-	 * @param Viewport the viewport being drawn in
-	 * @param Canvas the render interface to draw with
-	 */
-	UE_DEPRECATED(4.25, "Please override the other prototype of the Draw function.")
-	virtual void Draw(UObject* Object, int32 X, int32 Y, uint32 Width, uint32 Height, FRenderTarget* Viewport, FCanvas* Canvas) { Draw(Object, X, Y, Width, Height, Viewport, Canvas, false); }
-
 	/**
 	 * Draws a thumbnail for the object that was specified.
 	 *
@@ -80,18 +66,6 @@ public:
 	 */
 	virtual void Draw(UObject* Object, int32 X, int32 Y, uint32 Width, uint32 Height, FRenderTarget* Viewport, FCanvas* Canvas, bool bAdditionalViewFamily) PURE_VIRTUAL(UThumbnailRenderer::Draw, );
 
-	/**
-	 * Checks to see if the specified asset supports realtime thumbnails, which will cause them to always be rerendered to reflect any changes
-	 * made to the asset. If this is false, thumbnails should render once and then not update again.
-	 * For most renderers, this should remain as true.
-	 *
-	 * @param	Object	The asset to draw the thumbnail for
-	 *
-	 * @return True if the thumbnail needs to always be redrawn, false if it can be just drawn once and then reused.
-	 */
-	UE_DEPRECATED(5.2, "Please override GetThumbnailRenderFrequency instead, AllowsRealtimeThumbnails is not used.")
-	virtual bool AllowsRealtimeThumbnails(UObject* Object) const { return true; }
-
 	/**
 	 * Override this method to control the render frequency for thumbnails in your editor
 	 * Generally speaking you should use Realtime if you want an animated thumbnail, but if not it is expensive to use (per-frame draws)
@@ -108,9 +82,6 @@ protected:
 	/** Renders the thumbnail's view family. */
 	UNREALED_API static void RenderViewFamily(FCanvas* Canvas, class FSceneViewFamily* ViewFamily, class FSceneView* View);
 
-	UE_DEPRECATED(5.0, "Please use the prototype of RenderViewFamily which takes a non-const view in parameter : like the view family the view is still not quite final until SetupView is called.")
-	UNREALED_API static void RenderViewFamily(FCanvas* Canvas, class FSceneViewFamily* ViewFamily);
-
 	UNREALED_API static struct FGameTime GetTime();
 };
 
diff --git a/Engine/Source/Editor/UnrealEd/Private/AssetThumbnail.cpp b/Engine/Source/Editor/UnrealEd/Private/AssetThumbnail.cpp
--- a/Engine/Source/Editor/UnrealEd/Private/AssetThumbnail.cpp
+++ b/Engine/Source/Editor/UnrealEd/Private/AssetThumbnail.cpp
@@ -58,6 +58,8 @@
 
 #include "Misc/ObjectThumbnail.inl"
 
+DEFINE_LOG_CATEGORY_STATIC(LogThumbnails, Log, All);
+
 namespace AssetThumbnailPool
 {
 	const FObjectThumbnail* LoadThumbnailsFromPackage(const FAssetData& AssetData, FThumbnailMap& OutThumbnailMap)
@@ -2035,9 +2037,72 @@ TStatId FAssetThumbnailPool::GetStatId() const
 	RETURN_QUICK_DECLARE_CYCLE_STAT( FAssetThumbnailPool, STATGROUP_Tickables );
 }
 
+void FAssetThumbnailPool::ProcessPendingObjectThumbnails()
+{
+	const bool bShadersAreCompiling = (GShaderCompilingManager && GShaderCompilingManager->IsCompiling());
+
+	if (bShadersAreCompiling)
+	{
+		for ( TMap< FSoftObjectPath, ThumbnailTools::FAsyncObjectThumbnail >::TIterator PendingObjectThumbnailIt = PendingObjectThumbnails.CreateIterator(); PendingObjectThumbnailIt; ++PendingObjectThumbnailIt )
+		{
+			FSoftObjectPath ObjectPath = PendingObjectThumbnailIt.Key();
+
+			if ( UObject* Asset = ObjectPath.ResolveObject() )
+			{
+				UE_LOG(LogThumbnails, Verbose, TEXT("Shaders are compiling, throwing away pending object thumbnail for %s and retrying"), *Asset->GetFullName());
+				RefreshThumbnailsFor(Asset);
+			}
+		}
+
+		PendingObjectThumbnails.Empty();
+	}
+	else
+	{
+		for ( TMap< FSoftObjectPath, ThumbnailTools::FAsyncObjectThumbnail >::TIterator PendingObjectThumbnailIt = PendingObjectThumbnails.CreateIterator(); PendingObjectThumbnailIt; ++PendingObjectThumbnailIt )
+		{
+			FSoftObjectPath ObjectPath = PendingObjectThumbnailIt.Key();
+
+			if ( UObject* Asset = ObjectPath.ResolveObject() )
+			{
+				if ( !PendingObjectThumbnailIt.Value().IsObjectThumbnailReady() && PendingObjectThumbnailIt.Value().IsTextureReadbackReady() )
+				{
+					UE_LOG(LogThumbnails, Verbose, TEXT("Enqueue thumbnail readback data copy for %s"), *Asset->GetFullName());
+					PendingObjectThumbnailIt.Value().EnqueueReadbackDataCopy();
+				}
+				else if ( PendingObjectThumbnailIt.Value().IsObjectThumbnailReady() )
+				{
+					LocalObjectThumbnailsCache.Add(ObjectPath, PendingObjectThumbnailIt.Value().ConsumeObjectThumbnailUnchecked());
+					UE_LOG(LogThumbnails, Verbose, TEXT("Cached object thumbnail for %s"), *Asset->GetFullName());
+				
+					PendingObjectThumbnailIt.RemoveCurrent();
+				}
+			}
+			else
+			{
+				UE_LOG(LogThumbnails, Verbose, TEXT("Invalid pending object thumbnail for %s, removing."), *ObjectPath.ToString());
+				PendingObjectThumbnailIt.RemoveCurrent();
+			}
+		}
+	}
+}
+
+void FAssetThumbnailPool::ProcessRecentlyLoadedAssets()
+{
+	// If there were any assets loaded since last frame that we are currently displaying thumbnails for, push them on the render stack now.
+	if ( RecentlyLoadedAssets.Num() > 0 )
+	{
+		for ( const FSoftObjectPath& RecentlyLoadedAsset : RecentlyLoadedAssets )
+		{
+			RefreshThumbnailsFor(RecentlyLoadedAsset);
+		}
+
+		RecentlyLoadedAssets.Empty();
+	}
+}
+
 bool FAssetThumbnailPool::IsTickable() const
 {
-	return RecentlyLoadedAssets.Num() > 0 || ThumbnailsToRenderStack.Num() > 0 || RealTimeThumbnails.Num() > 0 || bWereShadersCompilingLastFrame || (GShaderCompilingManager && GShaderCompilingManager->IsCompiling());
+	return RecentlyLoadedAssets.Num() > 0 || ThumbnailsToRenderStack.Num() > 0 || RealTimeThumbnails.Num() > 0 || PendingObjectThumbnails.Num() > 0;
 }
 
 UE_TRACE_EVENT_BEGIN(Cpu, FAssetThumbnailPool_Tick_RenderingThumbnail, NoSync)
@@ -2053,29 +2118,10 @@ void FAssetThumbnailPool::Tick( float DeltaTime )
 		return;
 	}
 
-	const bool bAreShadersCompiling = (GShaderCompilingManager && GShaderCompilingManager->IsCompiling());
-	if (bWereShadersCompilingLastFrame && !bAreShadersCompiling)
-	{
-		ThumbnailsToRenderStack.Reset();
-		// Reschedule visible thumbnails to be rerendered now that shaders are finished compiling
-		for (auto ThumbIt = ThumbnailToTextureMap.CreateIterator(); ThumbIt; ++ThumbIt)
-		{
-			ThumbnailsToRenderStack.Push(ThumbIt.Value());
-		}
-	}
-	bWereShadersCompilingLastFrame = bAreShadersCompiling;
-
 	TRACE_CPUPROFILER_EVENT_SCOPE(FAssetThumbnailPool::Tick);
-	// If there were any assets loaded since last frame that we are currently displaying thumbnails for, push them on the render stack now.
-	if ( RecentlyLoadedAssets.Num() > 0 )
-	{
-		for ( int32 LoadedAssetIdx = 0; LoadedAssetIdx < RecentlyLoadedAssets.Num(); ++LoadedAssetIdx )
-		{
-			RefreshThumbnailsFor(RecentlyLoadedAssets[LoadedAssetIdx]);
-		}
-
-		RecentlyLoadedAssets.Empty();
-	}
+	
+	ProcessPendingObjectThumbnails();
+	ProcessRecentlyLoadedAssets();
 
 	// If we have dynamic thumbnails and we are done rendering the last batch of dynamic thumbnails, start a new batch as long as real-time thumbnails are enabled
 	const bool bIsInPIEOrSimulate = GEditor->PlayWorld != NULL || GEditor->bIsSimulatingInEditor;
@@ -2105,6 +2151,9 @@ void FAssetThumbnailPool::Tick( float DeltaTime )
 	// If there are any thumbnails to render, pop one off the stack and render it.
 	if( ThumbnailsToRenderStack.Num() + RealTimeThumbnailsToRender.Num() > 0 )
 	{
+		TArray<TSharedRef<FThumbnailInfo>, TInlineAllocator<4>> ThumbnailsToRetry;
+		TArray<TSharedRef<FThumbnailInfo>, TInlineAllocator<4>> RealTimeThumbnailsToRetry;
+
 		double FrameStartTime = FPlatformTime::Seconds();
 		// Render as many thumbnails as we are allowed to
 		while( ThumbnailsToRenderStack.Num() + RealTimeThumbnailsToRender.Num() > 0 && FPlatformTime::Seconds() - FrameStartTime < MaxFrameTimeAllowance )
@@ -2129,7 +2178,7 @@ void FAssetThumbnailPool::Tick( float DeltaTime )
 
 			if( Info.IsValid() )
 			{
-				bool bIsAssetStillCompiling = false;
+				bool bRetryOnFailure = false;
 				TSharedRef<FThumbnailInfo> InfoRef = Info.ToSharedRef();
 
 				if ( InfoRef->AssetData.IsValid() )
@@ -2151,10 +2200,10 @@ void FAssetThumbnailPool::Tick( float DeltaTime )
 						}
 					}
 
-					bool bLoadedThumbnail = LoadThumbnail(InfoRef, bIsAssetStillCompiling, CustomThumbnailAsset);
+					bool bLoadedThumbnail = LoadThumbnail(InfoRef, bRetryOnFailure, CustomThumbnailAsset);
 
 					// If we failed to load a custom thumbnail, then load the custom thumbnail's asset and try again
-					if (!bLoadedThumbnail && !bIsAssetStillCompiling && CustomThumbnailAsset.IsValid())
+					if (!bLoadedThumbnail && !bRetryOnFailure && CustomThumbnailAsset.IsValid())
 					{
 						if (UPackage* CustomThumbnailPackage = CustomThumbnailAsset.GetPackage())
 						{
@@ -2169,8 +2218,8 @@ void FAssetThumbnailPool::Tick( float DeltaTime )
 
 							if (CustomThumbnail)
 							{
-								bLoadedThumbnail = LoadThumbnail(InfoRef, bIsAssetStillCompiling, CustomThumbnailAsset);
-								if (!bIsAssetStillCompiling)
+								bLoadedThumbnail = LoadThumbnail(InfoRef, bRetryOnFailure, CustomThumbnailAsset);
+								if (!bRetryOnFailure)
 								{
 									// Clear RF_Standalone flag on the loaded custom thumbnail asset so it gets GCed eventually
 									CustomThumbnail->ClearFlags(RF_Standalone);
@@ -2190,138 +2239,321 @@ void FAssetThumbnailPool::Tick( float DeltaTime )
 					// Do not send a failure event for this asset yet if shaders are still compiling or the asset itself is compiling.
 					// The failure event will disable the rendering of this asset for good and we need to have a chance to 
 					// re-render it when everything settles down.
-					else if (!bAreShadersCompiling && !bIsAssetStillCompiling)
+					else if (!bRetryOnFailure)
 					{
 						// Notify listeners that a thumbnail render has failed
 						ThumbnailRenderFailedEvent.Broadcast(InfoRef->AssetData);
+						UE_LOG(LogThumbnails, Verbose, TEXT("No thumbnail found for %s"), *InfoRef->AssetData.GetFullName());
+					}
+					else // bRetryOnFailure
+					{
+						if (bIsRealTimeThumbnail)
+						{
+							RealTimeThumbnailsToRetry.Push(InfoRef);
+						}
+						else
+						{
+							ThumbnailsToRetry.Push(InfoRef);
+						}
 					}
 				}
 			}
 		}
+
+		ThumbnailsToRenderStack.Append(ThumbnailsToRetry);
+		RealTimeThumbnailsToRender.Append(RealTimeThumbnailsToRetry);
 	}
 }
 
-bool FAssetThumbnailPool::LoadThumbnail(TSharedRef<FThumbnailInfo> ThumbnailInfo, bool& bIsAssetStillCompiling, const FAssetData& CustomAssetToRender)
+bool FAssetThumbnailPool::LoadThumbnail(TSharedRef<FThumbnailInfo> ThumbnailInfo, bool& bOutRetryOnFailure, const FAssetData& CustomAssetToRender)
+{
+	const FAssetData& AssetData = CustomAssetToRender.IsValid() ? CustomAssetToRender : ThumbnailInfo->AssetData;
+
+	constexpr bool bLoad = false;
+	UObject* Asset = AssetData.FastGetAsset(bLoad);
+	const bool bIsAssetLoaded = IsValid(Asset);
+
+	if (bIsAssetLoaded)
+	{
+		return LoadThumbnailForLoadedAsset(ThumbnailInfo, bOutRetryOnFailure, CustomAssetToRender);
+	}
+	else
+	{
+		return LoadThumbnailForUnloadedAsset(ThumbnailInfo, CustomAssetToRender);
+	}
+}
+
+bool FAssetThumbnailPool::LoadThumbnailForLoadedAsset(const TSharedRef<FThumbnailInfo>& ThumbnailInfo, bool& bOutRetryOnFailure, const FAssetData& CustomAssetToRender)
 {
 	const FAssetData& AssetData = CustomAssetToRender.IsValid() ? CustomAssetToRender : ThumbnailInfo->AssetData;
 
-	FThumbnailMap ThumbnailMap;
 	const FObjectThumbnail* FoundThumbnail = nullptr;
+	FThumbnailMap ThumbnailMap;
+
+	TArray<TFunction<bool()>, TInlineAllocator<3>> LoadThumbnailSourcesPriority;
+
+	TFunction<bool()> ThumbnailFromCacheSource = 
+		[this, ThumbnailInfo, CustomAssetToRender]()
+		{
+			if (LoadThumbnailFromObjectThumbnailInCache(ThumbnailInfo, CustomAssetToRender))
+			{
+				UE_LOG(LogThumbnails, Verbose, TEXT("Thumbnail found in cache for %s"), *ThumbnailInfo->AssetData.GetFullName());
+				return true;
+			}
+
+			return false;
+		};
+
+	TFunction<bool()> ThumbnailFromDiskSource =
+		[this, ThumbnailInfo, CustomAssetToRender]()
+		{
+			if (LoadThumbnailFromObjectThumbnailOnDisk(ThumbnailInfo, CustomAssetToRender))
+			{
+				UE_LOG(LogThumbnails, Verbose, TEXT("Thumbnail found on disk for %s"), *ThumbnailInfo->AssetData.GetFullName());
+				return true;
+			}
+
+			return false;
+		};
+
+	TFunction<bool()> ThumbnailFromRenderSource = 
+		[this, ThumbnailInfo, &bOutRetryOnFailure, CustomAssetToRender]()
+		{
+			if (MakeFreshThumbnailFromRender(ThumbnailInfo, bOutRetryOnFailure, CustomAssetToRender))
+			{
+				UE_LOG(LogThumbnails, Verbose, TEXT("Fresh thumbnail rendered for %s"), *ThumbnailInfo->AssetData.GetFullName());
+				return true;
+			}
+
+			return false;
+		};
+
+	const bool bIsCookedAsset = AssetData.PackageFlags & PKG_Cooked;
+	const bool bIsInitialThumbnailLoad = (ThumbnailInfo->LastUpdateTime <= 0.0);
+
+	if (bIsCookedAsset)
+	{
+		if (bIsInitialThumbnailLoad)
+		{
+			LoadThumbnailSourcesPriority.Add(MoveTemp(ThumbnailFromCacheSource));
+			LoadThumbnailSourcesPriority.Add(MoveTemp(ThumbnailFromDiskSource));
+			LoadThumbnailSourcesPriority.Add(MoveTemp(ThumbnailFromRenderSource));
+		}
+		else
+		{
+			LoadThumbnailSourcesPriority.Add(MoveTemp(ThumbnailFromDiskSource));
+			LoadThumbnailSourcesPriority.Add(MoveTemp(ThumbnailFromRenderSource));
+			LoadThumbnailSourcesPriority.Add(MoveTemp(ThumbnailFromCacheSource));
+		}
+	}
+	else // Uncooked asset
+	{
+		if (bIsInitialThumbnailLoad)
+		{
+			LoadThumbnailSourcesPriority.Add(MoveTemp(ThumbnailFromCacheSource));
+			LoadThumbnailSourcesPriority.Add(MoveTemp(ThumbnailFromRenderSource));
+			LoadThumbnailSourcesPriority.Add(MoveTemp(ThumbnailFromDiskSource));
+		}
+		else
+		{
+			LoadThumbnailSourcesPriority.Add(MoveTemp(ThumbnailFromRenderSource));
+			LoadThumbnailSourcesPriority.Add(MoveTemp(ThumbnailFromCacheSource));
+			LoadThumbnailSourcesPriority.Add(MoveTemp(ThumbnailFromDiskSource));
+		}
+	}
+
+	for (const TFunction<bool()>& LoadThumbnailFromSource : LoadThumbnailSourcesPriority)
+	{
+		if (LoadThumbnailFromSource())
+		{
+			return true;
+		}
+		else if (bOutRetryOnFailure)
+		{
+			return false;
+		}
+	}
+
+	return false;
+}
 
-	// Prioritize thumbnail found from the on-disk package if it's cooked because the asset cannot change and 
-	// rendering it without editor-only data might not give the same result as when it was rendered uncooked.
-	if (AssetData.PackageFlags & PKG_Cooked)
+bool FAssetThumbnailPool::LoadThumbnailForUnloadedAsset(const TSharedRef<FThumbnailInfo>& ThumbnailInfo, const FAssetData& CustomAssetToRender)
+{
+	if (LoadThumbnailFromObjectThumbnailInCache(ThumbnailInfo, CustomAssetToRender))
+	{
+		UE_LOG(LogThumbnails, Verbose, TEXT("Thumbnail found in cache for %s"), *ThumbnailInfo->AssetData.GetFullName());
+		return true;
+	}
+
+	if (LoadThumbnailFromObjectThumbnailOnDisk(ThumbnailInfo, CustomAssetToRender))
+	{
+		UE_LOG(LogThumbnails, Verbose, TEXT("Thumbnail found on disk for %s"), *ThumbnailInfo->AssetData.GetFullName());
+		return true;
+	}
+
+	return false;
+}
+
+bool FAssetThumbnailPool::LoadThumbnailFromObjectThumbnailInCache(const TSharedRef<FThumbnailInfo>& ThumbnailInfo, const FAssetData& CustomAssetToRender)
+{
+	const FAssetData& AssetData = CustomAssetToRender.IsValid() ? CustomAssetToRender : ThumbnailInfo->AssetData;
+
+	const FObjectThumbnail* FoundObjectThumbnail = LocalObjectThumbnailsCache.Find(ThumbnailInfo->AssetData.GetSoftObjectPath());
+
+	if (!FoundObjectThumbnail || FoundObjectThumbnail->GetImageWidth() != ThumbnailInfo->Width || FoundObjectThumbnail->GetImageHeight() != ThumbnailInfo->Height)
 	{
-		FoundThumbnail = AssetThumbnailPool::LoadThumbnailsFromPackage(AssetData, ThumbnailMap);
+		FoundObjectThumbnail = ThumbnailTools::FindCachedThumbnail(AssetData.GetFullName());
 	}
 
-	if (!FoundThumbnail)
+	if (FoundObjectThumbnail && !FoundObjectThumbnail->IsDirty())
+	{
+		return CopyObjectThumbnailImageDataToSlateTextureData(ThumbnailInfo, *FoundObjectThumbnail);
+	}
+	else
 	{
-		// Render a fresh thumbnail from the loaded asset if possible
-		UObject* Asset = AssetData.FastGetAsset();
-		if (Asset && !IsValidChecked(Asset))
+		if (FoundObjectThumbnail && FoundObjectThumbnail->IsDirty())
 		{
-			Asset = nullptr;
+			UE_LOG(LogThumbnails, Verbose, TEXT("Ignoring cached object thumbnail for %s because its dirty"), *AssetData.GetFullName());
 		}
 
-		const bool bAreShadersCompiling = (GShaderCompilingManager && GShaderCompilingManager->IsCompiling());
-		if (Asset && !bAreShadersCompiling)
+		return false;
+	}
+}
+
+bool FAssetThumbnailPool::LoadThumbnailFromObjectThumbnailOnDisk(const TSharedRef<FThumbnailInfo>& ThumbnailInfo, const FAssetData& CustomAssetToRender)
+{
+	const FAssetData& AssetData = CustomAssetToRender.IsValid() ? CustomAssetToRender : ThumbnailInfo->AssetData;
+
+	FThumbnailMap ThumbnailMap;
+	const FObjectThumbnail* FoundObjectThumbnail = AssetThumbnailPool::LoadThumbnailsFromPackage(AssetData, ThumbnailMap);
+
+	if (FoundObjectThumbnail)
+	{
+		return CopyObjectThumbnailImageDataToSlateTextureData(ThumbnailInfo, *FoundObjectThumbnail);
+	}
+	else
+	{
+		return false;
+	}
+}
+
+bool FAssetThumbnailPool::MakeFreshThumbnailFromRender(const TSharedRef<FThumbnailInfo>& ThumbnailInfo, bool& bOutRetryOnFailure, const FAssetData& CustomAssetToRender)
+{
+	const FAssetData& AssetData = CustomAssetToRender.IsValid() ? CustomAssetToRender : ThumbnailInfo->AssetData;
+
+	// Render a fresh thumbnail from the loaded asset if possible
+	UObject* Asset = AssetData.FastGetAsset();
+	if (!IsValid(Asset))
+	{
+		return false;
+	}
+
+	const bool bAreShadersCompiling = (GShaderCompilingManager && GShaderCompilingManager->IsCompiling());
+	if (!bAreShadersCompiling)
+	{
+		// Avoid rendering the thumbnail of an asset that is currently edited asynchronously
+		const IInterface_AsyncCompilation* Interface_AsyncCompilation = Cast<IInterface_AsyncCompilation>(Asset);
+		const bool bIsAssetStillCompiling = Interface_AsyncCompilation && Interface_AsyncCompilation->IsCompiling();
+
+		if (!bIsAssetStillCompiling)
 		{
-			//Avoid rendering the thumbnail of an asset that is currently edited asynchronously
-			const IInterface_AsyncCompilation* Interface_AsyncCompilation = Cast<IInterface_AsyncCompilation>(Asset);
-			bIsAssetStillCompiling = Interface_AsyncCompilation && Interface_AsyncCompilation->IsCompiling();
-			if (!bIsAssetStillCompiling)
+			FThumbnailRenderingInfo* RenderInfo = GUnrealEd->GetThumbnailManager()->GetRenderingInfo(Asset);
+			if (RenderInfo != nullptr && RenderInfo->Renderer != nullptr)
 			{
-				FThumbnailRenderingInfo* RenderInfo = GUnrealEd->GetThumbnailManager()->GetRenderingInfo(Asset);
-				if (RenderInfo != nullptr && RenderInfo->Renderer != nullptr)
-				{
-					FThumbnailInfo_RenderThread ThumbInfo = ThumbnailInfo.Get();
+				FThumbnailInfo_RenderThread ThumbInfo = ThumbnailInfo.Get();
 
-					auto EnqueueThumbnailRender = [Asset, ThumbInfo, ThumbnailInfo]()
+				auto EnqueueThumbnailRender = [this, Asset, ThumbInfo, ThumbnailInfo]()
+				{
+					ENQUEUE_RENDER_COMMAND(SyncSlateTextureCommand)(
+					[ThumbInfo](FRHICommandListImmediate& RHICmdList)
 					{
-						ENQUEUE_RENDER_COMMAND(SyncSlateTextureCommand)(
-						[ThumbInfo](FRHICommandListImmediate& RHICmdList)
+						if (ThumbInfo.ThumbnailTexture->GetTypedResource() != ThumbInfo.ThumbnailRenderTarget->GetTextureRHI())
 						{
-							if (ThumbInfo.ThumbnailTexture->GetTypedResource() != ThumbInfo.ThumbnailRenderTarget->GetTextureRHI())
-							{
-								ThumbInfo.ThumbnailTexture->ClearTextureData();
-								ThumbInfo.ThumbnailTexture->ReleaseRHI();
-								ThumbInfo.ThumbnailTexture->SetRHIRef(ThumbInfo.ThumbnailRenderTarget->GetTextureRHI(), ThumbInfo.Width, ThumbInfo.Height);
-							}
-						});
+							ThumbInfo.ThumbnailTexture->ClearTextureData();
+							ThumbInfo.ThumbnailTexture->ReleaseRHI();
+							ThumbInfo.ThumbnailTexture->SetRHIRef(ThumbInfo.ThumbnailRenderTarget->GetTextureRHI(), ThumbInfo.Width, ThumbInfo.Height);
+						}
+					});
 
-						ThumbnailTools::RenderThumbnail(
+					TOptional<ThumbnailTools::FAsyncObjectThumbnail> OptionalAsyncObjectThumbnail =
+						ThumbnailTools::RenderThumbnailAsync(
 							Asset,
 							ThumbnailInfo->Width,
 							ThumbnailInfo->Height,
 							ThumbnailTools::EThumbnailTextureFlushMode::NeverFlush,
 							ThumbnailInfo->ThumbnailRenderTarget
 						);
-					};
 
-					EThumbnailRenderFrequency ThumbnailRenderFrequency = RenderInfo->Renderer->GetThumbnailRenderFrequency(Asset);
+					if (OptionalAsyncObjectThumbnail)
+					{
+						PendingObjectThumbnails.Add(ThumbnailInfo->AssetData.GetSoftObjectPath(), MoveTemp(OptionalAsyncObjectThumbnail.GetValue()));
+						UE_LOG(LogThumbnails, Verbose, TEXT("Future Thumbnail Object being tracked for %s (%dx%d)"), *Asset->GetFullName(), ThumbnailInfo->Width, ThumbnailInfo->Height);
+					}
+				};
+
+				EThumbnailRenderFrequency ThumbnailRenderFrequency = RenderInfo->Renderer->GetThumbnailRenderFrequency(Asset);
 
-					switch (ThumbnailRenderFrequency)
+				switch (ThumbnailRenderFrequency)
+				{
+					case EThumbnailRenderFrequency::Realtime:
+					{
+						EnqueueThumbnailRender();
+						return true;
+					}
+					case EThumbnailRenderFrequency::OnPropertyChange:
 					{
-						case EThumbnailRenderFrequency::Realtime:
+						if (ThumbnailInfo->LastUpdateTime <= 0.0)
 						{
 							EnqueueThumbnailRender();
 							return true;
 						}
-						case EThumbnailRenderFrequency::OnPropertyChange:
-						{
-							if (ThumbnailInfo->LastUpdateTime <= 0.0f)
-							{
-								EnqueueThumbnailRender();
-								return true;
-							}
-							break;
-						}
-						case EThumbnailRenderFrequency::OnAssetSave:
-						{
-							// OnAssetSave is default behavior below, so nohing to do
-							break;
-						}
-						case EThumbnailRenderFrequency::Once:
-						{
-							// Eagerly return if we aren't interested in cached thumbnails
-							return true;
-						}
-						default:
-						{
-							break;
-						}
+						break;
+					}
+					case EThumbnailRenderFrequency::OnAssetSave:
+					{
+						// We don't allow rendering fresh thumbnails on demand for this frequency
+						return false;
+					}
+					case EThumbnailRenderFrequency::Once:
+					{
+						// Eagerly return if we aren't interested in cached thumbnails
+						return true;
+					}
+					default:
+					{
+						break;
 					}
 				}
 			}
 		}
+		else
+		{
+			UE_LOG(LogThumbnails, Verbose, TEXT("Asset is compiling, can't render thumbnail for %s. Will retry."), *AssetData.GetFullName());
+			bOutRetryOnFailure = true;
+		}
 	}
-
-	// If we could not render a fresh thumbnail, see if we already have a cached one to load
-	if (!FoundThumbnail)
+	else
 	{
-		FoundThumbnail = ThumbnailTools::FindCachedThumbnail(AssetData.GetFullName());
+		UE_LOG(LogThumbnails, Verbose, TEXT("Shaders are compiling, can't render thumbnail for %s. Will retry."), *AssetData.GetFullName());
+		bOutRetryOnFailure = true;
 	}
 
-	// If we don't have a thumbnail cached in memory, try to find it on disk
-	if (!FoundThumbnail && !(AssetData.PackageFlags & PKG_Cooked))
-	{
-		FoundThumbnail = AssetThumbnailPool::LoadThumbnailsFromPackage(AssetData, ThumbnailMap);
-	}
+	return false;
+}
 
-	if (FoundThumbnail)
-	{
-		FImageView Image = FoundThumbnail->GetImage();
+bool FAssetThumbnailPool::CopyObjectThumbnailImageDataToSlateTextureData(const TSharedRef<FThumbnailInfo>& ThumbnailInfo, const FObjectThumbnail& ObjectThumbnail)
+{
+	FImageView Image = ObjectThumbnail.GetImage();
 
-		if ( Image.GetNumPixels() > 0 )
-		{
-			// Make bulk data for updating the texture memory later
-			FSlateTextureData* BulkData = new FSlateTextureData(Image);
+	if ( Image.GetNumPixels() > 0 )
+	{
+		// Make bulk data for updating the texture memory later
+		FSlateTextureData* BulkData = new FSlateTextureData(Image);
 
-			// Update the texture RHI
-			FThumbnailInfo_RenderThread ThumbInfo = ThumbnailInfo.Get();
-			ENQUEUE_RENDER_COMMAND(ClearSlateTextureCommand)(
-				[ThumbInfo, BulkData](FRHICommandListImmediate& RHICmdList)
+		// Update the texture RHI
+		FThumbnailInfo_RenderThread ThumbInfo = ThumbnailInfo.Get();
+		ENQUEUE_RENDER_COMMAND(ClearSlateTextureCommand)(
+			[ThumbInfo, BulkData](FRHICommandListImmediate& RHICmdList)
 			{
 				if (ThumbInfo.ThumbnailTexture->GetTypedResource() == ThumbInfo.ThumbnailRenderTarget->GetTextureRHI())
 				{
@@ -2332,9 +2564,9 @@ bool FAssetThumbnailPool::LoadThumbnail(TSharedRef<FThumbnailInfo> ThumbnailInfo
 				ThumbInfo.ThumbnailTexture->UpdateRHI(RHICmdList);
 			});
 
-			return true;
-		}
+		return true;
 	}
+
 	return false;
 }
 
@@ -2367,7 +2599,7 @@ FSlateTexture2DRHIRef* FAssetThumbnailPool::AccessTexture( const FAssetData& Ass
 				break;
 			}
 		}
-		
+
 		// If we didn't find one and we have reached the pool limit, remove any thumbnail from the free thumbnails if any. This will free up a slot
 		// Otherwise, if there are no free thumbnails, evict the LRU used thumbnail to free up a slot.
 		// Note, with luck, the LRU element might be suitable for reuse
@@ -2429,10 +2661,11 @@ FSlateTexture2DRHIRef* FAssetThumbnailPool::AccessTexture( const FAssetData& Ass
 		ThumbnailInfo->AssetData = AssetData;
 	
 		// Mark this thumbnail as needing to be updated
-		ThumbnailInfo->LastUpdateTime = -1.f;
+		ThumbnailInfo->LastUpdateTime = -1.0;
 
 		// Request that the thumbnail be rendered as soon as possible
 		ThumbnailsToRenderStack.Push( ThumbnailRef );
+		UE_LOG(LogThumbnails, Verbose, TEXT("Queued thumbnail to render for %s (%dx%d)"), *ThumbnailInfo->AssetData.GetFullName(), Width, Height);
 	}
 
 	// This thumbnail was accessed, update its last time to the current time
@@ -2452,7 +2685,7 @@ void FAssetThumbnailPool::AddReferencer( const FAssetThumbnail& AssetThumbnail )
 	}
 
 	// Generate a key and look up the number of references in the RefCountMap
-	FThumbId ThumbId( AssetThumbnail.GetAssetData().GetSoftObjectPath(), Size.X, Size.Y ) ;
+	FThumbId ThumbId( AssetThumbnail.GetAssetData().GetSoftObjectPath(), Size.X, Size.Y );
 	int32* RefCountPtr = RefCountMap.Find(ThumbId);
 
 	if ( RefCountPtr )
@@ -2575,11 +2808,7 @@ void FAssetThumbnailPool::RefreshThumbnail( const TSharedPtr<FAssetThumbnail>& T
 	if ( ensure(AssetData.IsValid()) && ensure(Width > 0) && ensure(Height > 0) )
 	{
 		FThumbId ThumbId(AssetData.GetSoftObjectPath(), Width, Height);
-		const TSharedRef<FThumbnailInfo>* ThumbnailInfoPtr = ThumbnailToTextureMap.Find( ThumbId );
-		if ( ThumbnailInfoPtr )
-		{
-			ThumbnailsToRenderStack.AddUnique(*ThumbnailInfoPtr);
-		}
+		RefreshThumbnailFor(ThumbId);
 	}
 }
 
@@ -2597,6 +2826,7 @@ void FAssetThumbnailPool::SetRealTimeThumbnail( const TSharedPtr<FAssetThumbnail
 		{
 			if ( bRealTimeThumbnail )
 			{
+				UE_LOG(LogThumbnails, Verbose, TEXT("Setting real-time thumbnail ON for %s"), *(*ThumbnailInfoPtr)->AssetData.GetFullName());
 				RealTimeThumbnails.AddUnique(*ThumbnailInfoPtr);
 			}
 			else
@@ -2616,6 +2846,8 @@ void FAssetThumbnailPool::FreeThumbnail( const FSoftObjectPath& ObjectPath, uint
 		const TSharedRef<FThumbnailInfo>* ThumbnailInfoPtr = ThumbnailToTextureMap.Find(ThumbId);
 		if( ThumbnailInfoPtr )
 		{
+			UE_LOG(LogThumbnails, Verbose, TEXT("Freeing thumbnail for %s"), *(*ThumbnailInfoPtr)->AssetData.GetFullName());
+
 			TSharedRef<FThumbnailInfo> ThumbnailInfo = *ThumbnailInfoPtr;
 			ThumbnailToTextureMap.Remove(ThumbId);
 			ThumbnailsToRenderStack.Remove(ThumbnailInfo);
@@ -2631,8 +2863,18 @@ void FAssetThumbnailPool::FreeThumbnail( const FSoftObjectPath& ObjectPath, uint
 
 			FreeThumbnails.Add( ThumbnailInfo );
 		}
+
+		// Compress the FObjectThumbnail if there's one to reduce memory usage while the thumbnail isn't in the pool.
+		if (FObjectThumbnail* ObjectThumbnail = LocalObjectThumbnailsCache.Find(ObjectPath))
+		{
+			if (ObjectThumbnail->GetCompressedDataSize() == 0)
+			{
+				ObjectThumbnail->CompressImageData();
+			}
+
+			ObjectThumbnail->EmptyUncompressedImageData();
+		}
 	}
-			
 }
 
 void FAssetThumbnailPool::RefreshThumbnailsFor( const FSoftObjectPath& ObjectPath )
@@ -2641,20 +2883,64 @@ void FAssetThumbnailPool::RefreshThumbnailsFor( const FSoftObjectPath& ObjectPat
 	{
 		if ( ThumbIt.Key().ObjectPath == ObjectPath )
 		{
-			ThumbnailsToRenderStack.AddUnique( ThumbIt.Value() );
+			RefreshThumbnailFor(ThumbIt.Key());
 		}
 	}
 }
 
+void FAssetThumbnailPool::RefreshThumbnailFor( const FThumbId ThumbId )
+{
+	check(ThumbnailToTextureMap.Contains(ThumbId));
+
+	TSharedRef<FThumbnailInfo> ThumbnailInfo = ThumbnailToTextureMap[ThumbId];
+	ThumbnailsToRenderStack.AddUnique( ThumbnailInfo );
+
+	UE_LOG(LogThumbnails, Verbose, TEXT("Refresh thumbnail requested for %s"), *ThumbnailInfo->AssetData.GetFullName());
+}
+
 void FAssetThumbnailPool::OnAssetLoaded( UObject* Asset )
 {
-	if ( Asset != NULL )
+	if ( Asset != nullptr )
 	{
-		RecentlyLoadedAssets.Add( FSoftObjectPath(Asset) );
+		UE_LOG(LogThumbnails, Verbose, TEXT("OnAssetLoaded %s"), *Asset->GetFullName());
+		RecentlyLoadedAssets.Emplace( Asset );
 	}
 }
 
 void FAssetThumbnailPool::OnThumbnailDirtied( const FSoftObjectPath& ObjectPath )
 {
+	UE_LOG(LogThumbnails, Verbose, TEXT("OnThumbnailDirtied %s"), *ObjectPath.ToString());
+
+	// Dirty any cached ObjectThumbnail for this ObjectPath
+	if (FObjectThumbnail* ObjectThumbnail = LocalObjectThumbnailsCache.Find(ObjectPath))
+	{
+		ObjectThumbnail->MarkAsDirty();
+	}
+
 	RefreshThumbnailsFor( ObjectPath );
+
+	// Request a thumbnail refresh for our referencers (ie: refresh a material instance thumbnail when material is edited)
+	FAssetRegistryModule& AssetRegistryModule = FModuleManager::LoadModuleChecked<FAssetRegistryModule>(TEXT("AssetRegistry"));
+	const IAssetRegistry& AssetRegistry = AssetRegistryModule.GetRegistry();
+
+	TArray<FName> PackageReferencers;
+	if (AssetRegistry.GetReferencers(ObjectPath.GetLongPackageFName(), PackageReferencers))
+	{
+		for (const FName& PackageReferencer : PackageReferencers)
+		{
+			TArray<FAssetData> ReferencingAssets;
+			if (AssetRegistryModule.GetRegistry().GetAssetsByPackageName(PackageReferencer, ReferencingAssets))
+			{
+				for (const FAssetData& ReferencingAsset : ReferencingAssets)
+				{
+					if (FObjectThumbnail* ObjectThumbnail = LocalObjectThumbnailsCache.Find(ReferencingAsset.GetSoftObjectPath()))
+					{
+						ObjectThumbnail->MarkAsDirty();
+					}
+
+					RefreshThumbnailsFor(ReferencingAsset.GetSoftObjectPath());
+				}
+			}
+		}
+	}
 }
diff --git a/Engine/Source/Editor/UnrealEd/Private/ObjectTools.cpp b/Engine/Source/Editor/UnrealEd/Private/ObjectTools.cpp
--- a/Engine/Source/Editor/UnrealEd/Private/ObjectTools.cpp
+++ b/Engine/Source/Editor/UnrealEd/Private/ObjectTools.cpp
@@ -117,6 +117,7 @@
 #include "Settings/EditorStyleSettings.h"
 #include "ProfilingDebugging/AssetMetadataTrace.h"
 #include "Engine/DataTable.h"
+#include "RHIGPUReadback.h"
 
 DEFINE_LOG_CATEGORY_STATIC(LogObjectTools, Log, All);
 
@@ -5526,13 +5527,107 @@ UE_TRACE_EVENT_END()
 
 namespace ThumbnailTools
 {
+	FAsyncObjectThumbnail::FAsyncObjectThumbnail(int32 InWidth, int32 InHeight, TSharedRef<FRHIGPUTextureReadback> InTextureReadback)
+	: ImageWidth(InWidth)
+	, ImageHeight(InHeight)
+	, TextureReadback(InTextureReadback)
+	{}
 
-	/** Renders a thumbnail for the specified object */
-	void RenderThumbnail( UObject* InObject, const uint32 InImageWidth, const uint32 InImageHeight, EThumbnailTextureFlushMode::Type InFlushMode, FTextureRenderTargetResource* InTextureRenderTargetResource, FObjectThumbnail* OutThumbnail )
+	bool FAsyncObjectThumbnail::IsTextureReadbackReady() const { return TextureReadback->IsReady() && !ObjectThumbnailFuture.IsSet(); }
+	bool FAsyncObjectThumbnail::IsObjectThumbnailReady() const { return ObjectThumbnailFuture.IsSet() && ObjectThumbnailFuture->IsReady(); }
+
+	void FAsyncObjectThumbnail::EnqueueReadbackDataCopy()
+	{
+		check(TextureReadback->IsReady());
+		{
+			ObjectThumbnailPromise = MakeShared<TPromise<FObjectThumbnail>>();
+			ObjectThumbnailFuture = ObjectThumbnailPromise->GetFuture();
+
+			ENQUEUE_RENDER_COMMAND(CopyReadbackDataToThumbnailRTCommand)(
+			[ObjectThumbnailPromise = this->ObjectThumbnailPromise, ImageWidth = this->ImageWidth, ImageHeight = this->ImageHeight, TextureReadback = this->TextureReadback]
+			(FRHICommandListImmediate& RHICmdList)
+			{
+				FObjectThumbnail ObjectThumbnail;
+				ObjectThumbnail.SetImageSize(ImageWidth, ImageHeight);
+
+				int32 ReadbackWidth = 0, ReadbackHeight = 0;
+				void* ReadbackBuffer = TextureReadback->Lock(ReadbackWidth, &ReadbackHeight);
+
+				if (ReadbackBuffer)
+				{
+					TArray<uint8>& ImageData = ObjectThumbnail.AccessImageData();
+					ImageData.Empty();
+					ImageData.AddUninitialized(ObjectThumbnail.GetImageWidth() * ObjectThumbnail.GetImageHeight() * sizeof(FColor));
+
+					check(ReadbackBuffer && ObjectThumbnail.GetImageWidth() <= ReadbackWidth && ObjectThumbnail.GetImageHeight() <= ReadbackHeight);
+
+					if (ObjectThumbnail.GetImageWidth() == ReadbackWidth && ObjectThumbnail.GetImageHeight() == ReadbackHeight)
+					{
+						FMemory::Memcpy( (void*)ImageData.GetData(), ReadbackBuffer, ObjectThumbnail.GetImageWidth() * ObjectThumbnail.GetImageHeight() * sizeof(FColor) );
+					}
+					else // Copy one line at a time
+					{
+						int32 SrcPitchElem = ReadbackWidth;
+						int32 DstPitchElem = ObjectThumbnail.GetImageWidth();
+
+						const uint8* SrcColorData = reinterpret_cast<const uint8*>(ReadbackBuffer);
+
+						for (int32 RowIndex = 0; RowIndex < ObjectThumbnail.GetImageHeight(); RowIndex++)
+						{
+							const void* SrcPtr = SrcColorData + RowIndex * SrcPitchElem * sizeof(FColor);
+							void* DstPtr = ImageData.GetData() + RowIndex * DstPitchElem * sizeof(FColor);
+							FMemory::Memcpy(DstPtr, SrcPtr, DstPitchElem * sizeof(FColor));
+						}
+					}
+				}
+
+				TextureReadback->Unlock();
+
+				ObjectThumbnailPromise->SetValue(MoveTemp(ObjectThumbnail));
+			});
+		}
+	}
+
+	void ReadbackThumbnail(FTextureRenderTargetResource* RenderTargetResource, FObjectThumbnail& OutThumbnail)
+	{
+		const FIntRect InSrcRect(0,	0, OutThumbnail.GetImageWidth(), OutThumbnail.GetImageHeight());
+
+		TArray<uint8>& OutData = OutThumbnail.AccessImageData();
+
+		OutData.Empty();
+		OutData.AddUninitialized(OutThumbnail.GetImageWidth() * OutThumbnail.GetImageHeight() * sizeof(FColor));
+
+		// Copy the contents of the remote texture to system memory
+		// prefer GetRenderTargetImage()
+		RenderTargetResource->ReadPixelsPtr((FColor*)OutData.GetData(), FReadSurfaceDataFlags(), InSrcRect);
+	}
+
+	FAsyncObjectThumbnail ReadbackThumbnailAsync(FTextureRenderTargetResource* RenderTargetResource, const int32 ImageWidth, const int32 ImageHeight)
+	{
+		TSharedRef<FRHIGPUTextureReadback> TextureReadback = MakeShared<FRHIGPUTextureReadback>("ThumbnailReadback");
+		FAsyncObjectThumbnail AsyncObjectThumbnail(ImageWidth, ImageHeight, TextureReadback);
+
+		ENQUEUE_RENDER_COMMAND(ReadbackThumbnailRTCommand)(
+		[RenderTargetResource, ImageWidth, ImageHeight, TextureReadback]
+		(FRHICommandListImmediate& RHICmdList) mutable
+		{
+			RHICmdList.Transition(FRHITransitionInfo(RenderTargetResource->GetRenderTargetTexture(), ERHIAccess::SRVMask, ERHIAccess::CopySrc));
+
+			const FResolveRect SrcRect(0, 0, ImageWidth, ImageHeight);
+			TextureReadback->EnqueueCopy(RHICmdList, RenderTargetResource->GetRenderTargetTexture(), SrcRect);
+
+			RHICmdList.Transition(FRHITransitionInfo(RenderTargetResource->GetRenderTargetTexture(), ERHIAccess::CopySrc, ERHIAccess::SRVMask));
+		});
+
+		return AsyncObjectThumbnail;
+	}
+
+	bool RenderThumbnailInternal( UObject* InObject, uint32& InOutImageWidth, uint32& InOutImageHeight, EThumbnailTextureFlushMode::Type InFlushMode,
+								  FTextureRenderTargetResource*& InOutTextureRenderTargetResource )
 	{
 		if (!FApp::CanEverRender())
 		{
-			return;
+			return {};
 		}
 		
 #if CPUPROFILERTRACE_ENABLED
@@ -5543,29 +5638,24 @@ namespace ThumbnailTools
 		// Renderer must be initialized before generating thumbnails
 		check( GIsRHIInitialized );
 
-		// Store dimensions
-		if ( OutThumbnail )
-		{
-			OutThumbnail->SetImageSize( InImageWidth, InImageHeight );
-		}
-
 		// Grab the actual render target resource from the texture.  Note that we're absolutely NOT ALLOWED to
 		// dereference this pointer.  We're just passing it along to other functions that will use it on the render
 		// thread.  The only thing we're allowed to do is check to see if it's NULL or not.
-		FTextureRenderTargetResource* RenderTargetResource = InTextureRenderTargetResource;
+		FTextureRenderTargetResource* RenderTargetResource = InOutTextureRenderTargetResource;
 		if ( RenderTargetResource == NULL )
 		{
 			// No render target was supplied, just use a scratch texture render target
-			const uint32 MinRenderTargetSize = FMath::Max( InImageWidth, InImageHeight );
+			const uint32 MinRenderTargetSize = FMath::Max( InOutImageWidth, InOutImageHeight );
 			UTextureRenderTarget2D* RenderTargetTexture = GEditor->GetScratchRenderTarget( MinRenderTargetSize );
 			check( RenderTargetTexture != NULL );
 
 			// Make sure the input dimensions are OK.  The requested dimensions must be less than or equal to
 			// our scratch render target size.
-			check( InImageWidth <= RenderTargetTexture->GetSurfaceWidth() );
-			check( InImageHeight <= RenderTargetTexture->GetSurfaceHeight() );
+			check( InOutImageWidth <= RenderTargetTexture->GetSurfaceWidth() );
+			check( InOutImageHeight <= RenderTargetTexture->GetSurfaceHeight() );
 
 			RenderTargetResource = RenderTargetTexture->GameThread_GetRenderTargetResource();
+			InOutTextureRenderTargetResource = RenderTargetResource;
 		}
 		check( RenderTargetResource != NULL );
 
@@ -5614,9 +5704,9 @@ namespace ThumbnailTools
 
 			const float ZoomFactor = 1.0f;
 
-			uint32 DrawWidth = InImageWidth;
-			uint32 DrawHeight = InImageHeight;
-			if ( OutThumbnail )
+			uint32 DrawWidth = InOutImageWidth;
+			uint32 DrawHeight = InOutImageHeight;
+
 			{
 				// Find how big the thumbnail WANTS to be
 				uint32 DesiredWidth = 0;
@@ -5639,21 +5729,23 @@ namespace ThumbnailTools
 				if( DesiredWidth > 0 && DesiredHeight > 0 )
 				{
 					// Scale the desired size down if it's too big, preserving aspect ratio
-					if( DesiredWidth > InImageWidth )
+					if( DesiredWidth > InOutImageWidth )
 					{
-						DesiredHeight = ( DesiredHeight * InImageWidth ) / DesiredWidth;
-						DesiredWidth = InImageWidth;
+						DesiredHeight = ( DesiredHeight * InOutImageWidth ) / DesiredWidth;
+						DesiredWidth = InOutImageWidth;
 					}
-					if( DesiredHeight > InImageHeight )
+					if( DesiredHeight > InOutImageHeight )
 					{
-						DesiredWidth = ( DesiredWidth * InImageHeight ) / DesiredHeight;
-						DesiredHeight = InImageHeight;
+						DesiredWidth = ( DesiredWidth * InOutImageHeight ) / DesiredHeight;
+						DesiredHeight = InOutImageHeight;
 					}
 
 					// Update dimensions
 					DrawWidth = FMath::Max<uint32>(1, DesiredWidth);
 					DrawHeight = FMath::Max<uint32>(1, DesiredHeight);
-					OutThumbnail->SetImageSize( DrawWidth, DrawHeight );
+
+					InOutImageWidth = DrawWidth;
+					InOutImageHeight = DrawHeight;
 				}
 			}
 
@@ -5679,31 +5771,48 @@ namespace ThumbnailTools
 			// Tell the rendering thread to draw any remaining batched elements
 			Canvas.Flush_GameThread();
 
-
-			{
-				ENQUEUE_RENDER_COMMAND(UpdateThumbnailRTCommand)(
-					[RenderTargetResource](FRHICommandListImmediate& RHICmdList)
-					{
-						TransitionAndCopyTexture(RHICmdList, RenderTargetResource->GetRenderTargetTexture(), RenderTargetResource->TextureRHI, {});
-					});
-
-				if(OutThumbnail)
+			ENQUEUE_RENDER_COMMAND(UpdateThumbnailRTCommand)(
+				[RenderTargetResource] (FRHICommandListImmediate& RHICmdList)
 				{
-					const FIntRect InSrcRect(0,	0, OutThumbnail->GetImageWidth(), OutThumbnail->GetImageHeight());
+					TransitionAndCopyTexture(RHICmdList, RenderTargetResource->GetRenderTargetTexture(), RenderTargetResource->TextureRHI, {});
+				});
+		}
 
-					TArray<uint8>& OutData = OutThumbnail->AccessImageData();
+		return true;
+	}
 
-					OutData.Empty();
-					OutData.AddUninitialized(OutThumbnail->GetImageWidth() * OutThumbnail->GetImageHeight() * sizeof(FColor));
+	/** Renders a thumbnail for the specified object */
+	void RenderThumbnail( UObject* InObject, const uint32 InImageWidth, const uint32 InImageHeight, EThumbnailTextureFlushMode::Type InFlushMode, FTextureRenderTargetResource* InTextureRenderTargetResource, FObjectThumbnail* OutThumbnail )
+	{
+		uint32 ImageWidth = InImageWidth;
+		uint32 ImageHeight = InImageHeight;
 
-					// Copy the contents of the remote texture to system memory
-					// prefer GetRenderTargetImage()
-					RenderTargetResource->ReadPixelsPtr((FColor*)OutData.GetData(), FReadSurfaceDataFlags(), InSrcRect);
-				}
-			}
+		const bool bThumbnailRendered = RenderThumbnailInternal(InObject, ImageWidth, ImageHeight, InFlushMode, InTextureRenderTargetResource);
+
+		if (OutThumbnail && bThumbnailRendered)
+		{
+			OutThumbnail->SetImageSize(ImageWidth, ImageHeight);
+			ReadbackThumbnail(InTextureRenderTargetResource, *OutThumbnail);
 		}
 	}
 
+	TOptional<FAsyncObjectThumbnail> RenderThumbnailAsync( UObject* InObject, const uint32 InImageWidth, const uint32 InImageHeight, EThumbnailTextureFlushMode::Type InFlushMode,
+															   FTextureRenderTargetResource* InTextureRenderTargetResource )
+	{
+		uint32 ImageWidth = InImageWidth;
+		uint32 ImageHeight = InImageHeight;
+
+		const bool bThumbnailRendered = RenderThumbnailInternal(InObject, ImageWidth, ImageHeight, InFlushMode, InTextureRenderTargetResource);
+
+		if (bThumbnailRendered)
+		{
+			return ReadbackThumbnailAsync(InTextureRenderTargetResource, ImageWidth, ImageHeight);
+		}
+		else
+		{
+			return {};
+		}
+	}
 
 	/** Generates a thumbnail for the specified object and caches it */
 	FObjectThumbnail* GenerateThumbnailForObjectToSaveToDisk( UObject* InObject )
diff --git a/Engine/Source/Editor/UnrealEd/Private/ThumbnailRendering/AnimBlueprintThumbnailRenderer.cpp b/Engine/Source/Editor/UnrealEd/Private/ThumbnailRendering/AnimBlueprintThumbnailRenderer.cpp
--- a/Engine/Source/Editor/UnrealEd/Private/ThumbnailRendering/AnimBlueprintThumbnailRenderer.cpp
+++ b/Engine/Source/Editor/UnrealEd/Private/ThumbnailRendering/AnimBlueprintThumbnailRenderer.cpp
@@ -44,7 +44,10 @@ void UAnimBlueprintThumbnailRenderer::Draw(UObject* Object, int32 X, int32 Y, ui
 	UAnimBlueprint* AnimBlueprint = Cast<UAnimBlueprint>(Object);
 	if (AnimBlueprint && AnimBlueprint->GeneratedClass)
 	{
-		TSharedRef<FAnimBlueprintThumbnailScene> ThumbnailScene = ThumbnailScenes.EnsureThumbnailScene(AnimBlueprint->GeneratedClass);
+		if (!ThumbnailScene)
+		{
+			ThumbnailScene.Emplace();
+		}
 
 		if(ThumbnailScene->SetAnimBlueprint(AnimBlueprint))
 		{
@@ -64,7 +67,7 @@ void UAnimBlueprintThumbnailRenderer::Draw(UObject* Object, int32 X, int32 Y, ui
 
 void UAnimBlueprintThumbnailRenderer::BeginDestroy()
 {
-	ThumbnailScenes.Clear();
+	ThumbnailScene.Reset();
 
 	Super::BeginDestroy();
 }
diff --git a/Engine/Source/Editor/UnrealEd/Private/ThumbnailRendering/BlueprintThumbnailRenderer.cpp b/Engine/Source/Editor/UnrealEd/Private/ThumbnailRendering/BlueprintThumbnailRenderer.cpp
--- a/Engine/Source/Editor/UnrealEd/Private/ThumbnailRendering/BlueprintThumbnailRenderer.cpp
+++ b/Engine/Source/Editor/UnrealEd/Private/ThumbnailRendering/BlueprintThumbnailRenderer.cpp
@@ -9,19 +9,6 @@
 
 #include UE_INLINE_GENERATED_CPP_BY_NAME(BlueprintThumbnailRenderer)
 
-static FAutoConsoleCommand CVarClearCachedBPThumbnailScenes(
-	TEXT("Editor.Thumbnails.ClearCachedBPThumbnailScenes"),
-	TEXT("Clears cached scenes used for Blueprint thumbnail rendering."),
-	FConsoleCommandDelegate::CreateLambda([]()
-	{
-		for (TObjectIterator<UBlueprintThumbnailRenderer> ThumbnailRendererIt; ThumbnailRendererIt; ++ThumbnailRendererIt)
-		{
-			if (UBlueprintThumbnailRenderer* ThumbnailRenderer = *ThumbnailRendererIt)
-			{
-				ThumbnailRenderer->ClearThumbnailScenes();
-			}
-		}
-	}));
 
 UBlueprintThumbnailRenderer::UBlueprintThumbnailRenderer(const FObjectInitializer& ObjectInitializer)
 	: Super(ObjectInitializer)
@@ -61,7 +48,10 @@ void UBlueprintThumbnailRenderer::Draw(UObject* Object, int32 X, int32 Y, uint32
 		&& !Blueprint->HasAnyFlags(RF_Transient);
 	if (bIsBlueprintValid)
 	{
-		TSharedRef<FBlueprintThumbnailScene> ThumbnailScene = ThumbnailScenes.EnsureThumbnailScene(Blueprint->GeneratedClass);
+		if (!ThumbnailScene)
+		{
+			ThumbnailScene.Emplace();
+		}
 
 		ThumbnailScene->SetBlueprint(Blueprint);
 		FSceneViewFamilyContext ViewFamily( FSceneViewFamily::ConstructionValues( RenderTarget, ThumbnailScene->GetScene(), FEngineShowFlags(ESFIM_Game) )
@@ -78,7 +68,7 @@ void UBlueprintThumbnailRenderer::Draw(UObject* Object, int32 X, int32 Y, uint32
 void UBlueprintThumbnailRenderer::BeginDestroy()
 {
 	FKismetEditorUtilities::OnBlueprintUnloaded.RemoveAll(this);
-	ThumbnailScenes.Clear();
+	ThumbnailScene.Reset();
 
 	Super::BeginDestroy();
 }
@@ -87,8 +77,7 @@ void UBlueprintThumbnailRenderer::BlueprintChanged(UBlueprint* Blueprint)
 {
 	if (Blueprint && Blueprint->GeneratedClass)
 	{
-		TSharedPtr<FBlueprintThumbnailScene> ThumbnailScene = ThumbnailScenes.FindThumbnailScene(Blueprint->GeneratedClass);
-		if (ThumbnailScene.IsValid())
+		if (ThumbnailScene)
 		{
 			ThumbnailScene->BlueprintChanged(Blueprint);
 		}
@@ -99,11 +88,9 @@ void UBlueprintThumbnailRenderer::OnBlueprintUnloaded(UBlueprint* Blueprint)
 {
 	if (Blueprint && Blueprint->GeneratedClass)
 	{
-		ThumbnailScenes.RemoveThumbnailScene(Blueprint->GeneratedClass);
+		if (ThumbnailScene && ThumbnailScene->GetCurrentBlueprint() == Blueprint)
+		{
+			ThumbnailScene.Reset();
+		}
 	}
 }
-
-void UBlueprintThumbnailRenderer::ClearThumbnailScenes()
-{
-	ThumbnailScenes.Clear();
-}
diff --git a/Engine/Source/Editor/UnrealEd/Private/ThumbnailRendering/ClassThumbnailRenderer.cpp b/Engine/Source/Editor/UnrealEd/Private/ThumbnailRendering/ClassThumbnailRenderer.cpp
--- a/Engine/Source/Editor/UnrealEd/Private/ThumbnailRendering/ClassThumbnailRenderer.cpp
+++ b/Engine/Source/Editor/UnrealEd/Private/ThumbnailRendering/ClassThumbnailRenderer.cpp
@@ -41,7 +41,10 @@ void UClassThumbnailRenderer::Draw(UObject* Object, int32 X, int32 Y, uint32 Wid
 	UClass* Class = Cast<UClass>(Object);
 	if (Class != nullptr)
 	{
-		TSharedRef<FClassThumbnailScene> ThumbnailScene = ThumbnailScenes.EnsureThumbnailScene(Class);
+		if (!ThumbnailScene)
+		{
+			ThumbnailScene.Emplace();
+		}
 
 		ThumbnailScene->SetClass(Class);
 		FSceneViewFamilyContext ViewFamily( FSceneViewFamily::ConstructionValues( RenderTarget, ThumbnailScene->GetScene(), FEngineShowFlags(ESFIM_Game) )
@@ -58,12 +61,15 @@ void UClassThumbnailRenderer::Draw(UObject* Object, int32 X, int32 Y, uint32 Wid
 void UClassThumbnailRenderer::BeginDestroy()
 {
 	FKismetEditorUtilities::OnBlueprintGeneratedClassUnloaded.RemoveAll(this);
-	ThumbnailScenes.Clear();
+	ThumbnailScene.Reset();
 
 	Super::BeginDestroy();
 }
 
 void UClassThumbnailRenderer::OnBlueprintGeneratedClassUnloaded(UBlueprintGeneratedClass* BPGC)
 {
-	ThumbnailScenes.RemoveThumbnailScene(BPGC);
+	if (ThumbnailScene && ThumbnailScene->GetCurrentClass() == BPGC)
+	{
+		ThumbnailScene.Reset();
+	}
 }
diff --git a/Engine/Source/Editor/UnrealEd/Private/ThumbnailRendering/ThumbnailRenderer.cpp b/Engine/Source/Editor/UnrealEd/Private/ThumbnailRendering/ThumbnailRenderer.cpp
--- a/Engine/Source/Editor/UnrealEd/Private/ThumbnailRendering/ThumbnailRenderer.cpp
+++ b/Engine/Source/Editor/UnrealEd/Private/ThumbnailRendering/ThumbnailRenderer.cpp
@@ -15,21 +15,6 @@ UThumbnailRenderer::UThumbnailRenderer(const FObjectInitializer& ObjectInitializ
 {
 }
 
-// static
-PRAGMA_DISABLE_DEPRECATION_WARNINGS
-void UThumbnailRenderer::RenderViewFamily(FCanvas* Canvas, FSceneViewFamily* ViewFamily)
-{
-	if ((ViewFamily == nullptr) || ViewFamily->Views.IsEmpty() || (ViewFamily->Views[0] == nullptr))
-	{
-		return;
-	}
-
-	// The new prototype of RenderViewFamily takes a non-const FSceneView in parameter since view extensions can still modify it as well as as the view family before rendering : 
-	FSceneView* View = const_cast<FSceneView*>(ViewFamily->Views[0]);
-	UThumbnailRenderer::RenderViewFamily(Canvas, ViewFamily, View);
-}
-PRAGMA_ENABLE_DEPRECATION_WARNINGS
-
 // static
 void UThumbnailRenderer::RenderViewFamily(FCanvas* Canvas, FSceneViewFamily* ViewFamily, FSceneView* View)
 {
diff --git a/Engine/Source/Editor/UnrealEd/Public/AssetThumbnail.h b/Engine/Source/Editor/UnrealEd/Public/AssetThumbnail.h
--- a/Engine/Source/Editor/UnrealEd/Public/AssetThumbnail.h
+++ b/Engine/Source/Editor/UnrealEd/Public/AssetThumbnail.h
@@ -10,6 +10,7 @@
 #include "Widgets/SWidget.h"
 #include "Widgets/SNullWidget.h"
 #include "TickableEditorObject.h"
+#include "ObjectTools.h"
 
 class AActor;
 class FAssetThumbnailPool;
@@ -175,16 +176,16 @@ public:
 	UNREALED_API const FAssetData& GetAssetData() const;
 
 	/**
-	 * Sets the asset to render the thumnail for 
+	 * Sets the asset to render the thumbnail for 
 	 *
 	 * @param InAsset	The new asset
 	 */
 	UNREALED_API void SetAsset( const UObject* InAsset );
 
 	/**
-	 * Sets the asset to render the thumnail for
+	 * Sets the asset to render the thumbnail for
 	 *
-	 * @param InAssetData	Asset data containin the the new asset to render
+	 * @param InAssetData	Asset data containing the the new asset to render
 	 */
 	UNREALED_API void SetAsset( const FAssetData& InAssetData );
 
@@ -298,6 +299,7 @@ public:
 	UNREALED_API void SetRealTimeThumbnail(const TSharedPtr<FAssetThumbnail>& Thumbnail, bool bRealTimeThumbnail);
 
 private:
+	struct FThumbId;
 
 	/**
 	 * Releases all rendering resources held by the pool
@@ -315,6 +317,8 @@ private:
 
 	/** Adds the thumbnails associated with the object found at ObjectPath to the render stack */
 	void RefreshThumbnailsFor( const FSoftObjectPath& ObjectPath );
+	/** Adds the thumbnails associated with the ThumbId to the render stack */
+	void RefreshThumbnailFor( const FThumbId ThumbId );
 
 	/** Handler for when an asset is loaded */
 	void OnAssetLoaded( UObject* Asset );
@@ -322,6 +326,12 @@ private:
 	/** Handler for when a thumbnail gets flagged as dirty. Used to refresh the thumbnail. */
 	void OnThumbnailDirtied( const FSoftObjectPath& ObjectPath );
 
+	/** Copies any completed object thumbnails to the local object thumbnail cache */
+	void ProcessPendingObjectThumbnails();
+
+	/** Requests a refresh for recently loaded assets for which we have a thumbnail. Done outside of the OnAssetLoaded callback so as to minimize loading disruptions. */
+	void ProcessRecentlyLoadedAssets();
+
 private:
 	/** Information about a thumbnail */
 	struct FThumbnailInfo
@@ -348,12 +358,26 @@ private:
 	 * Assign a thumbnail from its render target and re-render it if necessary.
 	 *
 	 * @param ThumbnailInfo The thumbnail info to assign a texture to
-	 * @param bIsAssetStillCompiling If the asset we want to load the thumbnail is compiling, this flag will be set to true, it wont be touch in other cases.
+	 * @param bOutRetryOnFailure If the asset we want to load the thumbnail is compiling, this flag will be set to true, it wont be touch in other cases.
 	 * @param CustomAssetToRender The asset to render when generating the texture
 	 *
 	 * @return true if the thumbnail was assigned to a valid texture
 	 */
-	bool LoadThumbnail(TSharedRef<FThumbnailInfo> ThumbnailInfo, bool& bIsAssetStillCompiling, const FAssetData& CustomAssetToRender = FAssetData());
+	bool LoadThumbnail(TSharedRef<FThumbnailInfo> ThumbnailInfo, bool& bOutRetryOnFailure, const FAssetData& CustomAssetToRender = FAssetData());
+	bool LoadThumbnailForLoadedAsset(const TSharedRef<FThumbnailInfo>& ThumbnailInfo, bool& bOutRetryOnFailure, const FAssetData& CustomAssetToRender);
+	bool LoadThumbnailForUnloadedAsset(const TSharedRef<FThumbnailInfo>& ThumbnailInfo, const FAssetData& CustomAssetToRender);
+	
+	/** Tries to load the thumbnail from an object thumbnail cached locally or on the asset package. */
+	bool LoadThumbnailFromObjectThumbnailInCache(const TSharedRef<FThumbnailInfo>& ThumbnailInfo, const FAssetData& CustomAssetToRender);
+
+	/** Tries to load the thumbnail from the package on disk. */
+	bool LoadThumbnailFromObjectThumbnailOnDisk(const TSharedRef<FThumbnailInfo>& ThumbnailInfo, const FAssetData& CustomAssetToRender);
+
+	/** Tries to render a fresh thumbnail if the thumbnail render frequency allows it. */
+	bool MakeFreshThumbnailFromRender(const TSharedRef<FThumbnailInfo>& ThumbnailInfo, bool& bOutRetryOnFailure, const FAssetData& CustomAssetToRender);
+
+	/** Copies the image data from the object thumbnail to the slate texture in the thumbnail info. */
+	bool CopyObjectThumbnailImageDataToSlateTextureData(const TSharedRef<FThumbnailInfo>& ThumbnailInfo, const FObjectThumbnail& ObjectThumbnail);
 
 	struct FThumbnailInfo_RenderThread
 	{
@@ -418,21 +442,21 @@ private:
 	/** List of free thumbnails that can be reused */
 	TArray< TSharedRef<FThumbnailInfo> > FreeThumbnails;
 
+	/** Cached object thumbnails that we rendered. Ignored when a thumbnail refresh is explicitly requested or thumbnails were marked as dirty for the asset */
+	TMap< FSoftObjectPath, FObjectThumbnail > LocalObjectThumbnailsCache;
+
+	/** Object thumbnails that are waiting for their textures to be rendered */
+	TMap< FSoftObjectPath, ThumbnailTools::FAsyncObjectThumbnail > PendingObjectThumbnails;
+
 	/** A mapping of objects to the number of referencers */
 	TMap< FThumbId, int32 > RefCountMap;
 
 	/** A list of object paths for recently loaded assets whose thumbnails need to be refreshed. */
 	TArray<FSoftObjectPath> RecentlyLoadedAssets;
 
-	/** Attribute that determines if real-time thumbnails are allowed. Called every frame. */
-	TAttribute<bool> AreRealTimeThumbnailsAllowed;
-
 	/** Max number of thumbnails in the pool */
 	uint32 NumInPool;
 
-	/** Shaders are still building */
-	bool bWereShadersCompilingLastFrame = false;
-
 	/** Max number of dynamic thumbnails to update per frame */
 	uint32 MaxRealTimeThumbnailsPerFrame;
 
diff --git a/Engine/Source/Editor/UnrealEd/Public/ObjectTools.h b/Engine/Source/Editor/UnrealEd/Public/ObjectTools.h
--- a/Engine/Source/Editor/UnrealEd/Public/ObjectTools.h
+++ b/Engine/Source/Editor/UnrealEd/Public/ObjectTools.h
@@ -738,7 +738,7 @@ namespace ObjectTools
 }
 
 
-
+class FRHIGPUTextureReadback;
 
 namespace ThumbnailTools
 {
@@ -764,10 +764,58 @@ namespace ThumbnailTools
 	 * @param	InObject				the rendered thumbnail will represent this object
 	 * @param	InImageWidth			the maximum width of the thumbnail. It may be smaller when supplying a OutThumbnail.
 	 * @param	InImageHeight			the maximum height of the thumbnail. It may be smaller when supplying a OutThumbnail.
+	 * @param	InFlushMode				Flushes or not texture streaming before rendering
 	 * @param	InRenderTargetResource	if non-NULL, the render target that will receive the thumbnail.
 	 * @param	OutThumbnail			if non-NULL, the FObjectThumbnail that will receive the thumbnail data. This would also force a synchronization with the RenderThread
 	 */
-	UNREALED_API void RenderThumbnail( UObject* InObject, const uint32 InImageWidth, const uint32 InImageHeight, EThumbnailTextureFlushMode::Type InFlushMode, FTextureRenderTargetResource* InRenderTargetResource = NULL, FObjectThumbnail* OutThumbnail = NULL );
+	UNREALED_API void RenderThumbnail( UObject* InObject, const uint32 InImageWidth, const uint32 InImageHeight, EThumbnailTextureFlushMode::Type InFlushMode,
+									   FTextureRenderTargetResource* InRenderTargetResource = NULL, FObjectThumbnail* OutThumbnail = NULL );
+
+	/**
+	 * Async process to get a texture copied to an object thumbnail
+	 * 1. Wait for the TextureReadback to be ready
+	 * 2. Once TextureReadback is ready, enqueue a render command to copy its data to the object thumbnail
+	 * 3. Once the render command completes, ObjectThumbnail can be read
+	 */
+	struct FAsyncObjectThumbnail
+	{
+		explicit FAsyncObjectThumbnail(int32 InWidth, int32 InHeight, TSharedRef<FRHIGPUTextureReadback> InTextureReadback);
+
+		// Checks if the texture data is available to be readback
+		bool IsTextureReadbackReady() const;
+
+		// Enqueues a render command to copy the texture readback data to the ObjectThumbnail
+		void EnqueueReadbackDataCopy();
+
+		// ObjectThumbnail is available
+		bool IsObjectThumbnailReady() const;
+
+		/**
+		 * Returns a copy of ObjectThumbnail and empties the Future
+		 * IsObjectThumbnailReady must be true before calling this function
+		 */
+		FObjectThumbnail ConsumeObjectThumbnailUnchecked() { return ObjectThumbnailFuture->Consume(); }
+
+	private:
+		int32 ImageWidth, ImageHeight;
+
+		TSharedRef<FRHIGPUTextureReadback> TextureReadback;
+		TSharedPtr<TPromise<FObjectThumbnail>> ObjectThumbnailPromise;
+		TOptional<TFuture<FObjectThumbnail>> ObjectThumbnailFuture;
+	};
+
+	/**
+	 * Renders a thumbnail for the specified object.
+	 * A FObjectThumbnail will be filled once the rendering is complete.
+	 *
+	 * @param	InObject					the rendered thumbnail will represent this object
+	 * @param	InImageWidth				the maximum width of the thumbnail. It may be smaller when supplying a OutThumbnail.
+	 * @param	InImageHeight				the maximum height of the thumbnail. It may be smaller when supplying a OutThumbnail.
+	 * @param	InFlushMode					Flushes or not texture streaming before rendering
+	 * @param	InRenderTargetResource		if non-NULL, the render target that will receive the thumbnail.
+	 */
+	UNREALED_API TOptional<FAsyncObjectThumbnail> RenderThumbnailAsync( UObject* InObject, const uint32 InImageWidth, const uint32 InImageHeight, EThumbnailTextureFlushMode::Type InFlushMode,
+																			FTextureRenderTargetResource* InRenderTargetResource = nullptr );
 
 	/** Generates a thumbnail for the specified object and caches it. Used to generate the thumbnail that gets saved in a package. */
 	UNREALED_API FObjectThumbnail* GenerateThumbnailForObjectToSaveToDisk( UObject* InObject );
diff --git a/Engine/Source/Editor/UnrealEd/Public/ThumbnailHelpers.h b/Engine/Source/Editor/UnrealEd/Public/ThumbnailHelpers.h
--- a/Engine/Source/Editor/UnrealEd/Public/ThumbnailHelpers.h
+++ b/Engine/Source/Editor/UnrealEd/Public/ThumbnailHelpers.h
@@ -322,6 +322,7 @@ public:
 	/** Refreshes components for the specified blueprint */
 	UNREALED_API void BlueprintChanged(class UBlueprint* Blueprint);
 
+	UNREALED_API TWeakObjectPtr<class UBlueprint> GetCurrentBlueprint() const { return CurrentBlueprint; }
 protected:
 
 	/** Get the scene thumbnail info to use for the object currently being rendered */
@@ -341,6 +342,8 @@ public:
 	/** Sets the class use in the next CreateView() */
 	UNREALED_API void SetClass(class UClass* Class);
 
+	UNREALED_API UClass* GetCurrentClass() const { return CurrentClass; }
+
 protected:
 	/** Get the scene thumbnail info to use for the object currently being rendered */
 	UNREALED_API virtual USceneThumbnailInfo* GetSceneThumbnailInfo(const float TargetDistance) const override;
@@ -427,8 +430,8 @@ public:
 
 							if (ScenesToDelete.IsEmpty())
 							{
-								GEditor->GetTimerManager()->ClearTimer(SceneDeleteTimerHandle);
 								ScenesToDelete.Shrink();
+								GEditor->GetTimerManager()->ClearTimer(SceneDeleteTimerHandle); // TimerManager owns this lambda and clearing the timer will destroy it
 								SceneDeleteTimerHandle.Invalidate();
 							}
 						}), 2.0f/*InRate*/, true /*bInLoop*/, 1.0f/*InFirstDelay*/);
diff --git a/Engine/Source/Runtime/Core/Private/Misc/ObjectThumbnail.cpp b/Engine/Source/Runtime/Core/Private/Misc/ObjectThumbnail.cpp
--- a/Engine/Source/Runtime/Core/Private/Misc/ObjectThumbnail.cpp
+++ b/Engine/Source/Runtime/Core/Private/Misc/ObjectThumbnail.cpp
@@ -127,6 +127,11 @@ void FObjectThumbnail::DecompressImageData()
 	}
 }
 
+void FObjectThumbnail::EmptyUncompressedImageData()
+{
+	ImageData.Empty();
+}
+
 FThumbnailCompressionInterface* FObjectThumbnail::GetCompressor() const
 {
 	if (!CompressedImageData.IsEmpty() && ImageWidth > 0 && ImageHeight > 0)
diff --git a/Engine/Source/Runtime/Core/Public/Misc/ObjectThumbnail.h b/Engine/Source/Runtime/Core/Public/Misc/ObjectThumbnail.h
--- a/Engine/Source/Runtime/Core/Public/Misc/ObjectThumbnail.h
+++ b/Engine/Source/Runtime/Core/Public/Misc/ObjectThumbnail.h
@@ -211,6 +211,9 @@ public:
 	/** Decompress image data. */
 	CORE_API void DecompressImageData();
 
+	/** Empties the uncompressed image data. The next call to GetImage will have to decompress the image again. */
+	CORE_API void EmptyUncompressedImageData();
+
 	/**
 	 * Calculates the memory usage of this FObjectThumbnail.
 	 *
