Implement PropertyNames strings gather filter by outer types

#tests Compile UE project. Locally run Localize command with modified PropertyNames.ini file (add FieldOuterTypesToExclude=Function entry). Observe reduced number of strings being gathered.

diff --git a/Engine/Source/Developer/Localization/Private/LocalizationConfigurationScript.cpp b/Engine/Source/Developer/Localization/Private/LocalizationConfigurationScript.cpp
--- a/Engine/Source/Developer/Localization/Private/LocalizationConfigurationScript.cpp
+++ b/Engine/Source/Developer/Localization/Private/LocalizationConfigurationScript.cpp
@@ -458,6 +458,16 @@ namespace LocalizationConfigurationScript
 				ConfigSection.Add(TEXT("FieldOwnerTypesToExclude"), FieldOwnerTypeToExclude);
 			}
 
+			// Field Outer Type Filters
+			for (const FString& FieldOuterTypeToInclude : Target->Settings.GatherFromMetaData.FieldOuterTypesToInclude)
+			{
+				ConfigSection.Add(TEXT("FieldOuterTypesToInclude"), FieldOuterTypeToInclude);
+			}
+			for (const FString& FieldOuterTypeToExclude : Target->Settings.GatherFromMetaData.FieldOuterTypesToExclude)
+			{
+				ConfigSection.Add(TEXT("FieldOuterTypesToExclude"), FieldOuterTypeToExclude);
+			}
+
 			ConfigSection.Add( TEXT("ShouldGatherFromEditorOnlyData"), Target->Settings.GatherFromMetaData.ShouldGatherFromEditorOnlyData ? TEXT("true") : TEXT("false") );
 
 			Script.AddGatherTextStep(GatherTextStepIndex++, MoveTemp(ConfigSection));
diff --git a/Engine/Source/Developer/Localization/Public/LocalizationTargetTypes.h b/Engine/Source/Developer/Localization/Public/LocalizationTargetTypes.h
--- a/Engine/Source/Developer/Localization/Public/LocalizationTargetTypes.h
+++ b/Engine/Source/Developer/Localization/Public/LocalizationTargetTypes.h
@@ -306,6 +306,14 @@ struct FGatherTextFromMetaDataConfiguration
 	UPROPERTY(config, EditAnywhere, Category = "MetaData")
 	TArray<FString> FieldOwnerTypesToExclude;
 
+	/** List of outer types (eg, Property, Function, ScriptStruct, etc) of fields that should be included in the gather, or empty to include everything. */
+	UPROPERTY(config, EditAnywhere, Category = "MetaData")
+	TArray<FString> FieldOuterTypesToInclude;
+
+	/** List of outer types (eg, Property, Function, ScriptStruct, etc) of fields that should be excluded from the gather. */
+	UPROPERTY(config, EditAnywhere, Category = "MetaData")
+	TArray<FString> FieldOuterTypesToExclude;
+
 	/* If enabled, data that is specified as editor-only may be processed for gathering. */
 	UPROPERTY(config, EditAnywhere, Category = "Filter")
 	bool ShouldGatherFromEditorOnlyData;
diff --git a/Engine/Source/Editor/UnrealEd/Classes/Commandlets/GatherTextFromMetadataCommandlet.h b/Engine/Source/Editor/UnrealEd/Classes/Commandlets/GatherTextFromMetadataCommandlet.h
--- a/Engine/Source/Editor/UnrealEd/Classes/Commandlets/GatherTextFromMetadataCommandlet.h
+++ b/Engine/Source/Editor/UnrealEd/Classes/Commandlets/GatherTextFromMetadataCommandlet.h
@@ -99,4 +99,8 @@ private:
 	/** Array of field owner types (eg, UMyClass, FMyStruct, etc) that should have fields within them included or excluded in the current gather */
 	TArray<const UStruct*> FieldOwnerTypesToInclude;
 	TArray<const UStruct*> FieldOwnerTypesToExclude;
+
+	/** Array of field outer types (eg, FProperty, UFunction, UScriptStruct, etc) of fields that should be included or excluded in the current gather */
+	TArray<const UStruct*> FieldOuterTypesToInclude;
+	TArray<const UStruct*> FieldOuterTypesToExclude;
 };
diff --git a/Engine/Source/Editor/UnrealEd/Private/Commandlets/GatherTextFromMetadataCommandlet.cpp b/Engine/Source/Editor/UnrealEd/Private/Commandlets/GatherTextFromMetadataCommandlet.cpp
--- a/Engine/Source/Editor/UnrealEd/Private/Commandlets/GatherTextFromMetadataCommandlet.cpp
+++ b/Engine/Source/Editor/UnrealEd/Private/Commandlets/GatherTextFromMetadataCommandlet.cpp
@@ -299,6 +299,69 @@ int32 UGatherTextFromMetaDataCommandlet::Main( const FString& Params )
 		GetFieldOwnerTypesArrayFromConfig(TEXT("FieldOwnerTypesToExclude"), FieldOwnerTypesToExclude);
 	}
 
+	// FieldOuterTypesToInclude/FieldOuterTypesToExclude
+	{
+		auto GetFieldOuterTypesArrayFromConfig = [this, &SectionName, &GatherTextConfigPath](const TCHAR* InConfigKey, TArray<const UStruct*>& OutFieldOuterTypes)
+		{
+			TArray<FString> FieldOuterTypeStrs;
+			GetStringArrayFromConfig(*SectionName, InConfigKey, FieldOuterTypeStrs, GatherTextConfigPath);
+
+			if (FieldOuterTypeStrs.Num() == 0)
+			{
+				return;
+			}
+
+			TArray<const UStruct*> AllFieldOuterClassTypes;
+			TArray<const UStruct*> AllFieldOuterScriptStructTypes;
+			GetObjectsOfClass(UClass::StaticClass(), (TArray<UObject*>&)AllFieldOuterClassTypes, false);
+			GetObjectsOfClass(UScriptStruct::StaticClass(), (TArray<UObject*>&)AllFieldOuterScriptStructTypes, false);
+
+			for (const FString& FieldOuterTypeStr : FieldOuterTypeStrs)
+			{
+				const bool bIsWildcard = FieldOuterTypeStr.Contains(TEXT("*")) || FieldOuterTypeStr.Contains(TEXT("?"));
+				if (bIsWildcard)
+				{
+					auto AddFieldOuterTypesMatchingWildcard = [&FieldOuterTypeStr, &OutFieldOuterTypes](const TArray<const UStruct*>& AllFieldOuterTypes)
+					{
+						for (const UStruct* Candidate : AllFieldOuterTypes)
+						{
+							if (Candidate->GetName().MatchesWildcard(FieldOuterTypeStr))
+							{
+								OutFieldOuterTypes.Add(Candidate);
+							}
+						}
+					};
+					AddFieldOuterTypesMatchingWildcard(AllFieldOuterClassTypes);
+					AddFieldOuterTypesMatchingWildcard(AllFieldOuterScriptStructTypes);
+				}
+				else
+				{
+					const UStruct* OuterType = FindFirstObject<UStruct>(*FieldOuterTypeStr, EFindFirstObjectOptions::EnsureIfAmbiguous);
+					if (!OuterType)
+					{
+						UE_LOGFMT(LogGatherTextFromMetaDataCommandlet, Warning, "Field Outer Type {fieldOuter} not found (from {config} in section {section})",
+							("fieldOuter", *FieldOuterTypeStr),
+							("config", InConfigKey),
+							("section", *SectionName),
+							("id", GatherTextFromMetaDataCommandlet::LocalizationLogIdentifier)
+						);
+						continue;
+					}
+
+					OutFieldOuterTypes.Add(OuterType);
+
+					if (const UClass* OuterClass = Cast<UClass>(OuterType))
+					{
+						GetDerivedClasses(OuterClass, (TArray<UClass*>&)OutFieldOuterTypes);
+					}
+				}
+			}
+		};
+
+		GetFieldOuterTypesArrayFromConfig(TEXT("FieldOuterTypesToInclude"), FieldOuterTypesToInclude);
+		GetFieldOuterTypesArrayFromConfig(TEXT("FieldOuterTypesToExclude"), FieldOuterTypesToExclude);
+	}
+
 	FGatherParameters Arguments;
 	GetStringArrayFromConfig(*SectionName, TEXT("InputKeys"), Arguments.InputKeys, GatherTextConfigPath);
 	GetStringArrayFromConfig(*SectionName, TEXT("OutputNamespaces"), Arguments.OutputNamespaces, GatherTextConfigPath);
@@ -474,7 +537,36 @@ bool UGatherTextFromMetaDataCommandlet::ShouldGatherFromField(const FieldType* F
 		return true;
 	};
 
-	return (!bIsEditorOnly || ShouldGatherFromEditorOnlyData) && ShouldGatherFieldByType() && ShouldGatherFieldByOwnerType();
+	auto ShouldGatherFieldByOuterType = [this, Field]()
+	{
+		if (FieldOuterTypesToInclude.Num() == 0 && FieldOuterTypesToExclude.Num() == 0)
+		{
+			return true;
+		}
+
+		UObject* FieldOuter = nullptr;
+		if constexpr (TIsDerivedFrom<FieldType, FField>::IsDerived)
+		{
+			FieldOuter = Field->GetOwnerUObject();
+		}
+		else
+		{
+			FieldOuter = Field->GetOuter();
+		}
+
+		if (FieldOuter)
+		{
+			if (const UStruct* FieldOuterType = Cast<UStruct>(FieldOuter))
+			{
+				return (FieldOuterTypesToInclude.Num() == 0 || FieldOuterTypesToInclude.Contains(FieldOuterType))
+					&& (FieldOuterTypesToExclude.Num() == 0 || !FieldOuterTypesToExclude.Contains(FieldOuterType));
+			}
+		}
+
+		return true;
+	};
+
+	return (!bIsEditorOnly || ShouldGatherFromEditorOnlyData) && ShouldGatherFieldByType() && ShouldGatherFieldByOwnerType() && ShouldGatherFieldByOuterType();
 }
 
 template <typename FieldType>
