Fixing the editor compile error when UE_WITH_PACKAGE_ACCESS_TRACKING=0

diff --git a/Engine/Source/Editor/UnrealEd/Private/EditorDomain/EditorDomainUtils.cpp b/Engine/Source/Editor/UnrealEd/Private/EditorDomain/EditorDomainUtils.cpp
--- a/Engine/Source/Editor/UnrealEd/Private/EditorDomain/EditorDomainUtils.cpp
+++ b/Engine/Source/Editor/UnrealEd/Private/EditorDomain/EditorDomainUtils.cpp
@@ -689,12 +689,13 @@ FClassDigestData* FPrecacheClassDigest::GetRecursive(FTopLevelAssetPath ClassPat
 	{
 		UE_LOG(LogEditorDomain, Display, TEXT("[DebugNullCook] Calculating class digest for class %s"), *ClassPath.ToString());
 		UE_LOG(LogEditorDomain, Display, TEXT("[DebugNullCook] bAllowRedirects=%d"), bAllowRedirects);
-
+#if UE_WITH_PACKAGE_ACCESS_TRACKING
 		const PackageAccessTracking_Private::FTrackedData* TrackedData = PackageAccessTracking_Private::FPackageAccessRefScope::GetCurrentThreadAccumulatedData();
 		if (TrackedData)
 		{
 			UE_LOG(LogEditorDomain, Display, TEXT("[DebugNullCook] Package %s"), *TrackedData->PackageName.ToString());
 		}
+#endif
 
 		FDebug::DumpStackTraceToLog(ELogVerbosity::Display);
 	}
