Don't reset the cached world bounds for non-editor worlds.

diff --git a/Engine/Source/Editor/UnrealEd/Private/EditorViewportClient.cpp b/Engine/Source/Editor/UnrealEd/Private/EditorViewportClient.cpp
--- a/Engine/Source/Editor/UnrealEd/Private/EditorViewportClient.cpp
+++ b/Engine/Source/Editor/UnrealEd/Private/EditorViewportClient.cpp
@@ -1701,7 +1701,7 @@ void FEditorViewportClient::Tick(float DeltaTime)
 	
 	if (!CachedWorldBounds.IsValid)
 	{
-		if (UWorld* World = GetWorld())
+		if (UWorld* World = GetWorld(); World && !World->IsGameWorld())
 		{
 			if (UWorldPartition* WorldPartition = World->GetWorldPartition())
 			{
@@ -7372,7 +7372,7 @@ void FEditorViewportClient::OnObjectTransacted(UObject* Object, const FTransacti
 
 void FEditorViewportClient::OnPostLoadMap(uint32 MapChangeFlags)
 {
-	if (GWorld == GetWorld())
+	if (GWorld == GetWorld() && !GWorld->IsGameWorld())
 	{
 		CachedWorldBounds.Init();
 		OrthographicFarPlane.Reset();
