UE: Bump Cook incremental version because a primary asset id changed but was not detected

diff --git a/Engine/Source/Editor/UnrealEd/Private/Cooker/CookIncrementalVersion.cpp b/Engine/Source/Editor/UnrealEd/Private/Cooker/CookIncrementalVersion.cpp
--- a/Engine/Source/Editor/UnrealEd/Private/Cooker/CookIncrementalVersion.cpp
+++ b/Engine/Source/Editor/UnrealEd/Private/Cooker/CookIncrementalVersion.cpp
@@ -8,6 +8,6 @@ namespace UE::Cook
 {
 
 // Change CookIncrementalVersion to a new guid when all packages in an incremental cook need to be invalidated
-FGuid CookIncrementalVersion(0x366E3E75, 0x92EA490C, 0xA6C041B8, 0x26460007);
+FGuid CookIncrementalVersion(0x0EFCAD3B, 0x7BF04E78, 0x847FA620, 0xED21796E);
 
 }
