SourceControl: Only perform an `FUpdateStatus` in `AutomaticCheckoutOrPromptToRevertPackages` on packages that are not already checked out.

If the user moves an actor, the 'automatic checkout' already checks it out. The 'AutomicCheckoutOrPromptToRevertPackages', which is called on SAVE, does then not have to update its status again.

#rnx

diff --git a/Engine/Source/Editor/UnrealEd/Private/FileHelpers.cpp b/Engine/Source/Editor/UnrealEd/Private/FileHelpers.cpp
--- a/Engine/Source/Editor/UnrealEd/Private/FileHelpers.cpp
+++ b/Engine/Source/Editor/UnrealEd/Private/FileHelpers.cpp
@@ -5020,11 +5020,29 @@ bool FEditorFileUtils::AutomaticCheckoutOrPromptToRevertPackages(const TArray<UP
 	}
 
 	// Determine initial states.
-	SourceControlProvider.Execute(ISourceControlOperation::Create<FUpdateStatus>(), PackagesToCheckOut);
-
 	TArray<FSourceControlStateRef> InitialStates;
 	SourceControlProvider.GetState(PackagesToCheckOut, InitialStates, EStateCacheUsage::Use);
 
+	TArray<FString> Filenames;
+	for (const FSourceControlStateRef& State : InitialStates)
+	{
+		if (!State->IsCheckedOut() && !State->IsAdded())
+		{
+			Filenames.Add(State->GetFilename());
+		}
+	}
+
+	if (Filenames.Num() > 0)
+	{
+		InitialStates.Empty();
+
+		// Update states - this blocks the thread.
+		SourceControlProvider.Execute(ISourceControlOperation::Create<FUpdateStatus>(), Filenames);
+
+		// Refetch states.
+		SourceControlProvider.GetState(PackagesToCheckOut, InitialStates, EStateCacheUsage::Use);
+	}
+
 	TArray<UPackage*> PackagesCheckOutNeeded;
 	TArray<UPackage*> PackagesCheckOutImpossible;
 	TArray<UPackage*> PackagesCheckedOutAlready;
