UE: Bump incremental cook version. Staging crash

diff --git a/Engine/Source/Editor/UnrealEd/Private/Cooker/CookIncrementalVersion.cpp b/Engine/Source/Editor/UnrealEd/Private/Cooker/CookIncrementalVersion.cpp
--- a/Engine/Source/Editor/UnrealEd/Private/Cooker/CookIncrementalVersion.cpp
+++ b/Engine/Source/Editor/UnrealEd/Private/Cooker/CookIncrementalVersion.cpp
@@ -8,6 +8,6 @@ namespace UE::Cook
 {
 
 // Change CookIncrementalVersion to a new guid when all packages in an incremental cook need to be invalidated
-FGuid CookIncrementalVersion(0x0EFCAD3B, 0x7BF04E78, 0x847FA620, 0xED21796E);
+FGuid CookIncrementalVersion(0x7573B6B6, 0xADB64154, 0xAAC3EE05, 0xC14FDB66);
 
 }
