[ITF Gizmo] Component selection is incorrect after modifying a Spline Point in a component of the same Actor

diff --git a/Engine/Source/Editor/UnrealEd/Private/ComponentVisualizerManager.cpp b/Engine/Source/Editor/UnrealEd/Private/ComponentVisualizerManager.cpp
--- a/Engine/Source/Editor/UnrealEd/Private/ComponentVisualizerManager.cpp
+++ b/Engine/Source/Editor/UnrealEd/Private/ComponentVisualizerManager.cpp
@@ -9,6 +9,7 @@
 #include "SEditorViewport.h"
 #include "HAL/IConsoleManager.h"
 #include "EditorModeManager.h"
+#include "Elements/Component/ComponentElementData.h"
 #include "Elements/Framework/TypedElementSelectionSet.h"
 
 
@@ -30,7 +31,7 @@ bool FComponentVisualizerManager::HandleClick(FEditorViewportClient* InViewportC
 			return true;
 		}
 	}
-		
+
 	bool bHandled = HandleProxyForComponentVis(InViewportClient, HitProxy, Click);
 	if (bHandled && Click.GetKey() == EKeys::RightMouseButton)
 	{
@@ -52,6 +53,12 @@ bool FComponentVisualizerManager::HandleClick(FEditorViewportClient* InViewportC
 		}
 	}
 
+	// Make sure we are not keeping any active component visualizer around when Visualizer Manager is not handling the click
+	if (!bHandled)
+	{
+		ClearActiveComponentVis();
+	}
+
 	return bHandled;
 }
 
@@ -79,8 +86,11 @@ bool FComponentVisualizerManager::HandleProxyForComponentVis(FEditorViewportClie
 
 					if (bShouldSelectElement && ElementSelectionSet)
 					{
-						if (!ElementSelectionSet->IsElementSelected(SelectedComponentHandle, FTypedElementIsSelectedOptions()))
+						if (ElementSelectionSet->CanSelectElement(SelectedComponentHandle, FTypedElementSelectionOptions())
+							&& !ElementSelectionSet->IsElementSelected(SelectedComponentHandle, FTypedElementIsSelectedOptions()))
 						{
+							// Clear the selection and just select the component
+							ElementSelectionSet->ClearSelection(FTypedElementSelectionOptions());
 							ElementSelectionSet->SelectElement(SelectedComponentHandle, FTypedElementSelectionOptions());
 						}
 					}
