* Fixed resave commandlet failing to find any asset when using the asset registry when the package has a path relative to the engine dir
* Added -ForceMaterialAutoUsageFlag command line option, to be used in conjunction with -onlymaterials, to force set the corresponding option on all processed materials

diff --git a/Engine/Source/Editor/UnrealEd/Classes/Commandlets/ResavePackagesCommandlet.h b/Engine/Source/Editor/UnrealEd/Classes/Commandlets/ResavePackagesCommandlet.h
--- a/Engine/Source/Editor/UnrealEd/Classes/Commandlets/ResavePackagesCommandlet.h
+++ b/Engine/Source/Editor/UnrealEd/Classes/Commandlets/ResavePackagesCommandlet.h
@@ -137,6 +137,8 @@ protected:
 	bool bHLODMapCleanup;
 	FString ForceHLODSetupAsset;
 	FString HLODSkipToMap;
+	/** Only valid for bOnlyMaterials.If set, force the materials auto usage flag to the optional value. Don't change the flag's value otherwise */
+	TOptional<bool> ForceMaterialAutoUsageFlag; 
 	bool bForceUATEnvironmentVariableSet;
 	bool bResaveOnDemand = false;
 
diff --git a/Engine/Source/Editor/UnrealEd/Private/Commandlets/BaseIteratePackagesCommandlet.cpp b/Engine/Source/Editor/UnrealEd/Private/Commandlets/BaseIteratePackagesCommandlet.cpp
--- a/Engine/Source/Editor/UnrealEd/Private/Commandlets/BaseIteratePackagesCommandlet.cpp
+++ b/Engine/Source/Editor/UnrealEd/Private/Commandlets/BaseIteratePackagesCommandlet.cpp
@@ -514,6 +514,8 @@ void UBaseIteratePackagesCommandlet::LoadAndSaveOnePackage(const FString& Filena
 	}
 
 	FPackagePath PackagePath = FPackagePath::FromLocalPath(Filename);
+	// Convert this path to mounted before querying the asset registry or CollectionFilter because they use mounted paths 
+	PackagePath.TryConvertToMounted();
 	if (CollectionFilter.Num())
 	{
 		FString PackageNameToCreate = PackagePath.GetPackageNameOrFallback();
diff --git a/Engine/Source/Editor/UnrealEd/Private/Commandlets/ContentCommandlets.cpp b/Engine/Source/Editor/UnrealEd/Private/Commandlets/ContentCommandlets.cpp
--- a/Engine/Source/Editor/UnrealEd/Private/Commandlets/ContentCommandlets.cpp
+++ b/Engine/Source/Editor/UnrealEd/Private/Commandlets/ContentCommandlets.cpp
@@ -708,6 +708,8 @@ void UResavePackagesCommandlet::LoadAndSaveOnePackage(const FString& Filename)
 	}
 
 	FPackagePath PackagePath = FPackagePath::FromLocalPath(Filename);
+	// Convert this path to mounted before querying the asset registry or CollectionFilter because they use mounted paths 
+	PackagePath.TryConvertToMounted();
 	if (CollectionFilter.Num())
 	{
 		FString PackageNameToCreate = PackagePath.GetPackageNameOrFallback();
@@ -722,7 +724,6 @@ void UResavePackagesCommandlet::LoadAndSaveOnePackage(const FString& Filename)
 	{
 		FString PackageNameToCreate = PackagePath.GetPackageNameOrFallback();
 		FName PackageFNameToCreate(*PackageNameToCreate);
-
 		FAssetRegistryModule& AssetRegistryModule = FModuleManager::LoadModuleChecked<FAssetRegistryModule>(TEXT("AssetRegistry"));
 		IAssetRegistry& AssetRegistry = AssetRegistryModule.Get();
 		TArray<FAssetData> PackageAssetData;
@@ -1202,6 +1203,24 @@ int32 UResavePackagesCommandlet::Main( const FString& Params )
 	bSkipPayloadTrailers = Switches.Contains(TEXT("SkipPayloadTrailers"));
 	/** only process packages containing materials */
 	bOnlyMaterials = Switches.Contains(TEXT("onlymaterials"));
+	
+	/** force the materials auto usage flag */
+	{
+		int32 ForceAutoUsageFlagValue = 0;
+		if (FParse::Value(FCommandLine::Get(), TEXT("ForceMaterialAutoUsageFlag="), ForceAutoUsageFlagValue))
+		{
+			if (bOnlyMaterials)
+			{
+				ForceMaterialAutoUsageFlag = (ForceAutoUsageFlagValue != 0);
+			}
+			else
+			{
+				UE_LOG(LogContentCommandlet, Warning,
+					TEXT("ForceMaterialAutoUsageFlag should only be used with -onlymaterials. Ignoring."));
+			}
+		}
+
+	}
 	/** determine if we are building navigation data for the map packages on the pass. **/
 	bShouldBuildNavigationData = Switches.Contains(TEXT("BuildNavigationData"));
 
@@ -2402,7 +2421,16 @@ void UResavePackagesCommandlet::PerformAdditionalOperations(class UWorld* World,
 
 void UResavePackagesCommandlet::PerformAdditionalOperations( class UObject* Object, bool& bSavePackage )
 {
-
+	if (bOnlyMaterials && ForceMaterialAutoUsageFlag.IsSet())
+	{
+		if (UMaterial* Material = Cast<UMaterial>(Object); 
+			(Material != nullptr) && (Material->bAutomaticallySetUsageInEditor != *ForceMaterialAutoUsageFlag))
+		{
+			Material->MarkPackageDirty();
+			Material->bAutomaticallySetUsageInEditor = *ForceMaterialAutoUsageFlag;
+			bSavePackage = true;
+		}
+	}
 }
 
 
diff --git a/Engine/Source/Runtime/CoreUObject/Public/Misc/PackagePath.h b/Engine/Source/Runtime/CoreUObject/Public/Misc/PackagePath.h
--- a/Engine/Source/Runtime/CoreUObject/Public/Misc/PackagePath.h
+++ b/Engine/Source/Runtime/CoreUObject/Public/Misc/PackagePath.h
@@ -221,6 +221,12 @@ public:
 	 */ 
 	friend COREUOBJECT_API FArchive& operator<<(FArchive& Ar, FPackagePath& PackagePath);
 
+	/**
+	 * If this PackagePath is not yet internally stored as MountedPath, search the FPackageName MountPoints again for
+	 * its PackageName or LocalPath, and convert it to mounted if found
+	 */
+	COREUOBJECT_API bool TryConvertToMounted() const;
+
 	/**
 	 * Return whether this is a PackagePath with no path information and can never identify a package on disk
 	 *
@@ -455,12 +461,6 @@ private:
 		FStringView CustomExtension) const;
 	int32 GetStringDataLen() const;
 
-	/**
-	 * If this PackagePath is not yet internally stored as MountedPath, search the FPackageName MountPoints again for
-	 * its PackageName or LocalPath, and convert it to mounted if found
-	 */
-	bool TryConvertToMounted() const;
-
 	enum class EPackageIdType : uint8
 	{
 		/** Does not specify a Package */
