When using "save as" while in PIE, ask the user if they want to stop the PIE session to save or cancel the save and stay in PIE.
This is required because the "save as" operation will reload the map.

diff --git a/Engine/Source/Editor/UnrealEd/Private/FileHelpers.cpp b/Engine/Source/Editor/UnrealEd/Private/FileHelpers.cpp
--- a/Engine/Source/Editor/UnrealEd/Private/FileHelpers.cpp
+++ b/Engine/Source/Editor/UnrealEd/Private/FileHelpers.cpp
@@ -1420,6 +1420,13 @@ static bool SaveAsImplementation( UWorld* InWorld, const FString& DefaultFilenam
 {
 	UEditorLoadingSavingSettings* LoadingSavingSettings = GetMutableDefault<UEditorLoadingSavingSettings>();
 
+	// If a PIE world exists, warn the user that the PIE session will be terminated.
+	// Abort if the user refuses to terminate the PIE session.
+	if ( GEditor->ShouldAbortBecauseOfPIEWorld() )
+	{
+		return false;
+	}
+
 	// Get default path and filename. If no default filename was supplied, create one.
 	FString DefaultDirectory = FEditorDirectories::Get().GetLastDirectory(ELastDirectory::LEVEL);
 	FString Filename = FPaths::GetCleanFilename(DefaultFilename);
