Removed the "null graph found" warning, as in some context, it can result in a failed asset validation.

#tests Previously failing validation now succeed.
#rnx

diff --git a/Engine/Source/Editor/UnrealEd/Private/Kismet2/BlueprintEditorUtils.cpp b/Engine/Source/Editor/UnrealEd/Private/Kismet2/BlueprintEditorUtils.cpp
--- a/Engine/Source/Editor/UnrealEd/Private/Kismet2/BlueprintEditorUtils.cpp
+++ b/Engine/Source/Editor/UnrealEd/Private/Kismet2/BlueprintEditorUtils.cpp
@@ -6852,7 +6852,6 @@ void FBlueprintEditorUtils::CleanNullGraphReferencesRecursive(UEdGraph* Graph)
 		}
 		else
 		{
-			UE_LOG(LogBlueprint, Warning, TEXT("Found NULL graph reference in children of '%s', removing it!"), *Graph->GetPathName());
 			Graph->SubGraphs.RemoveAt(GraphIndex);
 		}
 	}
@@ -6869,7 +6868,6 @@ void FBlueprintEditorUtils::CleanNullGraphReferencesInArray(UBlueprint* Blueprin
 		}
 		else
 		{
-			UE_LOG(LogBlueprint, Warning, TEXT("Found NULL graph reference in '%s', removing it!"), *Blueprint->GetPathName());
 			GraphArray.RemoveAt(GraphIndex);
 		}
 	}
