UE: Bump incremental cook version due to stale material reference

diff --git a/Engine/Source/Editor/UnrealEd/Private/Cooker/CookIncrementalVersion.cpp b/Engine/Source/Editor/UnrealEd/Private/Cooker/CookIncrementalVersion.cpp
--- a/Engine/Source/Editor/UnrealEd/Private/Cooker/CookIncrementalVersion.cpp
+++ b/Engine/Source/Editor/UnrealEd/Private/Cooker/CookIncrementalVersion.cpp
@@ -8,6 +8,6 @@ namespace UE::Cook
 {
 
 // Change CookIncrementalVersion to a new guid when all packages in an incremental cook need to be invalidated
-FGuid CookIncrementalVersion(0x7573B6B6, 0xADB64154, 0xAAC3EE05, 0xC14FDB66);
+FGuid CookIncrementalVersion(0x06FC3DC8, 0xC9AA4203, 0xBEB48A2F, 0x2F8725DF);
 
 }
