Consistently filter whitespace only strings out of the localization gather

#rnx
#Synced-CL 48092843

diff --git a/Engine/Source/Editor/UnrealEd/Private/Commandlets/GatherTextFromMetadataCommandlet.cpp b/Engine/Source/Editor/UnrealEd/Private/Commandlets/GatherTextFromMetadataCommandlet.cpp
--- a/Engine/Source/Editor/UnrealEd/Private/Commandlets/GatherTextFromMetadataCommandlet.cpp
+++ b/Engine/Source/Editor/UnrealEd/Private/Commandlets/GatherTextFromMetadataCommandlet.cpp
@@ -463,7 +463,7 @@ void UGatherTextFromMetaDataCommandlet::GatherTextFromField(UField* Field, const
 					if (Enum->HasMetaData(*Arguments.InputKeys[j], i))
 					{
 						const FString& MetaDataValue = Enum->GetMetaData(*Arguments.InputKeys[j], i);
-						if (!MetaDataValue.IsEmpty())
+						if (FText::ShouldGatherForLocalization(MetaDataValue))
 						{
 							PatternArguments.Add(TEXT("MetaDataValue"), MetaDataValue);
 
@@ -580,7 +580,7 @@ void UGatherTextFromMetaDataCommandlet::GatherTextFromFieldImpl(FieldType* Field
 		if (Field->HasMetaData(*Arguments.InputKeys[i]))
 		{
 			const FString& MetaDataValue = Field->GetMetaData(*Arguments.InputKeys[i]);
-			if (!MetaDataValue.IsEmpty())
+			if (FText::ShouldGatherForLocalization(MetaDataValue))
 			{
 				PatternArguments.Add(TEXT("MetaDataValue"), MetaDataValue);
 
diff --git a/Engine/Source/Editor/UnrealEd/Private/Commandlets/GatherTextFromSourceCommandlet.cpp b/Engine/Source/Editor/UnrealEd/Private/Commandlets/GatherTextFromSourceCommandlet.cpp
--- a/Engine/Source/Editor/UnrealEd/Private/Commandlets/GatherTextFromSourceCommandlet.cpp
+++ b/Engine/Source/Editor/UnrealEd/Private/Commandlets/GatherTextFromSourceCommandlet.cpp
@@ -656,7 +656,7 @@ void UGatherTextFromSourceCommandlet::RunPass(EGatherSourcePasses Pass, bool Sho
 				{
 					for (const auto& ParsedStringTableEntryPair : ParsedStringTablePair.Value.TableEntries)
 					{
-						if (!ParsedStringTableEntryPair.Value.bIsEditorOnly || ParseCtxt.ShouldGatherFromEditorOnlyData)
+						if ((!ParsedStringTableEntryPair.Value.bIsEditorOnly || ParseCtxt.ShouldGatherFromEditorOnlyData) && FText::ShouldGatherForLocalization(ParsedStringTableEntryPair.Value.SourceString))
 						{
 							FManifestContext SourceContext;
 							SourceContext.Key = ParsedStringTableEntryPair.Key;
@@ -693,7 +693,10 @@ void UGatherTextFromSourceCommandlet::RunPass(EGatherSourcePasses Pass, bool Sho
 			// Submit mainpass results to manifest helper
 			for (const FManifestEntryResult& Result : Batch.MainpassResults)
 			{
-				GatherManifestHelper->AddSourceText(Result.Namespace, FLocItem(Result.Source), Result.Context, &Result.Description);
+				if (FText::ShouldGatherForLocalization(Result.Source))
+				{
+					GatherManifestHelper->AddSourceText(Result.Namespace, FLocItem(Result.Source), Result.Context, &Result.Description);
+				}
 			}
 		}
 
diff --git a/Engine/Source/Runtime/Core/Private/Internationalization/Text.cpp b/Engine/Source/Runtime/Core/Private/Internationalization/Text.cpp
--- a/Engine/Source/Runtime/Core/Private/Internationalization/Text.cpp
+++ b/Engine/Source/Runtime/Core/Private/Internationalization/Text.cpp
@@ -1320,21 +1320,28 @@ bool FText::ShouldGatherForLocalization() const
 		return false;
 	}
 
-	const FString& SourceString = GetTextData().GetSourceString();
+	return !(Flags & (ETextFlag::CultureInvariant | ETextFlag::Transient))
+		&& !IsFromStringTable() 
+		&& FText::ShouldGatherForLocalization(GetTextData().GetSourceString());
+}
 
-	auto IsAllWhitespace = [](const FString& String) -> bool
+bool FText::ShouldGatherForLocalization(FStringView SourceString)
+{
+	// Don't gather empty strings
+	if (SourceString.IsEmpty())
 	{
-		for(int32 i = 0; i < String.Len(); ++i)
+		return false;
+	}
+
+	// Don't gather strings that are whitespace only
+	for (const TCHAR Char : SourceString)
+	{
+		if (!FTextChar::IsWhitespace(Char))
 		{
-			if( !FTextChar::IsWhitespace( String[i] ) )
-			{
-				return false;
-			}
+			return true;
 		}
-		return true;
-	};
-
-	return !((Flags & ETextFlag::CultureInvariant) || (Flags & ETextFlag::Transient)) && !IsFromStringTable() && !SourceString.IsEmpty() && !IsAllWhitespace(SourceString);
+	}
+	return false;
 }
 
 void FText::PreloadStringTable( FName InTableId, const FTextKey& InKey )
diff --git a/Engine/Source/Runtime/Core/Public/Internationalization/Text.h b/Engine/Source/Runtime/Core/Public/Internationalization/Text.h
--- a/Engine/Source/Runtime/Core/Public/Internationalization/Text.h
+++ b/Engine/Source/Runtime/Core/Public/Internationalization/Text.h
@@ -821,6 +821,7 @@ public:
 	[[nodiscard]] CORE_API bool IsFromStringTable() const;
 
 	[[nodiscard]] CORE_API bool ShouldGatherForLocalization() const;
+	[[nodiscard]] static CORE_API bool ShouldGatherForLocalization(FStringView SourceString);
 	
 	static CORE_API void PreloadStringTable( FName InTableId, const FTextKey& InKey );
 
