[UnrealEd] Introduced ULevelFactory method that outputs created actors to fix copy&paste issue in which un-selectable actors were pasted in an incomplete state.

#Synced-CL 48152445

diff --git a/Engine/Source/Editor/UnrealEd/Classes/Factories/LevelFactory.h b/Engine/Source/Editor/UnrealEd/Classes/Factories/LevelFactory.h
--- a/Engine/Source/Editor/UnrealEd/Classes/Factories/LevelFactory.h
+++ b/Engine/Source/Editor/UnrealEd/Classes/Factories/LevelFactory.h
@@ -19,7 +19,9 @@ class ULevelFactory : public UFactory
 
 	//~ Begin UFactory Interface
 	virtual UObject* FactoryCreateText( UClass* Class, UObject* InParent, FName Name, EObjectFlags Flags, UObject* Context, const TCHAR* Type, const TCHAR*& Buffer, const TCHAR* BufferEnd, FFeedbackContext* Warn ) override;
-	//~ Begin UFactory Interface	
+	//~ End UFactory Interface
+
+	UNREALED_API UObject* LevelFactoryCreateText(UClass* Class, UObject* InParent, FName Name, EObjectFlags Flags, UObject* Context, const TCHAR* Type, const TCHAR*& Buffer, const TCHAR* BufferEnd, FFeedbackContext* Warn, TArray<AActor*>* OutActors);
 };
 
 
diff --git a/Engine/Source/Editor/UnrealEd/Private/EditorActor.cpp b/Engine/Source/Editor/UnrealEd/Private/EditorActor.cpp
--- a/Engine/Source/Editor/UnrealEd/Private/EditorActor.cpp
+++ b/Engine/Source/Editor/UnrealEd/Private/EditorActor.cpp
@@ -377,13 +377,12 @@ void UUnrealEdEngine::PasteActors(TArray<AActor*>& OutPastedActors, UWorld* InWo
 
 	// Import the actors.
 	TStrongObjectPtr<ULevelFactory> Factory(NewObject<ULevelFactory>());
-	Factory->FactoryCreateText(ULevel::StaticClass(), InWorld->GetCurrentLevel(), InWorld->GetCurrentLevel()->GetFName(), RF_Transactional, NULL, bDuplicate ? TEXT("move") : TEXT("paste"), Paste, Paste + FCString::Strlen(Paste), GWarn);
+	Factory->LevelFactoryCreateText(ULevel::StaticClass(), InWorld->GetCurrentLevel(), InWorld->GetCurrentLevel()->GetFName(), RF_Transactional, NULL, bDuplicate ? TEXT("move") : TEXT("paste"), Paste, Paste + FCString::Strlen(Paste), GWarn, &OutPastedActors);
 
 	// Reinstate old BSP update setting, and force a rebuild - any levels whose geometry has changed while pasting will be rebuilt
 	GetMutableDefault<ULevelEditorMiscSettings>()->bBSPAutoUpdate = bBSPAutoUpdate;
 	
-	// FactoryCreateText set the selection to the new actors, so copy that into OutPastedActors and restore the original selection
-	ActorSelection->GetSelectedObjects<AActor>(OutPastedActors);
+	// FactoryCreateText sets the selection to the new actors, so restore the original selection
 	FObjectReader(ActorSelection, OriginalSelectionState);
 
 	// Fire ULevel::LevelDirtiedEvent when falling out of scope.
diff --git a/Engine/Source/Editor/UnrealEd/Private/Factories/EditorFactories.cpp b/Engine/Source/Editor/UnrealEd/Private/Factories/EditorFactories.cpp
--- a/Engine/Source/Editor/UnrealEd/Private/Factories/EditorFactories.cpp
+++ b/Engine/Source/Editor/UnrealEd/Private/Factories/EditorFactories.cpp
@@ -719,7 +719,13 @@ ULevelFactory::ULevelFactory(const FObjectInitializer& ObjectInitializer)
 	bEditorImport = false;
 }
 
-UObject* ULevelFactory::FactoryCreateText
+UObject* ULevelFactory::FactoryCreateText(UClass* Class, UObject* InParent, FName Name, EObjectFlags Flags, UObject* Context, const TCHAR* Type,
+	const TCHAR*& Buffer, const TCHAR* BufferEnd, FFeedbackContext* Warn)
+{
+	return LevelFactoryCreateText(Class, InParent, Name, Flags, Context, Type, Buffer, BufferEnd, Warn, nullptr);
+}
+
+UObject* ULevelFactory::LevelFactoryCreateText
 (
 	UClass*				Class,
 	UObject*			InParent,
@@ -729,9 +735,14 @@ UObject* ULevelFactory::FactoryCreateText
 	const TCHAR*		Type,
 	const TCHAR*&		Buffer,
 	const TCHAR*		BufferEnd,
-	FFeedbackContext*	Warn
+	FFeedbackContext*	Warn,
+	TArray<AActor*>*    OutActors
 )
 {
+	if (OutActors)
+	{
+		OutActors->Reset();
+	}
 	GEditor->GetEditorSubsystem<UImportSubsystem>()->BroadcastAssetPreImport(this, Class, InParent, Name, Type);
 
 	UWorld* World = GWorld;
@@ -1451,6 +1462,17 @@ UObject* ULevelFactory::FactoryCreateText
 		GEditor->GetEditorSubsystem<UImportSubsystem>()->BroadcastAssetPostImport(this, World );
 	}
 
+	if (OutActors)
+	{
+		OutActors->SetNumUninitialized(NewActorMap.Num());
+		
+		int32 Index = 0;
+		for (const TPair<AActor*, FString>& Pair : NewActorMap)
+		{
+			(*OutActors)[Index++] = Pair.Key;
+		}
+	}
+
 	return World;
 }
 
