Fixed a crash that can be caused by editing comment nodes in materials after a copy paste.

#tests Ran in PIE and edited a comment
#rnx
#Synced-CL 47906194

diff --git a/Engine/Source/Editor/UnrealEd/Private/MaterialGraphNode_Comment.cpp b/Engine/Source/Editor/UnrealEd/Private/MaterialGraphNode_Comment.cpp
--- a/Engine/Source/Editor/UnrealEd/Private/MaterialGraphNode_Comment.cpp
+++ b/Engine/Source/Editor/UnrealEd/Private/MaterialGraphNode_Comment.cpp
@@ -169,13 +169,16 @@ void UMaterialGraphNode_Comment::ResizeNode(const FVector2f& NewSize)
 {
 	Super::ResizeNode(NewSize);
 
-	// Set position as well since may have been resized from top corner
-	MaterialExpressionComment->SizeX = NodeWidth;
-	MaterialExpressionComment->SizeY = NodeHeight;
-	MaterialExpressionComment->MaterialExpressionEditorX = NodePosX;
-	MaterialExpressionComment->MaterialExpressionEditorY = NodePosY;
-	MaterialExpressionComment->MarkPackageDirty();
-	MaterialDirtyDelegate.ExecuteIfBound();
+	if(MaterialExpressionComment)
+	{
+		// Set position as well since may have been resized from top corner
+		MaterialExpressionComment->SizeX = NodeWidth;
+		MaterialExpressionComment->SizeY = NodeHeight;
+		MaterialExpressionComment->MaterialExpressionEditorX = NodePosX;
+		MaterialExpressionComment->MaterialExpressionEditorY = NodePosY;
+		MaterialExpressionComment->MarkPackageDirty();
+		MaterialDirtyDelegate.ExecuteIfBound();
+	}
 }
 
 int32 UMaterialGraphNode_Comment::GetFontSize() const
