[LightMass] Fixed missing dynamic dependencies in the hard coded list of dlls to copy by reading the manifest instead.

#rnx
#Synced-CL 43344078

diff --git a/Engine/Source/Editor/UnrealEd/Private/Lightmass/Lightmass.cpp b/Engine/Source/Editor/UnrealEd/Private/Lightmass/Lightmass.cpp
--- a/Engine/Source/Editor/UnrealEd/Private/Lightmass/Lightmass.cpp
+++ b/Engine/Source/Editor/UnrealEd/Private/Lightmass/Lightmass.cpp
@@ -56,6 +56,7 @@
 #include "Misc/FileHelper.h"
 #include "StaticLightingBuildContext.h"
 #include "Engine/InstancedStaticMesh.h"
+#include "JsonObjectConverter.h"
 
 extern FSwarmDebugOptions GSwarmDebugOptions;
 
@@ -2940,13 +2941,16 @@ void FLightmassProcessor::IssueStaticShadowDepthMapTask(const ULightComponent* L
 class FEngineDependencyPaths
 {
 public:
-	FEngineDependencyPaths(std::initializer_list<const TCHAR*> Paths)
+	FEngineDependencyPaths(std::initializer_list<const TCHAR*> Paths, const TArray<FString>& InDynamicDependencyPaths = TArray<FString>())
 	{
-		Strings.Reserve(Paths.size());
 		for(const TCHAR* Path : Paths)
 		{
 			Strings.Add(FPaths::EngineDir() / Path);
 		}
+		for(const FString& Path : InDynamicDependencyPaths)
+		{
+			Strings.Add(FPaths::EngineDir() / Path);
+		}
 
 		RawStrings.Reserve(Strings.Num());
 		for(const FString& String : Strings)
@@ -2970,6 +2974,39 @@ private:
 	TArray<const TCHAR*> RawStrings;
 };
 
+static TArray<FString> GetDynamicDependencyPathsFromManifest(const TCHAR* BinDir)
+{
+	TArray<FString> DynamicDependencyPaths;
+	FString FileData;
+	if (FFileHelper::LoadFileToString(FileData, *(FPaths::EngineDir() / BinDir / "UnrealLightmass.target")))
+	{
+		TSharedRef<TJsonReader<>> Reader = TJsonReaderFactory<>::Create(FileData);
+		TSharedPtr<FJsonObject> RootJsonObject;
+		FJsonSerializer::Deserialize(Reader, RootJsonObject);
+		if (RootJsonObject)
+		{
+			const TArray<TSharedPtr<FJsonValue>>* BuildProducts = nullptr;
+			if (RootJsonObject->TryGetArrayField(TEXT("BuildProducts"), BuildProducts))
+			{
+				for (const auto& BuildProductField : *BuildProducts)
+				{
+					auto Obj = BuildProductField->AsObject();
+					FString TypeStr;
+					if (Obj && Obj->TryGetStringField(TEXT("Type"), TypeStr) && TypeStr == TEXT("DynamicLibrary"))
+					{
+						FString Tmp;
+						if (Obj->TryGetStringField(TEXT("Path"), Tmp) && Tmp.RemoveFromStart(TEXT("$(EngineDir)/")))
+						{
+							DynamicDependencyPaths.Emplace(Tmp);
+						}
+					}
+				}
+			}
+		}
+	}
+	return DynamicDependencyPaths;
+}
+
 bool FLightmassProcessor::BeginRun()
 {
 	{
@@ -2998,97 +3035,30 @@ bool FLightmassProcessor::BeginRun()
 
 	// Setup dependencies for 32bit.
 	static const FString LightmassExecutable32 = FPaths::EngineDir() / TEXT("Binaries/Win32/UnrealLightmass.exe");
-	static const FEngineDependencyPaths RequiredDependencyPaths32 =
+	const FEngineDependencyPaths RequiredDependencyPaths32(
 	{
 		TEXT("Binaries/DotNET/SwarmInterface.dll"),
-		TEXT("Binaries/Win32/AgentInterface.dll"),
-		TEXT("Binaries/Win32/UnrealLightmass-SwarmInterface.dll"),
-		TEXT("Binaries/Win32/UnrealLightmass-ApplicationCore.dll"),
-		TEXT("Binaries/Win32/UnrealLightmass-AutoRTFM.dll"),
-		TEXT("Binaries/Win32/UnrealLightmass-Core.dll"),
-		TEXT("Binaries/Win32/UnrealLightmass-CoreUObject.dll"),
-		TEXT("Binaries/Win32/UnrealLightmass-CorePreciseFP.dll"),
-		TEXT("Binaries/Win32/UnrealLightmass-DerivedDataCache.dll"),
-		TEXT("Binaries/Win32/UnrealLightmass-DevHttp.dll"),
-		TEXT("Binaries/Win32/UnrealLightmass-Sockets.dll"),
-		TEXT("Binaries/Win32/UnrealLightmass-Zen.dll"),
-		TEXT("Binaries/Win32/UnrealLightmass-Projects.dll"),
-		TEXT("Binaries/Win32/UnrealLightmass-Json.dll"),
-		TEXT("Binaries/Win32/UnrealLightmass-SSL.dll")
-		TEXT("Binaries/Win32/UnrealLightmass-BuildSettings.dll")
-	};
+	}, GetDynamicDependencyPathsFromManifest(TEXT("Binaries/Win32/")));
 
 	// Setup dependencies for 64bit.
 #if PLATFORM_WINDOWS
 	static const FString LightmassExecutable64 = FPaths::EngineDir() / TEXT("Binaries/Win64/UnrealLightmass.exe");
-	static const FEngineDependencyPaths RequiredDependencyPaths64 =
-	{
+	const FEngineDependencyPaths RequiredDependencyPaths64( 
+	{ 
 		TEXT("Binaries/DotNET/SwarmInterface.dll"),
-		TEXT("Binaries/Win64/AgentInterface.dll"),
-		TEXT("Binaries/Win64/UnrealLightmass-SwarmInterface.dll"),
-		TEXT("Binaries/Win64/UnrealLightmass-ApplicationCore.dll"),
-		TEXT("Binaries/Win64/UnrealLightmass-AutoRTFM.dll"),
-		TEXT("Binaries/Win64/UnrealLightmass-Core.dll"),
-		TEXT("Binaries/Win64/UnrealLightmass-CoreUObject.dll"),
-		TEXT("Binaries/Win64/UnrealLightmass-CorePreciseFP.dll"),
-		TEXT("Binaries/Win64/UnrealLightmass-DerivedDataCache.dll"),
-		TEXT("Binaries/Win64/UnrealLightmass-DevHttp.dll"),
-		TEXT("Binaries/Win64/UnrealLightmass-Sockets.dll"),
-		TEXT("Binaries/Win64/UnrealLightmass-Zen.dll"),
-		TEXT("Binaries/Win64/UnrealLightmass-Projects.dll"),
-		TEXT("Binaries/Win64/UnrealLightmass-Json.dll"),
-		TEXT("Binaries/Win64/UnrealLightmass-SSL.dll"),
-		TEXT("Binaries/Win64/UnrealLightmass-BuildSettings.dll"),
-		TEXT("Binaries/Win64/embree4.dll"),
-		TEXT("Binaries/Win64/tbb12.dll"),
-		TEXT("Binaries/Win64/tbbmalloc.dll")
-	};
+	}, GetDynamicDependencyPathsFromManifest(TEXT("Binaries/Win64/")));
 #elif PLATFORM_MAC
 	static const FString LightmassExecutable64 = FPaths::EngineDir() / TEXT("Binaries/Mac/UnrealLightmass");
-	static const FEngineDependencyPaths RequiredDependencyPaths64 =
-	{
+	const FEngineDependencyPaths RequiredDependencyPaths64( 
+	{ 
 		TEXT("Binaries/DotNET/Mac/AgentInterface.dll"),
-		TEXT("Binaries/Mac/UnrealLightmass-ApplicationCore.dylib"),
-		TEXT("Binaries/Mac/UnrealLightmass-AutoRTFM.dylib"),
-		TEXT("Binaries/Mac/UnrealLightmass-Core.dylib"),
-		TEXT("Binaries/Mac/UnrealLightmass-CoreUObject.dylib"),
-		TEXT("Binaries/Mac/UnrealLightmass-CorePreciseFP.dylib"),
-		TEXT("Binaries/Mac/UnrealLightmass-DerivedDataCache.dylib"),
-		TEXT("Binaries/Mac/UnrealLightmass-DevHttp.dylib"),
-		TEXT("Binaries/Mac/UnrealLightmass-Sockets.dylib"),
-		TEXT("Binaries/Mac/UnrealLightmass-Zen.dylib"),
-		TEXT("Binaries/Mac/UnrealLightmass-Json.dylib"),
-		TEXT("Binaries/Mac/UnrealLightmass-Projects.dylib"),
-		TEXT("Binaries/Mac/UnrealLightmass-SSL.dylib"),
-		TEXT("Binaries/Mac/UnrealLightmass-SwarmInterface.dylib"),
-		TEXT("Binaries/Mac/UnrealLightmass-BuildSettings.dylib"),
-		TEXT("Binaries/Mac/libembree4.4.3.3.dylib"),
-		TEXT("Binaries/Mac/libtbb.12.13.dylib"),
-		TEXT("Binaries/Mac/libtbbmalloc.2.13.dylib")
-	};
+	}, GetDynamicDependencyPathsFromManifest(TEXT("Binaries/Mac/")));
 #elif PLATFORM_LINUX
 	static const FString LightmassExecutable64 = FPaths::EngineDir() / TEXT("Binaries/Linux/UnrealLightmass");
-	static const FEngineDependencyPaths RequiredDependencyPaths64 =
-	{
+	const FEngineDependencyPaths RequiredDependencyPaths64( 
+	{ 
 		TEXT("Binaries/DotNET/Linux/AgentInterface.dll"),
-		TEXT("Binaries/Linux/libUnrealLightmass-ApplicationCore.so"),
-		TEXT("Binaries/Linux/libUnrealLightmass-AutoRTFM.so"),
-		TEXT("Binaries/Linux/libUnrealLightmass-Core.so"),
-		TEXT("Binaries/Linux/libUnrealLightmass-CoreUObject.so"),
-		TEXT("Binaries/Linux/libUnrealLightmass-CorePrecise.so"),
-		TEXT("Binaries/Linux/libUnrealLightmass-DerivedDataCache.so"),
-		TEXT("Binaries/Linux/libUnrealLightmass-DevHttp.so"),
-		TEXT("Binaries/Linux/libUnrealLightmass-Sockets.so"),
-		TEXT("Binaries/Linux/libUnrealLightmass-Zen.so"),
-		TEXT("Binaries/Linux/libUnrealLightmass-Json.so"),
-		TEXT("Binaries/Linux/libUnrealLightmass-Projects.so"),
-		TEXT("Binaries/Linux/libUnrealLightmass-SSL.so"),
-		TEXT("Binaries/Linux/libUnrealLightmass-SwarmInterface.so"),
-		TEXT("Binaries/Linux/libUnrealLightmass-Networking.so"),
-		TEXT("Binaries/Linux/libUnrealLightmass-Messaging.so"),
-		TEXT("Binaries/Linux/libUnrealLightmass-BuildSettings.so"),
-		TEXT("Plugins/Messaging/UdpMessaging/Binaries/Linux/libUnrealLightmass-UdpMessaging.so")
-	};
+	}, GetDynamicDependencyPathsFromManifest(TEXT("Binaries/Linux/")));
 #else // PLATFORM_LINUX
 #error "Unknown Lightmass platform"
 #endif
