Cooker: Profiling: Tweaks to provide analysis for why packages are kept in memory.
Add number of packages in memory to Post-GC memory diagnostics.
Tweak Post-GC memory diagnostics: Add number of packages, remove display of unhelpful spread variable, add display of %freed.
Tweak Post-GC memory diagnostics: Remove unhelpful list of which referencers are keeping which objects kept in memory, replace it with which referencers are keeping which packages in memory.
Bugfix for references to UPackages not being tracked in CookProfiling's ConstructObjectGraph.
Bugfix for FLowLevelMemTracker::DumpToLog to skip display of the "CustomName" aggregate tag, it's not meant to be shown to clients because it is just an implementation detail for LLM_DEFINE_TAG.
Bugfix for FPackageTracker::LoadedPackages; it was missing startup packages.
Show Post-GC diagnostics for all GCs when profiling memory.
#Synced-CL Unknown

diff --git a/Engine/Source/Editor/UnrealEd/Classes/CookOnTheSide/CookOnTheFlyServer.h b/Engine/Source/Editor/UnrealEd/Classes/CookOnTheSide/CookOnTheFlyServer.h
--- a/Engine/Source/Editor/UnrealEd/Classes/CookOnTheSide/CookOnTheFlyServer.h
+++ b/Engine/Source/Editor/UnrealEd/Classes/CookOnTheSide/CookOnTheFlyServer.h
@@ -392,6 +392,7 @@ private:
 	bool bUseSoftGC = false;
 	bool bWarnedExceededMaxMemoryWithinGCCooldown = false;
 	bool bGarbageCollectTypeSoft = false;
+	bool bSaveGCDataForEvaluation = false;
 
 	/**
 	 * The maximum number of packages that should be preloaded at once. Once this is full, packages in LoadPrepare will
@@ -910,10 +911,11 @@ public:
 
 	UNREALED_API void OnCookerStartCollectGarbage(uint32& ResultFlagsFromTick);
 	UNREALED_API void OnCookerEndCollectGarbage(uint32& ResultFlagsFromTick);
+	UNREALED_API int32 NumLoadedPackages() const;
 	UNREALED_API void EvaluateGarbageCollectionResults(bool bWasDueToOOM, bool bWasPartialGC, uint32 ResultFlags,
-		int32 NumObjectsBeforeGC, const FPlatformMemoryStats& MemStatsBeforeGC,
+		int32 NumObjectsBeforeGC, const FPlatformMemoryStats& MemStatsBeforeGC, int32 NumPackagesBeforeGC,
 		const FGenericMemoryStats& AllocatorStatsBeforeGC,
-		int32 NumObjectsAfterGC, const FPlatformMemoryStats& MemStatsAfterGC,
+		int32 NumObjectsAfterGC, const FPlatformMemoryStats& MemStatsAfterGC, int32 NumPackagesAfterGC,
 		const FGenericMemoryStats& AllocatorStatsAfterGC,
 		float GCDurationSeconds);
 	UNREALED_API bool NeedsDiagnosticSecondGC() const;
@@ -1558,11 +1560,15 @@ private:
 	/** Sets the current phase of the cook to the TargetPhase, and updates per-state data. */
 	void SetCookPhase(UE::Cook::ECookPhase TargetPhase);
 
+	/** Clear temporary data that is allocated during GarbageCollect for use by AddReferencedObjects and PostGC evaluation. */
+	void ClearGCData();
+
 	static UCookOnTheFlyServer* ActiveCOTFS;
 	uint32		StatLoadedPackageCount = 0;
 	uint32		StatSavedPackageCount = 0;
 	int32		PackageDataFromBaseGameNum = 0;
 	UE::Cook::FStatHistoryInt NumObjectsHistory;
+	UE::Cook::FStatHistoryInt NumPackagesHistory;
 	UE::Cook::FStatHistoryInt VirtualMemoryHistory;
 
 	FCriticalSection HiddenDependenciesLock;
diff --git a/Engine/Source/Editor/UnrealEd/Private/Commandlets/CookCommandlet.cpp b/Engine/Source/Editor/UnrealEd/Private/Commandlets/CookCommandlet.cpp
--- a/Engine/Source/Editor/UnrealEd/Private/Commandlets/CookCommandlet.cpp
+++ b/Engine/Source/Editor/UnrealEd/Private/Commandlets/CookCommandlet.cpp
@@ -733,6 +733,7 @@ void UCookCommandlet::ConditionalCollectGarbage(uint32 TickResults, UCookOnTheFl
 
 	const FPlatformMemoryStats MemStatsBeforeGC = FPlatformMemory::GetStats();
 	int32 NumObjectsBeforeGC = GUObjectArray.GetObjectArrayNumMinusAvailable();
+	int32 NumPackagesBeforeGC = COTFS.NumLoadedPackages();
 	FGenericMemoryStats AllocatorStatsBeforeGC;
 	UE::Private::GMalloc->GetAllocatorStats(AllocatorStatsBeforeGC);
 #if ENABLE_LOW_LEVEL_MEM_TRACKER
@@ -741,8 +742,8 @@ void UCookCommandlet::ConditionalCollectGarbage(uint32 TickResults, UCookOnTheFl
 	COTFS.SetGarbageCollectType(TickResults);
 
 	UE_LOG(LogCookCommandlet, Display, TEXT("GarbageCollection...%s (%s)"), *GCType, *GCReason);
+	TGuardValue<bool> SoftGCGuard(UPackage::bSupportCookerSoftGC, true);
 	{
-		TGuardValue<bool> SoftGCGuard(UPackage::bSupportCookerSoftGC, true);
 #if !NO_LOGGING
 		ELogVerbosity::Type SavedVerbosity = LogGarbage.GetVerbosity();
 		bool bSuppressGCLogs = bSoft && bPeriodic && SavedVerbosity > ELogVerbosity::Warning;
@@ -769,9 +770,9 @@ void UCookCommandlet::ConditionalCollectGarbage(uint32 TickResults, UCookOnTheFl
 		}
 	}
 
-	COTFS.ClearGarbageCollectType();
 	FPlatformMemoryStats MemStatsAfterGC = FPlatformMemory::GetStats();
 	int32 NumObjectsAfterGC = GUObjectArray.GetObjectArrayNumMinusAvailable();
+	int32 NumPackagesAfterGC = COTFS.NumLoadedPackages();
 	FGenericMemoryStats AllocatorStatsAfterGC;
 	UE::Private::GMalloc->GetAllocatorStats(AllocatorStatsAfterGC);
 #if ENABLE_LOW_LEVEL_MEM_TRACKER
@@ -781,8 +782,10 @@ void UCookCommandlet::ConditionalCollectGarbage(uint32 TickResults, UCookOnTheFl
 	float GCDurationSeconds = FPlatformTime::Seconds() - GCStartTime;
 	bool bWasDueToOOM = (TickResults & UCookOnTheFlyServer::COSR_RequiresGC_OOM) != 0;
 	COTFS.EvaluateGarbageCollectionResults(bWasDueToOOM, bPartialGC, TickResults,
-		NumObjectsBeforeGC, MemStatsBeforeGC, AllocatorStatsBeforeGC,
-		NumObjectsAfterGC, MemStatsAfterGC, AllocatorStatsAfterGC, GCDurationSeconds);
+		NumObjectsBeforeGC, MemStatsBeforeGC, NumPackagesBeforeGC, AllocatorStatsBeforeGC,
+		NumObjectsAfterGC, MemStatsAfterGC, NumPackagesAfterGC, AllocatorStatsAfterGC, GCDurationSeconds);
+
+	COTFS.ClearGarbageCollectType();
 
 	// Only check stalling in this is the director and if garbage collection was triggered by OOM.
 	bool bRunStallDetector = !COTFS.IsCookWorkerMode() && bWasDueToOOM;
diff --git a/Engine/Source/Editor/UnrealEd/Private/CookOnTheFlyServer.cpp b/Engine/Source/Editor/UnrealEd/Private/CookOnTheFlyServer.cpp
--- a/Engine/Source/Editor/UnrealEd/Private/CookOnTheFlyServer.cpp
+++ b/Engine/Source/Editor/UnrealEd/Private/CookOnTheFlyServer.cpp
@@ -6025,6 +6025,7 @@ void UCookOnTheFlyServer::Initialize( ECookMode::Type DesiredCookMode, ECookInit
 	{
 		TSet<FName> StartupPackages;
 		PackageTracker->InitializeTracking(StartupPackages);
+		NumPackagesHistory.Initialize(PackageTracker->NumLoadedPackages());
 		CookByTheBookOptions->StartupPackages = MoveTemp(StartupPackages);
 	}
 
@@ -8187,6 +8188,7 @@ void UCookOnTheFlyServer::BlockOnAssetRegistry(TConstArrayView<FString> Commandl
 	{
 		TSet<FName> StartupPackages;
 		PackageTracker->InitializeTracking(StartupPackages);
+		NumPackagesHistory.Initialize(PackageTracker->NumLoadedPackages());
 		CookByTheBookOptions->StartupPackages = MoveTemp(StartupPackages);
 	}
 }
diff --git a/Engine/Source/Editor/UnrealEd/Private/Cooker/CookGarbageCollect.cpp b/Engine/Source/Editor/UnrealEd/Private/Cooker/CookGarbageCollect.cpp
--- a/Engine/Source/Editor/UnrealEd/Private/Cooker/CookGarbageCollect.cpp
+++ b/Engine/Source/Editor/UnrealEd/Private/Cooker/CookGarbageCollect.cpp
@@ -374,6 +374,7 @@ void ConstructSoftGCPackageToObjectList(TArray<UObject*>& PackageToObjectListBuf
 void UCookOnTheFlyServer::PollGarbageCollection(UE::Cook::FTickStackData& StackData)
 {
 	NumObjectsHistory.AddInstance(GUObjectArray.GetObjectArrayNumMinusAvailable());
+	NumPackagesHistory.AddInstance(PackageTracker->NumLoadedPackages());
 	VirtualMemoryHistory.AddInstance(FPlatformMemory::GetStats().UsedVirtual);
 
 	if (IsCookFlagSet(ECookInitializationFlags::TestCook))
@@ -608,11 +609,15 @@ bool UCookOnTheFlyServer::PumpHasExceededMaxMemory(uint32& OutResultFlags)
 void UCookOnTheFlyServer::SetGarbageCollectType(uint32 ResultFlagsFromTick)
 {
 	bGarbageCollectTypeSoft = (ResultFlagsFromTick & COSR_RequiresGC_Soft);
+	bSaveGCDataForEvaluation = true;
 }
 
 void UCookOnTheFlyServer::ClearGarbageCollectType()
 {
 	bGarbageCollectTypeSoft = false;
+	bSaveGCDataForEvaluation = false;
+
+	ClearGCData();
 }
 
 void UCookOnTheFlyServer::PreGarbageCollect()
@@ -626,6 +631,7 @@ void UCookOnTheFlyServer::PreGarbageCollect()
 	}
 
 	NumObjectsHistory.AddInstance(GUObjectArray.GetObjectArrayNumMinusAvailable());
+	NumPackagesHistory.AddInstance(PackageTracker->NumLoadedPackages());
 	VirtualMemoryHistory.AddInstance(FPlatformMemory::GetStats().UsedVirtual);
 	TArray<UPackage*> GCKeepPackages;
 	TArray<UE::Cook::FPackageData*> GCKeepPackageDatas;
@@ -826,6 +832,7 @@ void UCookOnTheFlyServer::PostGarbageCollect()
 	PackageTracker->SetCollectingGarbage(false);
 
 	NumObjectsHistory.AddInstance(GUObjectArray.GetObjectArrayNumMinusAvailable());
+	NumPackagesHistory.AddInstance(PackageTracker->NumLoadedPackages());
 	VirtualMemoryHistory.AddInstance(FPlatformMemory::GetStats().UsedVirtual);
 
 	TSet<UObject*> SaveQueueObjectsThatStillExist;
@@ -963,9 +970,10 @@ void UCookOnTheFlyServer::PostGarbageCollect()
 
 	// Only after running all possible callbacks that need our links for diagnostics, clear the list of temporary
 	// references that we created for the garbage collection.
-	GCKeepObjects.Empty();
-	UPackage::SoftGCPackageToObjectList.Empty();
-	SoftGCPackageToObjectListBuffer.Empty();
+	if (!bSaveGCDataForEvaluation)
+	{
+		ClearGCData();
+	}
 
 	CookedPackageCountSinceLastGC = 0;
 
@@ -974,6 +982,13 @@ void UCookOnTheFlyServer::PostGarbageCollect()
 	bShaderCompilerWasActiveOnPreviousBusyReport = true;
 }
 
+void UCookOnTheFlyServer::ClearGCData()
+{
+	GCKeepObjects.Empty();
+	UPackage::SoftGCPackageToObjectList.Empty();
+	SoftGCPackageToObjectListBuffer.Empty();
+}
+
 bool UCookOnTheFlyServer::NeedsDiagnosticSecondGC() const
 {
 	return GCDiagnosticContext->NeedsDiagnosticSecondGC();
@@ -989,10 +1004,15 @@ void UCookOnTheFlyServer::OnCookerEndCollectGarbage(uint32& ResultFlagsFromTick)
 	GCDiagnosticContext->OnCookerEndCollectGarbage(*this, ResultFlagsFromTick);
 }
 
+int32 UCookOnTheFlyServer::NumLoadedPackages() const
+{
+	return PackageTracker->NumLoadedPackages();
+}
+
 void UCookOnTheFlyServer::EvaluateGarbageCollectionResults(bool bWasDueToOOM, bool bWasPartialGC, uint32 ResultFlags,
-	int32 NumObjectsBeforeGC, const FPlatformMemoryStats& MemStatsBeforeGC,
+	int32 NumObjectsBeforeGC, const FPlatformMemoryStats& MemStatsBeforeGC, int32 NumPackagesBeforeGC,
 	const FGenericMemoryStats& AllocatorStatsBeforeGC,
-	int32 NumObjectsAfterGC, const FPlatformMemoryStats& MemStatsAfterGC,
+	int32 NumObjectsAfterGC, const FPlatformMemoryStats& MemStatsAfterGC, int32 NumPackagesAfterGC,
 	const FGenericMemoryStats& AllocatorStatsAfterGC,
 	float GCDurationSeconds)
 {
@@ -1044,16 +1064,26 @@ void UCookOnTheFlyServer::EvaluateGarbageCollectionResults(bool bWasDueToOOM, bo
 	{
 		return;
 	}
-	if (!bWasDueToOOM)
+
+	// Only show diagnostics if LLM is on, because they are somewhat expensive. We could add a separate setting
+	// for this, but it's more convenient to combine it with the LLM enabled setting
+#if ENABLE_LOW_LEVEL_MEM_TRACKER
+	bool bShowExtendedDiagnostics = FLowLevelMemTracker::Get().IsEnabled();
+#else
+	constexpr bool bShowExtendedDiagnostics = false;
+#endif
+	if (!bWasDueToOOM && !bShowExtendedDiagnostics)
 	{
 		return;
 	}
 
 	int64 NumObjectsMin = NumObjectsHistory.GetMinimum();
 	int64 NumObjectsMax = NumObjectsHistory.GetMaximum();
-	int64 NumObjectsSpread = NumObjectsMax - NumObjectsMin;
 	int64 NumObjectsFreed = NumObjectsBeforeGC - NumObjectsAfterGC;
 	int64 NumObjectsCapacity = GUObjectArray.GetObjectArrayEstimatedAvailable() + GUObjectArray.GetObjectArrayNumMinusAvailable();
+	int64 NumPackagesMin = NumPackagesHistory.GetMinimum();
+	int64 NumPackagesMax = NumPackagesHistory.GetMaximum();
+	int64 NumPackagesFreed = NumPackagesBeforeGC - NumPackagesAfterGC;
 	int64 VirtualMemMin = VirtualMemoryHistory.GetMinimum();
 	int64 VirtualMemMax = VirtualMemoryHistory.GetMaximum();
 	int64 VirtualMemSpread = VirtualMemMax - VirtualMemMin;
@@ -1061,8 +1091,9 @@ void UCookOnTheFlyServer::EvaluateGarbageCollectionResults(bool bWasDueToOOM, bo
 	int64 VirtualMemAfterGC = MemStatsAfterGC.UsedVirtual;
 	int64 VirtualMemFreed = MemStatsBeforeGC.UsedVirtual - MemStatsAfterGC.UsedVirtual;
 
-	int64 ExpectedObjectsFreed = static_cast<int64>(MemoryExpectedFreedToSpreadRatio * static_cast<float>(NumObjectsSpread));
-	double ExpectedMemFreed = MemoryExpectedFreedToSpreadRatio * VirtualMemSpread;
+	int64 ExpectedObjectsFreed = static_cast<int64>(MemoryExpectedFreedToSpreadRatio
+		* static_cast<float>(NumObjectsBeforeGC - NumObjectsMin));
+	double ExpectedMemFreed = MemoryExpectedFreedToSpreadRatio * (VirtualMemBeforeGC - VirtualMemMin);
 	static bool bCookMemoryAnalysis = FParse::Param(FCommandLine::Get(), TEXT("CookMemoryAnalysis"));
 #if ENABLE_LOW_LEVEL_MEM_TRACKER
 	// When tracking memory with LLM, always show the memory status.
@@ -1080,26 +1111,26 @@ void UCookOnTheFlyServer::EvaluateGarbageCollectionResults(bool bWasDueToOOM, bo
 			TEXT("\t\tCapacity:         %10d\n")
 			TEXT("\t\tBefore GC:        %10d\n")
 			TEXT("\t\tAfter GC:         %10d\n")
-			TEXT("\t\tFreed by GC:      %10d\n")
+			TEXT("\t\tFreed by GC:      %10.1f%%\n")
+			TEXT("\tNumPackages:\n")
+			TEXT("\t\tBefore GC:        %10d\n")
+			TEXT("\t\tAfter GC:         %10d\n")
+			TEXT("\t\tFreed by GC:      %10.1f%%\n")
 			TEXT("\tVirtual Memory:\n")
 			TEXT("\t\tBefore GC:        %10" INT64_FMT " MB\n")
 			TEXT("\t\tAfter GC:         %10" INT64_FMT " MB\n")
-			TEXT("\t\tFreed by GC:      %10" INT64_FMT " MB"),
+			TEXT("\t\tFreed by GC:      %10.1f%%"),
 			(bWasSoftGC ? TEXT("Soft") : (bWasPartialGC ? TEXT("Partial") : TEXT("Full"))),
 			GCDurationSeconds,
-			NumObjectsCapacity, (int64)NumObjectsBeforeGC, (int64)NumObjectsAfterGC, NumObjectsFreed,
-			VirtualMemBeforeGC / BytesPerMeg, VirtualMemAfterGC / BytesPerMeg, VirtualMemFreed / BytesPerMeg
+			NumObjectsCapacity, (int64)NumObjectsBeforeGC, (int64)NumObjectsAfterGC,
+			100.f * ((float)NumObjectsFreed)/NumObjectsBeforeGC,
+			(int64) NumPackagesBeforeGC, (int64) NumPackagesAfterGC,
+			100.f * ((float)NumPackagesFreed)/NumPackagesBeforeGC,
+			VirtualMemBeforeGC / BytesPerMeg, VirtualMemAfterGC / BytesPerMeg,
+			100.f * ((float)VirtualMemFreed) / VirtualMemBeforeGC
 		);
 	};
 
-	// Only show diagnostics if LLM is on, because they are somewhat expensive. We could add a separate setting
-	// for this, but it's more convenient to combine it with the LLM enabled setting
-#if ENABLE_LOW_LEVEL_MEM_TRACKER
-	bool bShowExtendedDiagnostics = FLowLevelMemTracker::Get().IsEnabled();
-#else
-	constexpr bool bShowExtendedDiagnostics = false;
-#endif
-
 	if (!bWasSoftGC)
 	{
 		bool bWasImpactful =
@@ -1129,21 +1160,30 @@ void UCookOnTheFlyServer::EvaluateGarbageCollectionResults(bool bWasDueToOOM, bo
 			TEXT("\t\tCapacity:         %10" INT64_FMT "\n")
 			TEXT("\t\tProcess Min:      %10" INT64_FMT "\n")
 			TEXT("\t\tProcess Max:      %10" INT64_FMT "\n")
-			TEXT("\t\tProcess Spread:   %10" INT64_FMT "\n")
 			TEXT("\t\tBefore GC:        %10" INT64_FMT "\n")
 			TEXT("\t\tAfter GC:         %10" INT64_FMT "\n")
-			TEXT("\t\tFreed by GC:      %10" INT64_FMT ""),
-			NumObjectsCapacity, NumObjectsMin, NumObjectsMax, NumObjectsSpread,
-			(int64)NumObjectsBeforeGC, (int64)NumObjectsAfterGC, NumObjectsFreed);
+			TEXT("\t\tFreed by GC:      %10.1f%%"),
+			NumObjectsCapacity, NumObjectsMin, NumObjectsMax,
+			(int64)NumObjectsBeforeGC, (int64)NumObjectsAfterGC,
+			100.f * ((float)NumObjectsFreed)/NumObjectsBeforeGC);
+		UE_LOG(LogCook, Display, TEXT("\tMemoryAnalysis: NumPackages:\n")
+			TEXT("\t\tProcess Min:      %10" INT64_FMT "\n")
+			TEXT("\t\tProcess Max:      %10" INT64_FMT "\n")
+			TEXT("\t\tBefore GC:        %10" INT64_FMT "\n")
+			TEXT("\t\tAfter GC:         %10" INT64_FMT "\n")
+			TEXT("\t\tFreed by GC:      %10.1f%%"),
+			NumPackagesMin, NumPackagesMax,
+			(int64)NumPackagesBeforeGC, (int64)NumPackagesAfterGC,
+			100.f * ((float)NumPackagesFreed)/NumPackagesBeforeGC);
 		UE_LOG(LogCook, Display, TEXT("\tMemoryAnalysis: Virtual Memory:\n")
 			TEXT("\t\tProcess Min:      %10" INT64_FMT " MB\n")
 			TEXT("\t\tProcess Max:      %10" INT64_FMT " MB\n")
-			TEXT("\t\tProcess Spread:   %10" INT64_FMT " MB\n")
 			TEXT("\t\tBefore GC:        %10" INT64_FMT " MB\n")
 			TEXT("\t\tAfter GC:         %10" INT64_FMT " MB\n")
-			TEXT("\t\tFreed by GC:      %10" INT64_FMT " MB"),
-			VirtualMemMin / BytesPerMeg, VirtualMemMax / BytesPerMeg, VirtualMemSpread / BytesPerMeg,
-			VirtualMemBeforeGC / BytesPerMeg, VirtualMemAfterGC / BytesPerMeg, VirtualMemFreed / BytesPerMeg);
+			TEXT("\t\tFreed by GC:      %10.1f%%"),
+			VirtualMemMin / BytesPerMeg, VirtualMemMax / BytesPerMeg,
+			VirtualMemBeforeGC / BytesPerMeg, VirtualMemAfterGC / BytesPerMeg,
+			100.f * ((float)VirtualMemFreed)/VirtualMemBeforeGC);
 		auto AllocatorStatsToString = [](const FGenericMemoryStats& AllocatorStats)
 		{
 			TStringBuilder<256> Writer;
@@ -1159,12 +1199,6 @@ void UCookOnTheFlyServer::EvaluateGarbageCollectionResults(bool bWasDueToOOM, bo
 			*AllocatorStatsToString(AllocatorStatsAfterGC));
 
 
-#if ENABLE_LOW_LEVEL_MEM_TRACKER
-		bool bExtendedMemoryAnalysisEnabled = FLowLevelMemTracker::Get().IsEnabled();
-#else
-		constexpr bool bExtendedMemoryAnalysisEnabled = false;
-#endif
-
 		// Only show diagnostics if LLM is on, because they are somewhat expensive. We could add a separate setting
 		// for this, but it's more convenient to combine it with the LLM enabled setting
 		if (!bShowExtendedDiagnostics)
@@ -1174,13 +1208,7 @@ void UCookOnTheFlyServer::EvaluateGarbageCollectionResults(bool bWasDueToOOM, bo
 		else
 		{
 			UE_LOG(LogCook, Display, TEXT("See log for memory use information for UObject classes and LLM tags."));
-			{
-				TGuardValue<bool> SoftGCGuard(UPackage::bSupportCookerSoftGC, true);
-				ConstructSoftGCPackageToObjectList(this->SoftGCPackageToObjectListBuffer);
-				UE::Cook::DumpObjClassList(CookByTheBookOptions->SessionStartupObjects);
-				UPackage::SoftGCPackageToObjectList.Empty();
-				SoftGCPackageToObjectListBuffer.Empty();
-			}
+			UE::Cook::DumpPackageReferencers(TEXT("Non-startup Packages"), CookByTheBookOptions->SessionStartupObjects);
 			GLog->Logf(TEXT("Memory Analysis: LLM Tags:"));
 #if ENABLE_LOW_LEVEL_MEM_TRACKER
 			if (FLowLevelMemTracker::Get().IsEnabled())
@@ -1227,7 +1255,8 @@ void UCookOnTheFlyServer::EvaluateGarbageCollectionResults(bool bWasDueToOOM, bo
 				});
 			if (PackagesReferencedOutsideOfCooker.Num() > 0)
 			{
-				UE::Cook::DumpPackageReferencers(PackagesReferencedOutsideOfCooker);
+				UE::Cook::DumpPackageReferencers(TEXT("SoftGCPackages"),
+					CookByTheBookOptions->SessionStartupObjects, PackagesReferencedOutsideOfCooker);
 			}
 		}
 	}
diff --git a/Engine/Source/Editor/UnrealEd/Private/Cooker/CookProfiling.cpp b/Engine/Source/Editor/UnrealEd/Private/Cooker/CookProfiling.cpp
--- a/Engine/Source/Editor/UnrealEd/Private/Cooker/CookProfiling.cpp
+++ b/Engine/Source/Editor/UnrealEd/Private/Cooker/CookProfiling.cpp
@@ -15,6 +15,7 @@
 #include "HAL/FileManager.h"
 #include "Misc/CommandLine.h"
 #include "Misc/OutputDevice.h"
+#include "Misc/PackageName.h"
 #include "Misc/PathViews.h"
 #include "Misc/StringBuilder.h"
 #include "PackageBuildDependencyTracker.h"
@@ -510,6 +511,13 @@ void ConstructObjectGraph(TConstArrayView<UObject*> AllObjects,
 				{
 					SourceObject->GetClass()->CallAddReferencedObjects(SourceObject, Collector);
 				}
+				// TODO: Handle elements in the token stream not covered by serialize, such as GetOuter()
+				// In the meantime we handle Outer explicitly.
+				UObject* Outer = SourceObject->GetOuter();
+				if (Outer)
+				{
+					TargetObjects.Add(Outer);
+				}
 			}
 		}
 
@@ -774,7 +782,8 @@ void DumpObjClassList(TConstArrayView<FWeakObjectPtr> InitialObjects)
 	}
 }
 
-void DumpPackageReferencers(TConstArrayView<UPackage*> Packages)
+void DumpPackageReferencers(const TCHAR* DebugNameOfPackageList,
+	TConstArrayView<FWeakObjectPtr> InitialObjects, TConstArrayView<UPackage*> Packages)
 {
 	using namespace Algo::Graph;
 
@@ -782,46 +791,147 @@ void DumpPackageReferencers(TConstArrayView<UPackage*> Packages)
 	FObjectGraphProfileData ProfileData;
 	ConstructObjectGraphProfileData(TConstArrayView<FWeakObjectPtr>(), ProfileData);
 
-	// List all roots that cause any of the Packages to remain alive, and count how many packages each one causes
-	TMap<FVertex, int32> Roots;
-	TMap<FVertex, FVertex> RootExamples;
-	int32 Unexpected = 0;
-	for (UPackage* Package : Packages)
+	TArray<UPackage*> AllPackagesInMemory;
+	if (Packages.IsEmpty())
 	{
-		FVertex* FoundVertex = ProfileData.VertexOfObject.Find(Package);
-		if (!FoundVertex)
+		// Ignore packages containing any of the InitialObjects
+		TSet<UPackage*> InitialPackages;
+		for (const FWeakObjectPtr& WeakObject : InitialObjects)
 		{
-			++Unexpected;
-			continue;
+			UObject* Object = WeakObject.Get();
+			if (Object)
+			{
+				InitialPackages.Add(Object->GetPackage());
+			}
 		}
-		FVertex PackageVertex = *FoundVertex;
-		FVertex RootOfThisVertex = ProfileData.RootOfVertex[PackageVertex];
-		if (RootOfThisVertex == InvalidVertex)
+		for (TObjectIterator<UPackage> Iter; Iter; ++Iter)
 		{
-			++Unexpected;
-			continue;
+			UPackage* Package = *Iter;
+			if (FPackageName::IsScriptPackage(WriteToString<256>(Package->GetFName())))
+			{
+				continue;
+			}
+			if (InitialPackages.Contains(Package))
+			{
+				continue;
+			}
+			AllPackagesInMemory.Add(Package);
+		}
+		Packages = AllPackagesInMemory;
+	}
+
+	// Store the root keeping each package in memory and calculate the memory size used by each package.
+	struct FPackageDumpData
+	{
+		UPackage* Package = nullptr;
+		FVertex Vertex = InvalidVertex;
+		FVertex RootVertex = InvalidVertex;
+		uint64 Size = 0;
+	};
+	struct FRootDumpData
+	{
+		FVertex Vertex = InvalidVertex;
+		TArray<FPackageDumpData*> Packages;
+		uint64 Size = 0;
+	};
+	uint64 TotalPackageSize = 0;
+	TMap<UPackage*, TUniquePtr<FPackageDumpData>> PackageDatas;
+	TMap<FVertex, FRootDumpData> RootDatas;
+	for (UPackage* Package : Packages)
+	{
+		TUniquePtr<FPackageDumpData>& PackageData = PackageDatas.FindOrAdd(Package);
+		if (!PackageData)
+		{
+			PackageData.Reset(new FPackageDumpData);
+			PackageData->Package = Package;
 		}
-		int32& RootCount = Roots.FindOrAdd(RootOfThisVertex);
-		if (RootCount == 0)
+		FVertex* FoundVertex = ProfileData.VertexOfObject.Find(Package);
+		PackageData->Vertex = FoundVertex ? *FoundVertex : InvalidVertex;
+		PackageData->RootVertex = PackageData->Vertex != InvalidVertex
+			? ProfileData.RootOfVertex[PackageData->Vertex] : InvalidVertex;
+
+		PackageData->Size += Package->GetClass()->GetStructureSize();
+		PackageData->Size += Package->GetResourceSizeBytes(EResourceSizeMode::Exclusive);
+		ForEachObjectWithPackage(Package, [&PackageData](UObject* Object)
+			{
+				PackageData->Size += Object->GetClass()->GetStructureSize();
+				PackageData->Size += Object->GetResourceSizeBytes(EResourceSizeMode::Exclusive);
+				return true;
+			}, true /* bIncludeNestedObjects */);
+
+		FRootDumpData& RootData = RootDatas.FindOrAdd(PackageData->RootVertex);
+		if (RootData.Packages.IsEmpty())
 		{
-			RootExamples.Add(RootOfThisVertex, PackageVertex);
+			RootData.Vertex = PackageData->RootVertex;
 		}
-		++RootCount;
+		RootData.Packages.Add(PackageData.Get());
+		RootData.Size += PackageData->Size;
+		TotalPackageSize += PackageData->Size;
 	}
 
-	LogAr.Logf(TEXT("Memory Analysis: Referencers of SoftGCPackages:"));
-	Roots.ValueSort([](int32 A, int32 B) { return A > B; });
-	for (TPair<FVertex, int32>& Pair : Roots)
+	LogAr.Logf(TEXT("Memory Analysis: Referencers of %s"), DebugNameOfPackageList);
+	if (RootDatas.Num() == 0)
+	{
+		LogAr.Logf(TEXT("No %s found to report"), DebugNameOfPackageList);
+		return;
+	}
+	LogAr.Logf(TEXT("%d Referencers are holding in memory %d packages totalling %d MiB in size."),
+		RootDatas.Num(), PackageDatas.Num(), (int32)(TotalPackageSize / 1024 / 1024));
+
+	// Sort the roots by size
+	RootDatas.ValueSort([&ProfileData](const FRootDumpData& A, const FRootDumpData& B)
+		{
+			if (A.Size != B.Size)
+			{
+				return A.Size > B.Size;
+			}
+			TStringBuilder<256> AName;
+			TStringBuilder<256> BName;
+			A.Vertex != InvalidVertex ? ProfileData.AppendVertexName(A.Vertex, AName)
+				: (void) AName.Append(TEXT("<Unknown>"));
+			B.Vertex != InvalidVertex ? ProfileData.AppendVertexName(B.Vertex, BName)
+				: (void) BName.Append(TEXT("<Unknown>"));
+			return AName.ToView().Compare(BName) < 0;
+		});
+
+	uint64 MaxSize = 0;
+	int32 NumSoFar = 0;
+	constexpr float MinimumRelativeSizeToDisplay = 0.10f;
+	constexpr int PackagesPerRoot = 5;
+	for (TPair<FVertex, FRootDumpData>& Pair : RootDatas)
 	{
+		FRootDumpData& RootData = Pair.Value;
+		++NumSoFar;
+		if (MaxSize == 0)
+		{
+			MaxSize = RootData.Size; // Elements are sorted by size, first element in iteration is the largest.
+		}
+		if (RootData.Size < (uint64)(MinimumRelativeSizeToDisplay * MaxSize))
+		{
+			LogAr.Logf(TEXT("\t... (%d more referencers each responsible for <= %d MiB in size)"),
+				RootDatas.Num() - NumSoFar, (int32)(RootData.Size / 1024 / 1024));
+			break;
+		}
+
 		TStringBuilder<256> ReferencerName;
-		ProfileData.AppendVertexName(Pair.Key, ReferencerName);
-		LogAr.Logf(TEXT("\t%5d: %s"), Pair.Value, *ReferencerName);
-		FVertex* ExampleVertexPtr = RootExamples.Find(Pair.Key);
-		if (ExampleVertexPtr)
+		RootData.Vertex != InvalidVertex ? ProfileData.AppendVertexName(RootData.Vertex, ReferencerName)
+			: (void) ReferencerName.Append(TEXT("<Unknown>"));
+		LogAr.Logf(TEXT("\t%s"), *ReferencerName);
+		check(!RootData.Packages.IsEmpty()); // Creation of each rootdata includes adding a package to it.
+		Algo::Sort(RootData.Packages, [](FPackageDumpData* A, FPackageDumpData* B)
+			{
+				if (A->Size != B->Size)
+				{
+					return A->Size < B->Size;
+				}
+				return A->Package->GetFName().LexicalLess(B->Package->GetFName());
+			});
+		FPackageDumpData* ExamplePackage = RootData.Packages[0];
+		if (ExamplePackage->Vertex != InvalidVertex)
 		{
 			TStringBuilder<256> Chain;
-			FObjectReferencer Link = ProfileData.AliveReason[*ExampleVertexPtr];
-			ProfileData.AllObjects[*ExampleVertexPtr]->GetFullName(Chain);
+			FObjectReferencer Link = ProfileData.AliveReason[ExamplePackage->Vertex];
+			ProfileData.AllObjects[ExamplePackage->Vertex]->GetFullName(Chain);
 			while (Link.GetLinkType() == EObjectReferencerType::Referenced)
 			{
 				Chain << TEXT(" <- ");
@@ -830,28 +940,23 @@ void DumpPackageReferencers(TConstArrayView<UPackage*> Packages)
 			}
 			Chain << TEXT(" <- ");
 			Link.ToString(Chain, ProfileData);
-			LogAr.Logf(TEXT("\t\t     Ex: %s"), *Chain);
+			LogAr.Logf(TEXT("\t\tExample Chain: %s"), *Chain);
 		}
-	}
-
-	if (Unexpected > 0)
-	{
-		LogAr.Logf(TEXT("Memory Analysis: Unknown referenced SoftGCPackages:"));
-		for (UPackage* Package : Packages)
+		LogAr.Logf(TEXT("\t\t%d MiB, %d Packages"), (int32)(RootData.Size / 1024 / 1024), RootData.Packages.Num());
+		int32 NumPackagesSoFar = 0;
+		uint64 SizeSoFar = 0;
+		for (FPackageDumpData* PackageData : RootData.Packages)
 		{
-			FVertex* FoundVertex = ProfileData.VertexOfObject.Find(Package);
-			if (!FoundVertex)
+			++NumPackagesSoFar;
+			if (NumPackagesSoFar > PackagesPerRoot)
 			{
-				LogAr.Logf(TEXT("%s: unknown, we did not create a vertex for the package"), *Package->GetName());
-				continue;
-			}
-			FVertex PackageVertex = *FoundVertex;
-
-			if (ProfileData.AliveReason[PackageVertex].GetLinkType() == EObjectReferencerType::Unknown)
-			{
-				LogAr.Logf(TEXT("%s: no reference found"), *Package->GetName());
-				continue;
+				LogAr.Logf(TEXT("\t\t\t... (%d more packages totalling %d MiB in size)"),
+					RootData.Packages.Num() - NumPackagesSoFar, (int32)((RootData.Size - SizeSoFar) / 1024 / 1024));
+				break;
 			}
+			SizeSoFar += PackageData->Size;
+			LogAr.Logf(TEXT("\t\t\t%4d MiB: %s"),
+				(int32)(PackageData->Size / 1024 / 1024), *WriteToString<256>(PackageData->Package->GetFName()));
 		}
 	}
 }
diff --git a/Engine/Source/Editor/UnrealEd/Private/Cooker/CookProfiling.h b/Engine/Source/Editor/UnrealEd/Private/Cooker/CookProfiling.h
--- a/Engine/Source/Editor/UnrealEd/Private/Cooker/CookProfiling.h
+++ b/Engine/Source/Editor/UnrealEd/Private/Cooker/CookProfiling.h
@@ -202,7 +202,8 @@ namespace UE::Cook
 void DumpObjClassList(TConstArrayView<FWeakObjectPtr> SessionStartupObjects);
 
 /** Used for profiling after a soft garbage collect. Find the referencers keeping each package in memory. */
-void DumpPackageReferencers(TConstArrayView<UPackage*> Packages);
+void DumpPackageReferencers(const TCHAR* DebugNameOfPackageList, TConstArrayView<FWeakObjectPtr> InitialObjects, 
+	TConstArrayView<UPackage*> Packages = TConstArrayView<UPackage*>());
 
 
 }
\ No newline at end of file
diff --git a/Engine/Source/Editor/UnrealEd/Private/Cooker/PackageTracker.cpp b/Engine/Source/Editor/UnrealEd/Private/Cooker/PackageTracker.cpp
--- a/Engine/Source/Editor/UnrealEd/Private/Cooker/PackageTracker.cpp
+++ b/Engine/Source/Editor/UnrealEd/Private/Cooker/PackageTracker.cpp
@@ -109,6 +109,11 @@ void FPackageTracker::InitializeTracking(TSet<FName>& OutStartupPackages)
 		}
 		PackageStream.Add({ Pair.Key, MoveTemp(Pair.Value), EPackageStreamEvent::PackageLoad });
 	}
+
+	for (TObjectIterator<UPackage> Iter; Iter; ++Iter)
+	{
+		LoadedPackages.Add(*Iter);
+	}
 }
 
 TArray<FPackageStreamEvent> FPackageTracker::GetPackageStream()
diff --git a/Engine/Source/Runtime/Core/Private/HAL/LowLevelMemTracker.cpp b/Engine/Source/Runtime/Core/Private/HAL/LowLevelMemTracker.cpp
--- a/Engine/Source/Runtime/Core/Private/HAL/LowLevelMemTracker.cpp
+++ b/Engine/Source/Runtime/Core/Private/HAL/LowLevelMemTracker.cpp
@@ -1129,6 +1129,10 @@ void FLowLevelMemTracker::DumpToLog(EDumpFormat DumpFormat, FOutputDevice* Outpu
 		TArray<FTagLine> TagLines;
 		for (const UE::LLMPrivate::FTagData* TagData : GetTrackedTags(TrackerType, TagSet))
 		{
+			if (!TagData->IsReportable())
+			{
+				continue;
+			}
 			int64 CurrentAmount = GetTagAmountForTracker(TrackerType, TagData, SizeParams);
 			int64 PeakAmount = GetTagAmountForTracker(TrackerType, TagData, PeakSizeParams);
 			FString TagName = GetTagUniqueName(TagData).ToString();
