Fixed objects translating in the opposite direction than what is expected in the Combined Translate and Rotate Mode when the axis display coordinate system is LUF and the user is translating on the Y axis.

#Synced-CL 47947116

diff --git a/Engine/Source/Editor/UnrealEd/Private/UnrealWidget.cpp b/Engine/Source/Editor/UnrealEd/Private/UnrealWidget.cpp
--- a/Engine/Source/Editor/UnrealEd/Private/UnrealWidget.cpp
+++ b/Engine/Source/Editor/UnrealEd/Private/UnrealWidget.cpp
@@ -411,6 +411,11 @@ void FWidget::ConvertMouseToAxis_TranslateRotateZ(FVector2D TangentDir, FVector2
 		const FVector2D EffectiveDelta = OutDrag.X * XAxisDir + OutDrag.Y * YAxisDir + OutDrag.Z * ZAxisDir;
 		InOutDelta                     = FVector(EffectiveDelta.X, -EffectiveDelta.Y, 0.0f);
 
+		if (AxisDisplayInfo::GetAxisDisplayCoordinateSystem() == EAxisList::LeftUpForward)
+		{
+			OutDrag.Y = -OutDrag.Y;
+		}
+
 		// Transform drag delta into world space
 		OutDrag = CustomCoordSystem.TransformPosition(OutDrag);
 	}
