Fixes for reinstantiation problem of the SceneGraphTestAll project
From a previous change in the reinstanciation, we might be replacing early some of the reinstanciated objects references. This check was detecting this case. Now we need just need to handle it correctly.

diff --git a/Engine/Source/Editor/UnrealEd/Private/Kismet2/KismetReinstanceUtilities.cpp b/Engine/Source/Editor/UnrealEd/Private/Kismet2/KismetReinstanceUtilities.cpp
--- a/Engine/Source/Editor/UnrealEd/Private/Kismet2/KismetReinstanceUtilities.cpp
+++ b/Engine/Source/Editor/UnrealEd/Private/Kismet2/KismetReinstanceUtilities.cpp
@@ -3399,7 +3399,6 @@ void FBlueprintCompileReinstancer::PreCreateSubObjectsForReinstantiation_Inner(
 	GetObjectsWithOuter(NewUObject, ContainedNewSubObjects, /*bIncludeNestedObjects*/false);
 
 	// Pre-create all non default subobjects to prevent re-instancing them as an old classes 
-	TMap<UObject*, UObject*> ReferenceReplacementMap;
 	for (int32 i = 0; i < ContainedOldSubObjects.Num(); ++i)
 	{
 		UObject* OldSubObject = ContainedOldSubObjects[i];
@@ -3435,10 +3434,13 @@ void FBlueprintCompileReinstancer::PreCreateSubObjectsForReinstantiation_Inner(
 				}
 				(*AlreadyCreatedSubObject)->Rename(nullptr, NewUObject, REN_DoNotDirty | REN_DontCreateRedirectors | REN_NonTransactional);
 			}
+			// Could still be already re-instantiated but be also replaced
+			else if (UObject*const* AlreadyReplacedSubObject = OldToNewInstanceMap && OldToNewInstanceMap->FindKey(OldSubObject) ? &OldSubObject : nullptr)
+			{
+				(*AlreadyReplacedSubObject)->Rename(nullptr, NewUObject, REN_DoNotDirty | REN_DontCreateRedirectors | REN_NonTransactional);
+			}
 			else
 			{
-				checkf( !OldToNewInstanceMap || !OldToNewInstanceMap->FindKey(OldSubObject), TEXT("For performance reason, let's assume the any old sub object will be before its replacement in the contained subobject list"));
-
 				// No need to handled invalid objects from this point on
 				if (!IsValid(OldSubObject))
 				{
