[ContentCommandlets]
- Fix issue where GetPackageLinker could corrupt async loading

#Synced-CL 0

diff --git a/Engine/Source/Editor/UnrealEd/Private/Commandlets/ContentCommandlets.cpp b/Engine/Source/Editor/UnrealEd/Private/Commandlets/ContentCommandlets.cpp
--- a/Engine/Source/Editor/UnrealEd/Private/Commandlets/ContentCommandlets.cpp
+++ b/Engine/Source/Editor/UnrealEd/Private/Commandlets/ContentCommandlets.cpp
@@ -777,6 +777,10 @@ void UResavePackagesCommandlet::LoadAndSaveOnePackage(const FString& Filename)
 		// Get the package linker.
 		VerboseMessage(TEXT("Pre GetPackageLinker"));
 
+		// Load the package before accessing its linker to prevent interference with async loads in case
+		// that package was already queued and being loaded asynchronously.
+		LoadPackage(nullptr, PackagePath, LOAD_NoVerify);
+
 		FLinkerLoad* Linker = GetPackageLinker(nullptr, PackagePath, LOAD_NoVerify, nullptr);
 
 		// Bail early if we don't have a valid linker (package was out of date, etc)
@@ -3107,6 +3111,10 @@ int32 UWrangleContentCommandlet::Main( const FString& Params )
 				}
 			}
 
+			// Load the package before accessing its linker to prevent interference with async loads in case
+			// that package was already queued and being loaded asynchronously.
+			LoadPackage(nullptr, PackagePath, LOAD_Quiet | LOAD_NoWarn | LOAD_NoVerify);
+
 			FLinkerLoad* Linker = GetPackageLinker(nullptr, PackagePath, LOAD_Quiet | LOAD_NoWarn | LOAD_NoVerify, nullptr);
 
 			// go through the exports in the package, looking for public objects
