Cooker: Profiling: Tweak the garbage collection analysis to provide the extended analysis as well for SoftGC and other non-OOM sources of GC, when cooker memory profiling is enabled
#rnx

#Synced-CL 45140074

diff --git a/Engine/Source/Editor/UnrealEd/Private/Cooker/CookGarbageCollect.cpp b/Engine/Source/Editor/UnrealEd/Private/Cooker/CookGarbageCollect.cpp
--- a/Engine/Source/Editor/UnrealEd/Private/Cooker/CookGarbageCollect.cpp
+++ b/Engine/Source/Editor/UnrealEd/Private/Cooker/CookGarbageCollect.cpp
@@ -1065,15 +1065,88 @@ void UCookOnTheFlyServer::EvaluateGarbageCollectionResults(bool bWasDueToOOM, bo
 		return;
 	}
 
-	// Only show diagnostics if LLM is on, because they are somewhat expensive. We could add a separate setting
-	// for this, but it's more convenient to combine it with the LLM enabled setting
+	// Mark the packages we freed so we can give a warning to diagnose why they are still referenced if they
+	// get loaded again.
+	if (bWasSoftGC)
+	{
+		PackageTracker->AddExpectedNeverLoadPackages(ExpectedFreedPackageNames);
+	}
+
+	enum class EDisplayLevel : uint8
+	{
+		None,
+		Simple,
+		Full,
+		FullWithDiagnostics,
+	};
+
+	static EDisplayLevel LevelGCOther = EDisplayLevel::None;
+	static EDisplayLevel LevelGCSuccess = EDisplayLevel::None;
+	// LevelGCFail defined below as the return value of ParseCommandLineAndReturnLevelGCFail.
+	auto ParseCommandLineAndReturnLevelGCFail = []()
+	{
+		EDisplayLevel LevelGCFailLocal;
 #if ENABLE_LOW_LEVEL_MEM_TRACKER
-	bool bShowExtendedDiagnostics = FLowLevelMemTracker::Get().IsEnabled();
-#else
-	constexpr bool bShowExtendedDiagnostics = false;
+		if (FLowLevelMemTracker::Get().IsEnabled())
+		{
+			LevelGCOther = EDisplayLevel::FullWithDiagnostics;
+			LevelGCSuccess = EDisplayLevel::FullWithDiagnostics;
+			LevelGCFailLocal = EDisplayLevel::FullWithDiagnostics;
+		}
+		else
 #endif
-	if (!bWasDueToOOM && !bShowExtendedDiagnostics)
+		{
+			LevelGCOther = EDisplayLevel::None;
+			LevelGCSuccess = EDisplayLevel::Simple;
+			LevelGCFailLocal = EDisplayLevel::Full;
+		}
+		// Command line overrides the defaults
+		auto TryParseDisplayLevel = [](const TCHAR* FlagPattern, EDisplayLevel& OutLevel)
+		{
+			FString Text;
+			bool bParsedValue = FParse::Value(FCommandLine::Get(), *FString::Printf(TEXT("-%s="), FlagPattern), Text);
+			if (!bParsedValue && FParse::Param(FCommandLine::Get(), FlagPattern))
+			{
+				Text = TEXT("Full");
+			}
+			if (Text.IsEmpty())
+			{
+				return false;
+			}
+			for (const auto& Pair : std::initializer_list<TPair<EDisplayLevel, const ANSICHAR*>> {
+				{ EDisplayLevel::None, "None" },
+				{ EDisplayLevel::Simple, "Simple" },
+				{ EDisplayLevel::FullWithDiagnostics, "Full" } })
+			{
+				if (Text == Pair.Value)
+				{
+					OutLevel = Pair.Key;
+					return true;
+				}
+			}
+			UE_LOG(LogCook, Error, TEXT("Invalid commandline %s argument for -%s=%s. Valid values are { None, Simple, Full }."),
+				*Text, FlagPattern, *Text);
+			return false;
+		};
+
+		EDisplayLevel LevelAll;
+		if (TryParseDisplayLevel(TEXT("CookMemoryAnalysis"), LevelAll))
+		{
+			LevelGCOther = LevelAll;
+			LevelGCSuccess = LevelAll;
+			LevelGCFailLocal = LevelAll;
+		}
+		TryParseDisplayLevel(TEXT("CookMemoryAnalysis.GCOther"), LevelGCOther);
+		TryParseDisplayLevel(TEXT("CookMemoryAnalysis.GCSuccess"), LevelGCSuccess);
+		TryParseDisplayLevel(TEXT("CookMemoryAnalysis.GCFail"), LevelGCFailLocal);
+		return LevelGCFailLocal;
+	};
+	static EDisplayLevel LevelGCFail = ParseCommandLineAndReturnLevelGCFail();
+
+	if ((!bWasDueToOOM && LevelGCOther == EDisplayLevel::None)
+		|| (bWasDueToOOM && LevelGCSuccess == EDisplayLevel::None && LevelGCFail == EDisplayLevel::None))
 	{
+		// No diagnostics requested for this case, early exit to avoid any further cost.
 		return;
 	}
 
@@ -1094,17 +1167,34 @@ void UCookOnTheFlyServer::EvaluateGarbageCollectionResults(bool bWasDueToOOM, bo
 	int64 ExpectedObjectsFreed = static_cast<int64>(MemoryExpectedFreedToSpreadRatio
 		* static_cast<float>(NumObjectsBeforeGC - NumObjectsMin));
 	double ExpectedMemFreed = MemoryExpectedFreedToSpreadRatio * (VirtualMemBeforeGC - VirtualMemMin);
-	static bool bCookMemoryAnalysis = FParse::Param(FCommandLine::Get(), TEXT("CookMemoryAnalysis"));
-#if ENABLE_LOW_LEVEL_MEM_TRACKER
-	// When tracking memory with LLM, always show the memory status.
-	const bool bAlwaysShowAnalysis = FLowLevelMemTracker::Get().IsEnabled() || bCookMemoryAnalysis;
-#else
-	const bool bAlwaysShowAnalysis = bCookMemoryAnalysis;
-#endif		
+
+	FString HeaderMessage = TEXT("GarbageCollection Results:");
+	EDisplayLevel DisplayLevel = EDisplayLevel::None;
+	bool bWasImpactful = true;
+	if (!bWasDueToOOM)
+	{
+		DisplayLevel = LevelGCOther;
+	}
+	else
+	{
+		bWasImpactful =
+			(NumObjectsFreed >= ExpectedObjectsFreed || NumObjectsBeforeGC - NumObjectsMin < ExpectedObjectsFreed) &&
+			(VirtualMemFreed >= ExpectedMemFreed || VirtualMemBeforeGC - VirtualMemMin <= ExpectedMemFreed);
+		if (!bWasImpactful)
+		{
+			HeaderMessage = TEXT("GarbageCollection Results: Garbage Collection was not very impactful.");
+		}
+		DisplayLevel = bWasImpactful ? LevelGCSuccess : LevelGCFail;
+	}
+
 	constexpr int32 BytesPerMeg = 1000000;
-	auto DisplaySimpleSummary = [&]()
+	if (DisplayLevel == EDisplayLevel::None)
 	{
-		UE_LOG(LogCook, Display, TEXT("GarbageCollection Results:\n")
+		// No text
+	}
+	else if (DisplayLevel == EDisplayLevel::Simple)
+	{
+		UE_LOG(LogCook, Display, TEXT("%s\n")
 			TEXT("\tType: %s\n")
 			TEXT("\tDuration: %.3fs\n")
 			TEXT("\tNumObjects:\n")
@@ -1120,6 +1210,7 @@ void UCookOnTheFlyServer::EvaluateGarbageCollectionResults(bool bWasDueToOOM, bo
 			TEXT("\t\tBefore GC:        %10" INT64_FMT " MB\n")
 			TEXT("\t\tAfter GC:         %10" INT64_FMT " MB\n")
 			TEXT("\t\tFreed by GC:      %10.1f%%"),
+			*HeaderMessage,
 			(bWasSoftGC ? TEXT("Soft") : (bWasPartialGC ? TEXT("Partial") : TEXT("Full"))),
 			GCDurationSeconds,
 			NumObjectsCapacity, (int64)NumObjectsBeforeGC, (int64)NumObjectsAfterGC,
@@ -1129,28 +1220,11 @@ void UCookOnTheFlyServer::EvaluateGarbageCollectionResults(bool bWasDueToOOM, bo
 			VirtualMemBeforeGC / BytesPerMeg, VirtualMemAfterGC / BytesPerMeg,
 			100.f * ((float)VirtualMemFreed) / VirtualMemBeforeGC
 		);
-	};
-
-	if (!bWasSoftGC)
+	}
+	else
 	{
-		bool bWasImpactful =
-			(NumObjectsFreed >= ExpectedObjectsFreed || NumObjectsBeforeGC - NumObjectsMin < ExpectedObjectsFreed) &&
-			(VirtualMemFreed >= ExpectedMemFreed || VirtualMemBeforeGC - VirtualMemMin <= ExpectedMemFreed);
-
-		if ((!bWasDueToOOM || bWasImpactful) && !bAlwaysShowAnalysis)
-		{
-			DisplaySimpleSummary();
-			return;
-		}
-
-		if (bWasDueToOOM && !bWasImpactful)
-		{
-			UE_LOG(LogCook, Display, TEXT("GarbageCollection Results: Garbage Collection was not very impactful."));
-		}
-		else
-		{
-			UE_LOG(LogCook, Display, TEXT("GarbageCollection Results:"));
-		}
+		check(DisplayLevel == EDisplayLevel::Full || DisplayLevel == EDisplayLevel::FullWithDiagnostics);
+		UE_LOG(LogCook, Display, TEXT("%s"), *HeaderMessage);
 		UE_LOG(LogCook, Display, TEXT("\tMemoryAnalysis: General:\n")
 			TEXT("\t\tType: %s\n")
 			TEXT("\tDuration: %.3fs"),
@@ -1198,16 +1272,14 @@ void UCookOnTheFlyServer::EvaluateGarbageCollectionResults(bool bWasDueToOOM, bo
 		UE_LOG(LogCook, Display, TEXT("\tMemoryAnalysis: Allocator Stats After:%s"),
 			*AllocatorStatsToString(AllocatorStatsAfterGC));
 
-
-		// Only show diagnostics if LLM is on, because they are somewhat expensive. We could add a separate setting
-		// for this, but it's more convenient to combine it with the LLM enabled setting
-		if (!bShowExtendedDiagnostics)
+		if (DisplayLevel == EDisplayLevel::Full)
 		{
-			UE_LOG(LogCook, Display, TEXT("Extended memory diagnostics are disabled. Run with -llm or -trace=memtag to log information for UObject classes and LLM tags."));
+			UE_LOG(LogCook, Display,
+				TEXT("Extended memory diagnostics are disabled. Run with -llm or -trace=memtag or -CookMemoryAnalysis for more diagnostics."));
 		}
-		else
+		else if (!bWasSoftGC)
 		{
-			UE_LOG(LogCook, Display, TEXT("See log for memory use information for UObject classes and LLM tags."));
+			UE_LOG(LogCook, Display, TEXT("See log for package info and LLM tag diagnostics."));
 			UE::Cook::DumpPackageReferencers(TEXT("Non-startup Packages"), CookByTheBookOptions->SessionStartupObjects);
 			GLog->Logf(TEXT("Memory Analysis: LLM Tags:"));
 #if ENABLE_LOW_LEVEL_MEM_TRACKER
@@ -1221,16 +1293,7 @@ void UCookOnTheFlyServer::EvaluateGarbageCollectionResults(bool bWasDueToOOM, bo
 				GLog->Logf(TEXT("LLM Tags are not displayed because llm is disabled. Run with -llm or -trace=memtag to see llm tags."));
 			}
 		}
-	}
-	else
-	{
-		DisplaySimpleSummary();
-
-		// Mark the packages we freed so we can give a warning to diagnose why they are still referenced if they
-		// get loaded again.
-		PackageTracker->AddExpectedNeverLoadPackages(ExpectedFreedPackageNames);
-
-		if (bShowExtendedDiagnostics)
+		else // bWasSoftGC
 		{
 			// If some packages we expected to be freed were not freed, show the reference chains for why
 			// they were not freed.
