Dataflow : update vertex attribute (weightmap) paint tool to remember color ramp changes and add presets to it as well

#Synced-CL 48328177

diff --git a/Engine/Plugins/Experimental/Dataflow/Source/DataflowEditor/Private/Dataflow/DataflowColorRampCustomization.cpp b/Engine/Plugins/Experimental/Dataflow/Source/DataflowEditor/Private/Dataflow/DataflowColorRampCustomization.cpp
--- a/Engine/Plugins/Experimental/Dataflow/Source/DataflowEditor/Private/Dataflow/DataflowColorRampCustomization.cpp
+++ b/Engine/Plugins/Experimental/Dataflow/Source/DataflowEditor/Private/Dataflow/DataflowColorRampCustomization.cpp
@@ -42,6 +42,39 @@ namespace UE::Dataflow
 			}
 		}
 	}
+
+	TSharedRef<IPropertyTypeCustomization> FDataflowLinearColorRampCustomization::MakeInstance()
+	{
+		return MakeShareable(new FDataflowLinearColorRampCustomization);
+	}
+
+	void FDataflowLinearColorRampCustomization::CustomizeHeader(TSharedRef<IPropertyHandle> PropertyHandle, FDetailWidgetRow& HeaderRow, IPropertyTypeCustomizationUtils& CustomizationUtils)
+	{
+		if (const FStructProperty* const StructProperty = CastField<FStructProperty>(PropertyHandle->GetProperty()))
+		{
+			if (StructProperty->Struct)
+			{
+				void* Data = nullptr;
+
+				if (PropertyHandle->GetValueData(Data) == FPropertyAccess::Success)
+				{
+					FDataflowLinearColorRamp* ColorRamp = reinterpret_cast<FDataflowLinearColorRamp*>(Data);
+
+					TSharedPtr<SColorGradientEditor> GradientEditor;
+
+					HeaderRow
+						[
+							SAssignNew(GradientEditor, SColorGradientEditor)
+								.ViewMinInput(0.0f)
+								.ViewMaxInput(1.0f)
+								.ClampStopsToViewRange(true)
+						];
+
+					GradientEditor->SetCurveOwner(ColorRamp);
+				}
+			}
+		}
+	}
 }  
 
 
diff --git a/Engine/Plugins/Experimental/Dataflow/Source/DataflowEditor/Private/Dataflow/DataflowEditorModule.cpp b/Engine/Plugins/Experimental/Dataflow/Source/DataflowEditor/Private/Dataflow/DataflowEditorModule.cpp
--- a/Engine/Plugins/Experimental/Dataflow/Source/DataflowEditor/Private/Dataflow/DataflowEditorModule.cpp
+++ b/Engine/Plugins/Experimental/Dataflow/Source/DataflowEditor/Private/Dataflow/DataflowEditorModule.cpp
@@ -59,6 +59,7 @@ namespace UE::Dataflow::Private
 	static const FName DataflowVariableOverridesName = TEXT("DataflowVariableOverrides");
 	static const FName DataflowFreezeActionsName = TEXT("DataflowFreezeActions");
 	static const FName DataflowColorRampName = TEXT("DataflowColorRamp");
+	static const FName DataflowLinearColorRampName = TEXT("DataflowLinearColorRamp");
 
 	class FDataflowEditorWeightMapPaintToolActionCommands : public TInteractiveToolCommands<FDataflowEditorWeightMapPaintToolActionCommands>
 	{
@@ -238,6 +239,7 @@ void FDataflowEditorModule::StartupModule()
 		PropertyModule->RegisterCustomPropertyTypeLayout(DataflowFreezeActionsName, FOnGetPropertyTypeCustomizationInstance::CreateStatic(&UE::Dataflow::FFreezeActionsCustomization::MakeInstance));
 
 		PropertyModule->RegisterCustomPropertyTypeLayout(DataflowColorRampName, FOnGetPropertyTypeCustomizationInstance::CreateStatic(&UE::Dataflow::FDataflowColorRampCustomization::MakeInstance));
+		PropertyModule->RegisterCustomPropertyTypeLayout(DataflowLinearColorRampName, FOnGetPropertyTypeCustomizationInstance::CreateStatic(&UE::Dataflow::FDataflowLinearColorRampCustomization::MakeInstance));
 		PropertyModule->RegisterCustomClassLayout(UDataflowEditorVertexAttributePaintToolProperties::StaticClass()->GetFName(), FOnGetDetailCustomizationInstance::CreateStatic(&UE::Dataflow::Editor::FVertexAttributePaintToolDetailCustomization::MakeInstance));
 
 		UE::Dataflow::RegisterAnyTypeCustomizations(*PropertyModule);
@@ -312,6 +314,10 @@ void FDataflowEditorModule::ShutdownModule()
 		PropertyModule->UnregisterCustomPropertyTypeLayout(DataflowFunctionPropertyName);
 		PropertyModule->UnregisterCustomPropertyTypeLayout(DataflowVariableOverridesName);
 		PropertyModule->UnregisterCustomPropertyTypeLayout(DataflowFreezeActionsName);
+		PropertyModule->UnregisterCustomPropertyTypeLayout(DataflowColorRampName);
+		PropertyModule->UnregisterCustomPropertyTypeLayout(DataflowLinearColorRampName);
+		PropertyModule->UnregisterCustomClassLayout(UDataflowEditorVertexAttributePaintToolProperties::StaticClass()->GetFName());
+
 		UE::Dataflow::UnregisterAnyTypeCustomizations(*PropertyModule);
 	}
 
diff --git a/Engine/Plugins/Experimental/Dataflow/Source/DataflowEditor/Private/DataflowEditorTools/DataflowEditorVertexAttributePaintTool.cpp b/Engine/Plugins/Experimental/Dataflow/Source/DataflowEditor/Private/DataflowEditorTools/DataflowEditorVertexAttributePaintTool.cpp
--- a/Engine/Plugins/Experimental/Dataflow/Source/DataflowEditor/Private/DataflowEditorTools/DataflowEditorVertexAttributePaintTool.cpp
+++ b/Engine/Plugins/Experimental/Dataflow/Source/DataflowEditor/Private/DataflowEditorTools/DataflowEditorVertexAttributePaintTool.cpp
@@ -261,7 +261,7 @@ UDataflowEditorVertexAttributePaintToolProperties::UDataflowEditorVertexAttribut
 {
 	LoadConfig();
 
-	if (DisplayProperties.ColorRamp.ColorRamp.IsEmpty())
+	if (DisplayProperties.ColorRamp.IsEmpty())
 	{
 		CreateDefaultColorRamp();
 	}
@@ -275,6 +275,11 @@ UDataflowEditorVertexAttributePaintToolProperties::UDataflowEditorVertexAttribut
 }
 
 void UDataflowEditorVertexAttributePaintToolProperties::CreateDefaultColorRamp()
+{
+	ResetColorRamp();
+}
+
+void UDataflowEditorVertexAttributePaintToolProperties::ResetColorRamp()
 {
 	constexpr FLinearColor HeatMapColors[] =
 	{
@@ -290,13 +295,13 @@ void UDataflowEditorVertexAttributePaintToolProperties::CreateDefaultColorRamp()
 	ensure(NumColors > 1);
 
 	// Set ColorRamp
-	FDataflowColorCurveOwner& Ramp = DisplayProperties.ColorRamp.ColorRamp;
+	DisplayProperties.ColorRamp.Reset();
 	for (int32 Index = 0; Index < NumColors; ++Index)
 	{
 		const float Time = (float)Index / ((float)NumColors - 1.f);
-		Ramp.SetColorAtTime(Time, HeatMapColors[Index], /*bOnlyRGB*/true);
+		DisplayProperties.ColorRamp.SetColorAtTime(Time, HeatMapColors[Index], /*bOnlyRGB*/true);
 	}
-	Ramp.GetCurves()[3].CurveToEdit->AddKey(0.f, 1.f); // single alpha  value
+	DisplayProperties.ColorRamp.GetCurves()[3].CurveToEdit->AddKey(0.f, 1.f); // single alpha  value
 }
 
 const FDataflowEditorVertexAttributePaintToolBrushConfig& UDataflowEditorVertexAttributePaintToolProperties::GetBrushConfig(EDataflowEditorToolEditOperation BrushMode) const
@@ -710,8 +715,8 @@ void UDataflowEditorVertexAttributePaintTool::Setup()
 	// initialize properties
 	ToolProperties = NewObject<UDataflowEditorVertexAttributePaintToolProperties>(this);
 	{
-		ToolProperties->DisplayProperties.ColorRamp.ColorRamp.OnColorCurveChangedDelegate.AddUObject(this, &UDataflowEditorVertexAttributePaintTool::OnColorRampChanged);
-		//ToolProperties->RestoreProperties(this);
+		ToolProperties->RestoreProperties(this);
+		ToolProperties->DisplayProperties.ColorRamp.OnColorCurveChangedDelegate.AddUObject(this, &UDataflowEditorVertexAttributePaintTool::OnColorRampChanged);
 	}
 	AddToolPropertySource(ToolProperties);
 
@@ -783,7 +788,6 @@ void UDataflowEditorVertexAttributePaintTool::Setup()
 	//DynamicMeshComponent->FastNotifyVertexAttributesUpdated(EMeshRenderAttributeFlags::VertexColors);
 	GetToolManager()->PostInvalidation();
 
-	ToolProperties->LoadConfig();
 	SetBrushMode(ToolProperties->BrushProperties.BrushMode);
 
 	//SetPrimaryFalloffType(EMeshSculptFalloffType::Smooth);
@@ -2364,7 +2368,7 @@ void UDataflowEditorVertexAttributePaintTool::UpdateVertexColorOverlay(const TSe
 			case EDataflowEditorToolColorMode::Greyscale:
 				return FMath::Lerp(FLinearColor::Black, FLinearColor::White, Value);
 			case EDataflowEditorToolColorMode::Ramp:
-				return ToolProperties->DisplayProperties.ColorRamp.ColorRamp.GetLinearColorValue(Value);
+				return ToolProperties->DisplayProperties.ColorRamp.GetLinearColorValue(Value);
 			case EDataflowEditorToolColorMode::FullMaterial:
 				return FLinearColor::White;
 			}
diff --git a/Engine/Plugins/Experimental/Dataflow/Source/DataflowEditor/Private/DataflowEditorTools/DataflowEditorVertexAttributePaintToolDetailCustomization.cpp b/Engine/Plugins/Experimental/Dataflow/Source/DataflowEditor/Private/DataflowEditorTools/DataflowEditorVertexAttributePaintToolDetailCustomization.cpp
--- a/Engine/Plugins/Experimental/Dataflow/Source/DataflowEditor/Private/DataflowEditorTools/DataflowEditorVertexAttributePaintToolDetailCustomization.cpp
+++ b/Engine/Plugins/Experimental/Dataflow/Source/DataflowEditor/Private/DataflowEditorTools/DataflowEditorVertexAttributePaintToolDetailCustomization.cpp
@@ -10,9 +10,10 @@
 #include "DetailWidgetRow.h"
 #include "ModelingToolsEditorModeStyle.h"
 #include "SColorGradientEditor.h"
+#include "Widgets/Input/SButton.h"
+#include "Widgets/Input/SComboButton.h"
 #include "Widgets/Input/SSegmentedControl.h"
 #include "Widgets/Input/SSpinBox.h"
-#include "Widgets/Input/SButton.h"
 #include "Widgets/Input/SNumericEntryBox.h"
 #include "Framework/MultiBox/MultiBoxBuilder.h"
 #include "Drawing/MeshElementsVisualizer.h"
@@ -134,12 +135,175 @@ namespace UE::Dataflow::Editor
 
 	///////////////////////////////////////////////////////////////////////////////////////////////////////////////
 
+
+	TSharedRef<SWidget> FVertexAttributePaintToolDetailCustomization::BuildColorRampMenu() const
+	{
+		const bool bShouldCloseWindowAfterMenuSelection = true;
+		const bool bCloseSelfOnly = true;
+		const bool bSearchable = false;
+		FMenuBuilder MenuBuilder(true, nullptr, nullptr, bCloseSelfOnly, &FCoreStyle::Get(), bSearchable);
+
+		MenuBuilder.BeginSection(NAME_None, LOCTEXT("ColorRampOperations", "Color Ramp"));
+		{
+			MenuBuilder.AddMenuEntry
+			(
+				LOCTEXT("ColorRamp_ResetToDefault_Label", "Reset to default"),
+				LOCTEXT("ColorRamp_ResetToDefault_Tooltip", "Reset the ramp to its default"),
+				FSlateIcon(),
+				FUIAction(
+					FExecuteAction::CreateLambda([WeakToolProperties = ToolProperties]()
+						{
+							if (TStrongObjectPtr<UDataflowEditorVertexAttributePaintToolProperties> ToolProperties = WeakToolProperties.Pin())
+							{
+								ToolProperties->ResetColorRamp();
+								ToolProperties->SaveConfig();
+
+								FPropertyChangedEvent PropertyChangedEvent(FDataflowEditorVertexAttributePaintToolDisplayProperties::StaticStruct()->FindPropertyByName(GET_MEMBER_NAME_CHECKED(FDataflowEditorVertexAttributePaintToolDisplayProperties, ColorRamp)));
+								ToolProperties->PostEditChangeProperty(PropertyChangedEvent);
+							}
+						}),
+					FCanExecuteAction()
+				)
+			);
+
+			MenuBuilder.AddMenuEntry
+			(
+				LOCTEXT("ColorRamp_AddToPresets_Label", "Add to presets"),
+				LOCTEXT("ColorRamp_AddToPresets_Tooltip", "Add current ramp to presets"),
+				FSlateIcon(),
+				FUIAction(
+					FExecuteAction::CreateLambda([WeakToolProperties = ToolProperties]()
+						{
+							if (TStrongObjectPtr<UDataflowEditorVertexAttributePaintToolProperties> ToolProperties = WeakToolProperties.Pin())
+							{
+								ToolProperties->DisplayProperties.ColorRampPresets.Add(ToolProperties->DisplayProperties.ColorRamp);
+								ToolProperties->SaveConfig();
+
+								FPropertyChangedEvent PropertyChangedEvent(FDataflowEditorVertexAttributePaintToolDisplayProperties::StaticStruct()->FindPropertyByName(GET_MEMBER_NAME_CHECKED(FDataflowEditorVertexAttributePaintToolDisplayProperties, ColorRamp)));
+								ToolProperties->PostEditChangeProperty(PropertyChangedEvent);
+							}
+						}),
+					FCanExecuteAction()
+				)
+			);
+		}
+		MenuBuilder.EndSection();
+
+		auto MakeColorRampPresetWidget =
+			[WeakToolProperties = ToolProperties](int32 PresetIndex)
+			{
+				TSharedRef<SColorGradientEditor> GradientEditor = SNew(SColorGradientEditor)
+					.ViewMinInput(0.0f)
+					.ViewMaxInput(1.0f)
+					.ClampStopsToViewRange(true)
+					.DrawColorAndAlphaSeparate(false)
+					.IsEditingEnabled(false)
+					.HideStops(true)
+					.DesiredSizeOverride(FVector2D{ 256, 32 })
+					;
+
+				if (TStrongObjectPtr<UDataflowEditorVertexAttributePaintToolProperties> ToolProperties = WeakToolProperties.Pin())
+				{
+					GradientEditor->SetCurveOwner(&ToolProperties->DisplayProperties.ColorRampPresets[PresetIndex]);
+				}
+
+				return SNew(SButton)
+					.HAlign(HAlign_Center)
+					.VAlign(VAlign_Center)
+					.OnClicked_Lambda(
+					[WeakToolProperties, PresetIndex]()
+					{
+						if (TStrongObjectPtr<UDataflowEditorVertexAttributePaintToolProperties> ToolProperties = WeakToolProperties.Pin())
+						{
+							if (ToolProperties->DisplayProperties.ColorRampPresets.IsValidIndex(PresetIndex))
+							{
+								ToolProperties->DisplayProperties.ColorRamp = ToolProperties->DisplayProperties.ColorRampPresets[PresetIndex];
+								ToolProperties->SaveConfig();
+
+								FPropertyChangedEvent PropertyChangedEvent(FDataflowEditorVertexAttributePaintToolDisplayProperties::StaticStruct()->FindPropertyByName(GET_MEMBER_NAME_CHECKED(FDataflowEditorVertexAttributePaintToolDisplayProperties, ColorRamp)));
+								ToolProperties->PostEditChangeProperty(PropertyChangedEvent);
+
+								FSlateApplication::Get().DismissAllMenus();
+							}
+						}
+						return FReply::Handled();
+					})
+					[
+						GradientEditor
+					]
+					;
+			};
+
+		auto MakeColorRampPresetButton =
+			[WeakToolProperties = ToolProperties](int32 PresetIndex)
+			{
+				return SNew(SButton)
+					.HAlign(HAlign_Center)
+					.VAlign(VAlign_Center)
+					.ToolTipText(LOCTEXT("ColorRamp_PresetDelete_Tooltip", "Delete this preset from preset list."))
+					.OnClicked_Lambda(
+					[WeakToolProperties, PresetIndex]()
+					{
+						if (TStrongObjectPtr<UDataflowEditorVertexAttributePaintToolProperties> ToolProperties = WeakToolProperties.Pin())
+						{
+							if (ToolProperties->DisplayProperties.ColorRampPresets.IsValidIndex(PresetIndex))
+							{
+								ToolProperties->DisplayProperties.ColorRampPresets.RemoveAt(PresetIndex);
+								FPropertyChangedEvent PropertyChangedEvent(FDataflowEditorVertexAttributePaintToolDisplayProperties::StaticStruct()->FindPropertyByName(GET_MEMBER_NAME_CHECKED(FDataflowEditorVertexAttributePaintToolDisplayProperties, ColorRamp)));
+								ToolProperties->PostEditChangeProperty(PropertyChangedEvent);
+								ToolProperties->SaveConfig();
+
+								FSlateApplication::Get().DismissAllMenus();
+							}
+						}
+						return FReply::Handled();
+					})
+					[
+						SNew(SImage)
+						.ColorAndOpacity(FSlateColor::UseForeground())
+						.Image(FSlateIcon(FAppStyle::GetAppStyleSetName(), "GenericCommands.Delete").GetIcon())
+						//.DesiredSizeOverride(FVector2D(24.0, 24.f))
+					]
+					;
+			};
+
+		MenuBuilder.BeginSection(NAME_None, LOCTEXT("ColorRampPreset", "Presets"));
+		{
+			const int32 NumPresets = ToolProperties->DisplayProperties.ColorRampPresets.Num();
+			for (int32 PresetIndex = 0; PresetIndex < NumPresets; ++PresetIndex)
+			{
+				TSharedRef<SWidget> ColorRampPresetWidget = MakeColorRampPresetWidget(PresetIndex);
+				TSharedRef<SWidget> ColorRampPresetButton = MakeColorRampPresetButton(PresetIndex);
+
+				TSharedRef<SWidget> ColorRampPresetCombinedWidget =
+					SNew(SHorizontalBox)
+					+ SHorizontalBox::Slot()
+					.Padding(2.f)
+					.AutoWidth()
+					[
+						ColorRampPresetWidget
+					]
+					+ SHorizontalBox::Slot()
+					.Padding(2.f)
+					.AutoWidth()
+					[
+						ColorRampPresetButton
+					];
+	
+				MenuBuilder.AddWidget(ColorRampPresetCombinedWidget, FText::GetEmpty());
+			}
+		}
+		MenuBuilder.EndSection();
+
+		return MenuBuilder.MakeWidget();
+	}
+
 	void FVertexAttributePaintToolDetailCustomization::AddColorRampRow(IDetailCategoryBuilder& EditModeCategory) const
 	{
 		auto MakeColorRampWidget =
 			[](
 				EDataflowEditorToolColorMode ColorMode, 
-				FCurveOwnerInterface* CurveOwner,
+				FDataflowLinearColorRamp* ColorRamp,
 				TWeakObjectPtr<UDataflowEditorVertexAttributePaintToolProperties> ToolProperties
 				)
 			{
@@ -160,20 +324,31 @@ namespace UE::Dataflow::Editor
 								: EVisibility::Collapsed;
 						})
 					;
-				if (ToolProperties.IsValid())
+				if (ToolProperties.IsValid() && ColorRamp)
 				{
-					GradientEditor->SetCurveOwner(CurveOwner);
+					GradientEditor->SetCurveOwner(ColorRamp);
 				}
 				return GradientEditor;
 			};
 
+		auto MakeColorRampOptionMenuWidget = [this]()
+		{
+			return SNew(SComboButton)
+				.ComboButtonStyle(&FAppStyle::Get().GetWidgetStyle<FComboButtonStyle>("SimpleComboButtonWithIcon"))
+				.ForegroundColor(FSlateColor::UseStyle())
+				.ToolTipText(LOCTEXT("ColorRampOptionMenuToolTip", "Color Ramp Options Menu"))
+				.OnGetMenuContent(this, &FVertexAttributePaintToolDetailCustomization::BuildColorRampMenu);
+		};
+
 		if (ToolProperties.IsValid())
 		{
 			TSharedRef<SColorGradientEditor> GreyScaleRampWidget =
 				MakeColorRampWidget(EDataflowEditorToolColorMode::Greyscale, &ToolProperties->DisplayProperties.GreyScaleColorRamp, ToolProperties);
 
 			TSharedRef<SColorGradientEditor> CustomRampWidget =
-				MakeColorRampWidget(EDataflowEditorToolColorMode::Ramp, &ToolProperties->DisplayProperties.ColorRamp.ColorRamp, ToolProperties);
+				MakeColorRampWidget(EDataflowEditorToolColorMode::Ramp, &ToolProperties->DisplayProperties.ColorRamp, ToolProperties);
+
+			TSharedRef<SComboButton> CustomRampWidgetOptionMenu = MakeColorRampOptionMenuWidget();
 
 			TSharedRef<SColorGradientEditor> MaterialRampWidget =
 				MakeColorRampWidget(EDataflowEditorToolColorMode::FullMaterial, &ToolProperties->DisplayProperties.WhiteColorRamp, ToolProperties);
@@ -190,7 +365,19 @@ namespace UE::Dataflow::Editor
 				+ SVerticalBox::Slot()
 				.AutoHeight()
 				[
-					CustomRampWidget
+					SNew(SHorizontalBox)
+					+ SHorizontalBox::Slot()
+					.Padding(2.f)
+					.FillWidth(1.0f)
+					[
+						CustomRampWidget
+					]
+					+ SHorizontalBox::Slot()
+					.Padding(2.f)
+					.AutoWidth()
+					[
+						CustomRampWidgetOptionMenu
+					]
 				]
 				+ SVerticalBox::Slot()
 				.AutoHeight()
diff --git a/Engine/Plugins/Experimental/Dataflow/Source/DataflowEditor/Private/DataflowEditorTools/DataflowEditorVertexAttributePaintToolDetailCustomization.h b/Engine/Plugins/Experimental/Dataflow/Source/DataflowEditor/Private/DataflowEditorTools/DataflowEditorVertexAttributePaintToolDetailCustomization.h
--- a/Engine/Plugins/Experimental/Dataflow/Source/DataflowEditor/Private/DataflowEditorTools/DataflowEditorVertexAttributePaintToolDetailCustomization.h
+++ b/Engine/Plugins/Experimental/Dataflow/Source/DataflowEditor/Private/DataflowEditorTools/DataflowEditorVertexAttributePaintToolDetailCustomization.h
@@ -33,6 +33,7 @@ namespace UE::Dataflow::Editor
 		void AddEditingModeRow(IDetailCategoryBuilder& EditModeCategory) const;
 		void AddColorModeRow(IDetailCategoryBuilder& EditModeCategory) const;
 		void AddColorRampRow(IDetailCategoryBuilder& EditModeCategory) const;
+		TSharedRef<SWidget> BuildColorRampMenu() const;
 		
 		void AddBrushUI(IDetailLayoutBuilder& DetailBuilder) const;
 		void AddBrushModeRow(IDetailCategoryBuilder& BrushCategory) const;
diff --git a/Engine/Plugins/Experimental/Dataflow/Source/DataflowEditor/Public/Dataflow/DataflowColorRampCustomization.h b/Engine/Plugins/Experimental/Dataflow/Source/DataflowEditor/Public/Dataflow/DataflowColorRampCustomization.h
--- a/Engine/Plugins/Experimental/Dataflow/Source/DataflowEditor/Public/Dataflow/DataflowColorRampCustomization.h
+++ b/Engine/Plugins/Experimental/Dataflow/Source/DataflowEditor/Public/Dataflow/DataflowColorRampCustomization.h
@@ -13,8 +13,22 @@ namespace UE::Dataflow
 		/** Makes a new instance of this detail layout class for a specific detail view requesting it */
 		DATAFLOWEDITOR_API static TSharedRef<IPropertyTypeCustomization> MakeInstance();
 
-		/** IDetailCustomization interface */
+		//~ Begin - IPropertyTypeCustomization interface
 		DATAFLOWEDITOR_API virtual void CustomizeHeader(TSharedRef<IPropertyHandle> PropertyHandle, FDetailWidgetRow& HeaderRow, IPropertyTypeCustomizationUtils& CustomizationUtils) override;
 		virtual void CustomizeChildren(TSharedRef<IPropertyHandle> /*StructPropertyHandle*/, IDetailChildrenBuilder& /*ChildBuilder*/, IPropertyTypeCustomizationUtils& /*CustomizationUtils*/) override {}
+		//~ End - IPropertyTypeCustomization interface
+	};
+
+	class FDataflowLinearColorRampCustomization : public IPropertyTypeCustomization
+	{
+	public:
+
+		/** Makes a new instance of this detail layout class for a specific detail view requesting it */
+		DATAFLOWEDITOR_API static TSharedRef<IPropertyTypeCustomization> MakeInstance();
+
+		//~ Begin - IPropertyTypeCustomization interface
+		DATAFLOWEDITOR_API virtual void CustomizeHeader(TSharedRef<IPropertyHandle> PropertyHandle, FDetailWidgetRow& HeaderRow, IPropertyTypeCustomizationUtils& CustomizationUtils) override;
+		virtual void CustomizeChildren(TSharedRef<IPropertyHandle> /*StructPropertyHandle*/, IDetailChildrenBuilder& /*ChildBuilder*/, IPropertyTypeCustomizationUtils& /*CustomizationUtils*/) override {}
+		//~ End - IPropertyTypeCustomization interface
 	};
 }
diff --git a/Engine/Plugins/Experimental/Dataflow/Source/DataflowEditor/Public/DataflowEditorTools/DataflowEditorVertexAttributePaintTool.h b/Engine/Plugins/Experimental/Dataflow/Source/DataflowEditor/Public/DataflowEditorTools/DataflowEditorVertexAttributePaintTool.h
--- a/Engine/Plugins/Experimental/Dataflow/Source/DataflowEditor/Public/DataflowEditorTools/DataflowEditorVertexAttributePaintTool.h
+++ b/Engine/Plugins/Experimental/Dataflow/Source/DataflowEditor/Public/DataflowEditorTools/DataflowEditorVertexAttributePaintTool.h
@@ -179,11 +179,14 @@ struct FDataflowEditorVertexAttributePaintToolDisplayProperties
 
 	// Color ramp to use when the color mode is set to "Ramp"
 	UPROPERTY()
-	FDataflowColorRamp ColorRamp;
+	FDataflowLinearColorRamp ColorRamp;
+
+	UPROPERTY()
+	TArray<FDataflowLinearColorRamp> ColorRampPresets;
 
 	// Used by UI customization 
-	FDataflowColorCurveOwner GreyScaleColorRamp;
-	FDataflowColorCurveOwner WhiteColorRamp;
+	FDataflowLinearColorRamp GreyScaleColorRamp;
+	FDataflowLinearColorRamp WhiteColorRamp;
 };
 
 USTRUCT()
@@ -243,6 +246,8 @@ public:
 	UPROPERTY(config)
 	FDataflowEditorVertexAttributePaintToolBrushConfig BrushConfigRelax;
 
+	void ResetColorRamp();
+
 private:
 	void CreateDefaultColorRamp();
 
diff --git a/Engine/Source/Editor/UnrealEd/Private/SColorGradientEditor.cpp b/Engine/Source/Editor/UnrealEd/Private/SColorGradientEditor.cpp
--- a/Engine/Source/Editor/UnrealEd/Private/SColorGradientEditor.cpp
+++ b/Engine/Source/Editor/UnrealEd/Private/SColorGradientEditor.cpp
@@ -72,16 +72,17 @@ namespace FColorGradientEditorGeometry
 			FSlateLayoutTransform(FVector2D(0.0f, WidgetGeometry.GetLocalSize().Y - HandleHeight)) );
 	}
 
-	FGeometry GetGradientAreaPaintGeometry( const FGeometry& WidgetGeometry )
+	FGeometry GetGradientAreaPaintGeometry( const FGeometry& WidgetGeometry, bool bHideStops = false)
 	{
 		return WidgetGeometry.MakeChild(
-			FVector2D( WidgetGeometry.GetLocalSize().X - HandleWidth, WidgetGeometry.GetLocalSize().Y - (2 * HandleHeight) ),
-			FSlateLayoutTransform(FVector2D(HandleWidth / 2, HandleHeight)) );
+			FVector2D( WidgetGeometry.GetLocalSize().X - (bHideStops ? 0: HandleWidth), WidgetGeometry.GetLocalSize().Y - (bHideStops ? 0 : (2 * HandleHeight)) ),
+			bHideStops? FSlateLayoutTransform(): FSlateLayoutTransform(FVector2D(HandleWidth / 2, HandleHeight)) 
+		);
 	}
 
-	FTrackScaleInfo GetTimeScaleForWidgetGeometry(const FGeometry& WidgetGeometry, float ViewMinInput, float ViewMaxInput)
+	FTrackScaleInfo GetTimeScaleForWidgetGeometry(const FGeometry& WidgetGeometry, float ViewMinInput, float ViewMaxInput, bool bHideStops = false)
 	{
-		return FTrackScaleInfo(ViewMinInput, ViewMaxInput, 0.0f, 1.0f, GetGradientAreaPaintGeometry(WidgetGeometry).GetLocalSize());
+		return FTrackScaleInfo(ViewMinInput, ViewMaxInput, 0.0f, 1.0f, GetGradientAreaPaintGeometry(WidgetGeometry, bHideStops).GetLocalSize());
 	}
 
 	FTrackScaleInfo GetTimeScaleForAreaPaintGeomtry(const FGeometry& AreaGeometry, float ViewMinInput, float ViewMaxInput)
@@ -89,9 +90,9 @@ namespace FColorGradientEditorGeometry
 		return FTrackScaleInfo(ViewMinInput, ViewMaxInput, 0.0f, 1.0f, AreaGeometry.GetLocalSize());
 	}
 
-	float GetTimeForAbsolutePosition(const FGeometry& WidgetGeometry, float ViewMinInput, float ViewMaxInput, const FVector2D& AbsolutePosition)
+	float GetTimeForAbsolutePosition(const FGeometry& WidgetGeometry, float ViewMinInput, float ViewMaxInput, bool bHideStops, const FVector2D& AbsolutePosition)
 	{
-		FGeometry GradientAreaPaintGeometry = GetGradientAreaPaintGeometry(WidgetGeometry);
+		FGeometry GradientAreaPaintGeometry = GetGradientAreaPaintGeometry(WidgetGeometry, bHideStops);
 		FTrackScaleInfo TimeScaleInfo = FTrackScaleInfo(ViewMinInput, ViewMaxInput, 0.0f, 1.0f, GradientAreaPaintGeometry.GetLocalSize());
 		return TimeScaleInfo.LocalXToInput(GradientAreaPaintGeometry.AbsoluteToLocal(AbsolutePosition).X);
 	}
@@ -199,6 +200,15 @@ void FGradientStopMark::SetTime( float NewTime, FCurveOwnerInterface& CurveOwner
 	Time = NewTime;
 }
 
+SColorGradientEditor::SColorGradientEditor()
+	: DesiredSizeOverride(*this)
+{}
+
+bool SColorGradientEditor::ShouldHideStops() const
+{
+	return (bRequestHideStops.Get() && !IsEditingEnabled.Get());
+}
+
 void SColorGradientEditor::Construct( const FArguments& InArgs )
 {
 	IsEditingEnabled = InArgs._IsEditingEnabled;
@@ -211,6 +221,8 @@ void SColorGradientEditor::Construct( const FArguments& InArgs )
 	bDraggingStop = false;
 	DistanceDragged = 0.0f;
 	ContextMenuPosition = FVector2D::ZeroVector;
+	DesiredSizeOverride.Assign(*this, InArgs._DesiredSizeOverride);
+	bRequestHideStops = InArgs._HideStops;
 }
 
 int32 SColorGradientEditor::OnPaint( const FPaintArgs& Args, const FGeometry& AllottedGeometry, const FSlateRect& MyCullingRect, FSlateWindowElementList& OutDrawElements, int32 LayerId, const FWidgetStyle& InWidgetStyle, bool bParentEnabled ) const
@@ -219,10 +231,12 @@ int32 SColorGradientEditor::OnPaint( const FPaintArgs& Args, const FGeometry& Al
 
 	if( CurveOwner )
 	{
+		const bool bHideStops = ShouldHideStops();
+
 		// Split the geometry into areas for stops and the gradient
 		FGeometry ColorMarkAreaGeometry = FColorGradientEditorGeometry::GetColorMarkAreaPaintGeometry( AllottedGeometry );
 		FGeometry AlphaMarkAreaGeometry = FColorGradientEditorGeometry::GetAlphaMarkAreaPaintGeometry( AllottedGeometry );
-		FGeometry GradientAreaGeometry = FColorGradientEditorGeometry::GetGradientAreaPaintGeometry( AllottedGeometry );
+		FGeometry GradientAreaGeometry = FColorGradientEditorGeometry::GetGradientAreaPaintGeometry( AllottedGeometry , bHideStops);
 
 		bool bEnabled = ShouldBeEnabled( bParentEnabled );
 		ESlateDrawEffect DrawEffects = bEnabled ? ESlateDrawEffect::None : ESlateDrawEffect::DisabledEffect;
@@ -274,14 +288,14 @@ int32 SColorGradientEditor::OnPaint( const FPaintArgs& Args, const FGeometry& Al
 		}
 			
 		// Sample the curve every 2 units.  THe curve could be non-linear so sampling at each stop would display an incorrect gradient
-		for( int32 CurrentStep = Start; CurrentStep < Finish; CurrentStep+=2 )
+		for (int32 CurrentStep = Start; CurrentStep < Finish; CurrentStep += 2)
 		{
 			// Figure out the time from the current screen unit
 			const float Time = ScaleInfo.LocalXToInput(CurrentStep);
 
 			// Sample the curve
-			FLinearColor Color = CurveOwner->GetLinearColorValue( Time );
-			if( !bHasAnyAlphaKeys )
+			FLinearColor Color = CurveOwner->GetLinearColorValue(Time);
+			if (!bHasAnyAlphaKeys)
 			{
 				// Only show alpha if there is at least one key.  For some curves, alpha may not be important
 				Color.A = 1.0f;
@@ -303,7 +317,7 @@ int32 SColorGradientEditor::OnPaint( const FPaintArgs& Args, const FGeometry& Al
 			}
 		}
 
-		if ( ColorStops.Num() > 0 )
+		if (ColorStops.Num() > 0)
 		{
 			// Draw the color gradient
 			if (bDrawColorAndAlphaSeparate)
@@ -358,83 +372,85 @@ int32 SColorGradientEditor::OnPaint( const FPaintArgs& Args, const FGeometry& Al
 			}
 		}
 
-		// Get actual editable stop marks
-		TArray<FGradientStopMark> ColorMarks;
-		TArray<FGradientStopMark> AlphaMarks;
-		GetGradientStopMarks( ColorMarks, AlphaMarks );
-
-		// Draw each color stop
-		for( int32 ColorIndex = 0; ColorIndex < ColorMarks.Num(); ++ColorIndex )
+		if (!bHideStops)
 		{
-			const FGradientStopMark& Mark = ColorMarks[ColorIndex];
-
-			float XVal = ScaleInfo.InputToLocalX( Mark.Time );
+			// Get actual editable stop marks
+			TArray<FGradientStopMark> ColorMarks;
+			TArray<FGradientStopMark> AlphaMarks;
+			GetGradientStopMarks(ColorMarks, AlphaMarks);
 
-			// Dont draw stops which are not visible
-			if( XVal >= 0 && XVal <= ColorMarkAreaGeometry.GetLocalSize().X )
+			// Draw each color stop
+			for (int32 ColorIndex = 0; ColorIndex < ColorMarks.Num(); ++ColorIndex)
 			{
-				FLinearColor Color = CurveOwner->GetLinearColorValue( Mark.Time );
-				Color.A = 1.0f;
-				DrawGradientStopMark( Mark, ColorMarkAreaGeometry, XVal, Color, OutDrawElements, LayerId, MyCullingRect, DrawEffects, true, InWidgetStyle );
-			}
-		}
+				const FGradientStopMark& Mark = ColorMarks[ColorIndex];
 
-		// Draw each alpha stop
-		for( int32 ColorIndex = 0; ColorIndex < AlphaMarks.Num(); ++ColorIndex )
-		{
-			const FGradientStopMark& Mark = AlphaMarks[ColorIndex];
+				float XVal = ScaleInfo.InputToLocalX(Mark.Time);
 
-			float XVal = ScaleInfo.InputToLocalX( Mark.Time );
-		
-			// Dont draw stops which are not visible
-			if( XVal >= 0 && XVal <= AlphaMarkAreaGeometry.GetLocalSize().X )
-			{
-				float Alpha = CurveOwner->GetLinearColorValue( Mark.Time ).A;
-				DrawGradientStopMark( Mark, AlphaMarkAreaGeometry, XVal, FLinearColor( Alpha, Alpha, Alpha, 1.0f ), OutDrawElements, LayerId, MyCullingRect, DrawEffects, false, InWidgetStyle );
+				// Dont draw stops which are not visible
+				if (XVal >= 0 && XVal <= ColorMarkAreaGeometry.GetLocalSize().X)
+				{
+					FLinearColor Color = CurveOwner->GetLinearColorValue(Mark.Time);
+					Color.A = 1.0f;
+					DrawGradientStopMark(Mark, ColorMarkAreaGeometry, XVal, Color, OutDrawElements, LayerId, MyCullingRect, DrawEffects, true, InWidgetStyle);
+				}
 			}
-		}
 
-		// Draw some hint messages about how to add stops if no stops exist
-		if( ColorMarks.Num() == 0 && AlphaMarks.Num() == 0 && IsEditingEnabled.Get() == true )
-		{
-			static FString GradientColorMessage( LOCTEXT("ClickToAddColorStop", "Click in this area add color stops").ToString() );
-			static FString GradientAlphaMessage( LOCTEXT("ClickToAddAlphaStop", "Click in this area add opacity stops").ToString() );
-				
-			// Draw the text centered in the color region
+			// Draw each alpha stop
+			for (int32 ColorIndex = 0; ColorIndex < AlphaMarks.Num(); ++ColorIndex)
 			{
-				FVector2D StringSize = FontMeasureService->Measure(GradientColorMessage, FCoreStyle::GetDefaultFontStyle("Regular", 8));
-				FPaintGeometry PaintGeom = ColorMarkAreaGeometry.ToPaintGeometry(FSlateLayoutTransform(FVector2D((ColorMarkAreaGeometry.GetLocalSize().X - StringSize.X) * 0.5f, 1.0f)));
+				const FGradientStopMark& Mark = AlphaMarks[ColorIndex];
 
-				FSlateDrawElement::MakeText
-				( 
-					OutDrawElements, 
-					LayerId,
-					PaintGeom,
-					GradientColorMessage,
-					FCoreStyle::GetDefaultFontStyle("Regular", 8),
-					DrawEffects,
-					FLinearColor( .5f, .5f, .5f, .85f )
-				);	
+				float XVal = ScaleInfo.InputToLocalX(Mark.Time);
+
+				// Dont draw stops which are not visible
+				if (XVal >= 0 && XVal <= AlphaMarkAreaGeometry.GetLocalSize().X)
+				{
+					float Alpha = CurveOwner->GetLinearColorValue(Mark.Time).A;
+					DrawGradientStopMark(Mark, AlphaMarkAreaGeometry, XVal, FLinearColor(Alpha, Alpha, Alpha, 1.0f), OutDrawElements, LayerId, MyCullingRect, DrawEffects, false, InWidgetStyle);
+				}
 			}
 
-			// Draw the text centered in the alpha region
+			// Draw some hint messages about how to add stops if no stops exist
+			if (ColorMarks.Num() == 0 && AlphaMarks.Num() == 0 && IsEditingEnabled.Get() == true)
 			{
-				FVector2D StringSize = FontMeasureService->Measure(GradientAlphaMessage, FCoreStyle::GetDefaultFontStyle("Regular", 8));
-				FPaintGeometry PaintGeom = AlphaMarkAreaGeometry.ToPaintGeometry(FSlateLayoutTransform(FVector2D((AlphaMarkAreaGeometry.GetLocalSize().X - StringSize.X) * 0.5f, 1.0f)));
+				static FString GradientColorMessage(LOCTEXT("ClickToAddColorStop", "Click in this area add color stops").ToString());
+				static FString GradientAlphaMessage(LOCTEXT("ClickToAddAlphaStop", "Click in this area add opacity stops").ToString());
 
-				FSlateDrawElement::MakeText
-				( 
-					OutDrawElements, 
-					LayerId,
-					PaintGeom,
-					GradientAlphaMessage,
-					FCoreStyle::GetDefaultFontStyle("Regular", 8),
-					DrawEffects,
-					FLinearColor( .5f, .5f, .5f, .85f )
-				);	
+				// Draw the text centered in the color region
+				{
+					FVector2D StringSize = FontMeasureService->Measure(GradientColorMessage, FCoreStyle::GetDefaultFontStyle("Regular", 8));
+					FPaintGeometry PaintGeom = ColorMarkAreaGeometry.ToPaintGeometry(FSlateLayoutTransform(FVector2D((ColorMarkAreaGeometry.GetLocalSize().X - StringSize.X) * 0.5f, 1.0f)));
+
+					FSlateDrawElement::MakeText
+					(
+						OutDrawElements,
+						LayerId,
+						PaintGeom,
+						GradientColorMessage,
+						FCoreStyle::GetDefaultFontStyle("Regular", 8),
+						DrawEffects,
+						FLinearColor(.5f, .5f, .5f, .85f)
+					);
+				}
+
+				// Draw the text centered in the alpha region
+				{
+					FVector2D StringSize = FontMeasureService->Measure(GradientAlphaMessage, FCoreStyle::GetDefaultFontStyle("Regular", 8));
+					FPaintGeometry PaintGeom = AlphaMarkAreaGeometry.ToPaintGeometry(FSlateLayoutTransform(FVector2D((AlphaMarkAreaGeometry.GetLocalSize().X - StringSize.X) * 0.5f, 1.0f)));
+
+					FSlateDrawElement::MakeText
+					(
+						OutDrawElements,
+						LayerId,
+						PaintGeom,
+						GradientAlphaMessage,
+						FCoreStyle::GetDefaultFontStyle("Regular", 8),
+						DrawEffects,
+						FLinearColor(.5f, .5f, .5f, .85f)
+					);
+				}
 			}
 		}
-		
 	}
 
 	return LayerId;
@@ -529,7 +545,8 @@ FReply SColorGradientEditor::OnMouseMove( const FGeometry& MyGeometry, const FPo
 			else
 			{
 				// Already dragging a stop, move it
-				float MouseTime = FColorGradientEditorGeometry::GetTimeForAbsolutePosition(MyGeometry, ViewMinInput.Get(), ViewMaxInput.Get(), MouseEvent.GetScreenSpacePosition());
+				const bool bHideStops = ShouldHideStops();
+				float MouseTime = FColorGradientEditorGeometry::GetTimeForAbsolutePosition(MyGeometry, ViewMinInput.Get(), ViewMaxInput.Get(), bHideStops, MouseEvent.GetScreenSpacePosition());
 				if (bClampStopsToViewRange)
 				{
 					MouseTime = FMath::Clamp(MouseTime, ViewMinInput.Get(), ViewMaxInput.Get());
@@ -608,7 +625,8 @@ FReply SColorGradientEditor::OnMouseButtonUp( const FGeometry& MyGeometry, const
 				}
 				else
 				{
-					const FGeometry GradientAreaPaintGeometry = FColorGradientEditorGeometry::GetGradientAreaPaintGeometry(MyGeometry);
+					const bool bHideStops = ShouldHideStops();
+					const FGeometry GradientAreaPaintGeometry = FColorGradientEditorGeometry::GetGradientAreaPaintGeometry(MyGeometry, bHideStops);
 					if (GradientAreaPaintGeometry.IsUnderLocation(MouseEvent.GetScreenSpacePosition()))
 					{
 						OpenGradientOptionsMenu(MouseEvent);
@@ -655,6 +673,11 @@ void SColorGradientEditor::OnMouseLeave(const FPointerEvent& MouseEvent)
 
 FVector2D SColorGradientEditor::ComputeDesiredSize( float ) const
 {
+	const TOptional<FVector2D>& CurrentSizeOverride = DesiredSizeOverride.Get();
+	if (CurrentSizeOverride.IsSet())
+	{
+		return CurrentSizeOverride.GetValue();
+	}
 	return FVector2D( 1000, 75 );
 }
 	
@@ -1092,7 +1115,8 @@ FGradientStopMark SColorGradientEditor::AddStop( const FVector2D& Position, cons
 
 	CurveOwner->ModifyOwner();
 
-	float NewStopTime = FColorGradientEditorGeometry::GetTimeForAbsolutePosition(MyGeometry, ViewMinInput.Get(), ViewMaxInput.Get(), Position);
+	const bool bHideStops = ShouldHideStops();
+	float NewStopTime = FColorGradientEditorGeometry::GetTimeForAbsolutePosition(MyGeometry, ViewMinInput.Get(), ViewMaxInput.Get(), bHideStops, Position);
 
 	TArray<FRichCurveEditInfo> Curves = CurveOwner->GetCurves();
 
diff --git a/Engine/Source/Editor/UnrealEd/Public/SColorGradientEditor.h b/Engine/Source/Editor/UnrealEd/Public/SColorGradientEditor.h
--- a/Engine/Source/Editor/UnrealEd/Public/SColorGradientEditor.h
+++ b/Engine/Source/Editor/UnrealEd/Public/SColorGradientEditor.h
@@ -56,10 +56,14 @@ public:
 		SLATE_ATTRIBUTE( float, ViewMinInput )
 		SLATE_ATTRIBUTE( float, ViewMaxInput )
 		SLATE_ATTRIBUTE( bool, IsEditingEnabled )
+		SLATE_ATTRIBUTE(bool, HideStops)
 		SLATE_ARGUMENT( bool, ClampStopsToViewRange )
 		SLATE_ARGUMENT(bool, DrawColorAndAlphaSeparate )
+		SLATE_ATTRIBUTE(TOptional<FVector2D>, DesiredSizeOverride)
 	SLATE_END_ARGS()
 
+	UNREALED_API SColorGradientEditor();
+	UNREALED_API virtual ~SColorGradientEditor() = default;
 
 	UNREALED_API void Construct( const FArguments& InArgs );
 
@@ -199,6 +203,8 @@ private:
 	 */
 	UNREALED_API void MoveStop( FGradientStopMark& Mark, float NewTime );
 
+	bool ShouldHideStops() const;
+
 private:
 	/** The currently selected stop */
 	FGradientStopMark SelectedStop;
@@ -226,4 +232,8 @@ private:
 	bool bDraggingStop;
 	/** Do we draw a gradient for color and alpha separate or combined? */
 	bool bDrawColorAndAlphaSeparate;
+	/** True to request hide stops (will only be fullfilled when IsEditingEnabled = false) */
+	TAttribute<bool> bRequestHideStops;
+	/** desired size override */
+	TSlateAttribute<TOptional<FVector2D>, EInvalidateWidgetReason::Layout> DesiredSizeOverride;
 };
diff --git a/Engine/Source/Runtime/Experimental/Dataflow/Engine/Private/Dataflow/DataflowColorRamp.cpp b/Engine/Source/Runtime/Experimental/Dataflow/Engine/Private/Dataflow/DataflowColorRamp.cpp
--- a/Engine/Source/Runtime/Experimental/Dataflow/Engine/Private/Dataflow/DataflowColorRamp.cpp
+++ b/Engine/Source/Runtime/Experimental/Dataflow/Engine/Private/Dataflow/DataflowColorRamp.cpp
@@ -4,80 +4,65 @@
 
 #include UE_INLINE_GENERATED_CPP_BY_NAME(DataflowColorRamp)
 
-FDataflowColorCurveOwner::FDataflowColorCurveOwner()
-{
-	ConstCurves.SetNum(4);
-	Curves.SetNum(4);
-
-	RichCurves.SetNum(4);
-
-	for (int32 Idx = 0; Idx < 4; ++Idx)
-	{
-		Curves[Idx].CurveToEdit = &RichCurves[Idx];
-		ConstCurves[Idx].CurveToEdit = &RichCurves[Idx];
-	}
-}
 
-bool FDataflowColorCurveOwner::IsEmpty() const
+namespace UE::Dataflow::ColorMap::Private
 {
-	return (RichCurves.Num() == 0 || RichCurves[0].IsEmpty());
+	static const FName RedCurveName("R");
+	static const FName GreenCurveName("G");
+	static const FName BlueCurveName("B");
+	static const FName AlphaCurveName("A");
 }
 
-FDataflowColorCurveOwner::FDataflowColorCurveOwner(const FDataflowColorCurveOwner& Other)
-{
-	this->operator=(Other);
-}
-
-FDataflowColorCurveOwner& FDataflowColorCurveOwner::operator=(const FDataflowColorCurveOwner& Other)
+bool FDataflowColorCurveOwner::IsEmpty() const
 {
-	ConstCurves.SetNum(4);
-	Curves.SetNum(4);
-
-	RichCurves.SetNum(4);
-
-	for (int32 Idx = 0; Idx < 4; ++Idx)
-	{
-		if (Other.RichCurves.IsValidIndex(Idx))
-		{
-			RichCurves[Idx] = Other.RichCurves[Idx];
-		}
-		Curves[Idx].CurveToEdit = &RichCurves[Idx];
-		ConstCurves[Idx].CurveToEdit = &RichCurves[Idx];
-	}
-	return *this;
+	return (RichCurves[0].IsEmpty() 
+		&& RichCurves[1].IsEmpty() 
+		&& RichCurves[2].IsEmpty() 
+		&& RichCurves[3].IsEmpty());
 }
 
 void FDataflowColorCurveOwner::SetColorAtTime(float Time, const FLinearColor& Color, bool bOnlyRBG)
 {
-	if (ensure(Curves.Num() == 4))
+	RichCurves[0].AddKey(Time, Color.R);
+	RichCurves[1].AddKey(Time, Color.G);
+	RichCurves[2].AddKey(Time, Color.B);
+	if (!bOnlyRBG)
 	{
-		Curves[0].CurveToEdit->AddKey(Time, Color.R);
-		Curves[1].CurveToEdit->AddKey(Time, Color.G);
-		Curves[2].CurveToEdit->AddKey(Time, Color.B);
-		if (!bOnlyRBG)
-		{
-			Curves[3].CurveToEdit->AddKey(Time, Color.A);
-		}
+		RichCurves[3].AddKey(Time, Color.A);
 	}
 }
 
 TArray<FRichCurveEditInfoConst> FDataflowColorCurveOwner::GetCurves() const
 {
-	return ConstCurves;
+	using namespace UE::Dataflow::ColorMap;
+
+	TArray<FRichCurveEditInfoConst> Curves;
+	Curves.Add(FRichCurveEditInfoConst(&RichCurves[0], Private::RedCurveName));
+	Curves.Add(FRichCurveEditInfoConst(&RichCurves[1], Private::GreenCurveName));
+	Curves.Add(FRichCurveEditInfoConst(&RichCurves[2], Private::BlueCurveName));
+	Curves.Add(FRichCurveEditInfoConst(&RichCurves[3], Private::AlphaCurveName));
+	return Curves;
 }
 
-void FDataflowColorCurveOwner::GetCurves(TAdderReserverRef<FRichCurveEditInfoConst> OutCurves) const
+void FDataflowColorCurveOwner::GetCurves(TAdderReserverRef<FRichCurveEditInfoConst> Curves) const
 {
-	check(OutCurves.Num() == 0);
-	OutCurves.Reserve(ConstCurves.Num());
-	for (const FRichCurveEditInfoConst& ConstCurve : ConstCurves)
-	{
-		OutCurves.Add(ConstCurve);
-	}
+	using namespace UE::Dataflow::ColorMap;
+
+	Curves.Add(FRichCurveEditInfoConst(&RichCurves[0], Private::RedCurveName));
+	Curves.Add(FRichCurveEditInfoConst(&RichCurves[1], Private::GreenCurveName));
+	Curves.Add(FRichCurveEditInfoConst(&RichCurves[2], Private::BlueCurveName));
+	Curves.Add(FRichCurveEditInfoConst(&RichCurves[3], Private::AlphaCurveName));
 }
 
 TArray<FRichCurveEditInfo> FDataflowColorCurveOwner::GetCurves()
 {
+	using namespace UE::Dataflow::ColorMap;
+
+	TArray<FRichCurveEditInfo> Curves;
+	Curves.Add(FRichCurveEditInfo(&RichCurves[0], Private::RedCurveName));
+	Curves.Add(FRichCurveEditInfo(&RichCurves[1], Private::GreenCurveName));
+	Curves.Add(FRichCurveEditInfo(&RichCurves[2], Private::BlueCurveName));
+	Curves.Add(FRichCurveEditInfo(&RichCurves[3], Private::AlphaCurveName));
 	return Curves;
 }
 
@@ -86,7 +71,8 @@ void FDataflowColorCurveOwner::ModifyOwner()
 
 TArray<const UObject*> FDataflowColorCurveOwner::GetOwners() const
 {
-	return Owners;
+	static TArray<const UObject*> EmptyObjectArray;
+	return EmptyObjectArray;
 }
 
 void FDataflowColorCurveOwner::MakeTransactional()
@@ -113,20 +99,23 @@ bool FDataflowColorCurveOwner::IsLinearColorCurve() const
 FLinearColor FDataflowColorCurveOwner::GetLinearColorValue(float InTime) const
 {
 	return FLinearColor(
-		Curves[0].CurveToEdit->Eval(InTime),
-		Curves[1].CurveToEdit->Eval(InTime),
-		Curves[2].CurveToEdit->Eval(InTime),
-		Curves[3].CurveToEdit->Eval(InTime));
+		RichCurves[0].Eval(InTime),
+		RichCurves[1].Eval(InTime),
+		RichCurves[2].Eval(InTime),
+		RichCurves[3].Eval(InTime));
 }
 
 bool FDataflowColorCurveOwner::HasAnyAlphaKeys() const
 {
-	return Curves[3].CurveToEdit->GetNumKeys() != 0;
+	return RichCurves[3].GetNumKeys() != 0;
 }
 
 bool FDataflowColorCurveOwner::IsValidCurve(FRichCurveEditInfo CurveInfo)
 {
-	return Curves.Contains(CurveInfo);
+	return (CurveInfo.CurveToEdit == &RichCurves[0]
+		|| CurveInfo.CurveToEdit == &RichCurves[1]
+		|| CurveInfo.CurveToEdit == &RichCurves[2]
+		|| CurveInfo.CurveToEdit == &RichCurves[3]);
 }
 
 FLinearColor FDataflowColorCurveOwner::GetCurveColor(FRichCurveEditInfo CurveInfo) const
@@ -134,3 +123,124 @@ FLinearColor FDataflowColorCurveOwner::GetCurveColor(FRichCurveEditInfo CurveInf
 	return FLinearColor::White;
 }
 
+////////////////////////////////////////////////////////////////////////////////
+
+bool FDataflowLinearColorRamp::IsEmpty() const
+{
+	return (RichCurves[0].IsEmpty()
+		&& RichCurves[1].IsEmpty()
+		&& RichCurves[2].IsEmpty()
+		&& RichCurves[3].IsEmpty());
+}
+
+void FDataflowLinearColorRamp::Reset()
+{
+	RichCurves[0].Reset();
+	RichCurves[1].Reset();
+	RichCurves[2].Reset();
+	RichCurves[3].Reset();
+}
+
+void FDataflowLinearColorRamp::SetColorAtTime(float Time, const FLinearColor& Color, bool bOnlyRBG)
+{
+	RichCurves[0].AddKey(Time, Color.R);
+	RichCurves[1].AddKey(Time, Color.G);
+	RichCurves[2].AddKey(Time, Color.B);
+	if (!bOnlyRBG)
+	{
+		RichCurves[3].AddKey(Time, Color.A);
+	}
+}
+
+TArray<FRichCurveEditInfoConst> FDataflowLinearColorRamp::GetCurves() const
+{
+	using namespace UE::Dataflow::ColorMap;
+
+	TArray<FRichCurveEditInfoConst> Curves;
+	Curves.Add(FRichCurveEditInfoConst(&RichCurves[0], Private::RedCurveName));
+	Curves.Add(FRichCurveEditInfoConst(&RichCurves[1], Private::GreenCurveName));
+	Curves.Add(FRichCurveEditInfoConst(&RichCurves[2], Private::BlueCurveName));
+	Curves.Add(FRichCurveEditInfoConst(&RichCurves[3], Private::AlphaCurveName));
+	return Curves;
+}
+
+void FDataflowLinearColorRamp::GetCurves(TAdderReserverRef<FRichCurveEditInfoConst> Curves) const
+{
+	using namespace UE::Dataflow::ColorMap;
+
+	Curves.Add(FRichCurveEditInfoConst(&RichCurves[0], Private::RedCurveName));
+	Curves.Add(FRichCurveEditInfoConst(&RichCurves[1], Private::GreenCurveName));
+	Curves.Add(FRichCurveEditInfoConst(&RichCurves[2], Private::BlueCurveName));
+	Curves.Add(FRichCurveEditInfoConst(&RichCurves[3], Private::AlphaCurveName));
+}
+
+TArray<FRichCurveEditInfo> FDataflowLinearColorRamp::GetCurves()
+{
+	using namespace UE::Dataflow::ColorMap;
+
+	TArray<FRichCurveEditInfo> Curves;
+	Curves.Add(FRichCurveEditInfo(&RichCurves[0], Private::RedCurveName));
+	Curves.Add(FRichCurveEditInfo(&RichCurves[1], Private::GreenCurveName));
+	Curves.Add(FRichCurveEditInfo(&RichCurves[2], Private::BlueCurveName));
+	Curves.Add(FRichCurveEditInfo(&RichCurves[3], Private::AlphaCurveName));
+	return Curves;
+}
+
+void FDataflowLinearColorRamp::ModifyOwner()
+{
+}
+
+TArray<const UObject*> FDataflowLinearColorRamp::GetOwners() const
+{
+	static TArray<const UObject*> EmptyObjectArray;
+	return EmptyObjectArray;
+}
+
+void FDataflowLinearColorRamp::MakeTransactional()
+{
+}
+
+void FDataflowLinearColorRamp::OnCurveChanged(const TArray<FRichCurveEditInfo>& ChangedCurveEditInfos)
+{
+	TArray<FRichCurve*> ChangedCurves;
+	for (const FRichCurveEditInfo& ChangedCurveEditInfo : ChangedCurveEditInfos)
+	{
+		ChangedCurves.Add((FRichCurve*)ChangedCurveEditInfo.CurveToEdit);
+	}
+	if (ChangedCurves.Num() > 0)
+	{
+		OnColorCurveChangedDelegate.Broadcast(ChangedCurves);
+	}
+}
+
+bool FDataflowLinearColorRamp::IsLinearColorCurve() const
+{
+	return true;
+}
+
+FLinearColor FDataflowLinearColorRamp::GetLinearColorValue(float InTime) const
+{
+	return FLinearColor(
+		RichCurves[0].Eval(InTime),
+		RichCurves[1].Eval(InTime),
+		RichCurves[2].Eval(InTime),
+		RichCurves[3].Eval(InTime));
+}
+
+bool FDataflowLinearColorRamp::HasAnyAlphaKeys() const
+{
+	return RichCurves[3].GetNumKeys() != 0;
+}
+
+bool FDataflowLinearColorRamp::IsValidCurve(FRichCurveEditInfo CurveInfo)
+{
+	return (CurveInfo.CurveToEdit == &RichCurves[0]
+		|| CurveInfo.CurveToEdit == &RichCurves[1]
+		|| CurveInfo.CurveToEdit == &RichCurves[2]
+		|| CurveInfo.CurveToEdit == &RichCurves[3]);
+}
+
+FLinearColor FDataflowLinearColorRamp::GetCurveColor(FRichCurveEditInfo CurveInfo) const
+{
+	return FLinearColor::White;
+}
diff --git a/Engine/Source/Runtime/Experimental/Dataflow/Engine/Public/Dataflow/DataflowColorRamp.h b/Engine/Source/Runtime/Experimental/Dataflow/Engine/Public/Dataflow/DataflowColorRamp.h
--- a/Engine/Source/Runtime/Experimental/Dataflow/Engine/Public/Dataflow/DataflowColorRamp.h
+++ b/Engine/Source/Runtime/Experimental/Dataflow/Engine/Public/Dataflow/DataflowColorRamp.h
@@ -7,45 +7,37 @@
 
 #include "DataflowColorRamp.generated.h"
 
+#define UE_API DATAFLOWENGINE_API
+
 struct FDataflowColorCurveOwner : public FCurveOwnerInterface
 {
 public:
 	DECLARE_MULTICAST_DELEGATE_OneParam(FOnColorCurveChanged, TArray<FRichCurve*> /* ChangedCurves */);
 
 public:
-	DATAFLOWENGINE_API FDataflowColorCurveOwner();
-	DATAFLOWENGINE_API FDataflowColorCurveOwner(const FDataflowColorCurveOwner& Other);
-	DATAFLOWENGINE_API FDataflowColorCurveOwner& operator=(const FDataflowColorCurveOwner& Other);
-
-	FDataflowColorCurveOwner(FDataflowColorCurveOwner&& Other) = delete;
-	FDataflowColorCurveOwner& operator=(FDataflowColorCurveOwner&& Other) = delete;
-
-	DATAFLOWENGINE_API void SetColorAtTime(float Time, const FLinearColor& Color, bool bOnlyRBG);
-
-	DATAFLOWENGINE_API bool IsEmpty() const;
-
-	/** implemenmt all the necessary virtual function from FCurveOwnerInterface */
-	DATAFLOWENGINE_API virtual TArray<FRichCurveEditInfoConst> GetCurves() const override;
-	DATAFLOWENGINE_API virtual void GetCurves(TAdderReserverRef<FRichCurveEditInfoConst> Curves) const override;
-	DATAFLOWENGINE_API virtual TArray<FRichCurveEditInfo> GetCurves() override;
-	DATAFLOWENGINE_API virtual void ModifyOwner() override;
-	DATAFLOWENGINE_API virtual TArray<const UObject*> GetOwners() const override;
-	DATAFLOWENGINE_API virtual void MakeTransactional() override;
-	DATAFLOWENGINE_API virtual void OnCurveChanged(const TArray<FRichCurveEditInfo>& ChangedCurveEditInfos) override;
-	DATAFLOWENGINE_API virtual bool IsLinearColorCurve() const override;
-	DATAFLOWENGINE_API virtual FLinearColor GetLinearColorValue(float InTime) const override;
-	DATAFLOWENGINE_API virtual bool HasAnyAlphaKeys() const override;
-	DATAFLOWENGINE_API virtual bool IsValidCurve(FRichCurveEditInfo CurveInfo) override;
-	DATAFLOWENGINE_API virtual FLinearColor GetCurveColor(FRichCurveEditInfo CurveInfo) const override;
+	UE_API void SetColorAtTime(float Time, const FLinearColor& Color, bool bOnlyRBG);
+
+	UE_API bool IsEmpty() const;
+
+	//~ Begin FCurveOwnerInterface interface
+	UE_API virtual TArray<FRichCurveEditInfoConst> GetCurves() const override;
+	UE_API virtual void GetCurves(TAdderReserverRef<FRichCurveEditInfoConst> Curves) const override;
+	UE_API virtual TArray<FRichCurveEditInfo> GetCurves() override;
+	UE_API virtual void ModifyOwner() override;
+	UE_API virtual TArray<const UObject*> GetOwners() const override;
+	UE_API virtual void MakeTransactional() override;
+	UE_API virtual void OnCurveChanged(const TArray<FRichCurveEditInfo>& ChangedCurveEditInfos) override;
+	UE_API virtual bool IsLinearColorCurve() const override;
+	UE_API virtual FLinearColor GetLinearColorValue(float InTime) const override;
+	UE_API virtual bool HasAnyAlphaKeys() const override;
+	UE_API virtual bool IsValidCurve(FRichCurveEditInfo CurveInfo) override;
+	UE_API virtual FLinearColor GetCurveColor(FRichCurveEditInfo CurveInfo) const override;
+	//~ End FCurveOwnerInterface interface
 
 	FOnColorCurveChanged OnColorCurveChangedDelegate;
 
 private:
-	TArray<FRichCurveEditInfoConst> ConstCurves;
-	TArray<FRichCurveEditInfo> Curves;
-	TArray<const UObject*> Owners;
-
-	TArray<FRichCurve> RichCurves;
+	FRichCurve RichCurves[4];
 };
 
 USTRUCT()
@@ -54,4 +46,42 @@ struct FDataflowColorRamp
 	GENERATED_USTRUCT_BODY()
 
 	FDataflowColorCurveOwner ColorRamp;
-};
\ No newline at end of file
+};
+
+USTRUCT()
+struct FDataflowLinearColorRamp 
+#if CPP
+	: public FCurveOwnerInterface
+#endif
+{
+	GENERATED_USTRUCT_BODY()
+
+public:
+	UE_API void SetColorAtTime(float Time, const FLinearColor& Color, bool bOnlyRBG);
+	UE_API bool IsEmpty() const;
+	UE_API void Reset();
+
+	DECLARE_MULTICAST_DELEGATE_OneParam(FOnColorCurveChanged, TArray<FRichCurve*> /* ChangedCurves */);
+	FOnColorCurveChanged OnColorCurveChangedDelegate;
+
+	//~ Begin FCurveOwnerInterface interface
+	UE_API virtual TArray<FRichCurveEditInfoConst> GetCurves() const override;
+	UE_API virtual void GetCurves(TAdderReserverRef<FRichCurveEditInfoConst> Curves) const override;
+	UE_API virtual TArray<FRichCurveEditInfo> GetCurves() override;
+	UE_API virtual void ModifyOwner() override;
+	UE_API virtual TArray<const UObject*> GetOwners() const override;
+	UE_API virtual void MakeTransactional() override;
+	UE_API virtual void OnCurveChanged(const TArray<FRichCurveEditInfo>& ChangedCurveEditInfos) override;
+	UE_API virtual bool IsLinearColorCurve() const override;
+	UE_API virtual FLinearColor GetLinearColorValue(float InTime) const override;
+	UE_API virtual bool HasAnyAlphaKeys() const override;
+	UE_API virtual bool IsValidCurve(FRichCurveEditInfo CurveInfo) override;
+	UE_API virtual FLinearColor GetCurveColor(FRichCurveEditInfo CurveInfo) const override;
+	//~ End FCurveOwnerInterface interface
+
+private:
+	UPROPERTY()
+	FRichCurve RichCurves[4];
+};
+
+#undef UE_API
