[ITF Viewport Interactions] Fix a bug where orbiting to view the top or bottom of an object could cause view position errors. ToggleOrbitCamera() was being used too often, and UpdateCameraMovement() was running while orbit interactions were active.
- Remove viewport interaction Type name field. In most cases, those type names become the InteractionName.
- Refactor GroupName to a Groups array. Groups were originally being used as an identifier within the behavior source, but this identifier was never used. Groups are now used in new IsAnyInteractionInGroupEnabled/Active accessors on UViewportInteractionsBehaviorSource.
- Refactor the previous "IsViewportInteractionActive" function and usages to refer to "Enabled" instead. "Active" sounds like "the behavior is doing something right now", which was not what it was checking for. The new setup differentiates between referring to groups and to specific interactions.

#rnx
#Synced-CL 48396535

diff --git a/Engine/Source/Editor/Experimental/EditorInteractiveToolsFramework/Private/EditorDragTools/BoxSelectInteraction.cpp b/Engine/Source/Editor/Experimental/EditorInteractiveToolsFramework/Private/EditorDragTools/BoxSelectInteraction.cpp
--- a/Engine/Source/Editor/Experimental/EditorInteractiveToolsFramework/Private/EditorDragTools/BoxSelectInteraction.cpp
+++ b/Engine/Source/Editor/Experimental/EditorInteractiveToolsFramework/Private/EditorDragTools/BoxSelectInteraction.cpp
@@ -67,8 +67,7 @@ void RemoveHoverEffect(const FViewportHoverTarget& InHoverTarget)
 
 UBoxSelectInteraction::UBoxSelectInteraction()
 {
-	ToolType = UE::Editor::ViewportInteractions::BoxSelect;
-	InteractionName = TEXT("Box Select");
+	InteractionName = UE::Editor::ViewportInteractions::BoxSelect;
 
 	if (UViewportClickDragBehavior* ClickDragBehavior = ClickDragInputBehavior.Get())
 	{
diff --git a/Engine/Source/Editor/Experimental/EditorInteractiveToolsFramework/Private/EditorDragTools/DragToolInteraction.cpp b/Engine/Source/Editor/Experimental/EditorInteractiveToolsFramework/Private/EditorDragTools/DragToolInteraction.cpp
--- a/Engine/Source/Editor/Experimental/EditorInteractiveToolsFramework/Private/EditorDragTools/DragToolInteraction.cpp
+++ b/Engine/Source/Editor/Experimental/EditorInteractiveToolsFramework/Private/EditorDragTools/DragToolInteraction.cpp
@@ -109,7 +109,6 @@ UDragToolInteraction::UDragToolInteraction()
 	, bIsDragging(false)
 	, EditorViewportClientProxy(nullptr)
 {
-	GroupName = UE::Editor::ViewportInteractions::DragToolsGroupName;
 }
 
 void UDragToolInteraction::Initialize(UViewportInteractionsBehaviorSource* const InViewportInteractionsBehaviorSource)
diff --git a/Engine/Source/Editor/Experimental/EditorInteractiveToolsFramework/Private/EditorDragTools/DuplicateDragSelectionInteraction.cpp b/Engine/Source/Editor/Experimental/EditorInteractiveToolsFramework/Private/EditorDragTools/DuplicateDragSelectionInteraction.cpp
--- a/Engine/Source/Editor/Experimental/EditorInteractiveToolsFramework/Private/EditorDragTools/DuplicateDragSelectionInteraction.cpp
+++ b/Engine/Source/Editor/Experimental/EditorInteractiveToolsFramework/Private/EditorDragTools/DuplicateDragSelectionInteraction.cpp
@@ -14,9 +14,7 @@
 
 UDuplicateDragSelectionInteraction::UDuplicateDragSelectionInteraction()
 {
-	ToolType = UE::Editor::ViewportInteractions::DuplicateDrag;
-	InteractionName = TEXT("Duplicate Drag");
-	GroupName = UE::Editor::ViewportInteractions::DragToolsGroupName;
+	InteractionName = UE::Editor::ViewportInteractions::DuplicateDrag;
 }
 
 void UDuplicateDragSelectionInteraction::Initialize(UViewportInteractionsBehaviorSource* const InViewportInteractionsBehaviorSource)
diff --git a/Engine/Source/Editor/Experimental/EditorInteractiveToolsFramework/Private/EditorDragTools/FrustumSelectInteraction.cpp b/Engine/Source/Editor/Experimental/EditorInteractiveToolsFramework/Private/EditorDragTools/FrustumSelectInteraction.cpp
--- a/Engine/Source/Editor/Experimental/EditorInteractiveToolsFramework/Private/EditorDragTools/FrustumSelectInteraction.cpp
+++ b/Engine/Source/Editor/Experimental/EditorInteractiveToolsFramework/Private/EditorDragTools/FrustumSelectInteraction.cpp
@@ -139,8 +139,7 @@ void FViewportFrustum::Calculate(FConvexVolume& OutFrustum) const
 
 UFrustumSelectInteraction::UFrustumSelectInteraction()
 {
-	ToolType = UE::Editor::ViewportInteractions::FrustumSelect;
-	InteractionName = TEXT("Frustum Select");
+	InteractionName = UE::Editor::ViewportInteractions::FrustumSelect;
 
 	if (UViewportClickDragBehavior* ClickDragBehavior = ClickDragInputBehavior.Get())
 	{
diff --git a/Engine/Source/Editor/Experimental/EditorInteractiveToolsFramework/Private/EditorDragTools/MeasureToolInteraction.cpp b/Engine/Source/Editor/Experimental/EditorInteractiveToolsFramework/Private/EditorDragTools/MeasureToolInteraction.cpp
--- a/Engine/Source/Editor/Experimental/EditorInteractiveToolsFramework/Private/EditorDragTools/MeasureToolInteraction.cpp
+++ b/Engine/Source/Editor/Experimental/EditorInteractiveToolsFramework/Private/EditorDragTools/MeasureToolInteraction.cpp
@@ -18,8 +18,7 @@ UMeasureToolInteraction::UMeasureToolInteraction()
 	: PixelStart()
 	, PixelEnd()
 {
-	ToolType = UE::Editor::ViewportInteractions::Measure;
-	InteractionName = TEXT("Measure Tool");
+	InteractionName = UE::Editor::ViewportInteractions::Measure;
 
 	bUseSnapping = true;
 	bConvertDelta = false;
diff --git a/Engine/Source/Editor/Experimental/EditorInteractiveToolsFramework/Private/EditorDragTools/MoveCameraWithObjectInteraction.cpp b/Engine/Source/Editor/Experimental/EditorInteractiveToolsFramework/Private/EditorDragTools/MoveCameraWithObjectInteraction.cpp
--- a/Engine/Source/Editor/Experimental/EditorInteractiveToolsFramework/Private/EditorDragTools/MoveCameraWithObjectInteraction.cpp
+++ b/Engine/Source/Editor/Experimental/EditorInteractiveToolsFramework/Private/EditorDragTools/MoveCameraWithObjectInteraction.cpp
@@ -10,9 +10,7 @@
 
 UMoveCameraWithObjectInteraction::UMoveCameraWithObjectInteraction()
 {
-	ToolType = UE::Editor::ViewportInteractions::DragMovesCamera;
-	InteractionName = TEXT("Drag Moves Camera");
-	GroupName = UE::Editor::ViewportInteractions::DragToolsGroupName;
+	InteractionName = UE::Editor::ViewportInteractions::DragMovesCamera;
 }
 
 void UMoveCameraWithObjectInteraction::BeginDestroy()
diff --git a/Engine/Source/Editor/Experimental/EditorInteractiveToolsFramework/Private/EditorDragTools/ViewportChangeInteraction.cpp b/Engine/Source/Editor/Experimental/EditorInteractiveToolsFramework/Private/EditorDragTools/ViewportChangeInteraction.cpp
--- a/Engine/Source/Editor/Experimental/EditorInteractiveToolsFramework/Private/EditorDragTools/ViewportChangeInteraction.cpp
+++ b/Engine/Source/Editor/Experimental/EditorInteractiveToolsFramework/Private/EditorDragTools/ViewportChangeInteraction.cpp
@@ -11,8 +11,7 @@
 
 UViewportChangeInteraction::UViewportChangeInteraction()
 {
-	ToolType = UE::Editor::ViewportInteractions::ViewportChange;
-	InteractionName = TEXT("Viewport Change");
+	InteractionName = UE::Editor::ViewportInteractions::ViewportChange;
 
 	ViewOptionOffset = FVector2D::ZeroVector;
 	bUseSnapping = true;
diff --git a/Engine/Source/Editor/Experimental/EditorInteractiveToolsFramework/Private/ViewportInteractions/ViewportCameraRotateInteraction.cpp b/Engine/Source/Editor/Experimental/EditorInteractiveToolsFramework/Private/ViewportInteractions/ViewportCameraRotateInteraction.cpp
--- a/Engine/Source/Editor/Experimental/EditorInteractiveToolsFramework/Private/ViewportInteractions/ViewportCameraRotateInteraction.cpp
+++ b/Engine/Source/Editor/Experimental/EditorInteractiveToolsFramework/Private/ViewportInteractions/ViewportCameraRotateInteraction.cpp
@@ -9,9 +9,8 @@
 
 UViewportCameraRotateInteraction::UViewportCameraRotateInteraction()
 {
-	InteractionName = TEXT("Camera Rotate");
-	GroupName = UE::Editor::ViewportInteractions::CameraAndSelectionGroupName;
-	ToolType = UE::Editor::ViewportInteractions::CameraRotate;
+	InteractionName = UE::Editor::ViewportInteractions::CameraRotate;
+	Groups = { UE::Editor::ViewportInteractions::CameraFly };
 
 	UKeyInputBehavior* KeyInputBehavior = NewObject<UKeyInputBehavior>();
 	KeyInputBehavior->Initialize(this, GetKeys());
diff --git a/Engine/Source/Editor/Experimental/EditorInteractiveToolsFramework/Private/ViewportInteractions/ViewportCameraSpeedMouseWheelInteraction.cpp b/Engine/Source/Editor/Experimental/EditorInteractiveToolsFramework/Private/ViewportInteractions/ViewportCameraSpeedMouseWheelInteraction.cpp
--- a/Engine/Source/Editor/Experimental/EditorInteractiveToolsFramework/Private/ViewportInteractions/ViewportCameraSpeedMouseWheelInteraction.cpp
+++ b/Engine/Source/Editor/Experimental/EditorInteractiveToolsFramework/Private/ViewportInteractions/ViewportCameraSpeedMouseWheelInteraction.cpp
@@ -8,9 +8,8 @@
 
 UViewportCameraSpeedMouseWheelInteraction::UViewportCameraSpeedMouseWheelInteraction()
 {
-	ToolType = UE::Editor::ViewportInteractions::CameraTranslate;
-	InteractionName = TEXT("Camera Speed");
-	GroupName = UE::Editor::ViewportInteractions::CameraAndSelectionGroupName;
+	InteractionName	= TEXT("Mouse Wheel Camera Speed");
+	Groups = { UE::Editor::ViewportInteractions::CameraFly };
 
 	// Set camera speed using mouse wheel
 	UMouseWheelInputBehavior* MouseWheelInputBehavior = NewObject<UMouseWheelModifierInputBehavior>();
diff --git a/Engine/Source/Editor/Experimental/EditorInteractiveToolsFramework/Private/ViewportInteractions/ViewportCameraTranslateInteraction.cpp b/Engine/Source/Editor/Experimental/EditorInteractiveToolsFramework/Private/ViewportInteractions/ViewportCameraTranslateInteraction.cpp
--- a/Engine/Source/Editor/Experimental/EditorInteractiveToolsFramework/Private/ViewportInteractions/ViewportCameraTranslateInteraction.cpp
+++ b/Engine/Source/Editor/Experimental/EditorInteractiveToolsFramework/Private/ViewportInteractions/ViewportCameraTranslateInteraction.cpp
@@ -27,9 +27,8 @@ namespace UE::Editor::ViewportInteractions::Private
 
 UViewportCameraTranslateInteraction::UViewportCameraTranslateInteraction()
 {
-	ToolType = UE::Editor::ViewportInteractions::CameraTranslate;
-	InteractionName = TEXT("Camera Translate");
-	GroupName = UE::Editor::ViewportInteractions::CameraAndSelectionGroupName;
+	InteractionName = UE::Editor::ViewportInteractions::CameraTranslate;
+	Groups = { UE::Editor::ViewportInteractions::CameraFly };
 
 	// Camera translate
 	UKeyInputBehavior* KeyInputBehavior = NewObject<UKeyInputBehavior>();
diff --git a/Engine/Source/Editor/Experimental/EditorInteractiveToolsFramework/Private/ViewportInteractions/ViewportCommandsInteraction.cpp b/Engine/Source/Editor/Experimental/EditorInteractiveToolsFramework/Private/ViewportInteractions/ViewportCommandsInteraction.cpp
--- a/Engine/Source/Editor/Experimental/EditorInteractiveToolsFramework/Private/ViewportInteractions/ViewportCommandsInteraction.cpp
+++ b/Engine/Source/Editor/Experimental/EditorInteractiveToolsFramework/Private/ViewportInteractions/ViewportCommandsInteraction.cpp
@@ -7,7 +7,7 @@
 
 UViewportCommandsInteraction::UViewportCommandsInteraction()
 {
-	ToolType = UE::Editor::ViewportInteractions::Commands;
+	InteractionName = UE::Editor::ViewportInteractions::Commands;
 
 	UKeyInputBehavior* KeyInputBehavior = NewObject<UKeyInputBehavior>();
 	KeyInputBehavior->Initialize(this, GetKeys());
diff --git a/Engine/Source/Editor/Experimental/EditorInteractiveToolsFramework/Private/ViewportInteractions/ViewportDollyInteraction.cpp b/Engine/Source/Editor/Experimental/EditorInteractiveToolsFramework/Private/ViewportInteractions/ViewportDollyInteraction.cpp
--- a/Engine/Source/Editor/Experimental/EditorInteractiveToolsFramework/Private/ViewportInteractions/ViewportDollyInteraction.cpp
+++ b/Engine/Source/Editor/Experimental/EditorInteractiveToolsFramework/Private/ViewportInteractions/ViewportDollyInteraction.cpp
@@ -9,9 +9,7 @@
 
 UViewportDollyInteraction::UViewportDollyInteraction()
 {
-	ToolType = UE::Editor::ViewportInteractions::CameraDrag;
-
-	InteractionName = TEXT("Perspective Dolly");
+	InteractionName = UE::Editor::ViewportInteractions::PerspectiveDolly;
 	
 	ViewportClickDragBehavior->SetBindings({
 		EKeys::RightMouseButton,
diff --git a/Engine/Source/Editor/Experimental/EditorInteractiveToolsFramework/Private/ViewportInteractions/ViewportDragInteraction.cpp b/Engine/Source/Editor/Experimental/EditorInteractiveToolsFramework/Private/ViewportInteractions/ViewportDragInteraction.cpp
--- a/Engine/Source/Editor/Experimental/EditorInteractiveToolsFramework/Private/ViewportInteractions/ViewportDragInteraction.cpp
+++ b/Engine/Source/Editor/Experimental/EditorInteractiveToolsFramework/Private/ViewportInteractions/ViewportDragInteraction.cpp
@@ -22,8 +22,6 @@ UViewportDragInteraction::UViewportDragInteraction()
 	ViewportClickDragBehavior = ClickDragBehavior;
 
 	RegisterInputBehavior(ClickDragBehavior);
-
-	GroupName = UE::Editor::ViewportInteractions::CameraAndSelectionGroupName;
 }
 
 FInputRayHit UViewportDragInteraction::CanBeginClickDragSequence(const FInputDeviceRay& InPressPos)
diff --git a/Engine/Source/Editor/Experimental/EditorInteractiveToolsFramework/Private/ViewportInteractions/ViewportFOVInteraction.cpp b/Engine/Source/Editor/Experimental/EditorInteractiveToolsFramework/Private/ViewportInteractions/ViewportFOVInteraction.cpp
--- a/Engine/Source/Editor/Experimental/EditorInteractiveToolsFramework/Private/ViewportInteractions/ViewportFOVInteraction.cpp
+++ b/Engine/Source/Editor/Experimental/EditorInteractiveToolsFramework/Private/ViewportInteractions/ViewportFOVInteraction.cpp
@@ -8,10 +8,8 @@
 
 UViewportFOVInteraction::UViewportFOVInteraction()
 {
-	ToolType = UE::Editor::ViewportInteractions::FOV;
-
-	InteractionName = TEXT("FOV");
-	GroupName = UE::Editor::ViewportInteractions::CameraAndSelectionGroupName;
+	InteractionName = UE::Editor::ViewportInteractions::FOV;
+	Groups = { UE::Editor::ViewportInteractions::CameraFly };
 
 	// Separating keys which work only when Mouse Looking from numpad ones and arrows, which always work.
 	// This prevents this IKeyInputBehaviorTarget from capturing unwanted inputs
diff --git a/Engine/Source/Editor/Experimental/EditorInteractiveToolsFramework/Private/ViewportInteractions/ViewportInteraction.cpp b/Engine/Source/Editor/Experimental/EditorInteractiveToolsFramework/Private/ViewportInteractions/ViewportInteraction.cpp
--- a/Engine/Source/Editor/Experimental/EditorInteractiveToolsFramework/Private/ViewportInteractions/ViewportInteraction.cpp
+++ b/Engine/Source/Editor/Experimental/EditorInteractiveToolsFramework/Private/ViewportInteractions/ViewportInteraction.cpp
@@ -93,18 +93,6 @@ bool UViewportInteraction::IsMouseLooking() const
 void UViewportInteraction::SetEnabled(bool bInEnabled)
 {
 	bEnabled = bInEnabled;
-
-	if (UViewportInteractionsBehaviorSource* BehaviorSource = GetViewportInteractionsBehaviorSource())
-	{
-		if (ToolType != NAME_None)
-		{
-			BehaviorSource->SetViewportInteractionActive(ToolType, bEnabled);
-		}
-	}
-	else
-	{
-		UE_LOG(LogTemp,Log,TEXT("Trying to initialize an instance of %s with an invalid UViewportInteractionsBehaviorSource"), *GetClass()->GetName());
-	}
 }
 
 bool UViewportInteraction::IsEnabled() const
@@ -117,9 +105,9 @@ FName UViewportInteraction::GetInteractionName() const
 	return InteractionName.IsNone() ? StaticClass()->GetFName() : InteractionName;
 }
 
-FName UViewportInteraction::GetInteractionGroup() const
+const TArray<FName>& UViewportInteraction::GetInteractionGroups() const
 {
-	return GroupName.IsNone() ? UE::Editor::ViewportInteractions::DefaultGroupName : GroupName;
+	return Groups;
 }
 
 void UViewportInteraction::Initialize(UViewportInteractionsBehaviorSource* const InViewportInteractionsBehaviorSource)
diff --git a/Engine/Source/Editor/Experimental/EditorInteractiveToolsFramework/Private/ViewportInteractions/ViewportInteractionsBehaviorSource.cpp b/Engine/Source/Editor/Experimental/EditorInteractiveToolsFramework/Private/ViewportInteractions/ViewportInteractionsBehaviorSource.cpp
--- a/Engine/Source/Editor/Experimental/EditorInteractiveToolsFramework/Private/ViewportInteractions/ViewportInteractionsBehaviorSource.cpp
+++ b/Engine/Source/Editor/Experimental/EditorInteractiveToolsFramework/Private/ViewportInteractions/ViewportInteractionsBehaviorSource.cpp
@@ -373,10 +373,16 @@ UViewportInteraction* UViewportInteractionsBehaviorSource::AddInteraction(const
 
 	UE::Editor::ViewportInteractions::LOG(*FString::Printf(TEXT("Initializing Viewport Interaction \"%s\""), *InInteractionClass->GetName()));
 	ViewportInteractions.Add(Interaction->GetInteractionName(), Interaction);
+	
+	for (const FName& GroupName : Interaction->GetInteractionGroups())
+	{
+		FViewportInteractionGroup& Group = ViewportInteractionGroups.FindOrAdd(GroupName);
+		Group.Interactions.Add(Interaction);
+	}
 
 	for (UInputBehavior* const InputBehavior : Interaction->GetInputBehaviors())
 	{
-		BehaviorSet->Add(InputBehavior, Interaction, Interaction->GetInteractionGroup().ToString());
+		BehaviorSet->Add(InputBehavior, Interaction);
 	}
 
 	if (bInReregister)
@@ -401,6 +407,41 @@ UViewportInteraction* UViewportInteractionsBehaviorSource::FindInteraction(const
 	return nullptr;
 }
 
+UViewportInteraction* UViewportInteractionsBehaviorSource::FindInteraction(const FName& InName) const
+{
+	return ViewportInteractions.FindRef(InName);
+}
+
+bool UViewportInteractionsBehaviorSource::IsAnyInteractionInGroupEnabled(const FName& InName) const
+{
+	if (const FViewportInteractionGroup* Group = ViewportInteractionGroups.Find(InName))
+	{
+		for (const UViewportInteraction* Interaction : Group->Interactions)
+		{
+			if (Interaction->IsEnabled())
+			{
+				return true;
+			}
+		}
+	}
+	return false;
+}
+
+bool UViewportInteractionsBehaviorSource::IsAnyInteractionInGroupActive(const FName& InName) const
+{
+	if (const FViewportInteractionGroup* Group = ViewportInteractionGroups.Find(InName))
+	{
+		for (const UViewportInteraction* Interaction : Group->Interactions)
+		{
+			if (Interaction->IsActive())
+			{
+				return true;
+			}
+		}
+	}
+	return false;
+}
+
 void UViewportInteractionsBehaviorSource::RemoveInteraction(const FName InInteractionName, bool bInReregister)
 {
 	if (UViewportInteraction* Interaction = ViewportInteractions.FindRef(InInteractionName))
@@ -408,6 +449,15 @@ void UViewportInteractionsBehaviorSource::RemoveInteraction(const FName InIntera
 		UE::Editor::ViewportInteractions::LOG(*FString::Printf(TEXT("Removing Viewport Interaction \"%s\""), *Interaction->GetInteractionName().ToString()));
 
 		ViewportInteractions.Remove(Interaction->GetInteractionName());
+		
+		for (const FName& GroupName : Interaction->GetInteractionGroups())
+		{
+			if (FViewportInteractionGroup* Group = ViewportInteractionGroups.Find(GroupName))
+			{
+				Group->Interactions.Remove(Interaction);
+			}
+		}
+		
 		Interaction->Shutdown();
 
 		BehaviorSet->RemoveBySource(Interaction);
@@ -438,7 +488,7 @@ void UViewportInteractionsBehaviorSource::Reset()
 	}
 
 	ViewportInteractions.Reset();
-	ViewportInteractionsStatusMap.Reset();
+	ViewportInteractionGroups.Reset();
 	
 	for (const TObjectPtr<UInputBehavior>& Behavior : CoreBehaviors)
 	{
@@ -583,21 +633,6 @@ FEditorViewportClient* UViewportInteractionsBehaviorSource::GetHoveredEditorView
 	return EditorViewportClient;
 }
 
-void UViewportInteractionsBehaviorSource::SetViewportInteractionActive(FName InViewportInteraction, bool bInActive)
-{
-	ViewportInteractionsStatusMap.Emplace(InViewportInteraction, bInActive);
-}
-
-bool UViewportInteractionsBehaviorSource::IsViewportInteractionActive(FName InViewportInteraction)
-{
-	if (ViewportInteractionsStatusMap.Contains(InViewportInteraction))
-	{
-		return ViewportInteractionsStatusMap[InViewportInteraction];
-	}
-
-	return false;
-}
-
 void UViewportInteractionsBehaviorSource::SetAltKeyState(bool bInDown)
 {
 	bIsAltDown = bInDown;
diff --git a/Engine/Source/Editor/Experimental/EditorInteractiveToolsFramework/Private/ViewportInteractions/ViewportMoveYawInteraction.cpp b/Engine/Source/Editor/Experimental/EditorInteractiveToolsFramework/Private/ViewportInteractions/ViewportMoveYawInteraction.cpp
--- a/Engine/Source/Editor/Experimental/EditorInteractiveToolsFramework/Private/ViewportInteractions/ViewportMoveYawInteraction.cpp
+++ b/Engine/Source/Editor/Experimental/EditorInteractiveToolsFramework/Private/ViewportInteractions/ViewportMoveYawInteraction.cpp
@@ -11,9 +11,11 @@
 
 UViewportMoveYawInteraction::UViewportMoveYawInteraction()
 {
-	ToolType = UE::Editor::ViewportInteractions::CameraDrag;
-
-	InteractionName = TEXT("Move-Yaw");
+	InteractionName = UE::Editor::ViewportInteractions::PerspectiveMoveYaw;
+	Groups = {
+		UE::Editor::ViewportInteractions::CameraDrag,
+		UE::Editor::ViewportInteractions::CameraFly,
+	};
 
 	ViewportClickDragBehavior->SetBindings({ EKeys::LeftMouseButton });
 }
diff --git a/Engine/Source/Editor/Experimental/EditorInteractiveToolsFramework/Private/ViewportInteractions/ViewportOrbitInteraction.cpp b/Engine/Source/Editor/Experimental/EditorInteractiveToolsFramework/Private/ViewportInteractions/ViewportOrbitInteraction.cpp
--- a/Engine/Source/Editor/Experimental/EditorInteractiveToolsFramework/Private/ViewportInteractions/ViewportOrbitInteraction.cpp
+++ b/Engine/Source/Editor/Experimental/EditorInteractiveToolsFramework/Private/ViewportInteractions/ViewportOrbitInteraction.cpp
@@ -10,9 +10,8 @@
 
 UViewportOrbitInteraction::UViewportOrbitInteraction()
 {
-	ToolType = UE::Editor::ViewportInteractions::Orbit;
-
-	InteractionName = TEXT("Orbit");
+	InteractionName = UE::Editor::ViewportInteractions::Orbit;
+	Groups = { UE::Editor::ViewportInteractions::CameraDrag };
 	
 	using namespace UE::Editor::ViewportInteractions;
 	ViewportClickDragBehavior->SetBindings({
diff --git a/Engine/Source/Editor/Experimental/EditorInteractiveToolsFramework/Private/ViewportInteractions/ViewportOrthoPanInteraction.cpp b/Engine/Source/Editor/Experimental/EditorInteractiveToolsFramework/Private/ViewportInteractions/ViewportOrthoPanInteraction.cpp
--- a/Engine/Source/Editor/Experimental/EditorInteractiveToolsFramework/Private/ViewportInteractions/ViewportOrthoPanInteraction.cpp
+++ b/Engine/Source/Editor/Experimental/EditorInteractiveToolsFramework/Private/ViewportInteractions/ViewportOrthoPanInteraction.cpp
@@ -8,9 +8,8 @@
 
 UViewportOrthoPanInteraction::UViewportOrthoPanInteraction()
 {
-	ToolType = UE::Editor::ViewportInteractions::CameraDrag;
-
-	InteractionName = TEXT("Orthographic Panning");
+	InteractionName = UE::Editor::ViewportInteractions::OrthographicPan;
+	Groups = { UE::Editor::ViewportInteractions::CameraDrag };
 
 	ViewportClickDragBehavior->SetBindings({ EKeys::RightMouseButton });
 }
diff --git a/Engine/Source/Editor/Experimental/EditorInteractiveToolsFramework/Private/ViewportInteractions/ViewportPanInteraction.cpp b/Engine/Source/Editor/Experimental/EditorInteractiveToolsFramework/Private/ViewportInteractions/ViewportPanInteraction.cpp
--- a/Engine/Source/Editor/Experimental/EditorInteractiveToolsFramework/Private/ViewportInteractions/ViewportPanInteraction.cpp
+++ b/Engine/Source/Editor/Experimental/EditorInteractiveToolsFramework/Private/ViewportInteractions/ViewportPanInteraction.cpp
@@ -11,9 +11,11 @@
 
 UViewportPanInteraction::UViewportPanInteraction()
 {
-	ToolType = UE::Editor::ViewportInteractions::CameraDrag;
-
-	InteractionName = TEXT("Perspective Panning");
+	InteractionName = UE::Editor::ViewportInteractions::PerspectivePan;
+	Groups = {
+		UE::Editor::ViewportInteractions::CameraDrag,
+		UE::Editor::ViewportInteractions::CameraFly,
+	};
 
 	using namespace UE::Editor::ViewportInteractions;
 	ViewportClickDragBehavior->SetBindings({
diff --git a/Engine/Source/Editor/Experimental/EditorInteractiveToolsFramework/Private/ViewportInteractions/ViewportViewAngleInteraction.cpp b/Engine/Source/Editor/Experimental/EditorInteractiveToolsFramework/Private/ViewportInteractions/ViewportViewAngleInteraction.cpp
--- a/Engine/Source/Editor/Experimental/EditorInteractiveToolsFramework/Private/ViewportInteractions/ViewportViewAngleInteraction.cpp
+++ b/Engine/Source/Editor/Experimental/EditorInteractiveToolsFramework/Private/ViewportInteractions/ViewportViewAngleInteraction.cpp
@@ -10,9 +10,11 @@
 
 UViewportViewAngleInteraction::UViewportViewAngleInteraction()
 {
-	ToolType = UE::Editor::ViewportInteractions::CameraDrag;
-
-	InteractionName = TEXT("View Angle");
+	InteractionName = UE::Editor::ViewportInteractions::PerspectiveViewAngle;
+	Groups = {
+		UE::Editor::ViewportInteractions::CameraDrag,
+		UE::Editor::ViewportInteractions::CameraFly
+	};
 	
 	ViewportClickDragBehavior->SetBindings({
 		EKeys::RightMouseButton
diff --git a/Engine/Source/Editor/Experimental/EditorInteractiveToolsFramework/Private/ViewportInteractions/ViewportZoomInteraction.cpp b/Engine/Source/Editor/Experimental/EditorInteractiveToolsFramework/Private/ViewportInteractions/ViewportZoomInteraction.cpp
--- a/Engine/Source/Editor/Experimental/EditorInteractiveToolsFramework/Private/ViewportInteractions/ViewportZoomInteraction.cpp
+++ b/Engine/Source/Editor/Experimental/EditorInteractiveToolsFramework/Private/ViewportInteractions/ViewportZoomInteraction.cpp
@@ -12,10 +12,11 @@
 
 UViewportZoomInteraction::UViewportZoomInteraction()
 {
-	ToolType = UE::Editor::ViewportInteractions::Zoom;
-
-	InteractionName = TEXT("Zoom");
-	GroupName = UE::Editor::ViewportInteractions::CameraAndSelectionGroupName;
+	InteractionName = UE::Editor::ViewportInteractions::Zoom;
+	Groups = {
+		UE::Editor::ViewportInteractions::CameraDrag,
+		UE::Editor::ViewportInteractions::CameraFly
+	};
 
 	// Mousewheel zoom
 	UMouseWheelInputBehavior* MouseWheelInputBehavior = NewObject<UMouseWheelInputBehavior>();
diff --git a/Engine/Source/Editor/Experimental/EditorInteractiveToolsFramework/Public/ViewportInteractions/ViewportDragInteraction.h b/Engine/Source/Editor/Experimental/EditorInteractiveToolsFramework/Public/ViewportInteractions/ViewportDragInteraction.h
--- a/Engine/Source/Editor/Experimental/EditorInteractiveToolsFramework/Public/ViewportInteractions/ViewportDragInteraction.h
+++ b/Engine/Source/Editor/Experimental/EditorInteractiveToolsFramework/Public/ViewportInteractions/ViewportDragInteraction.h
@@ -27,6 +27,8 @@ private:
 	GENERATED_BODY()
 public:
 	UE_API UViewportDragInteraction();
+	
+	UE_API virtual bool IsActive() const override { return bIsDragging; }
 
 	//~ Begin IViewportClickDragBehaviorTarget
 	UE_API virtual FInputRayHit CanBeginClickDragSequence(const FInputDeviceRay& InPressPos) override;
diff --git a/Engine/Source/Editor/Experimental/EditorInteractiveToolsFramework/Public/ViewportInteractions/ViewportInteraction.h b/Engine/Source/Editor/Experimental/EditorInteractiveToolsFramework/Public/ViewportInteractions/ViewportInteraction.h
--- a/Engine/Source/Editor/Experimental/EditorInteractiveToolsFramework/Public/ViewportInteractions/ViewportInteraction.h
+++ b/Engine/Source/Editor/Experimental/EditorInteractiveToolsFramework/Public/ViewportInteractions/ViewportInteraction.h
@@ -43,17 +43,16 @@ static constexpr int DRAG_PRIORITY = FInputCapturePriority::DEFAULT_TOOL_PRIORIT
 // Camera Movement
 static FName Zoom = TEXT("Zoom");
 static FName Orbit = TEXT("Orbit");
-static FName CameraDrag= TEXT("CameraDrag");
+static FName PerspectiveViewAngle = TEXT("PerspectiveCameraAngle");
+static FName PerspectiveDolly = TEXT("PerspectiveDolly");
+static FName PerspectivePan = TEXT("PerspectivePan");
+static FName PerspectiveMoveYaw = TEXT("PerspectiveMoveYaw");
+static FName OrthographicPan = TEXT("OrthographicPan");
 static FName CameraTranslate = TEXT("CameraTranslate");
 static FName CameraRotate = TEXT("CameraRotate");
 static FName FOV = TEXT("FOV");
 static FName Commands = TEXT("Commands");
 
-// Clicking
-static FName SingleClick = TEXT("SingleClick");
-static FName DoubleClick = TEXT("DoubleClick");
-static FName ContextMenu = TEXT("ContextMenu");
-
 // Drag Tools
 static FName BoxSelect = TEXT("BoxSelect");
 static FName FrustumSelect = TEXT("FrustumSelect");
@@ -62,10 +61,13 @@ static FName ViewportChange = TEXT("ViewportChange");
 static FName DuplicateDrag = TEXT("DuplicateDrag");
 static FName DragMovesCamera = TEXT("DragMovesCamera");
 
-// We use these strings to group Behavior Inputs so that we can selectively remove them if needed, without touching other inputs we might add (e.g. the ones added by Initialize)
-static const FName DefaultGroupName = TEXT("ViewportInteractions");
-static const FName CameraAndSelectionGroupName = TEXT("CameraAndSelection");
-static const FName DragToolsGroupName = TEXT("DragTools");
+// Groups
+/** Interactions that click on the viewport */
+static const FName ViewportClick = TEXT("ViewportClick");
+/** Interactions that move the camera via mouse motion */
+static const FName CameraDrag = TEXT("CameraDrag");
+/** Interactions the enable FPS camera movement */
+static const FName CameraFly = TEXT("CameraFly");
 
 } // namespace UE::Editor::ViewportInteractions
 
@@ -96,6 +98,8 @@ public:
 
 	UE_API void SetEnabled(bool bInEnabled);
 	UE_API bool IsEnabled() const;
+	
+	virtual bool IsActive() const { return false; }
 
 	/** Called when the interaction is removed from its behavior source. Override to make sure the interaction is properly cleaned up */
 	virtual void Shutdown() {}
@@ -118,9 +122,9 @@ public:
 	UE_API FName GetInteractionName() const;
 
 	/**
-	 * The Group this interaction belongs to.
+	 * The Groups this interaction belongs to, if any
 	 */
-	UE_API FName GetInteractionGroup() const;
+	UE_API const TArray<FName>& GetInteractionGroups() const;
 
 	/**
 	 * Initialization code which can be expanded by derived classes
@@ -170,8 +174,6 @@ protected:
 
 	bool bEnabled = true;
 
-	FName ToolType;
-
 	/**
 	 * Name for this interaction, used by Viewport Interactions Behavior Source to identify it.
 	 * If NAME_None, static class name will be used.
@@ -179,10 +181,9 @@ protected:
 	FName InteractionName = NAME_None;
 
 	/**
-	 * The Group this interaction belongs to. It can be specified to allow the selective removal/unregistering of interactions.
-	 * If NAME_None, a default group will be used.
+	 * The groups this interaction belongs to. Interactions can belong to many groups.
 	 */
-	FName GroupName = NAME_None;
+	TArray<FName> Groups;
 
 private:
 	UE_API void OnUserDefinedChordChanged(const FUICommandInfo& InCommandInfo);
diff --git a/Engine/Source/Editor/Experimental/EditorInteractiveToolsFramework/Public/ViewportInteractions/ViewportInteractionsBehaviorSource.h b/Engine/Source/Editor/Experimental/EditorInteractiveToolsFramework/Public/ViewportInteractions/ViewportInteractionsBehaviorSource.h
--- a/Engine/Source/Editor/Experimental/EditorInteractiveToolsFramework/Public/ViewportInteractions/ViewportInteractionsBehaviorSource.h
+++ b/Engine/Source/Editor/Experimental/EditorInteractiveToolsFramework/Public/ViewportInteractions/ViewportInteractionsBehaviorSource.h
@@ -56,17 +56,26 @@ EDITORINTERACTIVETOOLSFRAMEWORK_API void LOG(const TCHAR* InMessage);
 
 DECLARE_MULTICAST_DELEGATE_OneParam(FOnMouseLookingChanged, bool /* bIsMouseLooking*/)
 
-    DECLARE_MULTICAST_DELEGATE_OneParam(FOnModifierKeyChanged, bool /* bIsDown*/)
-
-    /**
-     * This class hosts a list of ITF Viewport Interactions and a UBehaviorSet containing the UInteractionBehavior
-     * instances required by those Viewport Interactions.
-     * It also acts as a behavior target, registering the status of Shift, Alt, Ctrl modifier keys. That and other
-     * pieces of information can be accessed externally, so that e.g. Viewport Interactions can know about modifier
-     * states without having to implement themselves what would be duplicate logic to handle those inputs. This class
-     * can also be used to know whether the Viewport camera is being moved using the mouse,or to mark it as such.
-     */
-	UCLASS(Transient, MinimalAPI)
+DECLARE_MULTICAST_DELEGATE_OneParam(FOnModifierKeyChanged, bool /* bIsDown*/)
+
+USTRUCT()
+struct FViewportInteractionGroup
+{
+	GENERATED_BODY()
+	
+	UPROPERTY()
+	TArray<TObjectPtr<UViewportInteraction>> Interactions;
+};
+
+/**
+ * This class hosts a list of ITF Viewport Interactions and a UBehaviorSet containing the UInteractionBehavior
+ * instances required by those Viewport Interactions.
+ * It also acts as a behavior target, registering the status of Shift, Alt, Ctrl modifier keys. That and other
+ * pieces of information can be accessed externally, so that e.g. Viewport Interactions can know about modifier
+ * states without having to implement themselves what would be duplicate logic to handle those inputs. This class
+ * can also be used to know whether the Viewport camera is being moved using the mouse,or to mark it as such.
+ */
+UCLASS(Transient, MinimalAPI)
 class UViewportInteractionsBehaviorSource final
 	: public UObject
 	, public IInputBehaviorSource
@@ -127,6 +136,21 @@ public:
 	 */
 	EDITORINTERACTIVETOOLSFRAMEWORK_API UViewportInteraction* FindInteraction(const UClass* InInteractionClass) const;
 
+	/**
+	 * Finds an interaction with a given name if one exists
+	 */
+	EDITORINTERACTIVETOOLSFRAMEWORK_API UViewportInteraction* FindInteraction(const FName& InName) const;
+
+	/**
+	 * Returns true if any interaction in the given group is enabled. An enabled interaction can potentially take action.
+	 */
+	EDITORINTERACTIVETOOLSFRAMEWORK_API bool IsAnyInteractionInGroupEnabled(const FName& InName) const;
+	
+	/**
+	 * Returns true if any interaction in the given group is active. An active interaction is currently taking action.
+	 */
+	EDITORINTERACTIVETOOLSFRAMEWORK_API bool IsAnyInteractionInGroupActive(const FName& InName) const;
+
 	/**
 	 * Removes the specified interaction from the active ones.
 	 */
@@ -257,17 +281,6 @@ public:
 
 	FEditorViewportClient* GetHoveredEditorViewportClient() const;
 
-	/**
-	 * Marks the specified interaction type as active (meaning a viewport interaction of that type is currently active)
-	 * Internally keeps count of how many interactions of this type are active
-	 */
-	void SetViewportInteractionActive(FName InViewportInteraction, bool bInActive);
-
-	/**
-	 * Checks if there is any active viewport interaction of the specified type (at least one)
-	 */
-	bool EDITORINTERACTIVETOOLSFRAMEWORK_API IsViewportInteractionActive(FName InViewportInteraction);
-
 	FOnMouseLookingChanged& OnMouseLookingStateChanged()
 	{
 		return OnMouseLookingChangedDelegate;
@@ -295,7 +308,7 @@ private:
 	UPROPERTY(Transient)
 	TObjectPtr<UInputBehaviorSet> BehaviorSet;
 	
-	/** Behaviors that are garuanteed to always be present */
+	/** Behaviors that are guaranteed to always be present */
 	UPROPERTY(Transient)
 	TArray<TObjectPtr<UInputBehavior>> CoreBehaviors;
 
@@ -306,6 +319,12 @@ private:
 	UPROPERTY(Transient)
 	TMap<FName, TObjectPtr<UViewportInteraction>> ViewportInteractions;
 
+	/**
+	 * Groups <-> Interactions is a many-to-many relationship. This provides easy access to a list of interactions by group.
+	 */
+	UPROPERTY(Transient)
+	TMap<FName, FViewportInteractionGroup> ViewportInteractionGroups;
+	
 	void SetAltKeyState(bool bInDown);
 	void SetCtrlKeyState(bool bInDown);
 	void SetShiftKeyState(bool bInDown);
@@ -331,12 +350,6 @@ private:
 	UPROPERTY(Transient)
 	TWeakObjectPtr<const UEditorInteractiveToolsContext> EditorInteractiveToolsContextWeak;
 
-	/**
-	 * Keeping track of active Viewport Interactions, so that FEditorViewportClient can disable legacy behaviors
-	 * accordingly e.g. if viewport interactions for zooming are enabled, skip the legacy ones doing the same
-	 */
-	TMap<FName, bool> ViewportInteractionsStatusMap;
-
 	FOnMouseLookingChanged OnMouseLookingChangedDelegate;
 
 	FOnModifierKeyChanged OnAltKeyStateChangedDelegate;
diff --git a/Engine/Source/Editor/LevelEditor/Private/ViewportInteractions/LevelViewportClickSelection.cpp b/Engine/Source/Editor/LevelEditor/Private/ViewportInteractions/LevelViewportClickSelection.cpp
--- a/Engine/Source/Editor/LevelEditor/Private/ViewportInteractions/LevelViewportClickSelection.cpp
+++ b/Engine/Source/Editor/LevelEditor/Private/ViewportInteractions/LevelViewportClickSelection.cpp
@@ -94,8 +94,8 @@ void ProcessClick(FEditorViewportClient* InEditorViewportClient, const FViewport
 
 ULevelViewportClickSelection::ULevelViewportClickSelection()
 {
-	ToolType = UE::Editor::ViewportInteractions::SingleClick;
 	InteractionName = TEXT("Level Viewport Click Selection");
+	Groups = { UE::Editor::ViewportInteractions::ViewportClick };
 }
 
 void ULevelViewportClickSelection::ProcessClick_Internal(const FViewportClick& InViewportClick, HHitProxy* InHitProxy, FSceneView& View)
@@ -106,8 +106,8 @@ void ULevelViewportClickSelection::ProcessClick_Internal(const FViewportClick& I
 
 ULevelViewportClickContextMenu::ULevelViewportClickContextMenu()
 {
-	ToolType = UE::Editor::ViewportInteractions::ContextMenu;
 	InteractionName = TEXT("Level Viewport Context Menu");
+	Groups = { UE::Editor::ViewportInteractions::ViewportClick };
 
 	MouseButton = EKeys::RightMouseButton;
 
@@ -123,8 +123,8 @@ ULevelViewportDoubleClick::ULevelViewportDoubleClick()
 	DoubleClickInputBehavior->Initialize(this);
 	DoubleClickInputBehavior->SetDefaultPriority(FInputCapturePriority::DEFAULT_GIZMO_PRIORITY - 1);
 
-	ToolType = UE::Editor::ViewportInteractions::DoubleClick;
 	InteractionName = TEXT("Level Viewport Double Click");
+	Groups = { UE::Editor::ViewportInteractions::ViewportClick };
 
 	RegisterInputBehavior(DoubleClickInputBehavior);
 }
diff --git a/Engine/Source/Editor/UnrealEd/Private/EditorViewportClient.cpp b/Engine/Source/Editor/UnrealEd/Private/EditorViewportClient.cpp
--- a/Engine/Source/Editor/UnrealEd/Private/EditorViewportClient.cpp
+++ b/Engine/Source/Editor/UnrealEd/Private/EditorViewportClient.cpp
@@ -155,16 +155,53 @@ bool UsesViewportInteractions(const FEditorModeTools* InModeTools)
 }
 
 /**
- * Returns true if any viewport interaction of the specified type is active
+ * Returns true if a viewport interaction with the specified name is enabled and able to act 
+ */
+bool IsViewportInteractionEnabled(const FEditorModeTools* InModeTools, FName InViewportInteraction)
+{
+	if (InModeTools)
+	{
+		if (UViewportInteractionsBehaviorSource* InputToolsBehaviorSource = InModeTools->GetInteractiveToolsContext()->GetViewportInteractionsBehaviorSource())
+		{
+			if (UViewportInteraction* Interaction = InputToolsBehaviorSource->FindInteraction(InViewportInteraction))
+			{
+				return Interaction->IsEnabled();
+			}
+		}
+	}
+
+	return false;
+}
+
+/**
+ * Returns true if a viewport interaction with the specified name is active and currently acting 
  */
 bool IsViewportInteractionActive(const FEditorModeTools* InModeTools, FName InViewportInteraction)
 {
 	if (InModeTools)
 	{
-		if (UViewportInteractionsBehaviorSource* InputToolsBehaviorSource =
-				InModeTools->GetInteractiveToolsContext()->GetViewportInteractionsBehaviorSource())
+		if (UViewportInteractionsBehaviorSource* InputToolsBehaviorSource = InModeTools->GetInteractiveToolsContext()->GetViewportInteractionsBehaviorSource())
+		{
+			if (UViewportInteraction* Interaction = InputToolsBehaviorSource->FindInteraction(InViewportInteraction))
+			{
+				return Interaction->IsActive();
+			}
+		}
+	}
+
+	return false;
+}
+
+/**
+ * Returns true if any viewport interaction in the specified group is enabled and able to act. 
+ */
+bool IsViewportInteractionGroupEnabled(const FEditorModeTools* InModeTools, FName InViewportInteractionGroup)
+{
+	if (InModeTools)
+	{
+		if (UViewportInteractionsBehaviorSource* InputToolsBehaviorSource = InModeTools->GetInteractiveToolsContext()->GetViewportInteractionsBehaviorSource())
 		{
-			return InputToolsBehaviorSource->IsViewportInteractionActive(InViewportInteraction);
+			return InputToolsBehaviorSource->IsAnyInteractionInGroupEnabled(InViewportInteractionGroup);\
 		}
 	}
 
@@ -2125,7 +2162,7 @@ void FEditorViewportClient::UpdateCameraMovement( float DeltaTime )
 	const bool bIsITFInteractions = UsesViewportInteractions(ModeTools.Get());
 	
 	// We only want to move perspective cameras around like this
-	if( Viewport != nullptr && IsPerspective() && (bIsITFInteractions ? (ViewportNavigationHelper && ViewportNavigationHelper->OrbitDelta.IsNearlyZero()) : !ShouldOrbitCamera()))
+	if( Viewport && IsPerspective() && (bIsITFInteractions ? !IsViewportInteractionActive(ModeTools.Get(), UE::Editor::ViewportInteractions::Orbit) : !ShouldOrbitCamera()))
 	{
 		const bool bEnable = false;
 		ToggleOrbitCamera(bEnable);
@@ -2463,6 +2500,7 @@ bool FEditorViewportClient::ConsumeExternalCameraMovement()
 		
 		if (!Translation.IsNearlyZero() || !Rotation.IsNearlyZero())
 		{
+			ToggleOrbitCamera(false);
 			MoveViewportPerspectiveCamera(Translation * GetCameraSpeed(), Rotation);
 		}
 	}
@@ -2575,8 +2613,6 @@ void FEditorViewportClient::OrbitPerspectiveCamera(const FRotator& InOrbit)
 		 * Recalculates the view location according to the new SetViewRotation() did earlier.
 		 */
 		SetViewLocation(ViewTransform.ComputeOrbitMatrix().Inverse().GetOrigin());
-
-		ToggleOrbitCamera(false);
 	}
 	
 	PerspectiveCameraMoved();
@@ -2606,8 +2642,6 @@ void FEditorViewportClient::DollyPerspectiveCamera(double InDelta)
 	FVector TransformedDelta = RotMat.InverseFast().TransformVector(DeltaLocation);
 
 	SetViewLocation( OrbitMatrix.GetOrigin() + TransformedDelta );
-
-	ToggleOrbitCamera(false);
 }
 
 /**
@@ -2929,7 +2963,7 @@ void FEditorViewportClient::UpdateMouseDelta()
 void FEditorViewportClient::PeformDefaultCameraMovement(FVector& Drag, FRotator& Rot, FVector& Scale)
 {
 	// Check if there is an active ITF Orbit interaction, and if there's an orbit delta it might have generated 
-	const bool bUseITFOrbitDelta = IsViewportInteractionActive(ModeTools.Get(), UE::Editor::ViewportInteractions::Orbit)
+	const bool bUseITFOrbitDelta = IsViewportInteractionEnabled(ModeTools.Get(), UE::Editor::ViewportInteractions::Orbit)
 					  && !ViewportNavigationHelper->OrbitDelta.IsNearlyZero();
 
 	FVector DragDelta;
@@ -3720,10 +3754,10 @@ bool FEditorViewportClient::Internal_InputKey(const FInputKeyEventArgs& EventArg
 		bHasMouseMovedSinceClick = false;
 	}
 
-	const bool bIsMultiMouseButtonInteractionActive = IsViewportInteractionActive(ModeTools.Get(), UE::Editor::ViewportInteractions::Zoom) || IsViewportInteractionActive(ModeTools.Get(), UE::Editor::ViewportInteractions::CameraDrag);
+	const bool bIsMultiMouseButtonInteractionEnabled = IsViewportInteractionEnabled(ModeTools.Get(), UE::Editor::ViewportInteractions::CameraDrag);
 
 	// Make sure ITF viewport interactions dealing with multiple mouse buttons pressed at the same time DO NOT trigger StartTrackingDueToInput when a button is released
-	if (!bIsMultiMouseButtonInteractionActive)
+	if (!bIsMultiMouseButtonInteractionEnabled)
 	{
 		// Start tracking if any mouse button is down and it was a tracking event (MouseButton/Ctrl/Shift/Alt):
 		if ( InputState.IsAnyMouseButtonDown()
@@ -3735,10 +3769,10 @@ bool FEditorViewportClient::Internal_InputKey(const FInputKeyEventArgs& EventArg
 		}
 	}
 
-	const bool bIsITFSingleClickActive = IsViewportInteractionActive(ModeTools.Get(), UE::Editor::ViewportInteractions::SingleClick);
+	const bool bIsITFSingleClickEnabled = IsViewportInteractionGroupEnabled(ModeTools.Get(), UE::Editor::ViewportInteractions::ViewportClick);
 
 	// If we are tracking and no mouse button is down and this input event released the mouse button stop tracking and process any clicks if necessary
-	if ( bIsTracking && !InputState.IsAnyMouseButtonDown() && InputState.IsMouseButtonEvent() && !bIsITFSingleClickActive)
+	if ( bIsTracking && !InputState.IsAnyMouseButtonDown() && InputState.IsMouseButtonEvent() && !bIsITFSingleClickEnabled)
 	{
 		// Handle possible mouse click viewport
 		ProcessClickInViewport( InputState, *View );
