Fix crash when an ensure notification popup would try to be created during garbage collection

#Synced-CL 48387090

diff --git a/Engine/Source/Editor/UnrealEd/Private/UnrealEdMisc.cpp b/Engine/Source/Editor/UnrealEd/Private/UnrealEdMisc.cpp
--- a/Engine/Source/Editor/UnrealEd/Private/UnrealEdMisc.cpp
+++ b/Engine/Source/Editor/UnrealEd/Private/UnrealEdMisc.cpp
@@ -1386,13 +1386,6 @@ void FUnrealEdMisc::OnEnsureFailed(const ANSICHAR* Expr, const ANSICHAR* File, i
 		return;
 	}
 
-	auto EnsureNotify = [](const TCHAR* ErrorMsg, const TCHAR* NotifyMsg)
-		{
-			FMessageLog EnsureFailedLogger = FMessageLog("EnsureFailed");
-			EnsureFailedLogger.Error(FText::Format(NSLOCTEXT("Core", "EnsureMessage", "{0}"), FTextStringHelper::CreateFromBuffer(ErrorMsg)));
-			EnsureFailedLogger.Notify(FText::Format(NSLOCTEXT("Core", "EnsureFailed", "Ensure condition failed: {0}"), FTextStringHelper::CreateFromBuffer(NotifyMsg)));
-		};
-
 	FString NotifyMsg;
 	if (Msg != nullptr && FCString::Strlen(Msg) > 0)
 	{
@@ -1403,18 +1396,18 @@ void FUnrealEdMisc::OnEnsureFailed(const ANSICHAR* Expr, const ANSICHAR* File, i
 		NotifyMsg = Expr;
 	}
 
-	if (IsInGameThread())
-	{
-		EnsureNotify(CombinedErrorMsg, *NotifyMsg);
-	}
-	else
-	{
-		FString ErrorMsg(CombinedErrorMsg);
-		FSimpleDelegateGraphTask::CreateAndDispatchWhenReady
-		(
-			FSimpleDelegateGraphTask::FDelegate::CreateLambda([EnsureNotify, ErrorMsg, NotifyMsg]() { EnsureNotify(*ErrorMsg, *NotifyMsg); }), TStatId(), nullptr, ENamedThreads::GameThread
-		);
-	}
+	// Always delay, in case we're not in the gamethread or we're in a callstack that doesn't allow slate windows to be created
+	FSimpleDelegateGraphTask::CreateAndDispatchWhenReady(FSimpleDelegateGraphTask::FDelegate::CreateLambda([ErrorMsg = FString(CombinedErrorMsg), NotifyMsg = MoveTemp(NotifyMsg)]()
+		{
+			if (IsEngineExitRequested())
+			{
+				return;
+			}
+
+			FMessageLog EnsureFailedLogger = FMessageLog("EnsureFailed");
+			EnsureFailedLogger.Error(FTextStringHelper::CreateFromBuffer(*ErrorMsg));
+			EnsureFailedLogger.Notify(FText::Format(NSLOCTEXT("Core", "EnsureFailed", "Ensure condition failed: {0}"), FTextStringHelper::CreateFromBuffer(*NotifyMsg)));
+		}), TStatId(), nullptr, ENamedThreads::GameThread);
 }
 
 void FUnrealEdMisc::OnActiveTabChanged(TSharedPtr<SDockTab> PreviouslyActive, TSharedPtr<SDockTab> NewlyActivated)
