UE: Bump incremental cook version to provide clarity about the source of an issue related to LOD resources

diff --git a/Engine/Source/Editor/UnrealEd/Private/Cooker/CookIncrementalVersion.cpp b/Engine/Source/Editor/UnrealEd/Private/Cooker/CookIncrementalVersion.cpp
--- a/Engine/Source/Editor/UnrealEd/Private/Cooker/CookIncrementalVersion.cpp
+++ b/Engine/Source/Editor/UnrealEd/Private/Cooker/CookIncrementalVersion.cpp
@@ -8,6 +8,6 @@ namespace UE::Cook
 {
 
 // Change CookIncrementalVersion to a new guid when all packages in an incremental cook need to be invalidated
-FGuid CookIncrementalVersion(0x1C46FCEC, 0xAE4B81A7, 0xC3871E56, 0x00EA52EB);
+FGuid CookIncrementalVersion(0x25399965, 0xBCA94C3B, 0xBAA048FA, 0x2EB2824A);
 
 }
