[Input] Add an editor preference which can be used to set a forced "preferred" input API.

This should have been included as a part of 48361114, but the file was missed

#rnx
#Synced-CL 48353899

diff --git a/Engine/Source/Editor/UnrealEd/Classes/Settings/LevelEditorPlaySettings.h b/Engine/Source/Editor/UnrealEd/Classes/Settings/LevelEditorPlaySettings.h
--- a/Engine/Source/Editor/UnrealEd/Classes/Settings/LevelEditorPlaySettings.h
+++ b/Engine/Source/Editor/UnrealEd/Classes/Settings/LevelEditorPlaySettings.h
@@ -245,6 +245,9 @@ protected:
 	UPROPERTY(config, EditAnywhere, Category=PlayInEditor, meta=(EditCondition=bSimulateDeviceMappingPolicy))
 	EInputDeviceMappingPolicy SimulatedDeviceMappingPolicy = EInputDeviceMappingPolicy::UseManagedPlatformLogin;
 
+	UPROPERTY(config, EditAnywhere, Category = PlayInEditor, meta = (ConfigRestartRequired = true))
+	FString PreferredInputAPIOverride;
+
 public:
 
 	bool IsSimulateDeviceMappingPolicy() const { return bSimulateDeviceMappingPolicy; }
