[HLOD] Make sure newly created world package is saved before running Selective HLOD Build

#Synced-CL 47987596

diff --git a/Engine/Source/Editor/UnrealEd/Private/WorldPartition/WorldPartitionHLODsBuilder.cpp b/Engine/Source/Editor/UnrealEd/Private/WorldPartition/WorldPartitionHLODsBuilder.cpp
--- a/Engine/Source/Editor/UnrealEd/Private/WorldPartition/WorldPartitionHLODsBuilder.cpp
+++ b/Engine/Source/Editor/UnrealEd/Private/WorldPartition/WorldPartitionHLODsBuilder.cpp
@@ -198,6 +198,13 @@ bool UWorldPartitionHLODsBuilder::BuildWithFilters(UWorld* InWorld, const TArray
 {
 	bool bResult = false;
 
+	// Make sure that a newly created world was saved
+	if (InWorld->GetPackage()->HasAnyPackageFlags(PKG_NewlyCreated))
+	{
+		FMessageDialog::Open(EAppMsgType::Ok, LOCTEXT("NewMap", "New world must be saved before performing this operation."));
+		return false;
+	}
+
 	bool bUnsavedActorsPassingFilters = false;
 	for (TActorIterator<AActor> It(InWorld); It; ++It)
 	{
