EditorUsability : UnrealEd
- Added behavior selection for tabs navigation action in Editor Settings / Misc, before this change editor only navigates between assets and uses last access time to determine order, now browser style tab navigation is available, it uses window tab order (default is browser style behavior) CTRL+Tab goes forward, CTRL+SHIFT+Tab goes backward
- Added an option to hide tab navigation popup dialog widget and directly perform the action (default is showing this dialog)
- Before this change navigation only worked on asset tabs, you can now navigate between any tabs (editor settings, project settings, logs, ...)
- Reduced size of items in tab navigation dialog widget
- Fixed Scrollbar not working and Listview not scrolling into view

#Synced-CL 48229788

diff --git a/Engine/Source/Editor/UnrealEd/Classes/Editor/EditorPerProjectUserSettings.h b/Engine/Source/Editor/UnrealEd/Classes/Editor/EditorPerProjectUserSettings.h
--- a/Engine/Source/Editor/UnrealEd/Classes/Editor/EditorPerProjectUserSettings.h
+++ b/Engine/Source/Editor/UnrealEd/Classes/Editor/EditorPerProjectUserSettings.h
@@ -2,13 +2,21 @@
 
 #pragma once
 
-#include "CoreMinimal.h"
-#include "UObject/ObjectMacros.h"
-#include "UObject/Object.h"
 #include "Engine/EngineTypes.h"
 #include "Settings/EditorViewportSettings.h"
+#include "UObject/ObjectMacros.h"
+#include "UObject/Object.h"
 #include "EditorPerProjectUserSettings.generated.h"
 
+UENUM()
+enum class EGlobalTabSwitchingBehavior : uint8
+{
+	/** Legacy behavior where last accessed asset defines order of switching */
+	LastAccessedAsset,
+	/** Same behavior as browser where window tab order defines order of switching */
+	BrowserStyle
+};
+
 UCLASS(minimalapi, autoexpandcategories=(ViewportControls, ViewportLookAndFeel, LevelEditing, SourceControl, Content, Startup), hidecategories=(Object, Options, Grid, RotationGrid), config=EditorPerProjectUserSettings)
 class UEditorPerProjectUserSettings : public UObject
 {
@@ -100,6 +108,14 @@ class UEditorPerProjectUserSettings : public UObject
 	UPROPERTY(EditAnywhere, Config, Category = Behavior)
 	uint32 bConfirmEditorClose: 1;
 
+	/** Defines global tab switching behavior for the editor */
+	UPROPERTY(EditAnywhere, Config, Category = Behavior)
+	EGlobalTabSwitchingBehavior TabSwitchingBehavior = EGlobalTabSwitchingBehavior::BrowserStyle;
+
+	/** Whether to show or hide tab switching dialog when tab command is executed */
+	UPROPERTY(EditAnywhere, Config, Category = Behavior)
+	bool bShowTabSwitchingDialog = true;
+
 	/** Select to make Distributions use the curves, not the baked lookup tables. */
 	UPROPERTY(config)
 	uint32 bUseCurvesForDistributions:1; //(GDistributionType == 0)
diff --git a/Engine/Source/Editor/UnrealEd/Private/Toolkits/SGlobalTabSwitchingDialog.cpp b/Engine/Source/Editor/UnrealEd/Private/Toolkits/SGlobalTabSwitchingDialog.cpp
--- a/Engine/Source/Editor/UnrealEd/Private/Toolkits/SGlobalTabSwitchingDialog.cpp
+++ b/Engine/Source/Editor/UnrealEd/Private/Toolkits/SGlobalTabSwitchingDialog.cpp
@@ -1,31 +1,30 @@
 // Copyright Epic Games, Inc. All Rights Reserved.
 
 #include "Toolkits/SGlobalTabSwitchingDialog.h"
-#include "Modules/ModuleManager.h"
+
+#include "Editor.h"
+#include "EngineGlobals.h"
+#include "Editor/EditorPerProjectUserSettings.h"
 #include "Framework/Application/SlateApplication.h"
-#include "Widgets/Images/SImage.h"
-#include "Widgets/Text/STextBlock.h"
+#include "Framework/Docking/TabManager.h"
 #include "Framework/MultiBox/MultiBox.h"
 #include "Framework/MultiBox/MultiBoxBuilder.h"
-#include "Widgets/Input/SButton.h"
-#include "Widgets/Layout/SScrollBorder.h"
-#include "Widgets/Views/SListView.h"
-#include "Framework/Docking/TabManager.h"
+#include "LevelEditor.h"
+#include "Modules/ModuleManager.h"
 #include "Styling/AppStyle.h"
+#include "Subsystems/AssetEditorSubsystem.h"
+#include "ThumbnailRendering/ThumbnailManager.h"
 #include "Widgets/Docking/SDockTab.h"
-
-#include "EngineGlobals.h"
-#include "Editor.h"
-
-#include "LevelEditor.h"
+#include "Widgets/Images/SImage.h"
+#include "Widgets/Input/SButton.h"
+#include "Widgets/Text/STextBlock.h"
+#include "Widgets/Views/SListView.h"
 #include "WorkspaceMenuStructure.h"
 #include "WorkspaceMenuStructureModule.h"
 
 #if PLATFORM_MAC
 #include "Mac/MacApplication.h"
 #endif
-#include "Subsystems/AssetEditorSubsystem.h"
-#include "ThumbnailRendering/ThumbnailManager.h"
 
 #define LOCTEXT_NAMESPACE "SGlobalTabSwitchingDialog"
 
@@ -96,7 +95,7 @@ public:
 		FText AssetText = FText::Format(LOCTEXT("AssetEntryLabel", "{AssetName}{DirtyState}"), Args);
 
 		// Create a thumbnail to represent the asset type
-		const int32 ThumbnailSize = 48;
+		constexpr int32 ThumbnailSize = 32;
 
 		FAssetData AssetData(MyAsset);
 		Thumbnail = MakeShareable(new FAssetThumbnail(MyAsset, ThumbnailSize, ThumbnailSize, AssetThumbnailPool));
@@ -107,7 +106,7 @@ public:
 			+SHorizontalBox::Slot()
 			.AutoWidth()
 			.VAlign(VAlign_Center)
-			.Padding(4.0f, 4.0f)
+			.Padding(5.f)
 			[
 				SNew(SBox)
 				.WidthOverride(static_cast<float>(ThumbnailSize))
@@ -119,7 +118,7 @@ public:
 			+SHorizontalBox::Slot()
 			.FillWidth(1.0f)
 			.VAlign(VAlign_Center)
-			.Padding(8.0f, 0.0f, 8.0f, 0.0f)
+			.Padding(5.f)
 			[
 				SNew(STextBlock)
 				.TextStyle(FAppStyle::Get(), "ControlTabMenu.AssetNameStyle")
@@ -146,7 +145,14 @@ public:
 
 	virtual void ActivateTab() override
 	{
-		GEditor->GetEditorSubsystem<UAssetEditorSubsystem>()->FindEditorForAsset(MyAsset, /*bFocusIfOpen=*/ true);
+		if (const TSharedPtr<FTabManager> AssetEditorTabManager = GetAssociatedTabManager())
+		{
+			if (const TSharedPtr<SDockTab> DockTab = AssetEditorTabManager->GetOwnerTab())
+			{
+				// Don't use DrawAttention since it steals focus
+				DockTab->ActivateInParent(ETabActivationCause::SetDirectly);
+			}
+		}
 	}
 
 	virtual TSharedPtr<FTabManager> GetAssociatedTabManager() override
@@ -193,8 +199,14 @@ public:
 
 	virtual void ActivateTab() override
 	{
-		const TSharedPtr<FTabManager> LevelEditorTabManager = GetAssociatedTabManager();
-		FGlobalTabmanager::Get()->DrawAttention(LevelEditorTabManager->GetOwnerTab().ToSharedRef());
+		if (const TSharedPtr<FTabManager> LevelEditorTabManager = GetAssociatedTabManager())
+		{
+			if (const TSharedPtr<SDockTab> DockTab = LevelEditorTabManager->GetOwnerTab())
+			{
+				// Don't use DrawAttention since it steals focus
+				DockTab->ActivateInParent(ETabActivationCause::SetDirectly);
+			}
+		}
 	}
 
 	virtual TSharedPtr<FTabManager> GetAssociatedTabManager() override
@@ -211,6 +223,81 @@ protected:
 	}
 };
 
+class FTabSwitchingListItem_DockTab : public FTabSwitchingListItemBase
+{
+public:
+	FTabSwitchingListItem_DockTab(const TSharedPtr<SDockTab>& InTab)
+		: DockTab(InTab)
+	{
+	}
+
+	virtual TSharedRef<SWidget> CreateWidget(TSharedPtr<class FAssetThumbnailPool> AssetThumbnailPool) override
+	{
+		if (!DockTab.IsValid())
+		{
+			return SNullWidget::NullWidget;
+		}
+
+		constexpr int32 ThumbnailSize = 32;
+
+		return SNew(SHorizontalBox)
+			+SHorizontalBox::Slot()
+			.AutoWidth()
+			.VAlign(VAlign_Center)
+			.Padding(5.f)
+			[
+				SNew(SBox)
+				.WidthOverride(static_cast<float>(ThumbnailSize))
+				.HeightOverride(static_cast<float>(ThumbnailSize))
+				[
+					SNew(SImage)
+					.Image(DockTab->GetTabIcon())
+				]
+			]
+			+SHorizontalBox::Slot()
+			.FillWidth(1.0f)
+			.VAlign(VAlign_Center)
+			.Padding(5.f)
+			[
+				SNew(STextBlock)
+				.TextStyle(FAppStyle::Get(), "ControlTabMenu.AssetNameStyle")
+				.Text(DockTab->GetTabLabel())
+			];
+	}
+
+	virtual void ShowInContentBrowser() override
+	{
+		// Do nothing here
+	}
+
+	virtual FText GetTypeString() const override
+	{
+		return DockTab.IsValid() ? DockTab->GetTabLabel() : FText::GetEmpty();
+	}
+
+	virtual FText GetPathString() const override
+	{
+		return FText::GetEmpty();
+	}
+
+	virtual void ActivateTab() override
+	{
+		if (DockTab.IsValid())
+		{
+			// Don't use DrawAttention since it steals focus
+			DockTab->ActivateInParent(ETabActivationCause::SetDirectly);
+		}
+	}
+
+	virtual TSharedPtr<FTabManager> GetAssociatedTabManager() override
+	{
+		return DockTab.IsValid() ? DockTab->GetTabManagerPtr() : nullptr;
+	}
+
+protected:
+	TSharedPtr<SDockTab> DockTab;
+};
+
 //////////////////////////////////////////////////////////////////////////
 // FTabSwitchingListItem_Asset
 
@@ -245,7 +332,8 @@ FReply SGlobalTabSwitchingDialog::OnBrowseToSelectedAsset()
 	{
 		SelectedItem->ShowInContentBrowser();
 
-		FSlateApplication::Get().DismissAllMenus();
+		constexpr bool ActivateTab = false;
+		DismissDialog(ActivateTab);
 	}
 	return FReply::Handled();
 }
@@ -260,8 +348,9 @@ void SGlobalTabSwitchingDialog::OnMainTabListSelectionChanged(FTabListItemPtr In
 
 	if (SelectedItems.Num() > 0)
 	{
-		FTabListItemPtr SelectedItem = SelectedItems[0];
+		const FTabListItemPtr SelectedItem = SelectedItems[0];
 
+		const FText ItemPathString = SelectedItem->GetPathString();
 		NewTopContents = SelectedItem->CreateWidget(UThumbnailManager::Get().GetSharedThumbnailPool());
 
 		NewBottomContents =
@@ -273,6 +362,7 @@ void SGlobalTabSwitchingDialog::OnMainTabListSelectionChanged(FTabListItemPtr In
 				.ButtonStyle(FAppStyle::Get(), "HoverOnlyHyperlinkButton")
 				.VAlign(VAlign_Center)
 				.HAlign(HAlign_Center)
+				.Visibility(ItemPathString.IsEmpty() ? EVisibility::Hidden : EVisibility::Visible)
 				.OnClicked(this, &SGlobalTabSwitchingDialog::OnBrowseToSelectedAsset)
 				.ToolTipText(LOCTEXT("BrowseButtonToolTipText", "Browse to Asset in Content Browser"))
 				[
@@ -287,11 +377,10 @@ void SGlobalTabSwitchingDialog::OnMainTabListSelectionChanged(FTabListItemPtr In
 					+SHorizontalBox::Slot()
 					.AutoWidth()
 					.VAlign(VAlign_Center)
-					.Padding(FMargin(4.0f, 0.0f, 0.0f, 0.0f))
 					[
 						SNew(STextBlock)
 						.TextStyle(FAppStyle::Get(), "ControlTabMenu.AssetPathStyle")
-						.Text(SelectedItem->GetPathString())
+						.Text(ItemPathString)
 					]
 				]
 			]
@@ -316,7 +405,7 @@ void SGlobalTabSwitchingDialog::OnMainTabListSelectionChanged(FTabListItemPtr In
 
 void SGlobalTabSwitchingDialog::OnMainTabListItemClicked(FTabListItemPtr InItem)
 {
-	DismissDialog();
+	DismissDialog(InItem.IsValid());
 }
 
 void SGlobalTabSwitchingDialog::CycleSelection(bool bForwards)
@@ -341,17 +430,23 @@ void SGlobalTabSwitchingDialog::CycleSelection(bool bForwards)
 				FTabListItemPtr NewSelectedItem = MainTabsListDataSource[NewSelectionIndex];
 				MainTabsListWidget->SetSelection(NewSelectedItem);
 				MainTabsListWidget->RequestScrollIntoView(NewSelectedItem);
+
+				// Activate tab without giving focus, this will allow input events to still flow through this dialog widget
+				NewSelectedItem->ActivateTab();
 			}
 		}
 	}
 }
 
-void SGlobalTabSwitchingDialog::DismissDialog()
+void SGlobalTabSwitchingDialog::DismissDialog(bool bInActivateTab)
 {
-	FTabListItemPtr SelectedItem = GetMainTabListSelectedItem();
-	if (SelectedItem.IsValid())
+	if (bInActivateTab)
 	{
-		SelectedItem->ActivateTab();
+		const FTabListItemPtr SelectedItem = GetMainTabListSelectedItem();
+		if (SelectedItem.IsValid())
+		{
+			SelectedItem->ActivateTab();
+		}
 	}
 
 	FSlateApplication::Get().DismissAllMenus();
@@ -369,36 +464,72 @@ void SGlobalTabSwitchingDialog::Construct(const FArguments& InArgs, FVector2D In
 
 	TriggerChord = InTriggerChord;
 
-	// Populate the list with open asset editors
-	TArray<UObject*> OpenAssetList = GEditor->GetEditorSubsystem<UAssetEditorSubsystem>()->GetAllEditedAssets();
-	for (UObject* OpenAsset : OpenAssetList)
+	const UEditorPerProjectUserSettings* EditorUserSettings = GetDefault<UEditorPerProjectUserSettings>();
+
+	ensureMsgf(EditorUserSettings != nullptr, TEXT("Editor Per Project User Settings cannot be retrieved"));
+
+	if (!EditorUserSettings || EditorUserSettings->TabSwitchingBehavior == EGlobalTabSwitchingBehavior::LastAccessedAsset)
 	{
-		if (OpenAsset->GetOuter() != GetTransientPackage())
+		// Populate the list with open asset editors
+		TArray<UObject*> OpenAssetList = GEditor->GetEditorSubsystem<UAssetEditorSubsystem>()->GetAllEditedAssets();
+		for (UObject* OpenAsset : OpenAssetList)
 		{
-			MainTabsListDataSource.Add(MakeShareable(new FTabSwitchingListItem_Asset(OpenAsset)));
+			if (OpenAsset->GetOuter() != GetTransientPackage())
+			{
+				MainTabsListDataSource.Add(MakeShareable(new FTabSwitchingListItem_Asset(OpenAsset)));
+			}
 		}
+
+		MainTabsListDataSource.Add(FTabSwitchingListItem_World::MakeWorldItem());
+
+		// Sort the list by access time
+		MainTabsListDataSource.Sort([](const FTabListItemPtr& A, const FTabListItemPtr& B) { return (A->LastAccessTime > B->LastAccessTime); });
 	}
+	else
+	{
+		const TSharedPtr<SWindow> Window = FSlateApplication::Get().GetActiveTopLevelWindow();
 
-	MainTabsListDataSource.Add(FTabSwitchingListItem_World::MakeWorldItem());
+		if (ensureMsgf(Window.IsValid(), TEXT("Active top level window cannot be invalid")))
+		{
+			const TSharedRef<FGlobalTabmanager> TabManager = FGlobalTabmanager::Get();
+			const FGlobalTabmanager::FWindowMajorTabsInfo TabsInfo = TabManager->GetDockableWindowMajorTabsInfo(Window.ToSharedRef());
 
-	// Sort the list by access time
-	MainTabsListDataSource.Sort([](const FTabListItemPtr& A, const FTabListItemPtr& B) { return (A->LastAccessTime > B->LastAccessTime); });
+			if (TabsInfo.ActiveTab.IsValid() && !TabsInfo.DockTabs.IsEmpty())
+			{
+				// In this mode, active tab is the first to be added
+				const int32 IndexActiveTab = TabsInfo.DockTabs.Find(TabsInfo.ActiveTab);
+
+				// Add all right side tabs
+				for (int32 Index = IndexActiveTab; Index < TabsInfo.DockTabs.Num(); Index++)
+				{
+					MainTabsListDataSource.Add(MakeShared<FTabSwitchingListItem_DockTab>(TabsInfo.DockTabs[Index]));
+				}
+
+				// Add all left side tabs
+				for (int32 Index = 0; Index < IndexActiveTab; Index++)
+				{
+					MainTabsListDataSource.Add(MakeShared<FTabSwitchingListItem_DockTab>(TabsInfo.DockTabs[Index]));
+				}
+			}
+		}
+	}
 
 	// Create the widgets
 	NewTabItemToActivateDisplayBox =
 		SNew(SBox)
-		.Padding(FMargin(0.0f, 0.0f, 10.0f, 0.0f))
-		.HeightOverride(70.0f)
+		.Padding(5.f)
+		.HeightOverride(40.0f)
 		.VAlign(VAlign_Top);
 
 	NewTabItemToActivatePathBox =
 		SNew(SBox)
-		.Padding(FMargin(0.0f, 10.0f, 10.0f, 10.0f))
+		.Padding(5.f)
 		.HeightOverride(40.0f)
 		.VAlign(VAlign_Center);
 
 	MainTabsListWidget = SNew(STabListWidget)
 		.ListItemsSource(&MainTabsListDataSource)
+		.ConsumeMouseWheel(EConsumeMouseWheel::WhenScrollingPossible)
 		.OnGenerateRow(this, &SGlobalTabSwitchingDialog::OnGenerateTabSwitchListItemWidget)
 		.OnSelectionChanged(this, &SGlobalTabSwitchingDialog::OnMainTabListSelectionChanged)
 		.OnMouseButtonClick(this, &SGlobalTabSwitchingDialog::OnMainTabListItemClicked)
@@ -407,6 +538,7 @@ void SGlobalTabSwitchingDialog::Construct(const FArguments& InArgs, FVector2D In
 	TSharedRef<SWidget> DocumentTabList = SNew(SVerticalBox)
 		+SVerticalBox::Slot()
 		.AutoHeight()
+		.Padding(5.f)
 		[
 			SNew(STextBlock)
 			.TextStyle(FAppStyle::Get(), "ControlTabMenu.HeadingStyle")
@@ -414,37 +546,48 @@ void SGlobalTabSwitchingDialog::Construct(const FArguments& InArgs, FVector2D In
 		]
 		+SVerticalBox::Slot()
 		.FillHeight(1.0f)
+		.Padding(5.f)
 		[
-			SNew(SScrollBorder, MainTabsListWidget.ToSharedRef())
-			[
-				MainTabsListWidget.ToSharedRef()
-			]
+			MainTabsListWidget.ToSharedRef()
 		];
 
+	// Show or collapse children, this will still allow input events to flow through this dialog widget
+	const EVisibility ChildVisibility = (!EditorUserSettings || EditorUserSettings->bShowTabSwitchingDialog)
+		? EVisibility::Visible
+		: EVisibility::Collapsed;
+
 	ChildSlot
 	[
 		SNew(SBorder)
+		.Padding(5.f)
+		.Visibility(ChildVisibility)
 		.BorderImage(FAppStyle::Get().GetBrush("ControlTabMenu.Background"))
 		.ForegroundColor(FCoreStyle::Get().GetSlateColor("DefaultForeground"))
 		[
 			SNew(SBox)
 			.WidthOverride(static_cast<float>(InSize.X))
 			.HeightOverride(static_cast<float>(InSize.Y))
-			.Padding(FMargin(12.0f, 12.0f, 12.0f, 0.0f))
+			.Padding(5.f)
 			[
 				SNew(SVerticalBox)
 				+SVerticalBox::Slot()
+				.Padding(5.f)
 				.AutoHeight()
 				[
-					NewTabItemToActivateDisplayBox.ToSharedRef()
+					SNew(SBox)
+					.HeightOverride(50.f)
+					[
+						NewTabItemToActivateDisplayBox.ToSharedRef()
+					]
 				]
 				+SVerticalBox::Slot()
-				.Padding(12.f, 24.f)
+				.Padding(5.f)
 				.FillHeight(1.0f)
 				[
-						DocumentTabList
+					DocumentTabList
 				]
 				+SVerticalBox::Slot()
+				.Padding(5.f)
 				.AutoHeight()
 				[
 					NewTabItemToActivatePathBox.ToSharedRef()
@@ -469,19 +612,62 @@ FReply SGlobalTabSwitchingDialog::OnKeyUp(const FGeometry& MyGeometry, const FKe
 	
 	if (bCloseViaControl || bCloseViaCommand || bCloseViaAlt)
 	{
-		DismissDialog();
+		constexpr bool bActivateTab = true;
+		DismissDialog(bActivateTab);
 		return FReply::Handled();
 	}
 
+	// Reset processing delay in case user quickly re-presses keys
+	if (TriggerChord.Key == InKeyEvent.GetKey())
+	{
+		LastProcessedTime = 0.0;
+	}
+
 	return FReply::Unhandled();
 }
 
+FReply SGlobalTabSwitchingDialog::OnMouseWheel(const FGeometry& InGeometry, const FPointerEvent& InMouseEvent)
+{
+	if (MainTabsListWidget)
+	{
+		// This dialog only appears when CTRL+Tab is pressed and CTRL is held
+		// But listview only consumes mouse wheel event if CTRL is not pressed...
+		// So fake it until you make it
+		FPointerEvent ModifiedEvent(
+			InMouseEvent.GetUserIndex(),
+			InMouseEvent.GetPointerIndex(),
+			InMouseEvent.GetScreenSpacePosition(),
+			InMouseEvent.GetLastScreenSpacePosition(),
+			InMouseEvent.GetPressedButtons(),
+			FKey(),
+			InMouseEvent.GetWheelDelta(),
+			FModifierKeysState()
+		);
+
+		FReply Reply = MainTabsListWidget->OnMouseWheel(InGeometry, ModifiedEvent);
+
+		if (Reply.IsEventHandled())
+		{
+			return Reply;
+		}
+	}
+
+	return SCompoundWidget::OnMouseWheel(InGeometry, InMouseEvent);
+}
+
 FReply SGlobalTabSwitchingDialog::OnKeyDown(const FGeometry& MyGeometry, const FKeyEvent& InKeyEvent)
 {
 	if (TriggerChord.Key == InKeyEvent.GetKey())
 	{
-		const bool bCycleForward = !InKeyEvent.IsShiftDown();
-		CycleSelection(bCycleForward);
+		// Don't switch tab too quickly when holding keys down
+		constexpr double ProcessingDelay = 0.3;
+		const double CurrentTime = FPlatformTime::Seconds();
+		if (CurrentTime - LastProcessedTime > ProcessingDelay)
+		{
+			const bool bCycleForward = !InKeyEvent.IsShiftDown();
+			CycleSelection(bCycleForward);
+			LastProcessedTime = CurrentTime;
+		}
 	}
 
 	return FReply::Unhandled();
@@ -491,7 +677,8 @@ FReply SGlobalTabSwitchingDialog::OnPreviewKeyDown(const FGeometry& MyGeometry,
 {
 	if (InKeyEvent.GetKey() == EKeys::Escape)
 	{
-		FSlateApplication::Get().DismissAllMenus();
+		constexpr bool bActivateTab = false;
+		DismissDialog(bActivateTab);
 		return FReply::Handled();
 	}
 
diff --git a/Engine/Source/Editor/UnrealEd/Private/Toolkits/SGlobalTabSwitchingDialog.h b/Engine/Source/Editor/UnrealEd/Private/Toolkits/SGlobalTabSwitchingDialog.h
--- a/Engine/Source/Editor/UnrealEd/Private/Toolkits/SGlobalTabSwitchingDialog.h
+++ b/Engine/Source/Editor/UnrealEd/Private/Toolkits/SGlobalTabSwitchingDialog.h
@@ -2,14 +2,16 @@
 
 #pragma once
 
-#include "CoreMinimal.h"
-#include "Input/Reply.h"
+#include "AssetThumbnail.h"
 #include "Framework/Commands/InputChord.h"
+#include "InputCoreTypes.h"
+#include "Input/Events.h"
+#include "Input/Reply.h"
+#include "Templates/SharedPointer.h"
 #include "Widgets/DeclarativeSyntaxSupport.h"
 #include "Widgets/SCompoundWidget.h"
-#include "AssetThumbnail.h"
-#include "Widgets/Views/STableViewBase.h"
 #include "Widgets/Views/STableRow.h"
+#include "Widgets/Views/STableViewBase.h"
 
 class FTabSwitchingListItemBase;
 template <typename ItemType> class SListView;
@@ -23,17 +25,17 @@ public:
 	SLATE_BEGIN_ARGS(SGlobalTabSwitchingDialog){}
 	SLATE_END_ARGS()
 
-public:
 	void Construct(const FArguments& InArgs, FVector2D InSize, FInputChord InTriggerChord);
 
-	~SGlobalTabSwitchingDialog();
+	virtual ~SGlobalTabSwitchingDialog() override;
 
-	// SWidget interface
+	//~ Begin SWidget
 	virtual FReply OnPreviewKeyDown(const FGeometry& MyGeometry, const FKeyEvent& InKeyEvent) override;
 	virtual FReply OnKeyDown(const FGeometry& MyGeometry, const FKeyEvent& InKeyEvent) override;
 	virtual FReply OnKeyUp(const FGeometry& MyGeometry, const FKeyEvent& InKeyEvent) override;
 	virtual bool SupportsKeyboardFocus() const override { return true; }
-	// End of SWidget interface
+	virtual FReply OnMouseWheel(const FGeometry& InGeometry, const FPointerEvent& InMouseEvent) override;
+	//~ End SWidget
 
 	// Is an instance already open?
 	static bool IsAlreadyOpen()
@@ -45,19 +47,17 @@ private:
 	typedef TSharedPtr<class FTabSwitchingListItemBase> FTabListItemPtr;
 	typedef SListView<FTabListItemPtr> STabListWidget;
 
-private:
 	TSharedRef<ITableRow> OnGenerateTabSwitchListItemWidget(FTabListItemPtr InItem, const TSharedRef<STableViewBase>& OwnerTable);
 
 	void CycleSelection(bool bForwards);
 	void OnMainTabListSelectionChanged(FTabListItemPtr InItem, ESelectInfo::Type SelectInfo);
 	void OnMainTabListItemClicked(FTabListItemPtr InItem);
-	void DismissDialog();
+	void DismissDialog(bool bInActivateTab);
 	FReply OnBrowseToSelectedAsset();
 
 	FTabListItemPtr GetMainTabListSelectedItem() const;
 
-private:
-	// The chord that triggered the dialog (so we can handle the correct Tab/`/etc... key repeat, and dismiss on the correct control/command modifier release
+	/** The chord that triggered the dialog (so we can handle the correct Tab/`/etc... key repeat, and dismiss on the correct control/command modifier release */
 	FInputChord TriggerChord;
 
 	// The array of 'document' items
@@ -72,5 +72,8 @@ private:
 	// The container widget for the indication of the path to the asset that will be activated when the dialog closes
 	TSharedPtr<SBox> NewTabItemToActivatePathBox;
 
+	/** Used to delay same key repeat */
+	double LastProcessedTime = 0.f;
+
 	static bool bIsAlreadyOpen;
 };
diff --git a/Engine/Source/Runtime/Slate/Private/Framework/Docking/TabManager.cpp b/Engine/Source/Runtime/Slate/Private/Framework/Docking/TabManager.cpp
--- a/Engine/Source/Runtime/Slate/Private/Framework/Docking/TabManager.cpp
+++ b/Engine/Source/Runtime/Slate/Private/Framework/Docking/TabManager.cpp
@@ -3323,6 +3323,54 @@ void FGlobalTabmanager::SetShouldUseMiddleEllipsisForDockTabLabel(const bool bIn
 	bShouldUseMiddleEllipsisForDockTabLabel = bInShouldUseMiddleEllipsis;
 }
 
+TArray<TSharedPtr<SWindow>> FGlobalTabmanager::GetDockableWindows() const
+{
+	TArray<TSharedPtr<SWindow>> Windows;
+	Windows.Reserve(DockAreas.Num());
+	for (const TWeakPtr<SDockingArea>& DockingAreaWeak : DockAreas)
+	{
+		if (const TSharedPtr<SDockingArea> DockArea = DockingAreaWeak.Pin())
+		{
+			if (TSharedPtr<SWindow> Window = DockArea->GetParentWindow())
+			{
+				Windows.Add(Window);
+			}
+		}
+	}
+	return Windows;
+}
+
+FGlobalTabmanager::FWindowMajorTabsInfo FGlobalTabmanager::GetDockableWindowMajorTabsInfo(const TSharedRef<SWindow>& InWindow) const
+{
+	FWindowMajorTabsInfo TabsInfo;
+	TabsInfo.Window = InWindow;
+	if (const TSharedPtr<SDockingArea> DockingArea = GetDockingAreaForWindow(InWindow))
+	{
+		TArray<TSharedRef<SDockingTabStack>> TabStacks;
+		GetAllStacks(DockingArea.ToSharedRef(), TabStacks);
+
+		for (const TSharedRef<SDockingTabStack>& TabStack : TabStacks)
+		{
+			TSlotlessChildren<SDockTab> Tabs = TabStack->GetTabs();
+			for (int32 Index = 0; Index < Tabs.Num(); Index++)
+			{
+				TSharedRef<SDockTab> Tab = StaticCastSharedRef<SDockTab>(Tabs.GetChildAt(Index));
+				TabsInfo.DockTabs.Add(Tab);
+
+				if (Tab->IsForeground())
+				{
+					TabsInfo.ActiveTab = Tab;
+				}
+			}
+		}
+	}
+	else
+	{
+		UE_LOG(LogSlate, Warning, TEXT("Could not find docking area for window %s"), *InWindow->GetTitle().ToString())
+	}
+	return TabsInfo;
+}
+
 FDelegateHandle FGlobalTabmanager::OnActiveTabChanged_Subscribe( const FOnActiveTabChanged::FDelegate& InDelegate )
 {
 	return OnActiveTabChanged.Add( InDelegate );
diff --git a/Engine/Source/Runtime/Slate/Public/Framework/Docking/TabManager.h b/Engine/Source/Runtime/Slate/Public/Framework/Docking/TabManager.h
--- a/Engine/Source/Runtime/Slate/Public/Framework/Docking/TabManager.h
+++ b/Engine/Source/Runtime/Slate/Public/Framework/Docking/TabManager.h
@@ -1591,6 +1591,31 @@ public:
 	/** Used to override dockable area restoration behavior */
 	FOnOverrideDockableAreaRestore OnOverrideDockableAreaRestore_Handler;
 
+	/**
+	 * Get all dockable windows managed by this tab manager
+	 * @return Array of all windows that contain dock areas
+	 */
+	SLATE_API TArray<TSharedPtr<SWindow>> GetDockableWindows() const;
+
+	struct FWindowMajorTabsInfo
+	{
+		/** The window these info belongs to */
+		TSharedPtr<SWindow> Window;
+
+		/** Ordered tabs docked in the window */
+		TArray<TSharedPtr<SDockTab>> DockTabs;
+
+		/** Active/foreground tab in the window */
+		TSharedPtr<SDockTab> ActiveTab;
+	};
+
+	/**
+	 * Get all major (top-level) tabs info for a specific window
+	 * @param InWindow The window to query info from
+	 * @return Window tabs info struct
+	 */
+	SLATE_API FWindowMajorTabsInfo GetDockableWindowMajorTabsInfo(const TSharedRef<SWindow>& InWindow) const;
+
 protected:
 	SLATE_API virtual void OnTabForegrounded( const TSharedPtr<SDockTab>& NewForegroundTab, const TSharedPtr<SDockTab>& BackgroundedTab ) override;
 	SLATE_API virtual void OnTabRelocated( const TSharedRef<SDockTab>& RelocatedTab, const TSharedPtr<SWindow>& NewOwnerWindow ) override;
