Changed PakOrder specs to be sorted by alphabetical so that they order reasonably with respect to non-numerical suffixes.

diff --git a/Engine/Source/Editor/UnrealEd/Private/Commandlets/AssetRegistryGenerator.cpp b/Engine/Source/Editor/UnrealEd/Private/Commandlets/AssetRegistryGenerator.cpp
--- a/Engine/Source/Editor/UnrealEd/Private/Commandlets/AssetRegistryGenerator.cpp
+++ b/Engine/Source/Editor/UnrealEd/Private/Commandlets/AssetRegistryGenerator.cpp
@@ -567,16 +567,7 @@ bool FAssetRegistryGenerator::GenerateStreamingInstallManifest(int64 InOverrideC
 				}
 			}
 			
-			Algo::SortBy(FoundFiles, [](const FString& Filename) {
-				FString Number;
-				int32 Order = 0;
-				if (Filename.Split(TEXT("_"), nullptr, &Number, ESearchCase::IgnoreCase, ESearchDir::FromEnd))
-				{
-					Order = FCString::Atoi(*Number);
-				}
-				return Order;
-			});
-
+			FoundFiles.Sort();
 			if (FoundFiles.Num() > 0)
 			{
 				bHaveGameOpenOrder = bHaveGameOpenOrder || (&Spec != SecondaryOrderSpec);
diff --git a/Engine/Source/Programs/AutomationTool/Scripts/CopyBuildToStagingDirectory.Automation.cs b/Engine/Source/Programs/AutomationTool/Scripts/CopyBuildToStagingDirectory.Automation.cs
--- a/Engine/Source/Programs/AutomationTool/Scripts/CopyBuildToStagingDirectory.Automation.cs
+++ b/Engine/Source/Programs/AutomationTool/Scripts/CopyBuildToStagingDirectory.Automation.cs
@@ -84,10 +84,10 @@ namespace AutomationScripts
 
 			public int CompareTo(OrderFile Other)
 			{
-				int compare = (OrderSpecIndex.CompareTo(Other.OrderSpecIndex));
+				int compare = OrderSpecIndex.CompareTo(Other.OrderSpecIndex);
 				if (compare == 0)
 				{
-					return (AppendOrder.CompareTo(Other.AppendOrder));
+					return File.GetFileNameWithoutExtension().CompareTo(File.GetFileNameWithoutExtension());
 				}
 				return compare;
 			}
