Removed assert when trying to refresh a thumbnail for which we don't have a ThumbnailInfo as it can happen during a rename since the thumbid uses the asset name.

#Synced-CL 47977636

diff --git a/Engine/Source/Editor/UnrealEd/Private/AssetThumbnail.cpp b/Engine/Source/Editor/UnrealEd/Private/AssetThumbnail.cpp
--- a/Engine/Source/Editor/UnrealEd/Private/AssetThumbnail.cpp
+++ b/Engine/Source/Editor/UnrealEd/Private/AssetThumbnail.cpp
@@ -2890,12 +2890,12 @@ void FAssetThumbnailPool::RefreshThumbnailsFor( const FSoftObjectPath& ObjectPat
 
 void FAssetThumbnailPool::RefreshThumbnailFor( const FThumbId ThumbId )
 {
-	check(ThumbnailToTextureMap.Contains(ThumbId));
-
-	TSharedRef<FThumbnailInfo> ThumbnailInfo = ThumbnailToTextureMap[ThumbId];
-	ThumbnailsToRenderStack.AddUnique( ThumbnailInfo );
-
-	UE_LOG(LogThumbnails, Verbose, TEXT("Refresh thumbnail requested for %s"), *ThumbnailInfo->AssetData.GetFullName());
+	const TSharedRef<FThumbnailInfo>* ThumbnailInfoPtr = ThumbnailToTextureMap.Find( ThumbId );
+	if ( ThumbnailInfoPtr )
+	{
+		ThumbnailsToRenderStack.AddUnique( *ThumbnailInfoPtr );
+		UE_LOG(LogThumbnails, Verbose, TEXT("Refresh thumbnail requested for %s"), *(*ThumbnailInfoPtr)->AssetData.GetFullName());
+	}
 }
 
 void FAssetThumbnailPool::OnAssetLoaded( UObject* Asset )
