WorldPartitionBuilder Commandlet: Collect Garbage between each worlds to avoid issues related to Level reuse.

#Synced-CL 0

diff --git a/Engine/Source/Editor/UnrealEd/Private/Commandlets/WorldPartitionBuilderCommandlet.cpp b/Engine/Source/Editor/UnrealEd/Private/Commandlets/WorldPartitionBuilderCommandlet.cpp
--- a/Engine/Source/Editor/UnrealEd/Private/Commandlets/WorldPartitionBuilderCommandlet.cpp
+++ b/Engine/Source/Editor/UnrealEd/Private/Commandlets/WorldPartitionBuilderCommandlet.cpp
@@ -14,6 +14,7 @@
 #include "ProfilingDebugging/ScopedTimers.h"
 #include "WorldPartition/WorldPartition.h"
 #include "WorldPartition/WorldPartitionBuilder.h"
+#include "WorldPartition/WorldPartitionHelpers.h"
 #include "UObject/GCObjectScopeGuard.h"
 #include "Trace/Trace.h"
 
@@ -208,6 +209,9 @@ int32 UWorldPartitionBuilderCommandlet::Main(const FString& Params)
 			}
 		}
 		AdditionalWorldPackagesToProcess.Empty();
+
+		//Collect Garbage between each levels to avoid problems when reusing levels.
+		FWorldPartitionHelpers::DoCollectGarbage();
 	}
 
 	// Autosubmit
