Localization manifest file not existing in a plugin yet should not trigger an irrecuperable error during an import job in the case the plugin is simply brand new. Don't report load warnings as an error to the caller function.

#Synced-CL 48539938

diff --git a/Engine/Source/Editor/UnrealEd/Private/Commandlets/GenerateTextLocalizationResourceCommandlet.cpp b/Engine/Source/Editor/UnrealEd/Private/Commandlets/GenerateTextLocalizationResourceCommandlet.cpp
--- a/Engine/Source/Editor/UnrealEd/Private/Commandlets/GenerateTextLocalizationResourceCommandlet.cpp
+++ b/Engine/Source/Editor/UnrealEd/Private/Commandlets/GenerateTextLocalizationResourceCommandlet.cpp
@@ -173,7 +173,7 @@ int32 UGenerateTextLocalizationResourceCommandlet::Main(const FString& Params)
 		if (!LocTextHelper.LoadAll(ELocTextHelperLoadFlags::Load, &LoadError))
 		{
 			UE_LOG(LogGenerateTextLocalizationResourceCommandlet, Warning, TEXT("%s"), *LoadError.ToString());
-			return -1;
+			return 0;
 		}
 	}
 
diff --git a/Engine/Source/Editor/UnrealEd/Private/Commandlets/InternationalizationExportCommandlet.cpp b/Engine/Source/Editor/UnrealEd/Private/Commandlets/InternationalizationExportCommandlet.cpp
--- a/Engine/Source/Editor/UnrealEd/Private/Commandlets/InternationalizationExportCommandlet.cpp
+++ b/Engine/Source/Editor/UnrealEd/Private/Commandlets/InternationalizationExportCommandlet.cpp
@@ -196,7 +196,7 @@ int32 UInternationalizationExportCommandlet::Main(const FString& Params)
 					("error", *LoadError.ToString()),
 					("id", InternationalizationExportCommandlet::LocalizationLogIdentifier)
 				);
-				return -1;
+				return 0;
 			}
 		}
 
@@ -227,7 +227,7 @@ int32 UInternationalizationExportCommandlet::Main(const FString& Params)
 					("error", *LoadError.ToString()),
 					("id", InternationalizationExportCommandlet::LocalizationLogIdentifier)
 				);
-				return -1;
+				return 0;
 			}
 		}
 
