UE: Bump cook incremental version. Stale content warnings related to a native change

diff --git a/Engine/Source/Editor/UnrealEd/Private/Cooker/CookIncrementalVersion.cpp b/Engine/Source/Editor/UnrealEd/Private/Cooker/CookIncrementalVersion.cpp
--- a/Engine/Source/Editor/UnrealEd/Private/Cooker/CookIncrementalVersion.cpp
+++ b/Engine/Source/Editor/UnrealEd/Private/Cooker/CookIncrementalVersion.cpp
@@ -8,6 +8,6 @@ namespace UE::Cook
 {
 
 // Change CookIncrementalVersion to a new guid when all packages in an incremental cook need to be invalidated
-FGuid CookIncrementalVersion(0x25399965, 0xBCA94C3B, 0xBAA048FA, 0x2EB2824A);
+FGuid CookIncrementalVersion(0xC943DD90, 0x885C47BB, 0xB81BD5F3, 0xE1B8A202);
 
 }
