Change editor Ortho view modes to initialise in unlit mode and allow it to use the same depth range as Wireframe (as it is only lit render passes that break with an unbound depth range). This allows much easier management of the scene from these viewmodes and resolves problems in various editor tool debug windows (at the very least until a lit pass fix can be resolved).

#Synced-CL 48543872

diff --git a/Engine/Source/Editor/LevelEditor/Private/SLevelViewport.cpp b/Engine/Source/Editor/LevelEditor/Private/SLevelViewport.cpp
--- a/Engine/Source/Editor/LevelEditor/Private/SLevelViewport.cpp
+++ b/Engine/Source/Editor/LevelEditor/Private/SLevelViewport.cpp
@@ -316,7 +316,7 @@ void SLevelViewport::Construct(const FArguments& InArgs, const FAssetEditorViewp
 	FLevelEditorViewportInstanceSettings ViewportInstanceSettings;
 	ViewportInstanceSettings.ViewportType = InConstructionArguments.ViewportType;
 	ViewportInstanceSettings.PerspViewModeIndex = VMI_Lit;
-	ViewportInstanceSettings.OrthoViewModeIndex = VMI_BrushWireframe;
+	ViewportInstanceSettings.OrthoViewModeIndex = VMI_Unlit;
 	ViewportInstanceSettings.bIsRealtime = InConstructionArguments.bRealtime;
 	ConstructLevelEditorViewportClient(ViewportInstanceSettings);
 
@@ -3141,8 +3141,8 @@ FLevelEditorViewportInstanceSettings SLevelViewport::LoadLegacyConfigFromIni(con
 
 	if(!GConfig->GetInt(*IniSection, *(InConfigKey+TEXT(".OrthoViewModeIndex")), (int32&)ViewportInstanceSettings.OrthoViewModeIndex, GEditorPerProjectIni))
 	{
-		// Default to Brush Wireframe for an orthographic viewport
-		ViewportInstanceSettings.OrthoViewModeIndex = (ViewportInstanceSettings.ViewportType != LVT_Perspective && LegacyViewModeIndex != VMI_Unknown) ? LegacyViewModeIndex : VMI_BrushWireframe;
+		// Default to Unlit for an orthographic viewport
+		ViewportInstanceSettings.OrthoViewModeIndex = (ViewportInstanceSettings.ViewportType != LVT_Perspective && LegacyViewModeIndex != VMI_Unknown) ? LegacyViewModeIndex : VMI_Unlit;
 	}
 
 	{
diff --git a/Engine/Source/Editor/UnrealEd/Classes/Settings/LevelEditorViewportSettings.h b/Engine/Source/Editor/UnrealEd/Classes/Settings/LevelEditorViewportSettings.h
--- a/Engine/Source/Editor/UnrealEd/Classes/Settings/LevelEditorViewportSettings.h
+++ b/Engine/Source/Editor/UnrealEd/Classes/Settings/LevelEditorViewportSettings.h
@@ -107,7 +107,7 @@ struct FLevelEditorViewportInstanceSettings
 	FLevelEditorViewportInstanceSettings()
 		: ViewportType(LVT_Perspective)
 		, PerspViewModeIndex(VMI_Lit)
-		, OrthoViewModeIndex(VMI_BrushWireframe)
+		, OrthoViewModeIndex(VMI_Unlit)
 		, EditorShowFlagsString()
 		, GameShowFlagsString()
 		, BufferVisualizationMode()
diff --git a/Engine/Source/Editor/UnrealEd/Private/EditorViewportClient.cpp b/Engine/Source/Editor/UnrealEd/Private/EditorViewportClient.cpp
--- a/Engine/Source/Editor/UnrealEd/Private/EditorViewportClient.cpp
+++ b/Engine/Source/Editor/UnrealEd/Private/EditorViewportClient.cpp
@@ -1448,71 +1448,75 @@ FSceneView* FEditorViewportClient::CalcSceneView(FSceneViewFamily* ViewFamily, c
 			FMatrix::FReal ZScale = 0.5f / UE_OLD_WORLD_MAX;
 			FMatrix::FReal ZOffset = UE_OLD_WORLD_MAX;
 			
-			if (SupportsManualOrthographicDepthControl())
+			//Only lit mode render passes require the restricted depth, so wireframe and unlit can just use the standard path
+			if (!(ViewFamily->EngineShowFlags.Wireframe || ViewFamily->EngineShowFlags.UnlitViewmode))
 			{
-				FMinimalViewInfo CalculatePlanesViewInfo;
-				CalculatePlanesViewInfo.Rotation = ViewInitOptions.ViewRotationMatrix.Rotator();
-				CalculatePlanesViewInfo.AspectRatio = AspectRatio;
-				CalculatePlanesViewInfo.bConstrainAspectRatio = false;
-
-				CalculatePlanesViewInfo.ProjectionMode = ECameraProjectionMode::Orthographic;				
-				CalculatePlanesViewInfo.OrthoWidth = OrthoWidth;
-				CalculatePlanesViewInfo.bAutoCalculateOrthoPlanes = true;
-				CalculatePlanesViewInfo.AutoPlaneShift = 0.0f;
-				CalculatePlanesViewInfo.bUpdateOrthoPlanes = true;
-				CalculatePlanesViewInfo.bUseCameraHeightAsViewTarget = false;
-				CalculatePlanesViewInfo.OrthoNearClipPlane = -UE_OLD_WORLD_MAX;
-				CalculatePlanesViewInfo.OrthoFarClipPlane = UE_OLD_WORLD_MAX;
-
-				SDepthBar::FDepthSpace DepthSpace(GetViewLocation(), GetForwardVector(), GetWorldBounds());
-				
-				if (OrthographicNearPlane.IsSet())
-				{
-					CalculatePlanesViewInfo.OrthoNearClipPlane = OrthographicNearPlane.GetValue();
-				}
-				else if (!ViewFamily->EngineShowFlags.Wireframe)
-				{
-					CalculatePlanesViewInfo.OrthoNearClipPlane = DepthSpace.GetAlignedMin();
-				}
-				
-				if (OrthographicFarPlane.IsSet())
-				{
-					CalculatePlanesViewInfo.OrthoFarClipPlane = OrthographicFarPlane.GetValue();
-				}
-				else if (!ViewFamily->EngineShowFlags.Wireframe)
-				{
-					CalculatePlanesViewInfo.OrthoFarClipPlane = DepthSpace.GetAlignedMax();
-				}
-				
-				ViewInitOptions.UpdateOrthoPlanes(CalculatePlanesViewInfo);
-				const float Denominator = CalculatePlanesViewInfo.OrthoFarClipPlane - CalculatePlanesViewInfo.OrthoNearClipPlane;
-				if (Denominator > 0.0f)
+				if (SupportsManualOrthographicDepthControl())
 				{
-					ZScale = 1.0f / Denominator;
-					ZOffset = -CalculatePlanesViewInfo.OrthoNearClipPlane;
+					FMinimalViewInfo CalculatePlanesViewInfo;
+					CalculatePlanesViewInfo.Rotation = ViewInitOptions.ViewRotationMatrix.Rotator();
+					CalculatePlanesViewInfo.AspectRatio = AspectRatio;
+					CalculatePlanesViewInfo.bConstrainAspectRatio = false;
+
+					CalculatePlanesViewInfo.ProjectionMode = ECameraProjectionMode::Orthographic;
+					CalculatePlanesViewInfo.OrthoWidth = OrthoWidth;
+					CalculatePlanesViewInfo.bAutoCalculateOrthoPlanes = true;
+					CalculatePlanesViewInfo.AutoPlaneShift = 0.0f;
+					CalculatePlanesViewInfo.bUpdateOrthoPlanes = true;
+					CalculatePlanesViewInfo.bUseCameraHeightAsViewTarget = false;
+					CalculatePlanesViewInfo.OrthoNearClipPlane = -UE_OLD_WORLD_MAX;
+					CalculatePlanesViewInfo.OrthoFarClipPlane = UE_OLD_WORLD_MAX;
+
+					SDepthBar::FDepthSpace DepthSpace(GetViewLocation(), GetForwardVector(), GetWorldBounds());
+
+					if (OrthographicNearPlane.IsSet())
+					{
+						CalculatePlanesViewInfo.OrthoNearClipPlane = OrthographicNearPlane.GetValue();
+					}
+					else if (!ViewFamily->EngineShowFlags.Wireframe)
+					{
+						CalculatePlanesViewInfo.OrthoNearClipPlane = DepthSpace.GetAlignedMin();
+					}
+
+					if (OrthographicFarPlane.IsSet())
+					{
+						CalculatePlanesViewInfo.OrthoFarClipPlane = OrthographicFarPlane.GetValue();
+					}
+					else if (!ViewFamily->EngineShowFlags.Wireframe)
+					{
+						CalculatePlanesViewInfo.OrthoFarClipPlane = DepthSpace.GetAlignedMax();
+					}
+
+					ViewInitOptions.UpdateOrthoPlanes(CalculatePlanesViewInfo);
+					const float Denominator = CalculatePlanesViewInfo.OrthoFarClipPlane - CalculatePlanesViewInfo.OrthoNearClipPlane;
+					if (Denominator > 0.0f)
+					{
+						ZScale = 1.0f / Denominator;
+						ZOffset = -CalculatePlanesViewInfo.OrthoNearClipPlane;
+					}
 				}
-			}
-			else if(!ViewFamily->EngineShowFlags.Wireframe)
-			{
-				FMinimalViewInfo CalculatePlanesViewInfo;
-				CalculatePlanesViewInfo.Rotation = ViewInitOptions.ViewRotationMatrix.Rotator();
-				CalculatePlanesViewInfo.AspectRatio = AspectRatio;
-				CalculatePlanesViewInfo.bConstrainAspectRatio = false;
-
-				CalculatePlanesViewInfo.ProjectionMode = ECameraProjectionMode::Orthographic;				
-				CalculatePlanesViewInfo.OrthoWidth = OrthoWidth;
-				CalculatePlanesViewInfo.bAutoCalculateOrthoPlanes = true;
-				CalculatePlanesViewInfo.AutoPlaneShift = 0.0f;
-				CalculatePlanesViewInfo.bUpdateOrthoPlanes = true;
-				CalculatePlanesViewInfo.bUseCameraHeightAsViewTarget = true;
-				CalculatePlanesViewInfo.OrthoNearClipPlane = OrthoWidth * -CVarOrthoEditorDebugClipPlaneScale.GetValueOnAnyThread();
-				CalculatePlanesViewInfo.OrthoFarClipPlane = FarPlane - NearPlane + CalculatePlanesViewInfo.OrthoNearClipPlane;
-
-				CalculatePlanesViewInfo.AutoCalculateOrthoPlanes(ViewInitOptions);
-				if(ViewInitOptions.UpdateOrthoPlanes(CalculatePlanesViewInfo))
+				else
 				{
-					ZScale = 1.0f / (CalculatePlanesViewInfo.OrthoFarClipPlane - CalculatePlanesViewInfo.OrthoNearClipPlane);
-					ZOffset = -CalculatePlanesViewInfo.OrthoNearClipPlane;
+					FMinimalViewInfo CalculatePlanesViewInfo;
+					CalculatePlanesViewInfo.Rotation = ViewInitOptions.ViewRotationMatrix.Rotator();
+					CalculatePlanesViewInfo.AspectRatio = AspectRatio;
+					CalculatePlanesViewInfo.bConstrainAspectRatio = false;
+
+					CalculatePlanesViewInfo.ProjectionMode = ECameraProjectionMode::Orthographic;
+					CalculatePlanesViewInfo.OrthoWidth = OrthoWidth;
+					CalculatePlanesViewInfo.bAutoCalculateOrthoPlanes = true;
+					CalculatePlanesViewInfo.AutoPlaneShift = 0.0f;
+					CalculatePlanesViewInfo.bUpdateOrthoPlanes = true;
+					CalculatePlanesViewInfo.bUseCameraHeightAsViewTarget = true;
+					CalculatePlanesViewInfo.OrthoNearClipPlane = OrthoWidth * -CVarOrthoEditorDebugClipPlaneScale.GetValueOnAnyThread();
+					CalculatePlanesViewInfo.OrthoFarClipPlane = FarPlane - NearPlane + CalculatePlanesViewInfo.OrthoNearClipPlane;
+
+					CalculatePlanesViewInfo.AutoCalculateOrthoPlanes(ViewInitOptions);
+					if (ViewInitOptions.UpdateOrthoPlanes(CalculatePlanesViewInfo))
+					{
+						ZScale = 1.0f / (CalculatePlanesViewInfo.OrthoFarClipPlane - CalculatePlanesViewInfo.OrthoNearClipPlane);
+						ZOffset = -CalculatePlanesViewInfo.OrthoNearClipPlane;
+					}
 				}
 			}
 			
diff --git a/Engine/Source/Editor/UnrealEd/Private/EditorViewportLayoutOnePane.cpp b/Engine/Source/Editor/UnrealEd/Private/EditorViewportLayoutOnePane.cpp
--- a/Engine/Source/Editor/UnrealEd/Private/EditorViewportLayoutOnePane.cpp
+++ b/Engine/Source/Editor/UnrealEd/Private/EditorViewportLayoutOnePane.cpp
@@ -13,7 +13,7 @@
 TSharedRef<SWidget> FEditorViewportLayoutOnePane::MakeViewportLayout(TSharedPtr<FAssetEditorViewportLayout> InParentLayout, const FString& LayoutString)
 {
 	FEngineShowFlags OrthoShowFlags(ESFIM_Editor);
-	ApplyViewMode(VMI_BrushWireframe, false, OrthoShowFlags);
+	ApplyViewMode(VMI_Unlit, false, OrthoShowFlags);
 
 	FEngineShowFlags PerspectiveShowFlags(ESFIM_Editor);
 	ApplyViewMode(VMI_Lit, true, PerspectiveShowFlags);
diff --git a/Engine/Source/Editor/UnrealEd/Private/EditorViewportLayoutThreePanes.cpp b/Engine/Source/Editor/UnrealEd/Private/EditorViewportLayoutThreePanes.cpp
--- a/Engine/Source/Editor/UnrealEd/Private/EditorViewportLayoutThreePanes.cpp
+++ b/Engine/Source/Editor/UnrealEd/Private/EditorViewportLayoutThreePanes.cpp
@@ -20,7 +20,7 @@ namespace ViewportLayoutThreePanesDefs
 TSharedRef<SWidget> FEditorViewportLayoutThreePanes::MakeViewportLayout(TSharedPtr<FAssetEditorViewportLayout> InParentLayout, const FString& LayoutString)
 {
 	FEngineShowFlags OrthoShowFlags(ESFIM_Editor);	
-	ApplyViewMode(VMI_BrushWireframe, false, OrthoShowFlags);
+	ApplyViewMode(VMI_Unlit, false, OrthoShowFlags);
 
 	FEngineShowFlags PerspectiveShowFlags(ESFIM_Editor);	
 	ApplyViewMode(VMI_Lit, true, PerspectiveShowFlags);
