Don't add redirection dependencies for native package.
Change cook incremental version to trigger a full recook and remove all the native package redirect dependencies.

#rnx
#Synced-CL 46456044

diff --git a/Engine/Source/Editor/UnrealEd/Private/Cooker/CookIncrementalVersion.cpp b/Engine/Source/Editor/UnrealEd/Private/Cooker/CookIncrementalVersion.cpp
--- a/Engine/Source/Editor/UnrealEd/Private/Cooker/CookIncrementalVersion.cpp
+++ b/Engine/Source/Editor/UnrealEd/Private/Cooker/CookIncrementalVersion.cpp
@@ -8,6 +8,6 @@ namespace UE::Cook
 {
 
 // Change CookIncrementalVersion to a new guid when all packages in an incremental cook need to be invalidated
-FGuid CookIncrementalVersion(0xC943DD90, 0x885C47BB, 0xB81BD5F3, 0xE1B8A202);
+FGuid CookIncrementalVersion(0x623B7C93, 0x212E4016, 0xBAC806FE, 0x2DFB80B8);
 
 }
diff --git a/Engine/Source/Editor/UnrealEd/Private/Cooker/CookPackageArtifacts.cpp b/Engine/Source/Editor/UnrealEd/Private/Cooker/CookPackageArtifacts.cpp
--- a/Engine/Source/Editor/UnrealEd/Private/Cooker/CookPackageArtifacts.cpp
+++ b/Engine/Source/Editor/UnrealEd/Private/Cooker/CookPackageArtifacts.cpp
@@ -828,7 +828,7 @@ bool FBuildDependencySet::TryCollectInternal(FBuildResultDependenciesMap& InOutR
 		{
 			// Pull transient packages out of the runtime dependencies for performance; we don't need them for
 			// deciding what gets cooked. 
-			// Runtime and script dependencies also cause RedirectionTarget dependencies, so record those.
+			// Runtime dependencies also cause RedirectionTarget dependencies, so record those.
 			for (TArray<FName>::TIterator Iter(InOutRuntimeDependencies); Iter; ++Iter)
 			{
 				FName DependencyPackageName = *Iter;
@@ -837,10 +837,15 @@ bool FBuildDependencySet::TryCollectInternal(FBuildResultDependenciesMap& InOutR
 				{
 				case EPackageMountPoint::GeneratedContent:
 				case EPackageMountPoint::Content:
-				case EPackageMountPoint::Script:
 					// Keep it
 					RedirectionTargets.Add(DependencyPackageName);
 					break;
+
+				// Do not add redirection dependencies for native package. Those dependencies are too wide and already
+				// included in schema hash or package digest.
+				case EPackageMountPoint::Script:
+					break;
+
 				default:
 					Iter.RemoveCurrentSwap();
 					break;
