Cooker: Memory: Remove wasted memory in PackageBuildDependencyTracker by dropping the records when packages unload.
#rnx
#Synced-CL 45140074

diff --git a/Engine/Source/Editor/UnrealEd/Private/Cooker/PackageBuildDependencyTracker.cpp b/Engine/Source/Editor/UnrealEd/Private/Cooker/PackageBuildDependencyTracker.cpp
--- a/Engine/Source/Editor/UnrealEd/Private/Cooker/PackageBuildDependencyTracker.cpp
+++ b/Engine/Source/Editor/UnrealEd/Private/Cooker/PackageBuildDependencyTracker.cpp
@@ -15,6 +15,7 @@
 #include "Misc/PackageAccessTrackingOps.h"
 #include "UObject/Class.h"
 #include "UObject/Package.h"
+#include "UObject/UObjectGlobals.h"
 #endif
 
 #if UE_WITH_PACKAGE_ACCESS_TRACKING
@@ -29,6 +30,8 @@ void FPackageBuildDependencyTracker::Disable()
 	{
 		UE::CoreUObject::RemoveObjectHandleReadCallback(ObjectHandleReadHandle);
 		ObjectHandleReadHandle = UE::CoreUObject::FObjectHandleTrackingCallbackId {};
+		FCoreUObjectDelegates::GetPostGarbageCollect().RemoveAll(this);
+
 		bEnabled = false;
 	}
 }
@@ -87,6 +90,8 @@ TArray<TPair<FBuildDependencyAccessData, FResultProjectionList>> FPackageBuildDe
 FPackageBuildDependencyTracker::FPackageBuildDependencyTracker()
 {
 	ObjectHandleReadHandle = UE::CoreUObject::AddObjectHandleReadCallback(StaticOnObjectHandleRead);
+	FCoreUObjectDelegates::GetPostGarbageCollect().AddRaw(this, &FPackageBuildDependencyTracker::PostGarbageCollect);
+
 	bEnabled = true;
 }
 
@@ -218,6 +223,29 @@ void FPackageBuildDependencyTracker::StaticOnObjectHandleRead(const TArrayView<c
 	}
 }
 
+void FPackageBuildDependencyTracker::PostGarbageCollect()
+{
+	// We expect not to be asked for the records for any package no longer in memory;
+	// Only package saves trigger a request for records, and they only request records for themselves
+	// and for their build dependencies, and they and their build dependencies should be in memory when
+	// they are saved.
+	// Drop these no longer needed records to save memory.
+	// We will recreate them if the package loads again.
+	UE::TUniqueLock RecordsScopeLock(Singleton.RecordsLock);
+	for (TMap<FName, TMap<FBuildDependencyAccessData, FResultProjectionList>>::TIterator Iter(Records); Iter; ++Iter)
+	{
+		FName PackageName = Iter.Key();
+		if (!FindPackage(nullptr, *WriteToString<256>(PackageName)))
+		{
+			Iter.RemoveCurrent();
+		}
+	}
+	Singleton.LastReferencer = NAME_None;
+	Singleton.LastReferencerMap = nullptr;
+	Singleton.LastAccessData = FBuildDependencyAccessData{ NAME_None, nullptr };
+	Singleton.LastCookResultProjection = NAME_None;
+}
+
 void FResultProjectionList::AddProjectionAll(bool* bOutExists)
 {
 	if (bHasAll)
diff --git a/Engine/Source/Editor/UnrealEd/Private/Cooker/PackageBuildDependencyTracker.h b/Engine/Source/Editor/UnrealEd/Private/Cooker/PackageBuildDependencyTracker.h
--- a/Engine/Source/Editor/UnrealEd/Private/Cooker/PackageBuildDependencyTracker.h
+++ b/Engine/Source/Editor/UnrealEd/Private/Cooker/PackageBuildDependencyTracker.h
@@ -54,6 +54,7 @@ public:
 private:
 	FPackageBuildDependencyTracker();
 	virtual ~FPackageBuildDependencyTracker();
+	void PostGarbageCollect();
 
 	/** Track object reference reads */
 	static void StaticOnObjectHandleRead(const TArrayView<const UObject* const>& Objects);
