[ITF Gizmos] Reset `FEditorViewportClient::bDraggingByHandle` when tracking stops. Fix issue where ITF-based interactions that called into legacy transform manipulations still had this value set to true.

#rnx
#Synced-CL 48609818

diff --git a/Engine/Source/Editor/UnrealEd/Private/EditorViewportClient.cpp b/Engine/Source/Editor/UnrealEd/Private/EditorViewportClient.cpp
--- a/Engine/Source/Editor/UnrealEd/Private/EditorViewportClient.cpp
+++ b/Engine/Source/Editor/UnrealEd/Private/EditorViewportClient.cpp
@@ -3873,6 +3873,7 @@ void FEditorViewportClient::StopTracking()
 	}
 
 	bHasMouseMovedSinceClick = false;
+	bDraggingByHandle = false;
 }
 
 void FEditorViewportClient::AbortTracking()
