UE: Bump cook incremental version due to tag not getting picked up

diff --git a/Engine/Source/Editor/UnrealEd/Private/Cooker/CookIncrementalVersion.cpp b/Engine/Source/Editor/UnrealEd/Private/Cooker/CookIncrementalVersion.cpp
--- a/Engine/Source/Editor/UnrealEd/Private/Cooker/CookIncrementalVersion.cpp
+++ b/Engine/Source/Editor/UnrealEd/Private/Cooker/CookIncrementalVersion.cpp
@@ -8,6 +8,6 @@ namespace UE::Cook
 {
 
 // Change CookIncrementalVersion to a new guid when all packages in an incremental cook need to be invalidated
-FGuid CookIncrementalVersion(0x623B7C93, 0x212E4016, 0xBAC806FE, 0x2DFB80B8);
+FGuid CookIncrementalVersion(0x2E7B0688, 0x11D743CA, 0x85D432C5, 0x8BAF841F);
 
 }
