Refactor the recent changes to ortho editor view modes to ensure that the depth range UI still works with Wireframe and unlit modes. Also remove the CVar controlling the use of the recently introduced depth bar so it is always enabled.

#Synced-CL 48726168

diff --git a/Engine/Source/Editor/UnrealEd/Private/EditorViewportClient.cpp b/Engine/Source/Editor/UnrealEd/Private/EditorViewportClient.cpp
--- a/Engine/Source/Editor/UnrealEd/Private/EditorViewportClient.cpp
+++ b/Engine/Source/Editor/UnrealEd/Private/EditorViewportClient.cpp
@@ -130,12 +130,6 @@ static TAutoConsoleVariable<int32> CVarEditorViewGPUMirrorTest(
 	TEXT("difficult, due to lack of interactivity, and multiple scene renders happening per frame.\n"),
 	ECVF_RenderThreadSafe);
 
-static TAutoConsoleVariable<int32> CVarEditorViewportManualOrthographicDepthControl(
-	TEXT("r.Editor.ManualOrthoDepth"),
-	1,
-	TEXT("Enable or disable the new orthographic depth tools."),
-	ECVF_RenderThreadSafe);
-
 static bool GetDefaultLowDPIPreviewValue()
 {
 	static auto CVarEditorViewportHighDPIPtr = IConsoleManager::Get().FindConsoleVariable(TEXT("r.Editor.Viewport.HighDPI"));
@@ -1023,39 +1017,36 @@ void FEditorViewportClient::FocusViewportOnBox( const FBox& BoundingBox, bool bI
 		}
 		else
 		{
-			if (SupportsManualOrthographicDepthControl())
-			{
-				const FVector StartingLocation = ViewTransform.GetLocation();
-				const FVector Forward = GetForwardVector();
+			const FVector StartingLocation = ViewTransform.GetLocation();
+			const FVector Forward = GetForwardVector();
 				
-				// Adjust the camera by the current forward instantly
-				const FVector LocationDelta = Position - StartingLocation;
-				const FVector ForwardLocationDelta = LocationDelta.ProjectOnToNormal(Forward);
-				ViewTransform.SetLocation(StartingLocation + ForwardLocationDelta);
+			// Adjust the camera by the current forward instantly
+			const FVector LocationDelta = Position - StartingLocation;
+			const FVector ForwardLocationDelta = LocationDelta.ProjectOnToNormal(Forward);
+			ViewTransform.SetLocation(StartingLocation + ForwardLocationDelta);
 				
-				// Adjust the near & far planes so that they remain fixed in world space
-				const double ForwardDelta = ForwardLocationDelta.Dot(Forward);
+			// Adjust the near & far planes so that they remain fixed in world space
+			const double ForwardDelta = ForwardLocationDelta.Dot(Forward);
 				
-				if (OrthographicNearPlane.IsSet())
-				{
-					OrthographicNearPlane = OrthographicNearPlane.GetValue() - ForwardDelta;
-				}
+			if (OrthographicNearPlane.IsSet())
+			{
+				OrthographicNearPlane = OrthographicNearPlane.GetValue() - ForwardDelta;
+			}
 				
-				if (OrthographicFarPlane.IsSet())
-				{
-					OrthographicFarPlane = OrthographicFarPlane.GetValue() - ForwardDelta;
-				}
+			if (OrthographicFarPlane.IsSet())
+			{
+				OrthographicFarPlane = OrthographicFarPlane.GetValue() - ForwardDelta;
+			}
 				
-				if (GetClipToSelectionOnFocus())
-				{
-					// Shift the orthographic near view plane back so that it is just before the object
-					OrthographicNearPlane = -Radius;
+			if (GetClipToSelectionOnFocus())
+			{
+				// Shift the orthographic near view plane back so that it is just before the object
+				OrthographicNearPlane = -Radius;
 					
-					// Ensure the orthographic far view plane can display the object
-					if (OrthographicFarPlane.IsSet())
-					{
-						OrthographicFarPlane = FMath::Max(OrthographicFarPlane.GetValue(), Radius);
-					}
+				// Ensure the orthographic far view plane can display the object
+				if (OrthographicFarPlane.IsSet())
+				{
+					OrthographicFarPlane = FMath::Max(OrthographicFarPlane.GetValue(), Radius);
 				}
 			}
 		
@@ -1449,75 +1440,48 @@ FSceneView* FEditorViewportClient::CalcSceneView(FSceneViewFamily* ViewFamily, c
 			FMatrix::FReal ZOffset = UE_OLD_WORLD_MAX;
 			
 			//Only lit mode render passes require the restricted depth, so wireframe and unlit can just use the standard path
-			if (!(ViewFamily->EngineShowFlags.Wireframe || ViewFamily->EngineShowFlags.UnlitViewmode))
+			bool bShouldCalculateDepthRange = ViewFamily->ViewMode > VMI_Unlit;
+
+			FMinimalViewInfo CalculatePlanesViewInfo;
+			CalculatePlanesViewInfo.Rotation = ViewInitOptions.ViewRotationMatrix.Rotator();
+			CalculatePlanesViewInfo.AspectRatio = AspectRatio;
+			CalculatePlanesViewInfo.bConstrainAspectRatio = false;
+			CalculatePlanesViewInfo.ProjectionMode = ECameraProjectionMode::Orthographic;
+			CalculatePlanesViewInfo.OrthoWidth = OrthoWidth;
+			CalculatePlanesViewInfo.bAutoCalculateOrthoPlanes = bShouldCalculateDepthRange;
+			CalculatePlanesViewInfo.AutoPlaneShift = 0.0f;
+			CalculatePlanesViewInfo.bUpdateOrthoPlanes = bShouldCalculateDepthRange;
+
+			CalculatePlanesViewInfo.bUseCameraHeightAsViewTarget = false;
+			CalculatePlanesViewInfo.OrthoNearClipPlane = -UE_OLD_WORLD_MAX;
+			CalculatePlanesViewInfo.OrthoFarClipPlane = UE_OLD_WORLD_MAX;
+
+			SDepthBar::FDepthSpace DepthSpace(GetViewLocation(), GetForwardVector(), GetWorldBounds());
+
+			if (OrthographicNearPlane.IsSet())
 			{
-				if (SupportsManualOrthographicDepthControl())
-				{
-					FMinimalViewInfo CalculatePlanesViewInfo;
-					CalculatePlanesViewInfo.Rotation = ViewInitOptions.ViewRotationMatrix.Rotator();
-					CalculatePlanesViewInfo.AspectRatio = AspectRatio;
-					CalculatePlanesViewInfo.bConstrainAspectRatio = false;
-
-					CalculatePlanesViewInfo.ProjectionMode = ECameraProjectionMode::Orthographic;
-					CalculatePlanesViewInfo.OrthoWidth = OrthoWidth;
-					CalculatePlanesViewInfo.bAutoCalculateOrthoPlanes = true;
-					CalculatePlanesViewInfo.AutoPlaneShift = 0.0f;
-					CalculatePlanesViewInfo.bUpdateOrthoPlanes = true;
-					CalculatePlanesViewInfo.bUseCameraHeightAsViewTarget = false;
-					CalculatePlanesViewInfo.OrthoNearClipPlane = -UE_OLD_WORLD_MAX;
-					CalculatePlanesViewInfo.OrthoFarClipPlane = UE_OLD_WORLD_MAX;
-
-					SDepthBar::FDepthSpace DepthSpace(GetViewLocation(), GetForwardVector(), GetWorldBounds());
-
-					if (OrthographicNearPlane.IsSet())
-					{
-						CalculatePlanesViewInfo.OrthoNearClipPlane = OrthographicNearPlane.GetValue();
-					}
-					else if (!ViewFamily->EngineShowFlags.Wireframe)
-					{
-						CalculatePlanesViewInfo.OrthoNearClipPlane = DepthSpace.GetAlignedMin();
-					}
+				CalculatePlanesViewInfo.OrthoNearClipPlane = OrthographicNearPlane.GetValue();
+			}
+			else if (!bShouldCalculateDepthRange)
+			{
+				CalculatePlanesViewInfo.OrthoNearClipPlane = DepthSpace.GetAlignedMin();
+			}
 
-					if (OrthographicFarPlane.IsSet())
-					{
-						CalculatePlanesViewInfo.OrthoFarClipPlane = OrthographicFarPlane.GetValue();
-					}
-					else if (!ViewFamily->EngineShowFlags.Wireframe)
-					{
-						CalculatePlanesViewInfo.OrthoFarClipPlane = DepthSpace.GetAlignedMax();
-					}
+			if (OrthographicFarPlane.IsSet())
+			{
+				CalculatePlanesViewInfo.OrthoFarClipPlane = OrthographicFarPlane.GetValue();
+			}
+			else if (!bShouldCalculateDepthRange)
+			{
+				CalculatePlanesViewInfo.OrthoFarClipPlane = DepthSpace.GetAlignedMax();
+			}
 
-					ViewInitOptions.UpdateOrthoPlanes(CalculatePlanesViewInfo);
-					const float Denominator = CalculatePlanesViewInfo.OrthoFarClipPlane - CalculatePlanesViewInfo.OrthoNearClipPlane;
-					if (Denominator > 0.0f)
-					{
-						ZScale = 1.0f / Denominator;
-						ZOffset = -CalculatePlanesViewInfo.OrthoNearClipPlane;
-					}
-				}
-				else
-				{
-					FMinimalViewInfo CalculatePlanesViewInfo;
-					CalculatePlanesViewInfo.Rotation = ViewInitOptions.ViewRotationMatrix.Rotator();
-					CalculatePlanesViewInfo.AspectRatio = AspectRatio;
-					CalculatePlanesViewInfo.bConstrainAspectRatio = false;
-
-					CalculatePlanesViewInfo.ProjectionMode = ECameraProjectionMode::Orthographic;
-					CalculatePlanesViewInfo.OrthoWidth = OrthoWidth;
-					CalculatePlanesViewInfo.bAutoCalculateOrthoPlanes = true;
-					CalculatePlanesViewInfo.AutoPlaneShift = 0.0f;
-					CalculatePlanesViewInfo.bUpdateOrthoPlanes = true;
-					CalculatePlanesViewInfo.bUseCameraHeightAsViewTarget = true;
-					CalculatePlanesViewInfo.OrthoNearClipPlane = OrthoWidth * -CVarOrthoEditorDebugClipPlaneScale.GetValueOnAnyThread();
-					CalculatePlanesViewInfo.OrthoFarClipPlane = FarPlane - NearPlane + CalculatePlanesViewInfo.OrthoNearClipPlane;
-
-					CalculatePlanesViewInfo.AutoCalculateOrthoPlanes(ViewInitOptions);
-					if (ViewInitOptions.UpdateOrthoPlanes(CalculatePlanesViewInfo))
-					{
-						ZScale = 1.0f / (CalculatePlanesViewInfo.OrthoFarClipPlane - CalculatePlanesViewInfo.OrthoNearClipPlane);
-						ZOffset = -CalculatePlanesViewInfo.OrthoNearClipPlane;
-					}
-				}
+			ViewInitOptions.UpdateOrthoPlanes(CalculatePlanesViewInfo);
+			const float Denominator = CalculatePlanesViewInfo.OrthoFarClipPlane - CalculatePlanesViewInfo.OrthoNearClipPlane;
+			if (Denominator > 0.0f)
+			{
+				ZScale = 1.0f / Denominator;
+				ZOffset = -CalculatePlanesViewInfo.OrthoNearClipPlane;
 			}
 			
 			ViewInitOptions.ProjectionMatrix = FReversedZOrthoMatrix(
@@ -6043,11 +6007,6 @@ void FEditorViewportClient::OverrideFarClipPlane(const float InFarPlane)
 	FarPlane = InFarPlane;
 }
 
-bool FEditorViewportClient::SupportsManualOrthographicDepthControl() const
-{
-	return CVarEditorViewportManualOrthographicDepthControl.GetValueOnAnyThread() > 0;
-}
-
 TOptional<double> FEditorViewportClient::GetOrthographicNearPlaneOverride() const
 {
 	return OrthographicNearPlane;
diff --git a/Engine/Source/Editor/UnrealEd/Private/SEditorViewport.cpp b/Engine/Source/Editor/UnrealEd/Private/SEditorViewport.cpp
--- a/Engine/Source/Editor/UnrealEd/Private/SEditorViewport.cpp
+++ b/Engine/Source/Editor/UnrealEd/Private/SEditorViewport.cpp
@@ -364,7 +364,7 @@ void SEditorViewport::Tick( const FGeometry& AllottedGeometry, const double InCu
 {
 	LastTickTime = FPlatformTime::Seconds();
 	
-	if (!bShowDepthBar && Client->SupportsManualOrthographicDepthControl() && Client->IsOrtho() &&
+	if (!bShowDepthBar && Client->IsOrtho() &&
 		(Client->GetOrthographicNearPlaneOverride().IsSet() || Client->GetOrthographicFarPlaneOverride().IsSet()))
 	{
 		bShowDepthBar = true;
@@ -809,7 +809,7 @@ void SEditorViewport::OnFocusAndClipViewportToSelection()
 
 EVisibility SEditorViewport::OnGetDepthBarVisibility() const
 {
-	if (bHasSpaceForDepthBar && Client->SupportsManualOrthographicDepthControl() && Client->IsOrtho() && bShowDepthBar)
+	if (bHasSpaceForDepthBar && Client->IsOrtho() && bShowDepthBar)
 	{
 		return EVisibility::Visible;
 	}
diff --git a/Engine/Source/Editor/UnrealEd/Public/EditorViewportClient.h b/Engine/Source/Editor/UnrealEd/Public/EditorViewportClient.h
--- a/Engine/Source/Editor/UnrealEd/Public/EditorViewportClient.h
+++ b/Engine/Source/Editor/UnrealEd/Public/EditorViewportClient.h
@@ -1372,10 +1372,7 @@ public:
 	UNREALED_API float GetFarClipPlaneOverride() const;
 	
 	/** Override the far clipping plane. Set to a negative value to disable the override. */
-	UNREALED_API void OverrideFarClipPlane(const float InFarPlane);
-
-	/** Accessor for whether orthographic control should be enabled. */
-	bool SupportsManualOrthographicDepthControl() const; 
+	UNREALED_API void OverrideFarClipPlane(const float InFarPlane); 
 
 	/** Get the orthographic near plane override for this viewport. */
 	UNREALED_API TOptional<double> GetOrthographicNearPlaneOverride() const;
