Firefight: Bump CookIncrementalVersion to fix a cooking error happening on Main

diff --git a/Engine/Source/Editor/UnrealEd/Private/Cooker/CookIncrementalVersion.cpp b/Engine/Source/Editor/UnrealEd/Private/Cooker/CookIncrementalVersion.cpp
--- a/Engine/Source/Editor/UnrealEd/Private/Cooker/CookIncrementalVersion.cpp
+++ b/Engine/Source/Editor/UnrealEd/Private/Cooker/CookIncrementalVersion.cpp
@@ -8,6 +8,6 @@ namespace UE::Cook
 {
 
 // Change CookIncrementalVersion to a new guid when all packages in an incremental cook need to be invalidated
-FGuid CookIncrementalVersion(0x2E7B0688, 0x11D743CA, 0x85D432C5, 0x8BAF841F);
+FGuid CookIncrementalVersion(0xabdb8358, 0x11D743CA, 0x85D432C5, 0x8BAF841F);
 
 }
