Cook stats telemetry

* Created UnrealEd/Public/Cooker/CookProfiling.h and CookProfiling.cpp for exposing profiling data to other systems
* Renamed existing CookProfiling.h & CookProfiling.cpp as ...Private
* Per asset DDC usage is generated directly in the engine not via C#
* Telemetry data includes per asset DDC usage that was previously only available in CookStats.json
* Engine cook stats includes per asset DDC usage that was previously only available in CookStats.json

#Synced-CL 48597399

diff --git a/Engine/Plugins/Experimental/EditorTelemetry/Source/Private/EditorTelemetry.cpp b/Engine/Plugins/Experimental/EditorTelemetry/Source/Private/EditorTelemetry.cpp
--- a/Engine/Plugins/Experimental/EditorTelemetry/Source/Private/EditorTelemetry.cpp
+++ b/Engine/Plugins/Experimental/EditorTelemetry/Source/Private/EditorTelemetry.cpp
@@ -16,6 +16,7 @@
 #include "Editor/UnrealEdEngine.h"
 #include "Engine/AssetManager.h"
 #include "Experimental/ZenServerInterface.h"
+#include "Internationalization/FastDecimalFormat.h"
 #include "FileHelpers.h"
 #include "HttpManager.h"
 #include "HttpModule.h"
@@ -31,6 +32,7 @@
 #include "UObject/ICookInfo.h"
 #include "UnrealEdGlobals.h"
 #include "Virtualization/VirtualizationSystem.h"
+#include "Cooker/CookProfiling.h"
 
 namespace Private
 {
@@ -121,7 +123,7 @@ void FEditorTelemetry::RecordEvent_Cooking(TArray<FAnalyticsEventAttribute> Attr
 {
 #if ENABLE_COOK_STATS
 
-	const int SchemaVersion = 3;
+	const int SchemaVersion = 4;
 
 	Attributes.Emplace(TEXT("SchemaVersion"), SchemaVersion);
 
@@ -130,6 +132,29 @@ void FEditorTelemetry::RecordEvent_Cooking(TArray<FAnalyticsEventAttribute> Attr
 	// Sends each cook stat to the studio analytics system.
 	auto GatherAnalyticsAttributes = [&CookAttributes](const FString& StatName, const TArray<FCookStatsManager::StringKeyValue>& StatAttributes)
 	{
+		//Special handling for the Hierarchical profiling data
+		if (StatName == TEXT("Cook.Profile"))
+		{
+			// Hierarchical stats have 2 attributes per stat. 
+			// The path (signifies the hierarchy information) and the stat itself
+			if (StatAttributes.Num() != 2)
+			{
+				return;
+			}
+
+			// When generating the telemetry event we don't care about the hierarchy but only stat itself.
+			const FCookStatsManager::StringKeyValue& Stat = StatAttributes[1];
+
+			const FString FormattedAttrName = (StatName + "_" + Stat.Key).Replace(TEXT("."), TEXT("_"));
+
+			if (CookAttributes.Find(FormattedAttrName) == nullptr)
+			{
+				CookAttributes.Emplace(FormattedAttrName, Stat.Value.IsNumeric() ? FAnalyticsEventAttribute(FormattedAttrName, FCString::Atof(*Stat.Value)) : FAnalyticsEventAttribute(FormattedAttrName, Stat.Value));
+			}
+
+			return;
+		}
+
 		for (const auto& Attr : StatAttributes)
 		{
 			const FString FormattedAttrName = (StatName + "_" + Attr.Key).Replace(TEXT("."), TEXT("_"));
@@ -150,25 +175,7 @@ void FEditorTelemetry::RecordEvent_Cooking(TArray<FAnalyticsEventAttribute> Attr
 		Attributes.Emplace((*it).Value);
 	}
 
-	// Gather the DDC summary stats
-	FDerivedDataCacheSummaryStats SummaryStats;
-
-	GatherDerivedDataCacheSummaryStats(SummaryStats);
-
-	// Append to the attributes
-	for (const FDerivedDataCacheSummaryStat& Stat : SummaryStats.Stats)
-	{
-		FString AttributeName = TEXT("DDC_Summary") + Stat.Key.Replace(TEXT("."), TEXT("_"));
-
-		if (Stat.Value.IsNumeric())
-		{
-			Attributes.Emplace(AttributeName, FCString::Atof(*Stat.Value));
-		}
-		else
-		{
-			Attributes.Emplace(AttributeName, Stat.Value);
-		}
-	}
+	CookProfiling::GatherAccumulatedDDCStats(Attributes);
 
 #if UE_WITH_ZEN
 	// Gather Zen analytics
@@ -186,7 +193,42 @@ void FEditorTelemetry::RecordEvent_Cooking(TArray<FAnalyticsEventAttribute> Attr
 
 	FShaderStatsFunctions::GatherShaderAnalytics(Attributes);
 	
-	FStudioTelemetry::Get().RecordEvent(TEXT("Core.Cooking"), Attributes);
+	// Format the attribs into an easily digestible JSON format that splits each attribute into 
+	// a separate json object
+	// Example:
+	// [
+	// { "name": "PackTime", "value" : 13.34},
+	// { "name": "MaterialShader_Count", "value" : 12 },
+	// ...
+	// ]
+
+	FString CookStatsJsonStr;
+	Private::FAnalyticsJsonWriter CookStatsJsonWriter(&CookStatsJsonStr);
+	CookStatsJsonWriter.WriteArrayStart();
+	const FDecimalNumberFormattingRules FormattingRules;
+	for (auto& It : Attributes)
+	{
+		CookStatsJsonWriter.WriteObjectStart();
+		CookStatsJsonWriter.WriteValue("name", It.GetName());
+		const FString& StringValue = It.GetValue();
+
+		if (StringValue.IsNumeric())
+		{
+			double ValueAsDouble = 0;
+			LexFromString(ValueAsDouble, StringValue);
+			CookStatsJsonWriter.WriteValue("value", ValueAsDouble);
+		}
+		else
+		{
+			CookStatsJsonWriter.WriteValue("value", StringValue);
+		}
+		CookStatsJsonWriter.WriteObjectEnd();
+	}
+	CookStatsJsonWriter.WriteArrayEnd();
+	CookStatsJsonWriter.Close();
+
+	FStudioTelemetry::Get().RecordEvent(TEXT("Core.Cooking"), 
+		{ { TEXT("CookStats"), FJsonFragment(MoveTemp(CookStatsJsonStr)) } });
 #endif
 }
 
@@ -910,8 +952,6 @@ void FEditorTelemetry::StartSession()
 					return;
 				}
 				TArray<FAnalyticsEventAttribute> Attributes;
-				Attributes.Emplace(TEXT("MapName"), EditorMapName);
-
 				FEditorTelemetry::RecordEvent_Cooking(Attributes);
 				FEditorTelemetry::RecordEvent_CoreSystems(TEXT("Cooking"), Attributes);
 				FEditorTelemetry::RecordEvent_DDCVerification(TEXT("Cooking"), Attributes);
@@ -928,8 +968,6 @@ void FEditorTelemetry::StartSession()
 		CookingSpan = FStudioTelemetry::Get().StartSpan(TEXT("Cooking"));
 
 		TArray<FAnalyticsEventAttribute> Attributes;
-		Attributes.Emplace(TEXT("MapName"), EditorMapName);
-
 		CookingSpan->AddAttributes(Attributes);
 	});
 
diff --git a/Engine/Source/Developer/DerivedDataCache/Public/DerivedDataCacheUsageStats.h b/Engine/Source/Developer/DerivedDataCache/Public/DerivedDataCacheUsageStats.h
--- a/Engine/Source/Developer/DerivedDataCache/Public/DerivedDataCacheUsageStats.h
+++ b/Engine/Source/Developer/DerivedDataCache/Public/DerivedDataCacheUsageStats.h
@@ -299,7 +299,7 @@ public:
 	int64 TotalCount = 0;
 	double Efficiency = 0;
 
-	double GameThreadTimeSec;
+	double GameThreadTimeSec = 0;
 };
 
 struct FDerivedDataCacheResourceStatKeyFuncs : BaseKeyFuncs<FDerivedDataCacheResourceStat, FString, false>
diff --git a/Engine/Source/Editor/UnrealEd/Private/Commandlets/AssetRegistryGenerator.cpp b/Engine/Source/Editor/UnrealEd/Private/Commandlets/AssetRegistryGenerator.cpp
--- a/Engine/Source/Editor/UnrealEd/Private/Commandlets/AssetRegistryGenerator.cpp
+++ b/Engine/Source/Editor/UnrealEd/Private/Commandlets/AssetRegistryGenerator.cpp
@@ -10,7 +10,7 @@
 #include "CollectionManagerTypes.h"
 #include "Cooker/CookPackageData.h"
 #include "Cooker/CookPlatformManager.h"
-#include "Cooker/CookProfiling.h"
+#include "Cooker/CookProfilingPrivate.h"
 #include "Cooker/CookSandbox.h"
 #include "Cooker/CookWorkerClient.h"
 #include "Cooker/MPCollector.h"
diff --git a/Engine/Source/Editor/UnrealEd/Private/Commandlets/CookCommandlet.cpp b/Engine/Source/Editor/UnrealEd/Private/Commandlets/CookCommandlet.cpp
--- a/Engine/Source/Editor/UnrealEd/Private/Commandlets/CookCommandlet.cpp
+++ b/Engine/Source/Editor/UnrealEd/Private/Commandlets/CookCommandlet.cpp
@@ -11,7 +11,7 @@
 #include "AutoRTFM.h"
 #include "CookOnTheSide/CookOnTheFlyServer.h"
 #include "Cooker/CookImportsChecker.h"
-#include "Cooker/CookProfiling.h"
+#include "Cooker/CookProfilingPrivate.h"
 #include "CookerSettings.h"
 #include "Editor.h"
 #include "EngineGlobals.h"
diff --git a/Engine/Source/Editor/UnrealEd/Private/CookOnTheFlyServer.cpp b/Engine/Source/Editor/UnrealEd/Private/CookOnTheFlyServer.cpp
--- a/Engine/Source/Editor/UnrealEd/Private/CookOnTheFlyServer.cpp
+++ b/Engine/Source/Editor/UnrealEd/Private/CookOnTheFlyServer.cpp
@@ -34,7 +34,7 @@
 #include "Cooker/CookPackageData.h"
 #include "Cooker/CookPackagePreloader.h"
 #include "Cooker/CookPlatformManager.h"
-#include "Cooker/CookProfiling.h"
+#include "Cooker/CookProfilingPrivate.h"
 #include "Cooker/CookRequestCluster.h"
 #include "Cooker/CookRequests.h"
 #include "Cooker/CookSandbox.h"
diff --git a/Engine/Source/Editor/UnrealEd/Private/Cooker/CookGarbageCollect.cpp b/Engine/Source/Editor/UnrealEd/Private/Cooker/CookGarbageCollect.cpp
--- a/Engine/Source/Editor/UnrealEd/Private/Cooker/CookGarbageCollect.cpp
+++ b/Engine/Source/Editor/UnrealEd/Private/Cooker/CookGarbageCollect.cpp
@@ -10,7 +10,7 @@
 #include "Cooker/CookGenerationHelper.h"
 #include "Cooker/CookPackageData.h"
 #include "Cooker/CookPackagePreloader.h"
-#include "Cooker/CookProfiling.h"
+#include "Cooker/CookProfilingPrivate.h"
 #include "Cooker/CookTypes.h"
 #include "Cooker/PackageTracker.h"
 #include "CookOnTheSide/CookLog.h"
diff --git a/Engine/Source/Editor/UnrealEd/Private/Cooker/CookProfiling.cpp b/Engine/Source/Editor/UnrealEd/Private/Cooker/CookProfiling.cpp
--- a/Engine/Source/Editor/UnrealEd/Private/Cooker/CookProfiling.cpp
+++ b/Engine/Source/Editor/UnrealEd/Private/Cooker/CookProfiling.cpp
@@ -1,6 +1,7 @@
 // Copyright Epic Games, Inc. All Rights Reserved.
 
-#include "CookProfiling.h"
+#include "Cooker/CookProfiling.h" // Keep as first include to satisfy UBT
+#include "CookProfilingPrivate.h" 
 
 #include "Algo/GraphConvert.h"
 #include "Algo/Sort.h"
@@ -215,56 +216,422 @@ void ClearHierarchyTimers()
 #if ENABLE_COOK_STATS
 namespace DetailedCookStats
 {
+	// Cook.Profile.Params
 	FString CookProject;
 	FString CookCultures;
 	FString CookLabel;
 	FString TargetPlatforms;
-	double CookStartTime = 0.0;
-	double CookWallTimeSec = 0.0;
-	double StartupWallTimeSec = 0.0;
-	double StartCookByTheBookTimeSec = 0.0;
-	double TickCookOnTheSideTimeSec = 0.0;
-	double TickCookOnTheSideLoadPackagesTimeSec = 0.0;
-	double TickCookOnTheSideResolveRedirectorsTimeSec = 0.0;
-	double TickCookOnTheSideSaveCookedPackageTimeSec = 0.0;
-	double TickCookOnTheSidePrepareSaveTimeSec = 0.0;
-	double BlockOnAssetRegistryTimeSec = 0.0;
-	double GameCookModificationDelegateTimeSec = 0.0;
-	double TickLoopGCTimeSec = 0.0;
-	double TickLoopRecompileShaderRequestsTimeSec = 0.0;
-	double TickLoopShaderProcessAsyncResultsTimeSec = 0.0;
-	double TickLoopProcessDeferredCommandsTimeSec = 0.0;
-	double TickLoopTickCommandletStatsTimeSec = 0.0;
-	double TickLoopFlushRenderingCommandsTimeSec = 0.0;
-	double ShaderFlushTimeSec = 0.0;
-	double ValidationTimeSec = 0.0;
 	bool IsCookAll = false;
 	bool IsCookOnTheFly = false;
 	bool IsIterativeCook = false;
 	bool IsFastCook = false;
 	bool IsUnversioned = false;
 
-
-	// Stats tracked through FAutoRegisterCallback
+	// Cook.Profile
+	double CookStartTime;
+	double CookWallTimeSec;
+	double StartupWallTimeSec;
+	double StartCookByTheBookTimeSec;
+	double BlockOnAssetRegistryTimeSec;
+	double GameCookModificationDelegateTimeSec;
+	double TickCookOnTheSideTimeSec;
+	double TickCookOnTheSideLoadPackagesTimeSec;
+	double TickCookOnTheSideSaveCookedPackageTimeSec;
+	double TickCookOnTheSideResolveRedirectorsTimeSec;
+	double TickCookOnTheSidePrepareSaveTimeSec;
+	double TickLoopGCTimeSec;
+	double TickLoopRecompileShaderRequestsTimeSec;
+	double TickLoopShaderProcessAsyncResultsTimeSec;
+	double TickLoopProcessDeferredCommandsTimeSec;
+	double TickLoopTickCommandletStatsTimeSec;
+	double TickLoopFlushRenderingCommandsTimeSec;
+	double ShaderFlushTimeSec;
+	double ValidationTimeSec;
+
+	// CookOnTheFlyServer.Package.Load
+	int32 NumRequestedLoads = 0;
+	std::atomic<int32> NumDetectedLoads{ 0 };
+	
+	// CookOnTheFlyServer.Package.Save
+	uint32 NumPackagesIncrementallySkipped = 0;
+	
+	//"CookOnTheFlyServer.QueueSize"
 	int32 PeakRequestQueueSize = 0;
 	int32 PeakLoadQueueSize = 0;
 	int32 PeakSaveQueueSize = 0;
-	std::atomic<int32> NumDetectedLoads{ 0 };
-	int32 NumRequestedLoads = 0;
-	uint32 NumPackagesIncrementallySkipped = 0;
+	
 	TMap<FTopLevelAssetPath, uint32> NumPackagesIncrementallyModifiedByClass;
-	FCookStatsManager::FAutoRegisterCallback RegisterCookOnTheFlyServerStats([](FCookStatsManager::AddStatFuncRef AddStat)
+
+	FCookStatsManager::FAutoRegisterCallback RegisterStats([](FCookStatsManager::AddStatFuncRef AddStat)
 		{
-			AddStat(TEXT("Package.Load"), FCookStatsManager::CreateKeyValueArray(
-				TEXT("NumRequestedLoads"), NumRequestedLoads,
-				TEXT("NumPackagesLoaded"), NumDetectedLoads.load(),
-				TEXT("NumInlineLoads"), NumDetectedLoads - NumRequestedLoads));
-			AddStat(TEXT("Package.Save"), FCookStatsManager::CreateKeyValueArray(TEXT("NumPackagesIncrementallySkipped"), NumPackagesIncrementallySkipped));
-			AddStat(TEXT("CookOnTheFlyServer"), FCookStatsManager::CreateKeyValueArray(
-				TEXT("PeakRequestQueueSize"), PeakRequestQueueSize,
-				TEXT("PeakLoadQueueSize"), PeakLoadQueueSize,
-				TEXT("PeakSaveQueueSize"), PeakSaveQueueSize));
+
+#define COOK_PROFILE_STAT_TEXT(name) TEXT(#name), name
+
+			AddStat(TEXT("Cook.Profile.Params"),
+				FCookStatsManager::CreateKeyValueArray(
+					COOK_PROFILE_STAT_TEXT(CookProject),
+					COOK_PROFILE_STAT_TEXT(CookCultures),
+					COOK_PROFILE_STAT_TEXT(CookLabel),
+					COOK_PROFILE_STAT_TEXT(TargetPlatforms),
+					COOK_PROFILE_STAT_TEXT(IsCookAll),
+					COOK_PROFILE_STAT_TEXT(IsCookOnTheFly),
+					COOK_PROFILE_STAT_TEXT(IsIterativeCook),
+					COOK_PROFILE_STAT_TEXT(IsFastCook),
+					COOK_PROFILE_STAT_TEXT(IsUnversioned)
+					)
+			);
+
+			AddStat(TEXT("CookOnTheFlyServer.Package.Load"),
+				FCookStatsManager::CreateKeyValueArray(
+					TEXT("NumRequestedLoads"), NumRequestedLoads,
+					TEXT("NumPackagesLoaded"), NumDetectedLoads.load(),
+					TEXT("NumInlineLoads"), NumDetectedLoads - NumRequestedLoads
+				)
+			);
+
+			AddStat(TEXT("CookOnTheFlyServer.Package.Save"),
+				FCookStatsManager::CreateKeyValueArray(
+					TEXT("NumPackagesIncrementallySkipped"), NumPackagesIncrementallySkipped
+				)
+			);
+
+			AddStat(TEXT("CookOnTheFlyServer.QueueSize"),
+				FCookStatsManager::CreateKeyValueArray(
+					TEXT("PeakRequest"), PeakRequestQueueSize,
+					TEXT("PeakLoad"), PeakLoadQueueSize,
+					TEXT("PeakSave"), PeakSaveQueueSize
+				)
+			);
+
+			// Hierarchical profiling stats
+			const FString StatName(TEXT("Cook.Profile"));
+#define ADD_COOK_STAT_FLT(Path, Name) AddStat(StatName, FCookStatsManager::CreateKeyValueArray(TEXT("Path"), TEXT(Path), TEXT(#Name), Name))
+
+			ADD_COOK_STAT_FLT(" 0", CookWallTimeSec);
+			ADD_COOK_STAT_FLT(" 0. 0", StartupWallTimeSec);
+			double CookProcessingTime = CookWallTimeSec - StartupWallTimeSec;
+			ADD_COOK_STAT_FLT(" 0. 1", CookProcessingTime);
+			ADD_COOK_STAT_FLT(" 0. 1. 1", StartCookByTheBookTimeSec);
+			ADD_COOK_STAT_FLT(" 0. 1. 1. 0", BlockOnAssetRegistryTimeSec);
+			ADD_COOK_STAT_FLT(" 0. 1. 1. 1", GameCookModificationDelegateTimeSec);
+			ADD_COOK_STAT_FLT(" 0. 1. 2", TickCookOnTheSideTimeSec);
+			ADD_COOK_STAT_FLT(" 0. 1. 2. 0", TickCookOnTheSideLoadPackagesTimeSec);
+			ADD_COOK_STAT_FLT(" 0. 1. 2. 1", TickCookOnTheSideSaveCookedPackageTimeSec);
+			ADD_COOK_STAT_FLT(" 0. 1. 2. 1. 0", TickCookOnTheSideResolveRedirectorsTimeSec);
+			ADD_COOK_STAT_FLT(" 0. 1. 2. 2", TickCookOnTheSidePrepareSaveTimeSec);
+			ADD_COOK_STAT_FLT(" 0. 1. 3", TickLoopGCTimeSec);
+			ADD_COOK_STAT_FLT(" 0. 1. 4", TickLoopRecompileShaderRequestsTimeSec);
+			ADD_COOK_STAT_FLT(" 0. 1. 5", TickLoopShaderProcessAsyncResultsTimeSec);
+			ADD_COOK_STAT_FLT(" 0. 1. 6", TickLoopProcessDeferredCommandsTimeSec);
+			ADD_COOK_STAT_FLT(" 0. 1. 7", TickLoopTickCommandletStatsTimeSec);
+			ADD_COOK_STAT_FLT(" 0. 1. 8", TickLoopFlushRenderingCommandsTimeSec);
+			ADD_COOK_STAT_FLT(" 0. 1. 9", ShaderFlushTimeSec);
+			ADD_COOK_STAT_FLT(" 0. 1. 10", ValidationTimeSec);
 		});
+
+	void SendLogCookStats(ECookMode::Type CookMode)
+	{
+		if (IsCookingInEditor(CookMode))
+		{
+			return;
+		}
+
+		/** Used to store profile data for custom logging. */
+		struct FCookProfileData
+		{
+		public:
+			FCookProfileData(FString InPath, FString InKey, FString InValue) : Path(MoveTemp(InPath)), Key(MoveTemp(InKey)), Value(MoveTemp(InValue)) {}
+			FString Path;
+			FString Key;
+			FString Value;
+		};
+
+		// instead of printing the usage stats generically, we capture them so we can log a subset of them in an easy-to-read way.
+		TArray<FCookProfileData> CookProfileData;
+		TArray<FString> StatCategories;
+		TMap<FString, TArray<FCookStatsManager::StringKeyValue>> StatsInCategories;
+
+		/** this functor will take a collected cooker stat and log it out using some custom formatting based on known stats that are collected.. */
+		auto LogStatsFunc = [&CookProfileData, &StatCategories, &StatsInCategories]
+		(const FString& StatName, const TArray<FCookStatsManager::StringKeyValue>& StatAttributes)
+			{
+				// Some stats will use custom formatting to make a visibly pleasing summary.
+				bool bStatUsedCustomFormatting = false;
+
+				if (StatName == TEXT("DDC.Usage"))
+				{
+					// Don't even log this detailed DDC data. It's mostly only consumable by ingestion into pivot tools.
+					bStatUsedCustomFormatting = true;
+				}
+				else if (StatName.EndsWith(TEXT(".Usage"), ESearchCase::IgnoreCase))
+				{
+					// These are gathered through GatherResourceStats.
+					bStatUsedCustomFormatting = true;
+				}
+				else if (StatName == TEXT("DDC.Summary"))
+				{
+					bStatUsedCustomFormatting = true;
+				}
+				else if (StatName == TEXT("Cook.Profile"))
+				{
+					if (StatAttributes.Num() >= 2)
+					{
+						CookProfileData.Emplace(StatAttributes[0].Value, StatAttributes[1].Key, StatAttributes[1].Value);
+					}
+					bStatUsedCustomFormatting = true;
+				}
+
+				// if a stat doesn't use custom formatting, just spit out the raw info.
+				if (!bStatUsedCustomFormatting)
+				{
+					TArray<FCookStatsManager::StringKeyValue>& StatsInCategory = StatsInCategories.FindOrAdd(StatName);
+					if (StatsInCategory.Num() == 0)
+					{
+						StatCategories.Add(StatName);
+					}
+					StatsInCategory.Append(StatAttributes);
+				}
+			};
+
+		FDerivedDataCacheSummaryStats DDCSummaryStats;
+		GetDerivedDataCacheRef().GatherSummaryStats(DDCSummaryStats);
+
+		TArray<FDerivedDataCacheResourceStat> DDCResourceUsageStats;
+		GetDerivedDataCacheRef().GatherResourceStats(DDCResourceUsageStats);
+
+		FCookStatsManager::LogCookStats(LogStatsFunc);
+
+		TMap<FString, FDerivedDataCacheResourceStat> DDCStatsPerAssetType;
+		CookProfiling::CollectAccumulatedDDCStats(DDCResourceUsageStats, DDCStatsPerAssetType);
+
+		FString CookStatsFileName;
+		FString CookStatsJsonString;
+		TSharedPtr<TJsonWriter<TCHAR, TCondensedJsonPrintPolicy<TCHAR>>> CookStatsWriter{};
+		FStringBuilderBase Builder;
+		if (FParse::Value(FCommandLine::Get(), TEXT("-CookStatsFile="), CookStatsFileName))
+		{
+			uint32 MultiprocessId = UE::GetMultiprocessId();
+			if (MultiprocessId == 0)
+			{
+				CookStatsWriter = TJsonWriterFactory<TCHAR, TCondensedJsonPrintPolicy<TCHAR>>::Create(&CookStatsJsonString).ToSharedPtr();
+				CookStatsWriter->WriteObjectStart();
+			}
+			else
+			{
+				// Suppress the file creation on CookWorkers
+				// TODO: Replicate the information back to the CookDirector instead, UE-185774
+			}
+		}
+
+		UE_LOG(LogCookStats, Display, TEXT("Misc Cook Stats"));
+		UE_LOG(LogCookStats, Display, TEXT("==============="));
+		for (FString& StatCategory : StatCategories)
+		{
+			UE_LOG(LogCookStats, Display, TEXT("%s"), *StatCategory);
+			TArray<FCookStatsManager::StringKeyValue>& StatsInCategory = StatsInCategories.FindOrAdd(StatCategory);
+
+			// log each key/value pair, with the equal signs lined up.
+			for (const FCookStatsManager::StringKeyValue& StatKeyValue : StatsInCategory)
+			{
+				UE_LOG(LogCookStats, Display, TEXT("    %s=%s"), *StatKeyValue.Key, *StatKeyValue.Value);
+			}
+
+			if (CookStatsWriter)
+			{
+				CookStatsWriter->WriteObjectStart(StatCategory);
+				for (const FCookStatsManager::StringKeyValue& StatKeyValue : StatsInCategory)
+				{
+					CookStatsWriter->WriteValue(*StatKeyValue.Key, *StatKeyValue.Value);
+				}
+				CookStatsWriter->WriteObjectEnd();
+			}
+		}
+
+		// DDC Usage stats are custom formatted, and the above code just accumulated them into a TSet. Now log it with our special formatting for readability.
+		if (CookProfileData.Num() > 0)
+		{
+			UE_LOG(LogCookStats, Display, TEXT(""));
+			UE_LOG(LogCookStats, Display, TEXT("Cook Profile"));
+			UE_LOG(LogCookStats, Display, TEXT("============"));
+			for (const auto& ProfileEntry : CookProfileData)
+			{
+				UE_LOG(LogCookStats, Display, TEXT("%s.%s=%s"), *ProfileEntry.Path, *ProfileEntry.Key, *ProfileEntry.Value);
+			}
+
+			if (CookStatsWriter)
+			{
+				CookStatsWriter->WriteObjectStart("CookProfile");
+				for (const auto& ProfileEntry : CookProfileData)
+				{
+					CookStatsWriter->WriteObjectStart(ProfileEntry.Key);
+					CookStatsWriter->WriteValue(TEXT("Path"), ProfileEntry.Path);
+					CookStatsWriter->WriteValue(TEXT("Value"), ProfileEntry.Value);
+					CookStatsWriter->WriteObjectEnd();
+				}
+				CookStatsWriter->WriteObjectEnd();
+			}
+		}
+		if (DDCSummaryStats.Stats.Num() > 0)
+		{
+			UE_LOG(LogCookStats, Display, TEXT(""));
+			UE_LOG(LogCookStats, Display, TEXT("DDC Summary Stats"));
+			UE_LOG(LogCookStats, Display, TEXT("================="));
+			for (const auto& Attr : DDCSummaryStats.Stats)
+			{
+				UE_LOG(LogCookStats, Display, TEXT("%-16s=%10s"), *Attr.Key, *Attr.Value);
+			}
+
+			if (CookStatsWriter)
+			{
+				CookStatsWriter->WriteObjectStart("DDCSummaryStats");
+				for (const auto& Attr : DDCSummaryStats.Stats)
+				{
+					CookStatsWriter->WriteValue(*Attr.Key, *Attr.Value);
+				}
+				CookStatsWriter->WriteObjectEnd();
+			}
+		}
+
+		DumpDerivedDataBuildRemoteExecutorStats();
+
+		if (!DDCResourceUsageStats.IsEmpty())
+		{
+			Algo::SortBy(DDCResourceUsageStats, [](const FDerivedDataCacheResourceStat& Stat) { return Stat.BuildTimeSec + Stat.LoadTimeSec; }, TGreater());
+
+			UE_LOG(LogCookStats, Display, TEXT(""));
+			UE_LOG(LogCookStats, Display, TEXT("DDC Resource Stats"));
+			UE_LOG(LogCookStats, Display, TEXT("======================================================================================================="));
+			UE_LOG(LogCookStats, Display, TEXT("Asset Type                          Total Time (Sec)  GameThread Time (Sec)  Assets Built  MB Processed"));
+			UE_LOG(LogCookStats, Display, TEXT("----------------------------------  ----------------  ---------------------  ------------  ------------"));
+			for (const FDerivedDataCacheResourceStat& Stat : DDCResourceUsageStats)
+			{
+				UE_LOG(LogCookStats, Display, TEXT("%-34s  %16.2f  %21.2f  %12d  %12.2f"),
+					*Stat.AssetType, Stat.LoadTimeSec + Stat.BuildTimeSec, Stat.GameThreadTimeSec,
+					Stat.BuildCount, Stat.LoadSizeMB + Stat.BuildSizeMB);
+			}
+
+			if (CookStatsWriter)
+			{
+				CookStatsWriter->WriteObjectStart("DDCResourceStats");
+				for (const FDerivedDataCacheResourceStat& Stat : DDCResourceUsageStats)
+				{
+					CookStatsWriter->WriteObjectStart(*Stat.AssetType);
+					CookStatsWriter->WriteValue(TEXT("TotalTimeSec"), Stat.LoadTimeSec + Stat.BuildTimeSec);
+					CookStatsWriter->WriteValue(TEXT("GameThreadTimeSec"), Stat.GameThreadTimeSec);
+					CookStatsWriter->WriteValue(TEXT("AssetsBuilt"), Stat.BuildCount);
+					CookStatsWriter->WriteValue(TEXT("MBProcessed"), Stat.LoadSizeMB + Stat.BuildSizeMB);
+					CookStatsWriter->WriteObjectEnd();
+				}
+				CookStatsWriter->WriteObjectEnd();
+			}
+		}
+
+		if (!DDCStatsPerAssetType.IsEmpty())
+		{
+			UE_LOG(LogCookStats, Display, TEXT(""));
+			UE_LOG(LogCookStats, Display, TEXT("Accumulated DDC Resource Stats Per Asset Type"));
+			UE_LOG(LogCookStats, Display, TEXT("======================================================================================================="));
+
+			for (const TPair< FString, FDerivedDataCacheResourceStat>& Stat : DDCStatsPerAssetType)
+			{
+				UE_LOG(LogCookStats, Display, TEXT("%-34s  %16.2f  %21.2f  %12d  %12.2f"),
+					*Stat.Key, Stat.Value.LoadTimeSec + Stat.Value.BuildTimeSec, Stat.Value.GameThreadTimeSec,
+					Stat.Value.BuildCount, Stat.Value.LoadSizeMB + Stat.Value.BuildSizeMB);
+			}
+
+			if (CookStatsWriter)
+			{
+				CookStatsWriter->WriteObjectStart("DDCAccumulatedStats");
+				for (const TPair< FString, FDerivedDataCacheResourceStat>& Stat : DDCStatsPerAssetType)
+				{
+					CookStatsWriter->WriteObjectStart(*Stat.Key);
+					CookStatsWriter->WriteValue(TEXT("LoadTimeSec"), Stat.Value.LoadTimeSec);
+					CookStatsWriter->WriteValue(TEXT("LoadSizeMB"), Stat.Value.LoadSizeMB);
+					CookStatsWriter->WriteValue(TEXT("LoadCount"), Stat.Value.LoadCount);
+					CookStatsWriter->WriteValue(TEXT("BuildTimeSec"), Stat.Value.BuildTimeSec);
+					CookStatsWriter->WriteValue(TEXT("BuildSizeMB"), Stat.Value.BuildSizeMB);
+					CookStatsWriter->WriteValue(TEXT("BuildCount"), Stat.Value.BuildCount);
+					CookStatsWriter->WriteValue(TEXT("TotalCount"), Stat.Value.TotalCount);
+					CookStatsWriter->WriteValue(TEXT("Efficiency"), Stat.Value.Efficiency);
+					CookStatsWriter->WriteValue(TEXT("GameThreadTimeSec"), Stat.Value.GameThreadTimeSec);
+					CookStatsWriter->WriteValue(TEXT("TotalTimeSec"), Stat.Value.LoadTimeSec + Stat.Value.BuildTimeSec);
+					CookStatsWriter->WriteObjectEnd();
+				}
+				CookStatsWriter->WriteObjectEnd();
+			}
+		}
+
+		DumpBuildDependencyTrackerStats();
+
+		if (UE::Virtualization::IVirtualizationSystem::IsInitialized())
+		{
+			UE::Virtualization::IVirtualizationSystem::Get().DumpStats();
+		}
+
+		if (CookStatsWriter)
+		{
+			FShaderCompilerStats ShaderCompilerStats;
+			GShaderCompilingManager->GetLocalStats(ShaderCompilerStats);
+
+			TSharedPtr<FJsonObject> JsonShaderCompilerStats = ShaderCompilerStats.ToJson();
+			FJsonSerializer::Serialize(
+				MakeShared<FJsonValueObject>(JsonShaderCompilerStats),
+				TEXT("ShaderCompilerStats"),
+				CookStatsWriter.ToSharedRef(),
+				false
+			);
+		}
+
+		if (!CookStatsFileName.IsEmpty() && CookStatsWriter)
+		{
+			CookStatsWriter->WriteObjectEnd();
+			CookStatsWriter->Close();
+			TUniquePtr<FArchive> CookStatsJsonFile(IFileManager::Get().CreateFileWriter(*CookStatsFileName));
+			if (!CookStatsJsonFile)
+			{
+				UE_LOG(LogCookStats, Warning, TEXT("Could not write to CookStatsFile %s."), *CookStatsFileName);
+			}
+			else
+			{
+				CookStatsJsonFile->Serialize(TCHAR_TO_ANSI(*CookStatsJsonString), CookStatsJsonString.Len());
+				CookStatsJsonFile->Close();
+			}
+		}
+	}
+};
+
+void OutputIncrementalCookStats()
+{
+	using namespace DetailedCookStats;
+
+	uint32 NumPackagesIncrementallyModifiedByClassSum = 0;
+	TArray<TPair<FTopLevelAssetPath, uint32>> ClassAndNums = NumPackagesIncrementallyModifiedByClass.Array();
+	ClassAndNums.RemoveAllSwap(
+		[&NumPackagesIncrementallyModifiedByClassSum](const TPair<FTopLevelAssetPath, uint32>& Pair)
+		{
+			NumPackagesIncrementallyModifiedByClassSum += Pair.Value;
+			return Pair.Value == 0;
+		});
+	if (NumPackagesIncrementallyModifiedByClassSum > 0)
+	{
+		Algo::Sort(ClassAndNums,
+			[](const TPair<FTopLevelAssetPath, uint32>& A, const TPair<FTopLevelAssetPath, uint32>& B)
+			{
+				if (A.Value != B.Value)
+				{
+					return A.Value > B.Value;
+				}
+				return A.Key.Compare(B.Key) < 0;
+			});
+
+		UE_LOG(LogCookStats, Display, TEXT("Incrementally Modified Counts for Each Class:"));
+		UE_LOG(LogCookStats, Display, TEXT("\tTotal, %u"), NumPackagesIncrementallyModifiedByClassSum);
+		for (const TPair<FTopLevelAssetPath, uint32>& Pair : ClassAndNums)
+		{
+			UE_LOG(LogCookStats, Display, TEXT("\t%s, %u"),
+				Pair.Key.IsValid() ? *Pair.Key.ToString() : TEXT("<UnknownClass>"), Pair.Value);
+		}
+	}
 }
 #endif
 
@@ -915,7 +1282,7 @@ void DumpPackageReferencers(const TCHAR* DebugNameOfPackageList,
 
 		TStringBuilder<256> ReferencerName;
 		RootData.Vertex != InvalidVertex ? ProfileData.AppendVertexName(RootData.Vertex, ReferencerName)
-			: (void) ReferencerName.Append(TEXT("<Unknown>"));
+			: (void)ReferencerName.Append(TEXT("<Unknown>"));
 		LogAr.Logf(TEXT("\t%s"), *ReferencerName);
 		check(!RootData.Packages.IsEmpty()); // Creation of each rootdata includes adding a package to it.
 		Algo::Sort(RootData.Packages, [](FPackageDumpData* A, FPackageDumpData* B)
@@ -964,312 +1331,81 @@ void DumpPackageReferencers(const TCHAR* DebugNameOfPackageList,
 } // namespace UE::Cook
 
 #if ENABLE_COOK_STATS
-
-namespace DetailedCookStats
+namespace CookProfiling
 {
-
-FCookStatsManager::FAutoRegisterCallback RegisterCookStats([](FCookStatsManager::AddStatFuncRef AddStat)
-{
-	const FString StatName(TEXT("Cook.Profile"));
-	#define ADD_COOK_STAT_FLT(Path, Name) AddStat(StatName, FCookStatsManager::CreateKeyValueArray(TEXT("Path"), TEXT(Path), TEXT(#Name), Name))
-	ADD_COOK_STAT_FLT(" 0", CookWallTimeSec);
-	ADD_COOK_STAT_FLT(" 0. 0", StartupWallTimeSec);
-	ADD_COOK_STAT_FLT(" 0. 1", StartCookByTheBookTimeSec);
-	ADD_COOK_STAT_FLT(" 0. 1. 0", BlockOnAssetRegistryTimeSec);
-	ADD_COOK_STAT_FLT(" 0. 1. 1", GameCookModificationDelegateTimeSec);
-	ADD_COOK_STAT_FLT(" 0. 2", TickCookOnTheSideTimeSec);
-	ADD_COOK_STAT_FLT(" 0. 2. 0", TickCookOnTheSideLoadPackagesTimeSec);
-	ADD_COOK_STAT_FLT(" 0. 2. 1", TickCookOnTheSideSaveCookedPackageTimeSec);
-	ADD_COOK_STAT_FLT(" 0. 2. 1. 0", TickCookOnTheSideResolveRedirectorsTimeSec);
-	ADD_COOK_STAT_FLT(" 0. 2. 2", TickCookOnTheSidePrepareSaveTimeSec);
-	ADD_COOK_STAT_FLT(" 0. 3", TickLoopGCTimeSec);
-	ADD_COOK_STAT_FLT(" 0. 4", TickLoopRecompileShaderRequestsTimeSec);
-	ADD_COOK_STAT_FLT(" 0. 5", TickLoopShaderProcessAsyncResultsTimeSec);
-	ADD_COOK_STAT_FLT(" 0. 6", TickLoopProcessDeferredCommandsTimeSec);
-	ADD_COOK_STAT_FLT(" 0. 7", TickLoopTickCommandletStatsTimeSec);
-	ADD_COOK_STAT_FLT(" 0. 8", TickLoopFlushRenderingCommandsTimeSec);
-	ADD_COOK_STAT_FLT(" 0. 9", ShaderFlushTimeSec);
-	ADD_COOK_STAT_FLT(" 0. 10", ValidationTimeSec);
-	FString CookParameters; // Empty value to write a header with name "CookParameters"
-	ADD_COOK_STAT_FLT(" 1", CookParameters);
-	ADD_COOK_STAT_FLT(" 1. 0", TargetPlatforms);
-	ADD_COOK_STAT_FLT(" 1. 1", CookProject);
-	ADD_COOK_STAT_FLT(" 1. 2", CookCultures);
-	ADD_COOK_STAT_FLT(" 1. 3", IsCookAll);
-	ADD_COOK_STAT_FLT(" 1. 4", IsCookOnTheFly);
-	ADD_COOK_STAT_FLT(" 1. 5", IsIterativeCook);
-	ADD_COOK_STAT_FLT(" 1. 6", IsUnversioned);
-	ADD_COOK_STAT_FLT(" 1. 7", CookLabel);
-	ADD_COOK_STAT_FLT(" 1. 8", IsFastCook);
-		
-	#undef ADD_COOK_STAT_FLT
-});
-
-void SendLogCookStats(ECookMode::Type CookMode)
-{
-	if (IsCookingInEditor(CookMode))
+	const char* AssetTypesToAccumulate[] =
 	{
-		return;
-	}
-
-	/** Used to store profile data for custom logging. */
-	struct FCookProfileData
-	{
-	public:
-		FCookProfileData(FString InPath, FString InKey, FString InValue) : Path(MoveTemp(InPath)), Key(MoveTemp(InKey)), Value(MoveTemp(InValue)) {}
-		FString Path;
-		FString Key;
-		FString Value;
-	};
-
-	// instead of printing the usage stats generically, we capture them so we can log a subset of them in an easy-to-read way.
-	TArray<FDerivedDataCacheResourceStat> DDCResourceUsageStats;
-	TArray<FCookStatsManager::StringKeyValue> DDCSummaryStats;
-	TArray<FCookProfileData> CookProfileData;
-	TArray<FString> StatCategories;
-	TMap<FString, TArray<FCookStatsManager::StringKeyValue>> StatsInCategories;
-
-	/** this functor will take a collected cooker stat and log it out using some custom formatting based on known stats that are collected.. */
-	auto LogStatsFunc = [&DDCResourceUsageStats, &DDCSummaryStats, &CookProfileData, &StatCategories, &StatsInCategories]
-	(const FString& StatName, const TArray<FCookStatsManager::StringKeyValue>& StatAttributes)
-	{
-		// Some stats will use custom formatting to make a visibly pleasing summary.
-		bool bStatUsedCustomFormatting = false;
-
-		if (StatName == TEXT("DDC.Usage"))
-		{
-			// Don't even log this detailed DDC data. It's mostly only consumable by ingestion into pivot tools.
-			bStatUsedCustomFormatting = true;
-		}
-		else if (StatName.EndsWith(TEXT(".Usage"), ESearchCase::IgnoreCase))
-		{
-			// These are gathered through GatherResourceStats.
-			bStatUsedCustomFormatting = true;
-		}
-		else if (StatName == TEXT("DDC.Summary"))
-		{
-			DDCSummaryStats.Append(StatAttributes);
-			bStatUsedCustomFormatting = true;
-		}
-		else if (StatName == TEXT("Cook.Profile"))
-		{
-			if (StatAttributes.Num() >= 2)
-			{
-				CookProfileData.Emplace(StatAttributes[0].Value, StatAttributes[1].Key, StatAttributes[1].Value);
-			}
-			bStatUsedCustomFormatting = true;
-		}
-
-		// if a stat doesn't use custom formatting, just spit out the raw info.
-		if (!bStatUsedCustomFormatting)
-		{
-			TArray<FCookStatsManager::StringKeyValue>& StatsInCategory = StatsInCategories.FindOrAdd(StatName);
-			if (StatsInCategory.Num() == 0)
-			{
-				StatCategories.Add(StatName);
-			}
-			StatsInCategory.Append(StatAttributes);
-		}
+		"AnimationSequence",
+		"Audio",
+		"CardRepresentation",
+		"DistanceField",
+		"MaterialShader",
+		"MaterialTranslation",
+		"NiagaraScript",
+		"SkeletalMesh",
+		"StaticMesh",
+		"Texture",
+		"Other"
 	};
 
-	GetDerivedDataCacheRef().GatherResourceStats(DDCResourceUsageStats);
-
-	FCookStatsManager::LogCookStats(LogStatsFunc);
-
-	FString CookStatsFileName;
-	FString CookStatsJsonString;
-	TSharedPtr<TJsonWriter<TCHAR, TCondensedJsonPrintPolicy<TCHAR>>> CookStatsWriter{};
-	FStringBuilderBase Builder;
-	if (FParse::Value(FCommandLine::Get(), TEXT("-CookStatsFile="), CookStatsFileName))
+	void CollectAccumulatedDDCStats(
+		TArray<FDerivedDataCacheResourceStat> DDCResourceUsageStats,
+		TMap<FString, FDerivedDataCacheResourceStat>& StatsPerAssetType)
 	{
-		uint32 MultiprocessId = UE::GetMultiprocessId();
-		if (MultiprocessId == 0)
-		{
-			CookStatsWriter = TJsonWriterFactory<TCHAR, TCondensedJsonPrintPolicy<TCHAR>>::Create(&CookStatsJsonString).ToSharedPtr();
-			CookStatsWriter->WriteObjectStart();
-		}
-		else
+		for (const char* KnownAssetType : CookProfiling::AssetTypesToAccumulate)
 		{
-			// Suppress the file creation on CookWorkers
-			// TODO: Replicate the information back to the CookDirector instead, UE-185774
+			StatsPerAssetType.Add(KnownAssetType, FDerivedDataCacheResourceStat());
 		}
-	}
-
-	UE_LOG(LogCookStats, Display, TEXT("Misc Cook Stats"));
-	UE_LOG(LogCookStats, Display, TEXT("==============="));
-	for (FString& StatCategory : StatCategories)
-	{
-		UE_LOG(LogCookStats, Display, TEXT("%s"), *StatCategory);
-		TArray<FCookStatsManager::StringKeyValue>& StatsInCategory = StatsInCategories.FindOrAdd(StatCategory);
 
-		// log each key/value pair, with the equal signs lined up.
-		for (const FCookStatsManager::StringKeyValue& StatKeyValue : StatsInCategory)
-		{
-			UE_LOG(LogCookStats, Display, TEXT("    %s=%s"), *StatKeyValue.Key, *StatKeyValue.Value);
-		}
-
-		if (CookStatsWriter)
-		{
-			CookStatsWriter->WriteObjectStart(StatCategory);
-			for (const FCookStatsManager::StringKeyValue& StatKeyValue : StatsInCategory)
-			{
-				CookStatsWriter->WriteValue(*StatKeyValue.Key, *StatKeyValue.Value);
-			}
-			CookStatsWriter->WriteObjectEnd();
-		}
-	}
-
-	// DDC Usage stats are custom formatted, and the above code just accumulated them into a TSet. Now log it with our special formatting for readability.
-	if (CookProfileData.Num() > 0)
-	{
-		UE_LOG(LogCookStats, Display, TEXT(""));
-		UE_LOG(LogCookStats, Display, TEXT("Cook Profile"));
-		UE_LOG(LogCookStats, Display, TEXT("============"));
-		for (const auto& ProfileEntry : CookProfileData)
+		FDerivedDataCacheResourceStat* OtherAccumulator = StatsPerAssetType.Find("Other");
+		for (const FDerivedDataCacheResourceStat& Stat : DDCResourceUsageStats)
 		{
-			UE_LOG(LogCookStats, Display, TEXT("%s.%s=%s"), *ProfileEntry.Path, *ProfileEntry.Key, *ProfileEntry.Value);
-		}
+			FDerivedDataCacheResourceStat* Accumulator = StatsPerAssetType.Find(Stat.AssetType);
 
-		if (CookStatsWriter)
-		{
-			CookStatsWriter->WriteObjectStart("CookProfile");
-			for (const auto& ProfileEntry : CookProfileData)
+			if (!Accumulator)
 			{
-				CookStatsWriter->WriteObjectStart(ProfileEntry.Key);
-				CookStatsWriter->WriteValue(TEXT("Path"), ProfileEntry.Path);
-				CookStatsWriter->WriteValue(TEXT("Value"), ProfileEntry.Value);
-				CookStatsWriter->WriteObjectEnd();
+				Accumulator = OtherAccumulator;
 			}
-			CookStatsWriter->WriteObjectEnd();
-		}
-	}
-	if (DDCSummaryStats.Num() > 0)
-	{
-		UE_LOG(LogCookStats, Display, TEXT(""));
-		UE_LOG(LogCookStats, Display, TEXT("DDC Summary Stats"));
-		UE_LOG(LogCookStats, Display, TEXT("================="));
-		for (const auto& Attr : DDCSummaryStats)
-		{
-			UE_LOG(LogCookStats, Display, TEXT("%-16s=%10s"), *Attr.Key, *Attr.Value);
-		}
 
-		if (CookStatsWriter)
-		{
-			CookStatsWriter->WriteObjectStart("DDCSummaryStats");
-			for (const auto& Attr : DDCSummaryStats)
-			{
-				CookStatsWriter->WriteValue(*Attr.Key, *Attr.Value);
-			}
-			CookStatsWriter->WriteObjectEnd();
+			*Accumulator += Stat;
 		}
 	}
 
-	DumpDerivedDataBuildRemoteExecutorStats();
-
-	if (!DDCResourceUsageStats.IsEmpty())
+	void CollectAttributesForAccumulatedDDCUsageStats(
+		TMap<FString, FDerivedDataCacheResourceStat> StatsPerAssetType,
+		TArray<FAnalyticsEventAttribute>& Attributes)
 	{
-		Algo::SortBy(DDCResourceUsageStats, [](const FDerivedDataCacheResourceStat& Stat) { return Stat.BuildTimeSec + Stat.LoadTimeSec; }, TGreater());
+		const FString BaseName = TEXT("DDCAccumulatedStats_");
+		Attributes.Reserve(StatsPerAssetType.Num() * 10); // pre allocate ten attributes per asset type
 
-		UE_LOG(LogCookStats, Display, TEXT(""));
-		UE_LOG(LogCookStats, Display, TEXT("DDC Resource Stats"));
-		UE_LOG(LogCookStats, Display, TEXT("======================================================================================================="));
-		UE_LOG(LogCookStats, Display, TEXT("Asset Type                          Total Time (Sec)  GameThread Time (Sec)  Assets Built  MB Processed"));
-		UE_LOG(LogCookStats, Display, TEXT("----------------------------------  ----------------  ---------------------  ------------  ------------"));
-		for (const FDerivedDataCacheResourceStat& Stat : DDCResourceUsageStats)
+		for (auto& StatForAsset : StatsPerAssetType)
 		{
-			UE_LOG(LogCookStats, Display, TEXT("%-34s  %16.2f  %21.2f  %12d  %12.2f"),
-				*Stat.AssetType, Stat.LoadTimeSec + Stat.BuildTimeSec, Stat.GameThreadTimeSec,
-				Stat.BuildCount, Stat.LoadSizeMB + Stat.BuildSizeMB);
+			FDerivedDataCacheResourceStat& Stat = StatForAsset.Value;
+			FString AttrName = BaseName + StatForAsset.Key;
+			Attributes.Emplace(AttrName + TEXT("LoadTimeSec"), Stat.LoadTimeSec);
+			Attributes.Emplace(AttrName + TEXT("LoadSizeMB"), Stat.LoadSizeMB);
+			Attributes.Emplace(AttrName + TEXT("LoadCount"), Stat.LoadCount);
+			Attributes.Emplace(AttrName + TEXT("BuildTimeSec"), Stat.BuildTimeSec);
+			Attributes.Emplace(AttrName + TEXT("BuildSizeMB"), Stat.BuildSizeMB);
+			Attributes.Emplace(AttrName + TEXT("BuildCount"), Stat.BuildCount);
+			Attributes.Emplace(AttrName + TEXT("TotalCount"), Stat.TotalCount);
+			Attributes.Emplace(AttrName + TEXT("Efficiency"), Stat.Efficiency);
+			Attributes.Emplace(AttrName + TEXT("GameThreadTimeSec"), Stat.GameThreadTimeSec);
+			Attributes.Emplace(AttrName + TEXT("TotalTimeSec"), Stat.LoadTimeSec + Stat.BuildTimeSec);
 		}
-
-		if (CookStatsWriter)
-		{
-			CookStatsWriter->WriteObjectStart("DDCResourceStats");
-			for (const FDerivedDataCacheResourceStat& Stat : DDCResourceUsageStats)
-			{
-				CookStatsWriter->WriteObjectStart(*Stat.AssetType);
-				CookStatsWriter->WriteValue(TEXT("TotalTimeSec"), Stat.LoadTimeSec + Stat.BuildTimeSec);
-				CookStatsWriter->WriteValue(TEXT("GameThreadTimeSec"), Stat.GameThreadTimeSec);
-				CookStatsWriter->WriteValue(TEXT("AssetsBuilt"), Stat.BuildCount);
-				CookStatsWriter->WriteValue(TEXT("MBProcessed"), Stat.LoadSizeMB + Stat.BuildSizeMB);
-				CookStatsWriter->WriteObjectEnd();
-			}
-			CookStatsWriter->WriteObjectEnd();
-		}
-	}
-
-	DumpBuildDependencyTrackerStats();
-
-	if (UE::Virtualization::IVirtualizationSystem::IsInitialized())
-	{
-		UE::Virtualization::IVirtualizationSystem::Get().DumpStats();
 	}
 
-	if (CookStatsWriter)
+	void GatherAccumulatedDDCStats(
+		TArray<FAnalyticsEventAttribute>& Attributes)
 	{
-		FShaderCompilerStats ShaderCompilerStats;
-		GShaderCompilingManager->GetLocalStats(ShaderCompilerStats);
-
-		TSharedPtr<FJsonObject> JsonShaderCompilerStats = ShaderCompilerStats.ToJson();
-		FJsonSerializer::Serialize(
-			MakeShared<FJsonValueObject>(JsonShaderCompilerStats),
-			TEXT("ShaderCompilerStats"),
-			CookStatsWriter.ToSharedRef(),
-			false
-		);
-	}
-
-	if (!CookStatsFileName.IsEmpty())
-	{
-		CookStatsWriter->WriteObjectEnd();
-		CookStatsWriter->Close();
-		TUniquePtr<FArchive> CookStatsJsonFile(IFileManager::Get().CreateFileWriter(*CookStatsFileName));
-		if (!CookStatsJsonFile)
-		{
-			UE_LOG(LogCookStats, Warning, TEXT("Could not write to CookStatsFile %s."), *CookStatsFileName);
-		}
-		else
-		{
-			CookStatsJsonFile->Serialize(TCHAR_TO_ANSI(*CookStatsJsonString), CookStatsJsonString.Len());
-			CookStatsJsonFile->Close();
-		}
-	}
-}
+		TArray<FDerivedDataCacheResourceStat> DDCResourceUsageStats;
+		GetDerivedDataCacheRef().GatherResourceStats(DDCResourceUsageStats);
 
-} // namespace DetailedCookStats
+		TMap<FString, FDerivedDataCacheResourceStat> StatsPerAssetType;
+		CollectAccumulatedDDCStats(DDCResourceUsageStats, StatsPerAssetType);
 
-void OutputIncrementalCookStats()
-{
-	using namespace DetailedCookStats;
-
-	uint32 NumPackagesIncrementallyModifiedByClassSum = 0;
-	TArray<TPair<FTopLevelAssetPath, uint32>> ClassAndNums = NumPackagesIncrementallyModifiedByClass.Array();
-	ClassAndNums.RemoveAllSwap(
-		[&NumPackagesIncrementallyModifiedByClassSum](const TPair<FTopLevelAssetPath, uint32>& Pair)
-		{
-			NumPackagesIncrementallyModifiedByClassSum += Pair.Value;
-			return Pair.Value == 0;
-		});
-	if (NumPackagesIncrementallyModifiedByClassSum > 0)
-	{
-		Algo::Sort(ClassAndNums,
-			[](const TPair<FTopLevelAssetPath, uint32>& A, const TPair<FTopLevelAssetPath, uint32>& B)
-			{
-				if (A.Value != B.Value)
-				{
-					return A.Value > B.Value;
-				}
-				return A.Key.Compare(B.Key) < 0;
-			});
-
-		UE_LOG(LogCookStats, Display, TEXT("Incrementally Modified Counts for Each Class:"));
-		UE_LOG(LogCookStats, Display, TEXT("\tTotal, %u"), NumPackagesIncrementallyModifiedByClassSum);
-		for (const TPair<FTopLevelAssetPath, uint32>& Pair : ClassAndNums)
-		{
-			UE_LOG(LogCookStats, Display, TEXT("\t%s, %u"),
-				Pair.Key.IsValid() ? *Pair.Key.ToString() : TEXT("<UnknownClass>"), Pair.Value);
-		}
+		CollectAttributesForAccumulatedDDCUsageStats(StatsPerAssetType, Attributes);
 	}
-}
+};
 
 #endif // ENABLE_COOK_STATS
\ No newline at end of file
diff --git a/Engine/Source/Editor/UnrealEd/Private/Cooker/CookProfiling.h b/Engine/Source/Editor/UnrealEd/Private/Cooker/CookProfiling.h
deleted file mode 100644
diff --git a/Engine/Source/Editor/UnrealEd/Private/Cooker/CookProfilingPrivate.h b/Engine/Source/Editor/UnrealEd/Private/Cooker/CookProfilingPrivate.h
new file mode 100644
--- /dev/null
+++ b/Engine/Source/Editor/UnrealEd/Private/Cooker/CookProfilingPrivate.h
@@ -0,0 +1,216 @@
+// Copyright Epic Games, Inc. All Rights Reserved.
+
+#pragma once
+
+#include "Containers/ArrayView.h"
+#include "CookOnTheSide/CookOnTheFlyServer.h"
+#include "CoreTypes.h"
+#include "ProfilingDebugging/CookStats.h"
+#include "Stats/Stats.h"
+
+#include <atomic>
+
+#define OUTPUT_COOKTIMING ENABLE_COOK_STATS
+#define PROFILE_NETWORK 0
+
+#if ENABLE_COOK_STATS
+#include "Containers/Map.h"
+#include "UObject/TopLevelAssetPath.h"
+#endif
+
+#if OUTPUT_COOKTIMING
+#include "Trace/Trace.inl"
+#include "ProfilingDebugging/FormatArgsTrace.h"
+#include "ProfilingDebugging/ScopedTimers.h"
+#endif
+
+struct FWeakObjectPtr;
+
+#if OUTPUT_COOKTIMING
+
+void OutputHierarchyTimers();
+void OutputIncrementalCookStats();
+void ClearHierarchyTimers();
+
+struct FHierarchicalTimerInfo;
+
+struct FScopeTimer
+{
+public:
+	FScopeTimer(const FScopeTimer&) = delete;
+	FScopeTimer(FScopeTimer&&) = delete;
+
+	FScopeTimer(int InId, const char* InName, bool IncrementScope = false );
+
+	void Start();
+	void Stop();
+
+	~FScopeTimer();
+
+private:
+	uint64					StartTime = 0;
+	FHierarchicalTimerInfo* HierarchyTimerInfo;
+	FHierarchicalTimerInfo* PrevTimerInfo;
+};
+
+#define UE_CREATE_HIERARCHICAL_COOKTIMER(name, incrementScope) \
+			FScopeTimer UE_JOIN(__HierarchicalCookTimerScope, __LINE__)(__COUNTER__, #name, incrementScope); 
+#define UE_CREATE_TEXT_HIERARCHICAL_COOKTIMER(name, incrementScope) \
+			FScopeTimer UE_JOIN(__HierarchicalCookTimerScope, __LINE__)(__COUNTER__, name, incrementScope);
+
+// Emits trace events denoting scope/lifetime of an activity on the cooking channel.
+#define UE_SCOPED_COOKTIMER(name)						TRACE_CPUPROFILER_EVENT_SCOPE_ON_CHANNEL(name, CookChannel)
+// More expensive version that takes a dynamic TCHAR*
+#define UE_SCOPED_TEXT_COOKTIMER(name)					TRACE_CPUPROFILER_EVENT_SCOPE_TEXT_ON_CHANNEL(name, CookChannel)
+
+#define UE_SCOPED_COOKTIMER_AND_DURATION(name, durationStorage) \
+			FScopedDurationTimer UE_JOIN(__CookTimerScope, __LINE__)(durationStorage); \
+			UE_SCOPED_COOKTIMER(name)
+// More expensive version that takes a dynamic TCHAR*
+#define UE_SCOPED_TEXT_COOKTIMER_AND_DURATION(name, durationStorage) \
+			FScopedDurationTimer UE_JOIN(__CookTimerScope, __LINE__)(durationStorage); \
+			UE_SCOPED_TEXT_COOKTIMER(name)
+
+// Emits trace events denoting scope/lifetime of an activity on the cooking channel.  Also creates a named hierarchical timer that can be aggregated and reported at cook completion.
+#define UE_SCOPED_HIERARCHICAL_COOKTIMER(name) \
+			UE_CREATE_HIERARCHICAL_COOKTIMER(name, true); \
+			UE_JOIN(__HierarchicalCookTimerScope, __LINE__).Start(); \
+			UE_SCOPED_COOKTIMER(name)
+// More expensive version that takes a dynamic TCHAR*
+#define UE_SCOPED_TEXT_HIERARCHICAL_COOKTIMER(name) \
+			UE_CREATE_TEXT_HIERARCHICAL_COOKTIMER(name, true); \
+			UE_JOIN(__HierarchicalCookTimerScope, __LINE__).Start(); \
+			UE_SCOPED_TEXT_COOKTIMER(name)
+// Version that takes a separate variable for where the duration should be stored
+#define UE_SCOPED_HIERARCHICAL_COOKTIMER_AND_DURATION(name, durationStorage) \
+			FScopedDurationTimer UE_JOIN(__CookTimerScope, __LINE__)(durationStorage); \
+			UE_SCOPED_HIERARCHICAL_COOKTIMER(name)
+// More expensive version that takes a separate variable for where the duration should be stored and a dynamic TCHAR*
+#define UE_SCOPED_TEXT_HIERARCHICAL_COOKTIMER_AND_DURATION(name, durationStorage) \
+			FScopedDurationTimer UE_JOIN(__CookTimerScope, __LINE__)(durationStorage); \
+			UE_SCOPED_TEXT_HIERARCHICAL_COOKTIMER(name)
+
+#define UE_CUSTOM_COOKTIMER_LOG Cpu
+
+// Emits trace events denoting scope/lifetime of an activity on the cooking channel.
+#define UE_SCOPED_CUSTOM_COOKTIMER(name)				UE_TRACE_LOG_SCOPED_T(UE_CUSTOM_COOKTIMER_LOG, name, CookChannel)
+#define UE_SCOPED_CUSTOM_COOKTIMER_AND_DURATION(name, durationStorage) \
+			FScopedDurationTimer UE_JOIN(__CookTimerScope, __LINE__)(durationStorage); \
+			UE_SCOPED_CUSTOM_COOKTIMER(name)
+
+// Emits trace events denoting scope/lifetime of an activity on the cooking channel.  Also creates a named hierarchical timer that can be aggregated and reported at cook completion.
+#define UE_SCOPED_HIERARCHICAL_CUSTOM_COOKTIMER(name) \
+			UE_CREATE_HIERARCHICAL_COOKTIMER(name, true); \
+			UE_JOIN(__HierarchicalCookTimerScope, __LINE__).Start(); \
+			UE_SCOPED_CUSTOM_COOKTIMER(name)
+#define UE_SCOPED_HIERARCHICAL_CUSTOM_COOKTIMER_AND_DURATION(name, durationStorage) \
+			FScopedDurationTimer UE_JOIN(__CookTimerScope, __LINE__)(durationStorage); \
+			UE_SCOPED_HIERARCHICAL_CUSTOM_COOKTIMER(name)
+#define UE_ADD_CUSTOM_COOKTIMER_META(name, key, value) << name.key(value)
+
+#else
+#define UE_SCOPED_COOKTIMER(name)
+#define UE_SCOPED_TEXT_COOKTIMER(name)
+#define UE_SCOPED_COOKTIMER_AND_DURATION(name, durationStorage)
+#define UE_SCOPED_TEXT_COOKTIMER_AND_DURATION(name, durationStorage)
+#define UE_SCOPED_HIERARCHICAL_COOKTIMER(name)
+#define UE_SCOPED_TEXT_HIERARCHICAL_COOKTIMER(name)
+#define UE_SCOPED_HIERARCHICAL_COOKTIMER_AND_DURATION(name, durationStorage)
+#define UE_SCOPED_TEXT_HIERARCHICAL_COOKTIMER_AND_DURATION(name, durationStorage)
+
+#define UE_CUSTOM_COOKTIMER_LOG Cpu
+#define UE_SCOPED_CUSTOM_COOKTIMER(name)
+#define UE_SCOPED_CUSTOM_COOKTIMER_AND_DURATION(name, durationStorage)
+#define UE_SCOPED_HIERARCHICAL_CUSTOM_COOKTIMER(name)
+#define UE_SCOPED_HIERARCHICAL_CUSTOM_COOKTIMER_AND_DURATION(name, durationStorage)
+#define UE_ADD_CUSTOM_COOKTIMER_META(name, key, value)
+
+inline void OutputHierarchyTimers() {}
+inline void ClearHierarchyTimers() {}
+#endif
+
+#if ENABLE_COOK_STATS
+namespace DetailedCookStats
+{
+// Cook.Profile.Params
+extern FString CookProject;
+extern FString CookCultures;
+extern FString CookLabel;
+extern FString TargetPlatforms;
+extern bool IsCookAll;
+extern bool IsCookOnTheFly;
+extern bool IsIterativeCook;
+extern bool IsFastCook;
+extern bool IsUnversioned;
+
+// Cook.Profile
+extern double CookStartTime;
+extern double CookWallTimeSec;
+extern double StartupWallTimeSec;
+extern double StartCookByTheBookTimeSec;
+extern double TickCookOnTheSideTimeSec;
+extern double TickCookOnTheSideLoadPackagesTimeSec;
+extern double TickCookOnTheSideResolveRedirectorsTimeSec;
+extern double TickCookOnTheSideSaveCookedPackageTimeSec;
+extern double TickCookOnTheSidePrepareSaveTimeSec;
+extern double BlockOnAssetRegistryTimeSec;
+extern double GameCookModificationDelegateTimeSec;
+extern double TickLoopGCTimeSec;
+extern double TickLoopRecompileShaderRequestsTimeSec;
+extern double TickLoopShaderProcessAsyncResultsTimeSec;
+extern double TickLoopProcessDeferredCommandsTimeSec;
+extern double TickLoopTickCommandletStatsTimeSec;
+extern double TickLoopFlushRenderingCommandsTimeSec;
+extern double ShaderFlushTimeSec;
+extern double ValidationTimeSec;
+
+// CookOnTheFlyServer.Package.Load
+extern int32 NumRequestedLoads;
+extern std::atomic<int32> NumDetectedLoads;
+
+// CookOnTheFlyServer.Package.Save
+extern uint32 NumPackagesIncrementallySkipped;
+
+//"CookOnTheFlyServer.QueueSize"
+extern int32 PeakRequestQueueSize;
+extern int32 PeakLoadQueueSize;
+extern int32 PeakSaveQueueSize;
+
+extern TMap<FTopLevelAssetPath, uint32> NumPackagesIncrementallyModifiedByClass;
+
+void SendLogCookStats(ECookMode::Type CookMode);
+
+}
+#endif
+
+#if PROFILE_NETWORK
+double TimeTillRequestStarted = 0.0;
+double TimeTillRequestForfilled = 0.0;
+double TimeTillRequestForfilledError = 0.0;
+double WaitForAsyncFilesWrites = 0.0;
+FEvent* NetworkRequestEvent = nullptr;
+#endif
+
+DECLARE_STATS_GROUP(TEXT("Cooking"), STATGROUP_Cooking, STATCAT_Advanced);
+
+DECLARE_CYCLE_STAT(TEXT("Precache Derived data for platform"), STAT_TickPrecacheCooking, STATGROUP_Cooking);
+DECLARE_CYCLE_STAT(TEXT("Tick cooking"), STAT_TickCooker, STATGROUP_Cooking);
+
+namespace UE::Cook
+{
+
+/**
+ * Used for profiling memory after a garbage collect. Find how many instances of each class exist, and find the list
+ * of referencers that are keeping instances of each class in memory. Dump this information to the log (GLog).
+ * 
+ * @param SessionStartupObjects: The list of objects that existed at the start of the cook session. Objects in this
+ *                               list are removed from the results since we do not expect them to be garbage collected.
+ */
+void DumpObjClassList(TConstArrayView<FWeakObjectPtr> SessionStartupObjects);
+
+/** Used for profiling after a soft garbage collect. Find the referencers keeping each package in memory. */
+void DumpPackageReferencers(const TCHAR* DebugNameOfPackageList, TConstArrayView<FWeakObjectPtr> InitialObjects, 
+	TConstArrayView<UPackage*> Packages = TConstArrayView<UPackage*>());
+
+
+}
\ No newline at end of file
diff --git a/Engine/Source/Editor/UnrealEd/Private/Cooker/CookRequestCluster.cpp b/Engine/Source/Editor/UnrealEd/Private/Cooker/CookRequestCluster.cpp
--- a/Engine/Source/Editor/UnrealEd/Private/Cooker/CookRequestCluster.cpp
+++ b/Engine/Source/Editor/UnrealEd/Private/Cooker/CookRequestCluster.cpp
@@ -16,7 +16,7 @@
 #include "Cooker/CookLogPrivate.h"
 #include "Cooker/CookPackageData.h"
 #include "Cooker/CookPlatformManager.h"
-#include "Cooker/CookProfiling.h"
+#include "Cooker/CookProfilingPrivate.h"
 #include "Cooker/CookRequests.h"
 #include "Cooker/CookTypes.h"
 #include "Cooker/PackageTracker.h"
diff --git a/Engine/Source/Editor/UnrealEd/Private/Cooker/CookSavePackage.cpp b/Engine/Source/Editor/UnrealEd/Private/Cooker/CookSavePackage.cpp
--- a/Engine/Source/Editor/UnrealEd/Private/Cooker/CookSavePackage.cpp
+++ b/Engine/Source/Editor/UnrealEd/Private/Cooker/CookSavePackage.cpp
@@ -16,7 +16,7 @@
 #include "Cooker/CookPackageArtifacts.h"
 #include "Cooker/CookPackageData.h"
 #include "Cooker/CookPlatformManager.h"
-#include "Cooker/CookProfiling.h"
+#include "Cooker/CookProfilingPrivate.h"
 #include "Cooker/CookRequestCluster.h"
 #include "Cooker/CookRequests.h"
 #include "Cooker/CookTypes.h"
diff --git a/Engine/Source/Editor/UnrealEd/Private/Cooker/CookTypes.h b/Engine/Source/Editor/UnrealEd/Private/Cooker/CookTypes.h
--- a/Engine/Source/Editor/UnrealEd/Private/Cooker/CookTypes.h
+++ b/Engine/Source/Editor/UnrealEd/Private/Cooker/CookTypes.h
@@ -6,7 +6,7 @@
 #include "Containers/Map.h"
 #include "Containers/Set.h"
 #include "Containers/UnrealString.h"
-#include "Cooker/CookProfiling.h"
+#include "Cooker/CookProfilingPrivate.h"
 #include "Cooker/MPCollector.h"
 #include "CookOnTheSide/CookOnTheFlyServer.h" // ECookTickFlags
 #include "DerivedDataRequestOwner.h"
diff --git a/Engine/Source/Editor/UnrealEd/Private/Cooker/PackageTracker.cpp b/Engine/Source/Editor/UnrealEd/Private/Cooker/PackageTracker.cpp
--- a/Engine/Source/Editor/UnrealEd/Private/Cooker/PackageTracker.cpp
+++ b/Engine/Source/Editor/UnrealEd/Private/Cooker/PackageTracker.cpp
@@ -5,7 +5,7 @@
 #include "CookOnTheFlyServerInterface.h"
 #include "CookPackageData.h"
 #include "CookPlatformManager.h"
-#include "CookProfiling.h"
+#include "CookProfilingPrivate.h"
 #include "Engine/Level.h"
 #include "Misc/ScopeRWLock.h"
 #include "Misc/StringBuilder.h"
diff --git a/Engine/Source/Editor/UnrealEd/Public/Cooker/CookProfiling.h b/Engine/Source/Editor/UnrealEd/Public/Cooker/CookProfiling.h
new file mode 100644
--- /dev/null
+++ b/Engine/Source/Editor/UnrealEd/Public/Cooker/CookProfiling.h
@@ -0,0 +1,32 @@
+// Copyright Epic Games, Inc. All Rights Reserved.
+
+#pragma once
+
+#include "ProfilingDebugging/CookStats.h"
+
+#if ENABLE_COOK_STATS
+
+#include "AnalyticsEventAttribute.h"
+#include "Containers/Map.h"
+#include "Containers/Array.h"
+#include "Containers/UnrealString.h"
+#include "CoreTypes.h"
+#include "DerivedDataCacheUsageStats.h"
+
+namespace CookProfiling
+{
+	extern const char* AssetTypesToAccumulate[];
+
+	UNREALED_API void CollectAccumulatedDDCStats(
+		TArray<FDerivedDataCacheResourceStat> DDCResourceUsageStats,
+		TMap<FString, FDerivedDataCacheResourceStat>& StatsPerAssetType);
+
+	UNREALED_API void CollectAttributesForAccumulatedDDCUsageStats(
+		TMap<FString, FDerivedDataCacheResourceStat> StatsPerAssetType,
+		TArray<FAnalyticsEventAttribute>& Attributes);
+
+	UNREALED_API void GatherAccumulatedDDCStats(
+		TArray<FAnalyticsEventAttribute>& Attributes);
+}
+
+#endif
\ No newline at end of file
