[FBX Importer] Fixing re-import skeletalmesh with different Skeleton causing scrambled mesh results.

#rnx

diff --git a/Engine/Source/Editor/UnrealEd/Private/Fbx/FbxSkeletalMeshImport.cpp b/Engine/Source/Editor/UnrealEd/Private/Fbx/FbxSkeletalMeshImport.cpp
--- a/Engine/Source/Editor/UnrealEd/Private/Fbx/FbxSkeletalMeshImport.cpp
+++ b/Engine/Source/Editor/UnrealEd/Private/Fbx/FbxSkeletalMeshImport.cpp
@@ -2174,6 +2174,8 @@ USkeletalMesh* UnFbx::FFbxImporter::ImportSkeletalMesh(FImportSkeletalMeshArgs &
 		if (ExistSkelMeshDataPtr)
 		{
 			SkeletalMeshImportUtils::RestoreExistingSkelMeshData(ExistSkelMeshDataPtr, SkeletalMesh, ImportSkeletalMeshArgs.LodIndex, ImportOptions->bCanShowDialog, ImportOptions->bImportAsSkeletalSkinning, ImportOptions->bResetToFbxOnMaterialConflict);
+
+			SkeletalMesh->GetRefBasesInvMatrix().Reset();
 		}
 
 		SkeletalMesh->CalculateInvRefMatrices();
