Bump incremental cook version to address a stale shader library causing errors during staging.

#Synced-CL 48523383

diff --git a/Engine/Source/Editor/UnrealEd/Private/Cooker/CookIncrementalVersion.cpp b/Engine/Source/Editor/UnrealEd/Private/Cooker/CookIncrementalVersion.cpp
--- a/Engine/Source/Editor/UnrealEd/Private/Cooker/CookIncrementalVersion.cpp
+++ b/Engine/Source/Editor/UnrealEd/Private/Cooker/CookIncrementalVersion.cpp
@@ -8,6 +8,6 @@ namespace UE::Cook
 {
 
 // Change CookIncrementalVersion to a new guid when all packages in an incremental cook need to be invalidated
-FGuid CookIncrementalVersion(0x7F3A91C2, 0xD4B8821E, 0x5ACF33B7, 0x9E41F08C);
+FGuid CookIncrementalVersion(0x26D70737, 0x5D0B4412, 0xB62833CB, 0x3845E4FE);
 
 }
