UE: Bump the cook incremental version again. Shader change

diff --git a/Engine/Source/Editor/UnrealEd/Private/Cooker/CookIncrementalVersion.cpp b/Engine/Source/Editor/UnrealEd/Private/Cooker/CookIncrementalVersion.cpp
--- a/Engine/Source/Editor/UnrealEd/Private/Cooker/CookIncrementalVersion.cpp
+++ b/Engine/Source/Editor/UnrealEd/Private/Cooker/CookIncrementalVersion.cpp
@@ -8,6 +8,6 @@ namespace UE::Cook
 {
 
 // Change CookIncrementalVersion to a new guid when all packages in an incremental cook need to be invalidated
-FGuid CookIncrementalVersion(0x26D70737, 0x5D0B4412, 0xB62833CB, 0x3845E4FE);
+FGuid CookIncrementalVersion(0xB95A5251, 0x68EC4F0B, 0x8D584FDF, 0xFBF5AD9B);
 
 }
