EditorUsability : UnrealEd
- Improved ortho view grid density logic when in a mode where view background is irrelevant

#Synced-CL 48903505

diff --git a/Engine/Source/Editor/UnrealEd/Private/EditorComponents.cpp b/Engine/Source/Editor/UnrealEd/Private/EditorComponents.cpp
--- a/Engine/Source/Editor/UnrealEd/Private/EditorComponents.cpp
+++ b/Engine/Source/Editor/UnrealEd/Private/EditorComponents.cpp
@@ -543,26 +543,23 @@ void FEditorCommonDrawHelper::DrawOldGrid(const FSceneView* View,FPrimitiveDrawI
 
 void FEditorCommonDrawHelper::DrawGridSection(double ViewportGridY, double MaxSize, FVector* A,FVector* B, FVector::FReal* AX, FVector::FReal* BX,int32 Axis,const FSceneView* View,FPrimitiveDrawInterface* PDI )
 {
-	if( ViewportGridY == 0 )
+	if (FMath::IsNearlyZero(ViewportGridY))
 	{
 		// Don't draw zero-size grid
 		return;
 	}
 
-	// todo
-	int32 Exponent = GEditor->IsGridSizePowerOfTwo() ? 8 : 10;
-
 	const double SizeX = View->UnscaledViewRect.Width();
 	const double Zoom = (1.0 / View->ViewMatrices.GetProjectionMatrix().M[0][0]) * 2.0 / SizeX;
 	const double Dist = SizeX * Zoom / ViewportGridY;
 
 	// when the grid fades
-	static double Tweak = 4.0;
-
-	double IncValue = FMath::LogX(Exponent, Dist / (double)(SizeX / Tweak));
+	constexpr double Tweak = 6.0;
+	const int32 Exponent = GEditor->IsGridSizePowerOfTwo() ? 8 : 10;
+	const double IncValue = FMath::LogX(Exponent, Dist / (double)(SizeX / Tweak));
 	int32 IncScale = 1;
 
-	for(double x = 0; x < IncValue; ++x)
+	for (double x = 0; x < IncValue; ++x)
 	{
 		IncScale *= Exponent;
 	}
@@ -573,28 +570,6 @@ void FEditorCommonDrawHelper::DrawGridSection(double ViewportGridY, double MaxSi
 		return;
 	}
 
-	// 0 excluded for hard transitions .. 0.5 for very soft transitions
-	const double TransitionRegion = 0.5;
-
-	const double InvTransitionRegion = 1.0 / TransitionRegion;
-	const double Fract = IncValue - FMath::FloorToDouble(IncValue);
-	float AlphaA = FMath::Clamp<float>(static_cast<float>(Fract * InvTransitionRegion), 0.0f, 1.0f);
-	float AlphaB = FMath::Clamp<float>(static_cast<float>(InvTransitionRegion - Fract * InvTransitionRegion), 0.0f, 1.0f);
-
-	if(IncValue < -0.5)
-	{
-		// no fade in magnification case
-		AlphaA = 1.0f;
-		AlphaB = 1.0f;
-	}
-
-	const int32 MajorLineInterval = FMath::TruncToInt(GEditor->GetGridInterval());
-
-	const FLinearColor Background = View->BackgroundColor;
-
-	FLinearColor MajorColor = FMath::Lerp(Background, FLinearColor::White, 0.05f);
-	FLinearColor MinorColor = FMath::Lerp(Background, FLinearColor::White, 0.02f);
-
 	const FMatrix InvViewProjMatrix = View->ViewMatrices.GetInvViewProjectionMatrix();
 	int64 FirstLine = FMath::TruncToInt(InvViewProjMatrix.TransformPosition(FVector(-1, -1, 0.5f)).Component(Axis) / ViewportGridY);
 	int64 LastLine = FMath::TruncToInt(InvViewProjMatrix.TransformPosition(FVector(+1, +1, 0.5f)).Component(Axis) / ViewportGridY);
@@ -603,36 +578,48 @@ void FEditorCommonDrawHelper::DrawGridSection(double ViewportGridY, double MaxSi
 		Exchange(FirstLine, LastLine);
 	}
 
+	// Check if we're in a mode that is not using the view background (uses lighting)
+	const FLinearColor Background = View->BackgroundColor;
+	bool bIsLitMode = Background.GetLuminance() > 0.5;
+	if (View->Family 
+		&& View->Family->EngineShowFlags.Lighting > 0 
+		&& View->Family->EngineShowFlags.LightingOnlyOverride == 0)
+	{
+		bIsLitMode = true;
+	}
+
+	// Reduce grid density when further away by removing minor lines
+	const double Fract = IncValue - FMath::FloorToDouble(IncValue);
+	const float MinorLineAlpha = Fract < 0.5 || Zoom < 1.f ? 0.1 : 0;
+	constexpr float MajorLineAlpha = 0.25f;
+
+	// Use contrasting color based on lit/unlit mode
+	const FLinearColor LineColor = bIsLitMode ? FLinearColor::Black : FLinearColor::White;
+
 	// Draw major and minor grid lines
 	const int32 FirstLineClamped = FMath::Max<int32>(FirstLine - 1, -MaxSize / ViewportGridY) / IncScale;
 	const int32 LastLineClamped = FMath::Min<int32>(LastLine + 1, +MaxSize / ViewportGridY) / IncScale;
-	for( int32 LineIndex = FirstLineClamped; LineIndex <= LastLineClamped; ++LineIndex )
+	const int32 MajorLineInterval = FMath::TruncToInt(GEditor->GetGridInterval());
+	for (int32 LineIndex = FirstLineClamped; LineIndex <= LastLineClamped; ++LineIndex)
 	{
 		*AX = FPlatformMath::TruncToDouble(LineIndex * ViewportGridY) * IncScale;
 		*BX = FPlatformMath::TruncToDouble(LineIndex * ViewportGridY) * IncScale;
 
-		// Only minor lines fade out with ortho zoom distance.  Origin lines and major lines are drawn
-		// at 100% opacity, but with a brighter value
+		// Only minor lines fade out with ortho zoom distance
 		const bool bIsMajorLine = (MajorLineInterval == 0) || ((LineIndex % MajorLineInterval) == 0);
 
-		{
-			// Don't bother drawing the world origin line.  We'll do that later.
-			const bool bIsOriginLine = ( LineIndex == 0 );
-			if( !bIsOriginLine )
-			{
-				FLinearColor Color;
-				if( bIsMajorLine )
-				{
-					Color = FMath::Lerp( Background, MajorColor, AlphaA );
-				}
-				else
-				{
-					Color = FMath::Lerp( Background, MinorColor, AlphaB );
-				}
+		FLinearColor Color;
 
-				PDI->DrawLine(*A,*B,Color,SDPG_World);
-			}
+		if (bIsMajorLine)
+		{
+			Color = FMath::Lerp(FLinearColor::Transparent, LineColor, MajorLineAlpha);
+		}
+		else
+		{
+			Color = FMath::Lerp(FLinearColor::Transparent, LineColor, MinorLineAlpha);
 		}
+
+		PDI->DrawTranslucentLine(*A, *B, Color, SDPG_World);
 	}
 }
 
