Bump incremental cook version to address a stale shader library causing errors during staging.

#Synced-CL 47485904

diff --git a/Engine/Source/Editor/UnrealEd/Private/Cooker/CookIncrementalVersion.cpp b/Engine/Source/Editor/UnrealEd/Private/Cooker/CookIncrementalVersion.cpp
--- a/Engine/Source/Editor/UnrealEd/Private/Cooker/CookIncrementalVersion.cpp
+++ b/Engine/Source/Editor/UnrealEd/Private/Cooker/CookIncrementalVersion.cpp
@@ -8,6 +8,6 @@ namespace UE::Cook
 {
 
 // Change CookIncrementalVersion to a new guid when all packages in an incremental cook need to be invalidated
-FGuid CookIncrementalVersion(0xB95A5251, 0x68EC4F0B, 0x8D584FDF, 0xFBF5AD9B);
+FGuid CookIncrementalVersion(0x71B523B5, 0x596C49F0, 0x99F52FDC, 0x82EFF942);
 
 }
