IncrementalCook: Provide information on which packages were referenced and which were incrementally skipped to ICookArtifacts. CookArtifacts, including the ShaderLibrary, need this information to delete no-longer-referenced data.
#rnx
#Synced-CL 48806790

diff --git a/Engine/Source/Editor/UnrealEd/Classes/CookOnTheSide/CookOnTheFlyServer.h b/Engine/Source/Editor/UnrealEd/Classes/CookOnTheSide/CookOnTheFlyServer.h
--- a/Engine/Source/Editor/UnrealEd/Classes/CookOnTheSide/CookOnTheFlyServer.h
+++ b/Engine/Source/Editor/UnrealEd/Classes/CookOnTheSide/CookOnTheFlyServer.h
@@ -1104,6 +1104,11 @@ private:
 	 * incremental cook)
 	 */
 	void FinishPopulateShaderLibrary(const ITargetPlatform* TargetPlatform, FString const& Name);
+	/**
+	 * Prune the named ShaderLibrary of any data no longer referenced by assets in the current cook or
+	 * IncrementalCook oplog.
+	 */
+	void PruneShaderLibrary(const ITargetPlatform* TargetPlatform, const FString& Name);
 
 	/** Invokes the necessary FShaderCodeLibrary functions to save and close a named code library. */
 	void SaveShaderLibrary(const ITargetPlatform* TargetPlatform, FString const& Name);
diff --git a/Engine/Source/Editor/UnrealEd/Private/Commandlets/AssetRegistryGenerator.cpp b/Engine/Source/Editor/UnrealEd/Private/Commandlets/AssetRegistryGenerator.cpp
--- a/Engine/Source/Editor/UnrealEd/Private/Commandlets/AssetRegistryGenerator.cpp
+++ b/Engine/Source/Editor/UnrealEd/Private/Commandlets/AssetRegistryGenerator.cpp
@@ -777,7 +777,7 @@ bool FAssetRegistryGenerator::GenerateStreamingInstallManifest(int64 InOverrideC
 				PakChunkOptions += TEXT(" encryptionkeyguid=") + Guid.ToString();
 			}
 
-			// If this chunk has a seperate unique asset registry, add it to first subchunk's manifest here
+			// If this chunk has a separate unique asset registry, add it to first subchunk's manifest here
 			if (SubChunkIndex == 0)
 			{
 				// For chunks with unique asset registry name, pakchunkIndex should equal chunkid
@@ -1183,7 +1183,7 @@ FAssetPackageData* FAssetRegistryGenerator::GetAssetPackageData(const FName& Pac
 	return State.CreateOrGetAssetPackageData(PackageName);
 }
 
-void FAssetRegistryGenerator::UpdateKeptPackages()
+void FAssetRegistryGenerator::UpdateKeptPackages(TSet<FName>& InOutAllPackages)
 {
 	for (TPair<FName, FIterativelySkippedPackageUpdateData>& Pair : PreviousPackagesToUpdate)
 	{
@@ -1201,10 +1201,22 @@ void FAssetRegistryGenerator::UpdateKeptPackages()
 		{
 			State.AddReferencers(FAssetIdentifier(Pair.Key), Pair.Value.PackageReferencers);
 		}
+
+		// Keep each of currently unreferenced but previously referenced (or previously development-only)
+		// packages from an earlier incremental cook in our development AssetRegistry, but not in the list
+		// of referenced packages for the current cook.
+		InOutAllPackages.Add(Pair.Key); // May already be present
 	}
 	PreviousPackagesToUpdate.Empty();
 }
 
+TArray<FName> FAssetRegistryGenerator::GetPreviousIncrementalCookPackages()
+{
+	TArray<FName> PackageNames;
+	PreviousPackagesToUpdate.GetKeys(PackageNames);
+	return PackageNames;
+}
+
 TMap<int32, FString> FAssetRegistryGenerator::GetPakChunkIdToStringOverrideMap()
 {
 	TMap<int32, FString> ChunkIdStringOverride;
@@ -1774,8 +1786,8 @@ void FAssetRegistryGenerator::ComputePackageRemovals(const FAssetRegistryState&
 	OutRemovedPackages = RemovedPackageSet.Array();
 }
 
-void FAssetRegistryGenerator::FinalizeChunkIDs(const TSet<FName>& InCookedPackages,
-	const TSet<FName>& InDevelopmentOnlyPackages, UE::Cook::FCookSandbox& InSandboxFile,
+void FAssetRegistryGenerator::FinalizeChunkIDs(const TSet<FName>& CookedPackages,
+	const TSet<FName>& DevelopmentOnlyPackages, UE::Cook::FCookSandbox& InSandboxFile,
 	bool bGenerateStreamingInstallManifest, const TSet<FName>& StartupPackages)
 {
 	LLM_SCOPE_BYTAG(Cooker_GeneratedAssetRegistry);
@@ -1795,17 +1807,14 @@ void FAssetRegistryGenerator::FinalizeChunkIDs(const TSet<FName>& InCookedPackag
 		bGenerateChunks = ShouldPlatformGenerateStreamingInstallManifest(TargetPlatform);
 	}
 
-	CookedPackages = InCookedPackages;
-	DevelopmentOnlyPackages = InDevelopmentOnlyPackages;
-
-	// Possibly apply previous AssetData and AssetPackageData for packages kept from a previous cook
-	UpdateKeptPackages();
-
 	TSet<FName> AllPackages;
 	AllPackages.Append(CookedPackages);
 	AllPackages.Append(DevelopmentOnlyPackages);
 
-	// Prune our asset registry to cooked + dev only list
+	// Possibly apply previous AssetData and AssetPackageData for packages kept from a previous cook
+	UpdateKeptPackages(AllPackages);
+
+	// Prune our asset registry to union of cooked + dev only + (previously cooked or dev only)
 	{
 		FAssetRegistrySerializationOptions DevelopmentSaveOptions;
 		AssetRegistry.InitializeSerializationOptions(DevelopmentSaveOptions, TargetPlatform, UE::AssetRegistry::ESerializationTarget::ForDevelopment);
@@ -1951,7 +1960,7 @@ void FAssetRegistryGenerator::AddAssetToFileOrderRecursive(const FName& InPackag
 }
 
 bool FAssetRegistryGenerator::SaveAssetRegistry(const FString& SandboxPath, bool bSerializeDevelopmentAssetRegistry, 
-	bool bForceNoFilter, uint64& OutDevelopmentAssetRegistryHash)
+	bool bForceNoFilter, uint64& OutDevelopmentAssetRegistryHash, const TSet<FName>& CookedPackages)
 {
 	LLM_SCOPE_BYTAG(Cooker_GeneratedAssetRegistry);
 	UE_LOG(LogAssetRegistryGenerator, Display, TEXT("Saving asset registry v%d."), FAssetRegistryVersion::Type::LatestVersion);
diff --git a/Engine/Source/Editor/UnrealEd/Private/CookOnTheFlyServer.cpp b/Engine/Source/Editor/UnrealEd/Private/CookOnTheFlyServer.cpp
--- a/Engine/Source/Editor/UnrealEd/Private/CookOnTheFlyServer.cpp
+++ b/Engine/Source/Editor/UnrealEd/Private/CookOnTheFlyServer.cpp
@@ -808,7 +808,7 @@ void UCookOnTheFlyServer::AddCookOnTheFlyPlatformFromGameThread(ITargetPlatform*
 	InitializeShadersForCookOnTheFly(BeginContext.TargetPlatforms);
 	// SaveAssetRegistry is done in CookByTheBookFinished for CBTB, but we need at the start of CookOnTheFly to send as startup information to connecting clients
 	uint64 DevelopmentAssetRegistryHash = 0;
-	PlatformData->RegistryGenerator->SaveAssetRegistry(GetSandboxAssetRegistryFilename(), true, false, DevelopmentAssetRegistryHash);
+	PlatformData->RegistryGenerator->SaveAssetRegistry(GetSandboxAssetRegistryFilename(), true, false, DevelopmentAssetRegistryHash, TSet<FName>());
 
 	// Initialize systems that nothing in AddCookOnTheFlyPlatformFromGameThread references
 	// Functions in this section are not dependent upon each other and can be ordered arbitrarily or for async performance
@@ -9402,28 +9402,38 @@ void UCookOnTheFlyServer::CookByTheBookFinishedInternal()
 	struct FEndCookPlatformData
 	{
 		TSet<FName> CookedPackageNames;
+		TArray<FPackageData*> CookedPackageDatas;
+		TArray<FPackageData*> IgnorePackageDatas;
 	};
 	TArray<FEndCookPlatformData> EndCookPlatformDatas;
 	EndCookPlatformDatas.SetNum(PlatformManager->GetSessionPlatforms().Num());
 
+	// Call FinishPopulateShaderLibrary to load the previous version of the shader libraries for all platforms, and
+	// call UpdateOplogPackages to report which assets are still referenced, before pruning any no-longer-referenced
+	// shaders from any platforms. We do it in this order because platforms can share shaderformats, so we need to
+	// know if any platform is referencing shaderformat X before we can prune it.
 	int32 PlatformIndex = 0;
 	for (const ITargetPlatform* TargetPlatform : PlatformManager->GetSessionPlatforms())
 	{
 		FEndCookPlatformData& EndCookPlatformData = EndCookPlatformDatas[PlatformIndex++];
-		if (FindOrCreateSaveContext(TargetPlatform).PackageWriterCapabilities.bReadOnly)
+		FCookSavePackageContext& SaveContext = FindOrCreateSaveContext(TargetPlatform);
+		if (SaveContext.PackageWriterCapabilities.bReadOnly)
 		{
 			continue;
 		}
 
 		FPlatformData* PlatformData = PlatformManager->GetPlatformData(TargetPlatform);
 		FAssetRegistryGenerator& Generator = *PlatformData->RegistryGenerator;
-		TArray<FPackageData*> CookedPackageDatas;
-		TArray<FPackageData*> IgnorePackageDatas;
 
-		FString PlatformNameString = TargetPlatform->PlatformName();
+		if (bCacheShaderLibraries)
+		{
+			UE_SCOPED_HIERARCHICAL_COOKTIMER(FinishPopulateShaderLibrary);
+			FinishPopulateShaderLibrary(TargetPlatform, LibraryName);
+		}
+
+		TArray<FPackageData*>& CookedPackageDatas = EndCookPlatformData.CookedPackageDatas;
+		TArray<FPackageData*>& IgnorePackageDatas = EndCookPlatformData.IgnorePackageDatas;
 
-		TSet<FName>& CookedPackageNames = EndCookPlatformData.CookedPackageNames;
-		TSet<FName> IgnorePackageNames;
 		{
 			UE_SCOPED_HIERARCHICAL_COOKTIMER(CalculateReferencedSet);
 
@@ -9445,27 +9455,91 @@ void UCookOnTheFlyServer::CookByTheBookFinishedInternal()
 					}
 				}
 			}
+		}
 
-			CookedPackageNames.Reserve(CookedPackageDatas.Num());
+		// Call UpdateOplogPackages on all Cook Artifacts. We have to do it at this point to handle shader libraries
+		// correctly. FShaderLibraryCookArtifact needs UpdateOplogPackages called after FinishPopulateShaderLibrary,
+		// so that it has all of its data loaded and ready to be pruned. But it has to do this pruning before any calls
+		// to FinalizeChunkIds or SaveShaderLibrary; those functions need to write out the pruned version of the data.
+		// Another reason: We call Generator.GetPreviousIncrementalCookPackages to calculate the data we pass to 
+		// UpdateOplogPackages, and the data read by that function is cleared during FinalizeChunkIds.
+		if (!bLegacyBuildDependencies)
+		{
+			UE_SCOPED_HIERARCHICAL_COOKTIMER(CookArtifacts.UpdateOplogPackages);
+			using namespace UE::Cook::Artifact;
+			FUpdateOplogPackagesContext OplogContext(*this);
+			OplogContext.TargetPlatform = TargetPlatform;
 			for (FPackageData* PackageData : CookedPackageDatas)
 			{
-				CookedPackageNames.Add(PackageData->GetPackageName());
+				FOplogPackageData& Data = OplogContext.OplogPackages.FindOrAdd(PackageData->GetPackageName());
+				Data.bReferenced = true;
+				const FPackagePlatformData& PackagePlatformData = PackageData->FindOrAddPlatformData(TargetPlatform);
+				Data.bIncrementallyCooked = !PackagePlatformData.IsIncrementallySkipped();
+			}
+			for (FName PackageName : Generator.GetPreviousIncrementalCookPackages())
+			{
+				// We rely on Data.bReferenced being initialized to false for newly added elements. In that case,
+				// calling FindOrAdd is sufficient to add the entry we need because we only need to add entries for
+				// any non-referenced-set packages; all of the referenced-set packages were added in the
+				// for (PackageData : CookedPackageDatas) loop above.
+				OplogContext.OplogPackages.FindOrAdd(PackageName);
 			}
 
-			IgnorePackageNames.Reserve(IgnorePackageDatas.Num());
-			for (FPackageData* PackageData : IgnorePackageDatas)
+			for (const TRefCountPtr<UE::Cook::ICookArtifact>& Artifact : CookArtifacts)
 			{
-				IgnorePackageNames.Add(PackageData->GetPackageName());
+				Artifact->UpdateOplogPackages(OplogContext);
 			}
 		}
-		
-		if (bCacheShaderLibraries)
+	}
+
+	// Trigger the ShaderLibrary to prune the ShaderFormats used by each platform, now that we have reported
+	// which assets are still in use by each platform.
+	if (!bLegacyBuildDependencies && bCacheShaderLibraries)
+	{
+		UE_SCOPED_HIERARCHICAL_COOKTIMER(PruneShaderLibrary);
+
+		for (const ITargetPlatform* TargetPlatform : PlatformManager->GetSessionPlatforms())
 		{
-			UE_SCOPED_HIERARCHICAL_COOKTIMER(FinishPopulateShaderLibrary);
-			FinishPopulateShaderLibrary(TargetPlatform, LibraryName);
+			FCookSavePackageContext& SaveContext = FindOrCreateSaveContext(TargetPlatform);
+			if (SaveContext.PackageWriterCapabilities.bReadOnly)
+			{
+				continue;
+			}
+
+			PruneShaderLibrary(TargetPlatform, LibraryName);
 		}
+	}
 
+	PlatformIndex = 0;
+	for (const ITargetPlatform* TargetPlatform : PlatformManager->GetSessionPlatforms())
+	{
+		FEndCookPlatformData& EndCookPlatformData = EndCookPlatformDatas[PlatformIndex++];
 		FCookSavePackageContext& SaveContext = FindOrCreateSaveContext(TargetPlatform);
+		if (SaveContext.PackageWriterCapabilities.bReadOnly)
+		{
+			continue;
+		}
+
+		FPlatformData* PlatformData = PlatformManager->GetPlatformData(TargetPlatform);
+		FAssetRegistryGenerator& Generator = *PlatformData->RegistryGenerator;
+
+		FString PlatformNameString = TargetPlatform->PlatformName();
+
+		TArray<FPackageData*>& CookedPackageDatas = EndCookPlatformData.CookedPackageDatas;
+		TArray<FPackageData*>& IgnorePackageDatas = EndCookPlatformData.IgnorePackageDatas;
+		TSet<FName>& CookedPackageNames = EndCookPlatformData.CookedPackageNames;
+		TSet<FName> IgnorePackageNames;
+		CookedPackageNames.Reserve(CookedPackageDatas.Num());
+		for (FPackageData* PackageData : CookedPackageDatas)
+		{
+			CookedPackageNames.Add(PackageData->GetPackageName());
+		}
+		IgnorePackageNames.Reserve(IgnorePackageDatas.Num());
+		for (FPackageData* PackageData : IgnorePackageDatas)
+		{
+			IgnorePackageNames.Add(PackageData->GetPackageName());
+		}
+	
 		if (bSaveManifests || bSaveAssetRegistry)
 		{
 			UE_SCOPED_HIERARCHICAL_COOKTIMER(GeneratorPreSave);
@@ -9525,7 +9599,8 @@ void UCookOnTheFlyServer::CookByTheBookFinishedInternal()
 		if (bSaveAssetRegistry)
 		{
 			UE_SCOPED_HIERARCHICAL_COOKTIMER(SaveAssetRegistry);
-			Generator.SaveAssetRegistry(SandboxRegistryFilename, bSaveDevelopmentAssetRegistry, bForceNoFilterAssetsFromAssetRegistry, DevelopmentAssetRegistryHash);
+			Generator.SaveAssetRegistry(SandboxRegistryFilename, bSaveDevelopmentAssetRegistry,
+				bForceNoFilterAssetsFromAssetRegistry, DevelopmentAssetRegistryHash, CookedPackageNames);
 		}
 
 		if (bSaveManifests || bSaveAssetRegistry)
diff --git a/Engine/Source/Editor/UnrealEd/Private/Cooker/ShaderLibraryCooking.cpp b/Engine/Source/Editor/UnrealEd/Private/Cooker/ShaderLibraryCooking.cpp
--- a/Engine/Source/Editor/UnrealEd/Private/Cooker/ShaderLibraryCooking.cpp
+++ b/Engine/Source/Editor/UnrealEd/Private/Cooker/ShaderLibraryCooking.cpp
@@ -231,6 +231,11 @@ void FShaderLibraryCookArtifact::CleanIntermediateFiles(const ITargetPlatform* T
 	}
 }
 
+void FShaderLibraryCookArtifact::UpdateOplogPackages(UE::Cook::Artifact::FUpdateOplogPackagesContext& Context)
+{
+	FShaderLibraryCooker::UpdateOplogPackages(Context);
+}
+
 } // namespace UE::Cook
 
 void UCookOnTheFlyServer::RegisterShaderLibraryIncrementalCookArtifact(FBeginCookContext& BeginContext)
@@ -591,6 +596,11 @@ void UCookOnTheFlyServer::FinishPopulateShaderLibrary(const ITargetPlatform* Tar
 	FShaderLibraryCooker::FinishPopulateShaderLibrary(TargetPlatform, Name, ShaderCodeDir, MetaDataPath);
 }
 
+void UCookOnTheFlyServer::PruneShaderLibrary(const ITargetPlatform* TargetPlatform, const FString& Name)
+{
+	FShaderLibraryCooker::PruneShaderLibrary(TargetPlatform, Name);
+}
+
 void UCookOnTheFlyServer::SaveShaderLibrary(const ITargetPlatform* TargetPlatform, const FString& Name)
 {
 	FString ShaderCodeDir;
diff --git a/Engine/Source/Editor/UnrealEd/Private/Cooker/ShaderLibraryCooking.h b/Engine/Source/Editor/UnrealEd/Private/Cooker/ShaderLibraryCooking.h
--- a/Engine/Source/Editor/UnrealEd/Private/Cooker/ShaderLibraryCooking.h
+++ b/Engine/Source/Editor/UnrealEd/Private/Cooker/ShaderLibraryCooking.h
@@ -39,8 +39,9 @@ class FShaderLibraryCookArtifact : public ICookArtifact
 {
 public:
 	FShaderLibraryCookArtifact(UCookOnTheFlyServer& InCOTFS);
-	virtual FString GetArtifactName() const;
-	virtual void OnFullRecook(const ITargetPlatform* TargetPlatform);
+	virtual FString GetArtifactName() const override;
+	virtual void OnFullRecook(const ITargetPlatform* TargetPlatform) override;
+	virtual void UpdateOplogPackages(UE::Cook::Artifact::FUpdateOplogPackagesContext& Context) override;
 
 	void CleanIntermediateFiles(const ITargetPlatform* TargetPlatform);
 
diff --git a/Engine/Source/Editor/UnrealEd/Public/Commandlets/AssetRegistryGenerator.h b/Engine/Source/Editor/UnrealEd/Public/Commandlets/AssetRegistryGenerator.h
--- a/Engine/Source/Editor/UnrealEd/Public/Commandlets/AssetRegistryGenerator.h
+++ b/Engine/Source/Editor/UnrealEd/Public/Commandlets/AssetRegistryGenerator.h
@@ -145,6 +145,12 @@ public:
 	void ComputePackageRemovals(const FAssetRegistryState& PreviousState, TArray<FName>& OutRemovedPackages,
 		TMap<FName, FGeneratorPackageInfo>& OutGeneratorPackages, int32& OutNumNeverCookPlaceHolderPackages);
 
+	/**
+	 * Return the list of packages from Previous IncrementalCooks the Generator received from the PackageWriter during
+	 * StartCookByTheBook.
+	 */
+	TArray<FName> GetPreviousIncrementalCookPackages();
+
 	/**
 	 * FinalizeChunkIDs 
 	 * Create the list of packages to store in each Chunk; each Chunk corresponds to a pak file, or group of pak files if the chunk is split.
@@ -218,7 +224,8 @@ public:
 	/**
 	* Saves generated asset registry data for each platform.
 	*/
-	bool SaveAssetRegistry(const FString& SandboxPath, bool bSerializeDevelopmentAssetRegistry, bool bForceNoFilterAssets, uint64& OutDevelopmentAssetRegistryHash);
+	bool SaveAssetRegistry(const FString& SandboxPath, bool bSerializeDevelopmentAssetRegistry,
+		bool bForceNoFilterAssets, uint64& OutDevelopmentAssetRegistryHash, const TSet<FName>& CookedPackages);
 
 	/** 
 	 * Writes out CookerOpenOrder.log file 
@@ -288,8 +295,12 @@ private:
 	/**
 	 * Updates AssetData with previous TagsAndValues, and updates PackageData values with previous PackageData,
 	 * for all packages kept from a previous cook.
+	 * 
+	 * InOutAllPackages: Initially contains the list of packages that were referenced in the current cook. The function
+	 * appends to it to also include any packages that were present in a previous incremental cook so that we can keep
+	 * the information for them in case they are referenced again in a later cook.
 	 */
-	void UpdateKeptPackages();
+	void UpdateKeptPackages(TSet<FName>& InOutAllPackages);
 
 	static UE::Cook::FAssetRegistryPackageMessage RecordUpdateAssetRegistryData(FName PackageName,
 		const ITargetPlatform* TargetPlatform, UE::Cook::ECookResult CookResult,
@@ -323,12 +334,6 @@ private:
 		TArray<FAssetDependency> PackageReferencers;
 	};
 	TMap<FName, FIterativelySkippedPackageUpdateData> PreviousPackagesToUpdate;
-	/** List of packages that were successfully cooked */
-	TSet<FName> CookedPackages;
-	/** List of packages that were filtered out from cooking */
-	TSet<FName> DevelopmentOnlyPackages;
-	/** List of packages that were kept from a previous cook */
-	TArray<FName> KeptPackages;
 	/** Map of Package name to Sandbox Paths */
 	typedef TMap<FName, FString> FChunkPackageSet;
 	/** Holds a reference to asset registry */
diff --git a/Engine/Source/Runtime/CoreUObject/Public/Cooker/CookArtifact.h b/Engine/Source/Runtime/CoreUObject/Public/Cooker/CookArtifact.h
--- a/Engine/Source/Runtime/CoreUObject/Public/Cooker/CookArtifact.h
+++ b/Engine/Source/Runtime/CoreUObject/Public/Cooker/CookArtifact.h
@@ -3,6 +3,7 @@
 #pragma once
 
 #if WITH_EDITOR
+#include "Containers/Map.h"
 #include "Containers/UnrealString.h"
 #include "Misc/ConfigCacheIni.h"
 #include "Misc/Guid.h"
@@ -11,6 +12,7 @@
 class ITargetPlatform;
 class UCookOnTheFlyServer;
 namespace UE::Cook { class ICookInfo; }
+namespace UE::Cook::Artifact { struct FUpdateOplogPackagesContext; }
 
 namespace UE::Cook::Artifact
 {
@@ -55,6 +57,64 @@ private:
 	friend UCookOnTheFlyServer;
 };
 
+/**
+ * Information about the status of each OplogPackage that is held by FUpdateOplogPackagesContext.
+ * An OplogPackage is a package cooked in the current cook session or cooked in a previous incrementalcook session
+ * that is still available for future cooks that reference it.
+ */
+struct FOplogPackageData
+{
+public:
+	/**
+	 * True if the package was referenced in the current cook (and for DLC cooks is not a basegame package cooked
+	 * elsewhere), and will therefore be staged during staging.
+	 */
+	bool IsReferenced() const;
+	/**
+	 * True if the package was referenced, and was either not present in an earlier incremental cook (NewCooked) or
+	 * was found to be modified and was Recooked.
+	 */
+	bool WasIncrementallyCooked() const;
+	/**
+	 * True if the package was referenced, and was found to be cooked in a previous cook and not modified so it was
+	 * not Recooked.
+	 */
+	bool WasIncrementallySkipped() const;
+
+private:
+	bool bReferenced = false;
+	bool bIncrementallyCooked = false;
+
+	friend FUpdateOplogPackagesContext;
+	friend UCookOnTheFlyServer;
+};
+
+/** Argument and return value passing structure for ICookArtifact::UpdateOplogPackages. */
+struct FUpdateOplogPackagesContext
+{
+public:
+	ICookInfo& GetCookInfo() const;
+	/** The TargetPlatform being reported. Non-null. */
+	const ITargetPlatform* GetTargetPlatform() const;
+
+	/** 
+	 * Get the list of OplogPackages that were present in the current CookSession, and data about their cookstatus.
+	 * An OplogPackage is a package cooked in the current cook session or cooked in a previous incrementalcook session
+	 * that is still available for future cooks that reference it. Data includes whether it was referenced and whether
+	 * it was cooked, @see FOplogPackageData.
+	 */
+	const TMap<FName, FOplogPackageData>& GetOplogPackages() const;
+
+private:
+	FUpdateOplogPackagesContext(ICookInfo& InCookInfo);
+
+	TMap<FName, FOplogPackageData> OplogPackages;
+	ICookInfo& CookInfo;
+	const ITargetPlatform* TargetPlatform = nullptr;
+
+	friend UCookOnTheFlyServer;
+};
+
 } // namespace UE::Cook::Artifact
 
 namespace UE::Cook
@@ -109,6 +169,18 @@ public:
 	virtual void OnFullRecook(const ITargetPlatform* TargetPlatform)
 	{
 	}
+
+	/**
+	 * Called from cooker, on CookDirector only, during CookByTheBookFinished. The cooker provides the list
+	 * of Oplog packages. Oplog packages include every package cooked in the current session, but also any packages
+	 * cooked in previous sessions that still have IncrementalCook data stored for them and might be referenced
+	 * again in a future cook. Each package has flags for whether it was referenced and whether it was incrementally
+	 * recooked. The artifact should prune any information it has about packages no longer in the list, and may
+	 * need to prune stale information from packages that were incrementally recooked.
+	 */
+	virtual void UpdateOplogPackages(UE::Cook::Artifact::FUpdateOplogPackagesContext& Context)
+	{
+	}
 };
 
 
@@ -169,6 +241,40 @@ inline void FCompareSettingsContext::RequestFullRecook(bool bValue)
 	bRequestFullRecook = bValue;
 }
 
+inline FUpdateOplogPackagesContext::FUpdateOplogPackagesContext(ICookInfo& InCookInfo)
+	: CookInfo(InCookInfo)
+{
+}
+inline ICookInfo& FUpdateOplogPackagesContext::GetCookInfo() const
+{
+	return CookInfo;
+}
+
+inline const ITargetPlatform* FUpdateOplogPackagesContext::GetTargetPlatform() const
+{
+	return TargetPlatform;
+}
+
+inline const TMap<FName, FOplogPackageData>& FUpdateOplogPackagesContext::GetOplogPackages() const
+{
+	return OplogPackages;
+}
+
+inline bool FOplogPackageData::IsReferenced() const
+{
+	return bReferenced;
+}
+
+inline bool FOplogPackageData::WasIncrementallyCooked() const
+{
+	return bReferenced && bIncrementallyCooked;
+}
+
+inline bool FOplogPackageData::WasIncrementallySkipped() const
+{
+	return bReferenced && !bIncrementallyCooked;
+}
+
 } // namespace UE::Cook::Artifact
 
 #endif // WITH_EDITOR
diff --git a/Engine/Source/Runtime/RenderCore/Private/ShaderLibrary/ShaderCodeLibrary.cpp b/Engine/Source/Runtime/RenderCore/Private/ShaderLibrary/ShaderCodeLibrary.cpp
--- a/Engine/Source/Runtime/RenderCore/Private/ShaderLibrary/ShaderCodeLibrary.cpp
+++ b/Engine/Source/Runtime/RenderCore/Private/ShaderLibrary/ShaderCodeLibrary.cpp
@@ -62,6 +62,7 @@ ShaderCodeLibrary.cpp: Bound shader state cache implementation.
 #include "Serialization/CompactBinaryWriter.h"
 #include "RHIStrings.h"
 #include "CookArtifactReader.h"
+#include "Cooker/CookArtifact.h"
 #endif
 
 // allow introspection (e.g. dumping the contents) for easier debugging
@@ -2290,7 +2291,6 @@ bool FShaderLibraryCooker::AddShaderCode(EShaderPlatform ShaderPlatform, const F
 	return false;
 }
 
-#if WITH_EDITOR
 void FShaderLibraryCooker::CopyToCompactBinaryAndClear(FCbWriter& Writer, bool& bOutHasData,
 	bool& bOutRanOutOfRoom, int64 MaxShaderSize)
 {
@@ -2309,7 +2309,6 @@ bool FShaderLibraryCooker::AppendFromCompactBinary(FCbFieldView Field)
 	}
 	return false;
 }
-#endif
 
 bool FShaderLibraryCooker::NeedsShaderStableKeys(EShaderPlatform ShaderPlatform)
 {
@@ -2389,6 +2388,16 @@ void FShaderLibraryCooker::FinishPopulateShaderLibrary(const ITargetPlatform* Ta
 	}
 }
 
+void FShaderLibraryCooker::UpdateOplogPackages(UE::Cook::Artifact::FUpdateOplogPackagesContext& Context)
+{
+
+}
+
+void FShaderLibraryCooker::PruneShaderLibrary(const ITargetPlatform* TargetPlatform, const FString& Name)
+{
+
+}
+
 bool FShaderLibraryCooker::MergeShaderCodeArchive(const TArray<FString>& CookedMetadataDirs, const FString& OutputDir, TArray<FString>& OutWrittenFiles)
 {
 	OutWrittenFiles.Empty();
diff --git a/Engine/Source/Runtime/RenderCore/Public/ShaderCodeLibrary.h b/Engine/Source/Runtime/RenderCore/Public/ShaderCodeLibrary.h
--- a/Engine/Source/Runtime/RenderCore/Public/ShaderCodeLibrary.h
+++ b/Engine/Source/Runtime/RenderCore/Public/ShaderCodeLibrary.h
@@ -42,6 +42,7 @@ class UObject;
 class FCbFieldView;
 class FCbWriter;
 class ICookArtifactReader;
+namespace UE::Cook::Artifact { struct FUpdateOplogPackagesContext; }
 #endif
 
 struct FShaderCodeLibraryPipeline
@@ -331,13 +332,11 @@ struct FShaderLibraryCooker
 	// At cook time, add shader code to collection
 	static RENDERCORE_API bool AddShaderCode(EShaderPlatform ShaderPlatform, const FShaderMapResourceCode* Code, const FShaderMapAssetPaths& AssociatedAssets);
 
-#if WITH_EDITOR
 	/** Called from a CookWorker to send all contents of the ShaderLibrary to the CookDirector */
 	static RENDERCORE_API void CopyToCompactBinaryAndClear(FCbWriter& Writer, bool& bOutHasData, bool& bOutRanOutOfRoom, int64 MaxShaderSize);
 
 	/** Called On the CookDirector to receive ShaderLibrary contents from a CookWorker */
 	static RENDERCORE_API bool AppendFromCompactBinary(FCbFieldView Field);
-#endif
 
 	// We check this early in the callstack to avoid creating a bunch of FName and keys and things we will never save anyway. 
 	// Pass the shader platform to check or EShaderPlatform::SP_NumPlatforms to check if any of the registered types require
@@ -351,6 +350,15 @@ struct FShaderLibraryCooker
 	static RENDERCORE_API void FinishPopulateShaderLibrary(const ITargetPlatform* TargetPlatform, FString const& Name, FString const& SandboxDestinationPath,
 		FString const& SandboxMetadataPath);
 
+	/** Called by the cooker to report which assets are referenced in the current cook or in incrementalcook oplog. */
+	static RENDERCORE_API void UpdateOplogPackages(UE::Cook::Artifact::FUpdateOplogPackagesContext& Context);
+
+	/**
+	 * Called by the cooker to ask the library to prune the non-longer referenced shaders from its incremental cook library,
+	 * called only after UpdateOplogPackages has been called on every platform.
+	 */
+	static RENDERCORE_API void PruneShaderLibrary(const ITargetPlatform* TargetPlatform, const FString& Name);
+
 	/**
 	 * Given multiple Cooked Metadata directories will attempt to merge the ShaderByteCode and the ShaderStableInfo into the given OutputDir.
 	 * It would be expected that the OutputDir is another MetaData directory but this can be any dir.
@@ -406,4 +414,4 @@ struct FShaderLibraryCooker
 	// Create a smaller 'patch' library that only contains data from 'NewMetaDataDir' not contained in any of 'OldMetaDataDirs'
 	static RENDERCORE_API bool CreatePatchLibrary(TArray<FString> const& OldMetaDataDirs, FString const& NewMetaDataDir, FString const& OutDir, bool bNativeFormat, bool bNeedsDeterministicOrder);
 };
-#endif
+#endif // WITH_EDITOR
