AssetEditor: Add support for changing aspect ratio in the asset editor viewports

#Synced-CL 48921717

diff --git a/Engine/Source/Editor/UnrealEd/Private/SEditorViewport.cpp b/Engine/Source/Editor/UnrealEd/Private/SEditorViewport.cpp
--- a/Engine/Source/Editor/UnrealEd/Private/SEditorViewport.cpp
+++ b/Engine/Source/Editor/UnrealEd/Private/SEditorViewport.cpp
@@ -272,28 +272,15 @@ void SEditorViewport::Construct( const FArguments& InArgs )
 		}
 	}
 
-	if (bIsMainViewport)
-	{
-		// We wrap the main editor viewport in the middle of a 3x3 grid in order to restrict aspect ratio for preview platforms that request it.
-		VerticalBox->AddSlot()
-			[
-				SNew(SGlobalPlayWorldActions)
-					[
-						SNew(SEditorViewportGridPanel)
-						.ViewportWidget(ViewportWidget)
-					]
-			];
-	}
-	else
-	{
-		VerticalBox->AddSlot()
-			[
-				SNew(SGlobalPlayWorldActions)
-					[
-						ViewportWidget.ToSharedRef()
-					]
-			];
-	}
+	// We wrap the viewports in the middle of a 3x3 grid in order to restrict aspect ratio for preview platforms that request it or for manual aspect ratio restrictions.
+	VerticalBox->AddSlot()
+		[
+			SNew(SGlobalPlayWorldActions)
+				[
+					SAssignNew(ViewportAspectRatioManager, SEditorViewportGridPanel)
+					.ViewportWidget(ViewportWidget)
+				]
+		];
 	// clang-format on
 	
 	if (TSharedPtr<SWidget> DepthBarWidget = BuildViewportDepthBar())
@@ -1376,6 +1363,22 @@ TSharedPtr<SWidget> SEditorViewport::CreateViewportIndicatorWidget(const TAttrib
 		];
 }
 
+void SEditorViewport::SetAspectRatio(float AspectRatioIn)
+{	
+	ViewportAspectRatioManager->SetAspectRatioConstraint(AspectRatioIn);
+	ViewportAspectRatioManager->SetAspectRatioMode(SEditorViewportGridPanel::EAspectRatioMode::Constrained);
+}
+
+void SEditorViewport::SetUsePlatformAspectRatio()
+{
+	ViewportAspectRatioManager->SetAspectRatioMode(SEditorViewportGridPanel::EAspectRatioMode::Platform);
+}
+
+void SEditorViewport::UnsetAspectRatio()
+{
+	ViewportAspectRatioManager->SetAspectRatioMode(SEditorViewportGridPanel::EAspectRatioMode::Unconstrained);	
+}
+
 bool SEditorViewport::IsPreviewingScreenPercentage() const
 {
 	return Client->IsPreviewingScreenPercentage();
diff --git a/Engine/Source/Editor/UnrealEd/Private/SEditorViewportGridPanel.cpp b/Engine/Source/Editor/UnrealEd/Private/SEditorViewportGridPanel.cpp
--- a/Engine/Source/Editor/UnrealEd/Private/SEditorViewportGridPanel.cpp
+++ b/Engine/Source/Editor/UnrealEd/Private/SEditorViewportGridPanel.cpp
@@ -22,7 +22,9 @@ void SEditorViewportGridPanel::Construct(const FArguments& InArgs)
 	SuperArgs + SGridPanel::Slot(1, 1)[InArgs._ViewportWidget.Get().ToSharedRef()];
 	SGridPanel::Construct(SuperArgs);
 
-	DebugAspectRatio = 0.f;
+	PlatformAspectRatio = 0.f;
+	ForcedAspectRatio = 0.f;
+	AspectRatioMode = EAspectRatioMode::Platform;	
 	FSlateApplication::Get().OnConstrainedAspectRatioChanged.AddSP(this, &SEditorViewportGridPanel::UpdateAspectRatio);
 }
 
@@ -39,12 +41,17 @@ void SEditorViewportGridPanel::OnArrangeChildren(const FGeometry& AllottedGeomet
 		RowFillCoefficients[1].Set(1.f);
 		RowFillCoefficients[2].Set(0.f);
 
-		if (DebugAspectRatio != 0.f)
+		const float AspectRatio = 
+			AspectRatioMode == EAspectRatioMode::Platform ? PlatformAspectRatio : 
+			AspectRatioMode == EAspectRatioMode::Constrained ? ForcedAspectRatio :
+			0.0;		
+
+		if (AspectRatio != 0.f)
 		{
 			const float ContainerX = AllottedGeometry.GetLocalSize().X;
 			const float ContainerY = AllottedGeometry.GetLocalSize().Y;
-			const float NewX = ContainerY * DebugAspectRatio;
-			const float NewY = ContainerX / DebugAspectRatio;
+			const float NewX = ContainerY * AspectRatio;
+			const float NewY = ContainerX / AspectRatio;
 
 			if (NewX < ContainerX)
 			{
@@ -70,5 +77,20 @@ void SEditorViewportGridPanel::OnArrangeChildren(const FGeometry& AllottedGeomet
 
 void SEditorViewportGridPanel::UpdateAspectRatio(const float& AspectRatio)
 {
-	DebugAspectRatio = AspectRatio;
+	PlatformAspectRatio = AspectRatio;
+}
+
+void SEditorViewportGridPanel::SetAspectRatioMode(EAspectRatioMode ModeIn)
+{
+	AspectRatioMode = ModeIn;
+}
+
+SEditorViewportGridPanel::EAspectRatioMode SEditorViewportGridPanel::GetAspectRatioMode() const
+{
+	return AspectRatioMode;
+}
+
+void SEditorViewportGridPanel::SetAspectRatioConstraint(const float& AspectRatio)
+{
+	ForcedAspectRatio = AspectRatio;
 }
\ No newline at end of file
diff --git a/Engine/Source/Editor/UnrealEd/Private/SEditorViewportGridPanel.h b/Engine/Source/Editor/UnrealEd/Private/SEditorViewportGridPanel.h
--- a/Engine/Source/Editor/UnrealEd/Private/SEditorViewportGridPanel.h
+++ b/Engine/Source/Editor/UnrealEd/Private/SEditorViewportGridPanel.h
@@ -10,6 +10,14 @@ class SEditorViewportGridPanel : public SGridPanel
 	SLATE_DECLARE_WIDGET_API(SEditorViewportGridPanel, SGridPanel, UNREALED_API)
 
 public:
+
+	enum class EAspectRatioMode
+	{
+		Platform = 0,
+		Unconstrained = 1,
+		Constrained = 2
+	};
+
 	SLATE_BEGIN_ARGS(SEditorViewportGridPanel)
 		{}
 		SLATE_ATTRIBUTE(TSharedPtr<SWidget>, ViewportWidget);
@@ -17,9 +25,20 @@ public:
 
 	UNREALED_API void Construct(const FArguments& InArgs);
 	UNREALED_API virtual void OnArrangeChildren(const FGeometry& AllottedGeometry, FArrangedChildren& ArrangedChildren) const override;
+
+	// Hook to update the aspect ratio coming from the platform settings - called automatically externally to keep this up to date.
+	// Use SetAspectRatioConstraint and SetAspectRatioMode(EAspectRatioMode::Constrained) to set a fixed aspect ratio.
 	UNREALED_API void UpdateAspectRatio(const float& AspectRatio);
 
+	UNREALED_API void SetAspectRatioMode(EAspectRatioMode ModeIn);
+	UNREALED_API EAspectRatioMode GetAspectRatioMode() const;
+	UNREALED_API void SetAspectRatioConstraint(const float& AspectRatio);
+
 private:
 	TAttribute<TSharedPtr<SWidget>> ViewportWidget;
-	float DebugAspectRatio;
+	
+	EAspectRatioMode AspectRatioMode;
+	float PlatformAspectRatio;
+	float ForcedAspectRatio;
+
 };
\ No newline at end of file
diff --git a/Engine/Source/Editor/UnrealEd/Public/SEditorViewport.h b/Engine/Source/Editor/UnrealEd/Public/SEditorViewport.h
--- a/Engine/Source/Editor/UnrealEd/Public/SEditorViewport.h
+++ b/Engine/Source/Editor/UnrealEd/Public/SEditorViewport.h
@@ -22,6 +22,7 @@ class UViewportInteractionsBehaviorSource;
 struct FSlateBrush;
 struct FToolMenuContext;
 class SDepthBar;
+class SEditorViewportGridPanel;
 
 class SEditorViewport
 	: public SCompoundWidget
@@ -329,9 +330,15 @@ protected:
 	 */
 	static UNREALED_API TSharedPtr<SWidget> CreateViewportIndicatorWidget(const TAttribute<EVisibility>& InVisibility);
 
+	UNREALED_API void SetAspectRatio(float AspectRatioIn);
+	UNREALED_API void SetUsePlatformAspectRatio();
+	UNREALED_API void UnsetAspectRatio();
+
 protected:
 	TSharedPtr<SOverlay> ViewportOverlay;
 
+	TSharedPtr<SEditorViewportGridPanel> ViewportAspectRatioManager;
+
 	/** Viewport that renders the scene provided by the viewport client */
 	TSharedPtr<FSceneViewport> SceneViewport;
 	
