Allow ReplaceActors to copy object properties for blueprints when related by the native parent

#Synced-CL 48498703

diff --git a/Engine/Config/BaseEditorPerProjectUserSettings.ini b/Engine/Config/BaseEditorPerProjectUserSettings.ini
--- a/Engine/Config/BaseEditorPerProjectUserSettings.ini
+++ b/Engine/Config/BaseEditorPerProjectUserSettings.ini
@@ -208,6 +208,8 @@ bBSPAutoUpdate=True
 bAutoMoveBSPPivotOffset=False
 ; If true, Navigation will auto-update
 bNavigationAutoUpdate=True
+; If true, replacing actors allows copying properties when related by the native parent of the replaced class
+bReplaceCopiesPropertiesWhenRelatedByNativeParent=True
 ; If true, Replaces respects the scale of the original actor. If false, Replace sets the scale to 1.0
 bReplaceRespectsScale=True
 ; If false, references are not checked on delete.
diff --git a/Engine/Source/Editor/UnrealEd/Classes/Settings/LevelEditorMiscSettings.h b/Engine/Source/Editor/UnrealEd/Classes/Settings/LevelEditorMiscSettings.h
--- a/Engine/Source/Editor/UnrealEd/Classes/Settings/LevelEditorMiscSettings.h
+++ b/Engine/Source/Editor/UnrealEd/Classes/Settings/LevelEditorMiscSettings.h
@@ -39,6 +39,10 @@ public:
 	UPROPERTY(EditAnywhere, config, Category=Editing, meta=( DisplayName = "Update Navigation Automatically" ))
 	uint32 bNavigationAutoUpdate:1;
 
+	/** If enabled, replacing actors allows copying properties when related by the native parent of the replaced class. */
+	UPROPERTY(EditAnywhere, config, Category=Editing, meta=( DisplayName = "Copy Properties when Replacing with Different Blueprint" ))
+	uint32 bReplaceCopiesPropertiesWhenRelatedByNativeParent:1;
+
 	/** If enabled, replacing actors will respect the scale of the original actor.  Otherwise, the replaced actors will have a scale of 1.0 */
 	UPROPERTY(EditAnywhere, config, Category=Editing, meta=( DisplayName = "Preserve Actor Scale on Replace" ))
 	uint32 bReplaceRespectsScale:1;
diff --git a/Engine/Source/Editor/UnrealEd/Private/Subsystems/EditorActorSubsystem.cpp b/Engine/Source/Editor/UnrealEd/Private/Subsystems/EditorActorSubsystem.cpp
--- a/Engine/Source/Editor/UnrealEd/Private/Subsystems/EditorActorSubsystem.cpp
+++ b/Engine/Source/Editor/UnrealEd/Private/Subsystems/EditorActorSubsystem.cpp
@@ -2455,24 +2455,27 @@ void UEditorActorSubsystem::ReplaceActors(
 
 		// create the actor
 		NewActor = InFactory->CreateActor(Asset, Level, OldTransform, SpawnParams);
-		// For blueprints, try to copy over properties
-		if (bInCopySourceProperties && InFactory->IsA(UActorFactoryBlueprint::StaticClass()))
-		{
-			UBlueprint* Blueprint = CastChecked<UBlueprint>(Asset);
-			// Only try to copy properties if this blueprint is based on the actor
-			UClass* OldActorClass = OldActor->GetClass();
-			if (Blueprint->GeneratedClass->IsChildOf(OldActorClass) && NewActor != nullptr)
-			{
-				NewActor->UnregisterAllComponents();
-				UEngine::FCopyPropertiesForUnrelatedObjectsParams Options;
-				Options.bNotifyObjectReplacement = true;
-				UEditorEngine::CopyPropertiesForUnrelatedObjects(OldActor, NewActor, Options);
-				NewActor->RegisterAllComponents();
-			}
-		}
 
 		if (NewActor)
 		{
+			// For blueprints, optionally try to copy over properties
+			if (bInCopySourceProperties && InFactory->IsA(UActorFactoryBlueprint::StaticClass()))
+			{
+				// Only try to copy properties when the blueprint is a child of a shared native parent
+				const UBlueprint* const Blueprint = CastChecked<UBlueprint>(Asset);
+				UClass* const OldActorClass = OldActor->GetClass();
+				const bool bUseNativeParent = GetDefault<ULevelEditorMiscSettings>()->bReplaceCopiesPropertiesWhenRelatedByNativeParent;
+				const UClass* const ClassToCheck = bUseNativeParent ? GetParentNativeClass(OldActorClass) : OldActorClass;
+				if (Blueprint->GeneratedClass->IsChildOf(ClassToCheck))
+				{
+					NewActor->UnregisterAllComponents();
+					UEngine::FCopyPropertiesForUnrelatedObjectsParams Options;
+					Options.bNotifyObjectReplacement = true;
+					UEditorEngine::CopyPropertiesForUnrelatedObjects(OldActor, NewActor, Options);
+					NewActor->RegisterAllComponents();
+				}
+			}
+			
 			// The new actor might not have a root component
 			USceneComponent* const NewActorRootComponent = NewActor->GetRootComponent();
 			if (NewActorRootComponent)
