UE: Bump incremental version. Shaders that are not associated with any asset

diff --git a/Engine/Source/Editor/UnrealEd/Private/Cooker/CookIncrementalVersion.cpp b/Engine/Source/Editor/UnrealEd/Private/Cooker/CookIncrementalVersion.cpp
--- a/Engine/Source/Editor/UnrealEd/Private/Cooker/CookIncrementalVersion.cpp
+++ b/Engine/Source/Editor/UnrealEd/Private/Cooker/CookIncrementalVersion.cpp
@@ -8,6 +8,6 @@ namespace UE::Cook
 {
 
 // Change CookIncrementalVersion to a new guid when all packages in an incremental cook need to be invalidated
-FGuid CookIncrementalVersion(0xC70D6607, 0xCF3A48C0, 0xB484B5B8, 0x1DEA4AFD);
+FGuid CookIncrementalVersion(0xBE9A4C44, 0xDF2B4F35, 0xAF2D6B6D, 0xB70FC790);
 
 }
