Adding bespoke "UI Material" option in the content browser's new asset menu.

- Explicitly disables used with static mesh usage flag so all new UI materials don't automatically have extra static mesh shaders.
- Turn off automatically set usage flag that way users don't accidently turn on usage flags and add more shaders.
- Set the new material's domain to UI.
- Material is created with the following name: "M_UI_Material"

[RN] minor | Rendering - Materials and Shaders
#Synced-CL 48812887

diff --git a/Engine/Source/Editor/UnrealEd/Classes/Factories/UIMaterialFactoryNew.h b/Engine/Source/Editor/UnrealEd/Classes/Factories/UIMaterialFactoryNew.h
new file mode 100644
--- /dev/null
+++ b/Engine/Source/Editor/UnrealEd/Classes/Factories/UIMaterialFactoryNew.h
@@ -0,0 +1,27 @@
+// Copyright Epic Games, Inc. All Rights Reserved.
+
+//~=============================================================================
+// UIMaterialFactoryNew
+//~=============================================================================
+
+#pragma once
+
+#include "CoreMinimal.h"
+#include "UObject/ObjectMacros.h"
+#include "Factories/Factory.h"
+#include "UIMaterialFactoryNew.generated.h"
+
+UCLASS(hidecategories=Object, collapsecategories, MinimalAPI)
+class UUIMaterialFactoryNew : public UFactory
+{
+	GENERATED_UCLASS_BODY()
+
+	//~ Begin UFactory Interface
+	virtual FText GetDisplayName() const override;
+	virtual uint32 GetMenuCategories() const override;
+	virtual FText GetToolTip() const override;
+	virtual UObject* FactoryCreateNew(UClass* Class, UObject* InParent, FName Name, EObjectFlags Flags, UObject* Context, FFeedbackContext* Warn) override;
+	virtual bool ShouldShowInNewMenu() const override;
+	virtual FString GetDefaultNewAssetName() const override;
+	//~ End UFactory Interface	
+};
diff --git a/Engine/Source/Editor/UnrealEd/Private/Factories/EditorFactories.cpp b/Engine/Source/Editor/UnrealEd/Private/Factories/EditorFactories.cpp
--- a/Engine/Source/Editor/UnrealEd/Private/Factories/EditorFactories.cpp
+++ b/Engine/Source/Editor/UnrealEd/Private/Factories/EditorFactories.cpp
@@ -87,6 +87,7 @@
 #include "Factories/HapticFeedbackEffectSoundWaveFactory.h"
 #include "Factories/LevelFactory.h"
 #include "Factories/MaterialFactoryNew.h"
+#include "Factories/UIMaterialFactoryNew.h"
 #include "Factories/MaterialFunctionFactoryNew.h"
 #include "Factories/MaterialFunctionMaterialLayerFactory.h"
 #include "Factories/MaterialFunctionMaterialLayerBlendFactory.h"
@@ -553,6 +554,52 @@ bool UMaterialFactoryNew::ShouldShowInNewMenu() const
 	return Super::ShouldShowInNewMenu() && GetDefault<UMaterialEditorSettings>()->MaterialWizardConfiguration == nullptr;
 }
 
+/*------------------------------------------------------------------------------
+	UUIMaterialFactoryNew implementation.
+------------------------------------------------------------------------------*/
+UUIMaterialFactoryNew::UUIMaterialFactoryNew(const FObjectInitializer& ObjectInitializer)
+	: Super(ObjectInitializer)
+{
+	SupportedClass = UMaterial::StaticClass();
+	bCreateNew = true;
+	bEditAfterNew = true;
+}
+
+FText UUIMaterialFactoryNew::GetDisplayName() const
+{
+	return LOCTEXT("UIMaterialFactoryDescription", "UI Material");
+}
+
+uint32 UUIMaterialFactoryNew::GetMenuCategories() const
+{
+	return EAssetTypeCategories::Type::Materials;
+}
+
+FText UUIMaterialFactoryNew::GetToolTip() const
+{
+	return LOCTEXT("UIMaterialFactoryTooltip", "Creates a new Material with the UI Material Domain and turns off Used With Static Mesh usage flag.");
+}
+
+UObject* UUIMaterialFactoryNew::FactoryCreateNew(UClass* Class, UObject* InParent, FName Name, EObjectFlags Flags, UObject* Context, FFeedbackContext* Warn)
+{
+	UMaterial* NewMaterial = NewObject<UMaterial>(InParent, Class, Name, Flags);
+	NewMaterial->bAutomaticallySetUsageInEditor = false;
+	NewMaterial->MaterialDomain = EMaterialDomain::MD_UI;
+	NewMaterial->SetUsageByFlag(MATUSAGE_StaticMesh, false);
+	//NewMaterial->CacheShaders(EMaterialShaderPrecompileMode::Synchronous);
+
+	return NewMaterial;
+}
+
+bool UUIMaterialFactoryNew::ShouldShowInNewMenu() const
+{
+	return true;
+}
+
+FString UUIMaterialFactoryNew::GetDefaultNewAssetName() const
+{
+	return FString(TEXT("M_UI_Material"));
+}
 
 /*------------------------------------------------------------------------------
 	UMaterialFunctionFactoryNew implementation.
