Material previews fixed to work on materials and instances even when they disable static mesh usage.
Preview material set to support static mesh usage.
Material Instance editor creates a child MIC with static mesh usage always overridden on for use in preview.
#rnx
#Synced-CL 49018876

diff --git a/Engine/Source/Editor/MaterialEditor/Private/MaterialInstanceEditor.cpp b/Engine/Source/Editor/MaterialEditor/Private/MaterialInstanceEditor.cpp
--- a/Engine/Source/Editor/MaterialEditor/Private/MaterialInstanceEditor.cpp
+++ b/Engine/Source/Editor/MaterialEditor/Private/MaterialInstanceEditor.cpp
@@ -1279,7 +1279,7 @@ FLinearColor FMaterialInstanceEditor::GetWorldCentricTabColorScale() const
 
 UMaterialInterface* FMaterialInstanceEditor::GetMaterialInterface() const
 {
-	return MaterialEditorInstance->SourceInstance;
+	return MaterialEditorInstance->PreviewInstance ? MaterialEditorInstance->PreviewInstance : MaterialEditorInstance->SourceInstance;
 }
 
 void FMaterialInstanceEditor::NotifyPreChange(FProperty* PropertyThatChanged)
diff --git a/Engine/Source/Editor/UnrealEd/Classes/MaterialEditor/MaterialEditorInstanceConstant.h b/Engine/Source/Editor/UnrealEd/Classes/MaterialEditor/MaterialEditorInstanceConstant.h
--- a/Engine/Source/Editor/UnrealEd/Classes/MaterialEditor/MaterialEditorInstanceConstant.h
+++ b/Engine/Source/Editor/UnrealEd/Classes/MaterialEditor/MaterialEditorInstanceConstant.h
@@ -362,6 +362,9 @@ private:
 	UPROPERTY()
 	TObjectPtr<class UMaterialFunctionInstance> SourceFunction;	
 
+	UPROPERTY()
+	TObjectPtr<class UMaterialInstanceConstant> PreviewInstance;
+
 	UPROPERTY(transient, duplicatetransient)
 	TArray<FMaterialParameterInfo> VisibleExpressions;
 
diff --git a/Engine/Source/Editor/UnrealEd/Private/PreviewMaterial.cpp b/Engine/Source/Editor/UnrealEd/Private/PreviewMaterial.cpp
--- a/Engine/Source/Editor/UnrealEd/Private/PreviewMaterial.cpp
+++ b/Engine/Source/Editor/UnrealEd/Private/PreviewMaterial.cpp
@@ -51,6 +51,9 @@ public:
 	{
 	}
 
+	/** Always allow to be used with static mesh so that we can use in preview view. */
+	virtual bool IsUsedWithStaticMesh() const override { return true; }
+
 	/**
 	 * Should the shader for this material with the given platform, shader type and vertex 
 	 * factory type combination be compiled
@@ -579,8 +582,9 @@ void UMaterialEditorInstanceConstant::PostEditChangeProperty(FPropertyChangedEve
 
 		CopyToSourceInstance(bLayersParameterChanged);
 
-		// Tell our source instance to update itself so the preview updates.
 		SourceInstance->PostEditChangeProperty(PropertyChangedEvent);
+		// Also notify preview because it won't receive changes from parent during interactive edits.
+		PreviewInstance->PostEditChangeProperty(PropertyChangedEvent);
 
 		// Invalidate the streaming data so that it gets rebuilt.
 		SourceInstance->TextureStreamingData.Empty();
@@ -1036,6 +1040,12 @@ void UMaterialEditorInstanceConstant::SetSourceInstance(UMaterialInstanceConstan
 {
 	check(MaterialInterface);
 	SourceInstance = MaterialInterface;
+
+	PreviewInstance = NewObject<UMaterialInstanceConstant>();
+	PreviewInstance->SetParentEditorOnly(MaterialInterface);
+	PreviewInstance->BasePropertyOverrides.bOverride_UsageFlags = (1 << MATUSAGE_StaticMesh);
+	PreviewInstance->BasePropertyOverrides.UsageFlags = (1 << MATUSAGE_StaticMesh);
+
 	Parent = SourceInstance->Parent;
 	PhysMaterial = SourceInstance->PhysMaterial;
 	bNaniteOverride = SourceInstance->NaniteOverrideMaterial.bEnableOverride;
