Thumbnail rendering fixed for materials that don't support static mesh usage.
If necessary, the renderer internally creates a child material instance with static mesh usage turned on.
#rnx
#Synced-CL 49028617

diff --git a/Engine/Source/Editor/UnrealEd/Classes/ThumbnailRendering/MaterialInstanceThumbnailRenderer.h b/Engine/Source/Editor/UnrealEd/Classes/ThumbnailRendering/MaterialInstanceThumbnailRenderer.h
--- a/Engine/Source/Editor/UnrealEd/Classes/ThumbnailRendering/MaterialInstanceThumbnailRenderer.h
+++ b/Engine/Source/Editor/UnrealEd/Classes/ThumbnailRendering/MaterialInstanceThumbnailRenderer.h
@@ -30,6 +30,10 @@ class UMaterialInstanceThumbnailRenderer : public UDefaultSizedThumbnailRenderer
 	UNREALED_API virtual void BeginDestroy() override;
 
 private:
+	/** Material instance created only when the base material instance can't be used directly for thumbnail rendering. */
+	UPROPERTY()
+	TObjectPtr<class UMaterialInstanceConstant> OverrideMaterialInstance;
+
 	class FMaterialThumbnailScene* ThumbnailScene;
 	class FWidgetRenderer* WidgetRenderer;
 	struct FSlateMaterialBrush* UIMaterialBrush;
diff --git a/Engine/Source/Editor/UnrealEd/Private/ThumbnailRendering/MaterialInstanceThumbnailRenderer.cpp b/Engine/Source/Editor/UnrealEd/Private/ThumbnailRendering/MaterialInstanceThumbnailRenderer.cpp
--- a/Engine/Source/Editor/UnrealEd/Private/ThumbnailRendering/MaterialInstanceThumbnailRenderer.cpp
+++ b/Engine/Source/Editor/UnrealEd/Private/ThumbnailRendering/MaterialInstanceThumbnailRenderer.cpp
@@ -3,6 +3,7 @@
 #include "ThumbnailRendering/MaterialInstanceThumbnailRenderer.h"
 #include "Misc/App.h"
 #include "ShowFlags.h"
+#include "Materials/MaterialInstanceConstant.h"
 #include "Materials/MaterialInterface.h"
 #include "Materials/Material.h"
 #include "SceneView.h"
@@ -87,7 +88,25 @@ void UMaterialInstanceThumbnailRenderer::Draw(UObject* Object, int32 X, int32 Y,
 				ThumbnailScene = new FMaterialThumbnailScene();
 			}
 
-			ThumbnailScene->SetMaterialInterface(MatInst);
+			// If the material doesn't have the static mesh usage flag then create a child
+			// instance with the usage flag set and use that for the thumbnail rendering.
+			if (!MatInst->GetUsageByFlag(MATUSAGE_StaticMesh))
+			{
+				if (!OverrideMaterialInstance || OverrideMaterialInstance->Parent != MatInst)
+				{
+					OverrideMaterialInstance = NewObject<UMaterialInstanceConstant>();
+					OverrideMaterialInstance->SetParentEditorOnly(MatInst);
+					OverrideMaterialInstance->BasePropertyOverrides.bOverride_UsageFlags = (1 << MATUSAGE_StaticMesh);
+					OverrideMaterialInstance->BasePropertyOverrides.UsageFlags = (1 << MATUSAGE_StaticMesh);
+					OverrideMaterialInstance->UpdateStaticPermutation();
+				}
+				ThumbnailScene->SetMaterialInterface(OverrideMaterialInstance);
+			}
+			else
+			{
+				ThumbnailScene->SetMaterialInterface(MatInst);
+			}
+
 			FSceneViewFamilyContext ViewFamily( FSceneViewFamily::ConstructionValues( RenderTarget, ThumbnailScene->GetScene(), FEngineShowFlags(ESFIM_Game) )
 				.SetTime(UThumbnailRenderer::GetTime())
 				.SetAdditionalViewFamily(bAdditionalViewFamily));
