Bump cook incremental version for shader changes.

#Synced-CL 49027723

diff --git a/Engine/Source/Editor/UnrealEd/Private/Cooker/CookIncrementalVersion.cpp b/Engine/Source/Editor/UnrealEd/Private/Cooker/CookIncrementalVersion.cpp
--- a/Engine/Source/Editor/UnrealEd/Private/Cooker/CookIncrementalVersion.cpp
+++ b/Engine/Source/Editor/UnrealEd/Private/Cooker/CookIncrementalVersion.cpp
@@ -8,6 +8,6 @@ namespace UE::Cook
 {
 
 // Change CookIncrementalVersion to a new guid when all packages in an incremental cook need to be invalidated
-FGuid CookIncrementalVersion(0xBE9A4C44, 0xDF2B4F35, 0xAF2D6B6D, 0xB70FC790);
+FGuid CookIncrementalVersion(0xF5F6B1B6, 0x93C147B4, 0xBD91AC9D, 0x411692F2);
 
 }
