Material thumbnail renderer. Change local variable naming for clarity.
The objects are material interfaces not material instances.
#rnx
#Synced-CL 49028617

diff --git a/Engine/Source/Editor/UnrealEd/Private/ThumbnailRendering/MaterialInstanceThumbnailRenderer.cpp b/Engine/Source/Editor/UnrealEd/Private/ThumbnailRendering/MaterialInstanceThumbnailRenderer.cpp
--- a/Engine/Source/Editor/UnrealEd/Private/ThumbnailRendering/MaterialInstanceThumbnailRenderer.cpp
+++ b/Engine/Source/Editor/UnrealEd/Private/ThumbnailRendering/MaterialInstanceThumbnailRenderer.cpp
@@ -26,11 +26,11 @@ UMaterialInstanceThumbnailRenderer::UMaterialInstanceThumbnailRenderer(const FOb
 
 void UMaterialInstanceThumbnailRenderer::Draw(UObject* Object, int32 X, int32 Y, uint32 Width, uint32 Height, FRenderTarget* RenderTarget, FCanvas* Canvas, bool bAdditionalViewFamily)
 {
-	UMaterialInterface* MatInst = Cast<UMaterialInterface>(Object);
-	if (MatInst != nullptr)
+	UMaterialInterface* MaterialInterface = Cast<UMaterialInterface>(Object);
+	if (MaterialInterface != nullptr)
 	{
-		UMaterial* Mat = MatInst->GetMaterial();
-		if (Mat != nullptr && Mat->IsUIMaterial())
+		UMaterial* BaseMaterial = MaterialInterface->GetMaterial();
+		if (BaseMaterial != nullptr && BaseMaterial->IsUIMaterial())
 		{
 			if (WidgetRenderer == nullptr)
 			{
@@ -45,7 +45,7 @@ void UMaterialInstanceThumbnailRenderer::Draw(UObject* Object, int32 X, int32 Y,
 				check(UIMaterialBrush);
 			}
 		
-			UIMaterialBrush->SetMaterial(MatInst);
+			UIMaterialBrush->SetMaterial(MaterialInterface);
 
 			UTexture2D* CheckerboardTexture = UThumbnailManager::Get().CheckerboardTexture;
 			
@@ -90,12 +90,12 @@ void UMaterialInstanceThumbnailRenderer::Draw(UObject* Object, int32 X, int32 Y,
 
 			// If the material doesn't have the static mesh usage flag then create a child
 			// instance with the usage flag set and use that for the thumbnail rendering.
-			if (!MatInst->GetUsageByFlag(MATUSAGE_StaticMesh))
+			if (!MaterialInterface->GetUsageByFlag(MATUSAGE_StaticMesh))
 			{
-				if (!OverrideMaterialInstance || OverrideMaterialInstance->Parent != MatInst)
+				if (!OverrideMaterialInstance || OverrideMaterialInstance->Parent != MaterialInterface)
 				{
 					OverrideMaterialInstance = NewObject<UMaterialInstanceConstant>();
-					OverrideMaterialInstance->SetParentEditorOnly(MatInst);
+					OverrideMaterialInstance->SetParentEditorOnly(MaterialInterface);
 					OverrideMaterialInstance->BasePropertyOverrides.bOverride_UsageFlags = (1 << MATUSAGE_StaticMesh);
 					OverrideMaterialInstance->BasePropertyOverrides.UsageFlags = (1 << MATUSAGE_StaticMesh);
 					OverrideMaterialInstance->UpdateStaticPermutation();
@@ -104,7 +104,7 @@ void UMaterialInstanceThumbnailRenderer::Draw(UObject* Object, int32 X, int32 Y,
 			}
 			else
 			{
-				ThumbnailScene->SetMaterialInterface(MatInst);
+				ThumbnailScene->SetMaterialInterface(MaterialInterface);
 			}
 
 			FSceneViewFamilyContext ViewFamily( FSceneViewFamily::ConstructionValues( RenderTarget, ThumbnailScene->GetScene(), FEngineShowFlags(ESFIM_Game) )
@@ -112,7 +112,7 @@ void UMaterialInstanceThumbnailRenderer::Draw(UObject* Object, int32 X, int32 Y,
 				.SetAdditionalViewFamily(bAdditionalViewFamily));
 
 			ViewFamily.EngineShowFlags.DisableAdvancedFeatures();
-			ViewFamily.EngineShowFlags.SetSeparateTranslucency(ThumbnailScene->ShouldSetSeparateTranslucency(MatInst));
+			ViewFamily.EngineShowFlags.SetSeparateTranslucency(ThumbnailScene->ShouldSetSeparateTranslucency(MaterialInterface));
 			ViewFamily.EngineShowFlags.MotionBlur = 0;
 			ViewFamily.EngineShowFlags.AntiAliasing = 0;
 
@@ -125,9 +125,9 @@ void UMaterialInstanceThumbnailRenderer::Draw(UObject* Object, int32 X, int32 Y,
 
 bool UMaterialInstanceThumbnailRenderer::CanVisualizeAsset(UObject* Object)
 {
-	if (UMaterialInterface* MatInst = Cast<UMaterialInterface>(Object))
+	if (UMaterialInterface* MaterialInterface = Cast<UMaterialInterface>(Object))
 	{
-		if (MatInst->IsCompiling())
+		if (MaterialInterface->IsCompiling())
 		{
 			return false;
 		}
