UE: Bump incremental version. Shader change

diff --git a/Engine/Source/Editor/UnrealEd/Private/Cooker/CookIncrementalVersion.cpp b/Engine/Source/Editor/UnrealEd/Private/Cooker/CookIncrementalVersion.cpp
--- a/Engine/Source/Editor/UnrealEd/Private/Cooker/CookIncrementalVersion.cpp
+++ b/Engine/Source/Editor/UnrealEd/Private/Cooker/CookIncrementalVersion.cpp
@@ -8,6 +8,6 @@ namespace UE::Cook
 {
 
 // Change CookIncrementalVersion to a new guid when all packages in an incremental cook need to be invalidated
-FGuid CookIncrementalVersion(0xF5F6B1B6, 0x93C147B4, 0xBD91AC9D, 0x411692F2);
+FGuid CookIncrementalVersion(0x890DC5D0, 0x4FC046B9, 0xA6F1D755, 0x09E0E84D);
 
 }
