HarmonixNiagaraEditor

#tests PIE

#Synced-CL 48996228

diff --git a/Engine/Plugins/Animation/DeformerGraph/Source/OptimusEditor/Private/OptimusEditorMode.cpp b/Engine/Plugins/Animation/DeformerGraph/Source/OptimusEditor/Private/OptimusEditorMode.cpp
--- a/Engine/Plugins/Animation/DeformerGraph/Source/OptimusEditor/Private/OptimusEditorMode.cpp
+++ b/Engine/Plugins/Animation/DeformerGraph/Source/OptimusEditor/Private/OptimusEditorMode.cpp
@@ -129,6 +129,4 @@ void FOptimusEditorMode::RegisterTabFactories(TSharedPtr<FTabManager> InTabManag
 		// HostingApp->RegisterTabSpawners(InTabManager.ToSharedRef());
 		HostingApp->PushTabFactories(TabFactories);
 	}
-
-	FApplicationMode::RegisterTabFactories(InTabManager);
 }
diff --git a/Engine/Plugins/Animation/IKRig/Source/IKRigEditor/Private/RetargetEditor/IKRetargetApplicationMode.cpp b/Engine/Plugins/Animation/IKRig/Source/IKRigEditor/Private/RetargetEditor/IKRetargetApplicationMode.cpp
--- a/Engine/Plugins/Animation/IKRig/Source/IKRigEditor/Private/RetargetEditor/IKRetargetApplicationMode.cpp
+++ b/Engine/Plugins/Animation/IKRig/Source/IKRigEditor/Private/RetargetEditor/IKRetargetApplicationMode.cpp
@@ -121,7 +121,6 @@ void FIKRetargetApplicationMode::RegisterTabFactories(TSharedPtr<FTabManager> In
 	TSharedPtr<FIKRetargetEditor> IKRigEditor = IKRetargetEditorPtr.Pin();
 	IKRigEditor->RegisterTabSpawners(InTabManager.ToSharedRef());
 	IKRigEditor->PushTabFactories(TabFactories);
-	FApplicationMode::RegisterTabFactories(InTabManager);
 }
 
 #undef LOCTEXT_NAMESPACE
diff --git a/Engine/Plugins/Animation/IKRig/Source/IKRigEditor/Private/RigEditor/IKRigMode.cpp b/Engine/Plugins/Animation/IKRig/Source/IKRigEditor/Private/RigEditor/IKRigMode.cpp
--- a/Engine/Plugins/Animation/IKRig/Source/IKRigEditor/Private/RigEditor/IKRigMode.cpp
+++ b/Engine/Plugins/Animation/IKRig/Source/IKRigEditor/Private/RigEditor/IKRigMode.cpp
@@ -128,7 +128,6 @@ void FIKRigMode::RegisterTabFactories(TSharedPtr<FTabManager> InTabManager)
 	TSharedPtr<FIKRigEditorToolkit> IKRigEditor = IKRigEditorPtr.Pin();
 	IKRigEditor->RegisterTabSpawners(InTabManager.ToSharedRef());
 	IKRigEditor->PushTabFactories(TabFactories);
-	FApplicationMode::RegisterTabFactories(InTabManager);
 }
 
 #undef LOCTEXT_NAMESPACE
diff --git a/Engine/Plugins/Animation/LiveLinkHub/Source/LiveLinkHub/Private/LiveLinkHubApplication.cpp b/Engine/Plugins/Animation/LiveLinkHub/Source/LiveLinkHub/Private/LiveLinkHubApplication.cpp
--- a/Engine/Plugins/Animation/LiveLinkHub/Source/LiveLinkHub/Private/LiveLinkHubApplication.cpp
+++ b/Engine/Plugins/Animation/LiveLinkHub/Source/LiveLinkHub/Private/LiveLinkHubApplication.cpp
@@ -105,7 +105,7 @@ void FLiveLinkHubApplication::PushTabFactories(const FWorkflowAllowedTabSet& Fac
 void FLiveLinkHubApplication::AddApplicationMode(FName ModeName, TSharedRef<FApplicationMode> Mode)
 {
 	// Register tabs for a mode once it's registered so that other modes have access to that tab.
-	Mode->RegisterTabFactories(TabManager);
+	Mode->RegisterTabFactoriesWithManager(TabManager);
 	FWorkflowCentricApplication::AddApplicationMode(ModeName, MoveTemp(Mode));
 }
 
@@ -133,7 +133,7 @@ void FLiveLinkHubApplication::SetCurrentMode(FName NewMode)
 			if (CurrentAppModePtr.IsValid())
 			{
 				CurrentAppModePtr->PreDeactivateMode();
-				CurrentAppModePtr->DeactivateMode(TabManager);
+				CurrentAppModePtr->DeactivateMode();
 				RemoveToolbarExtender(CurrentAppModePtr->GetToolbarExtender());
 				RemoveAllToolbarWidgets();
 			}
@@ -141,7 +141,7 @@ void FLiveLinkHubApplication::SetCurrentMode(FName NewMode)
 			CurrentAppModePtr = NewModePtr;
 
 			// Activate the new layout
-			const TSharedRef<FTabManager::FLayout> NewLayout = CurrentAppModePtr->ActivateMode(TabManager);
+			const TSharedRef<FTabManager::FLayout> NewLayout = CurrentAppModePtr->ActivateMode();
 			RestoreFromLayout(NewLayout);
 
 			// Give the new mode a chance to do init
diff --git a/Engine/Plugins/Animation/LiveLinkHub/Source/LiveLinkHub/Public/LiveLinkHubApplicationMode.h b/Engine/Plugins/Animation/LiveLinkHub/Source/LiveLinkHub/Public/LiveLinkHubApplicationMode.h
--- a/Engine/Plugins/Animation/LiveLinkHub/Source/LiveLinkHub/Public/LiveLinkHubApplicationMode.h
+++ b/Engine/Plugins/Animation/LiveLinkHub/Source/LiveLinkHub/Public/LiveLinkHubApplicationMode.h
@@ -93,7 +93,6 @@ protected:
 				App->PushTabFactories(TabFactories, StaticCastSharedRef<FLiveLinkHubApplicationMode>(AsShared()));
 			}
 
-			FApplicationMode::RegisterTabFactories(InTabManager);
 			bRegisteredTabFactories = true;
 		}
 	}
diff --git a/Engine/Plugins/Animation/MLDeformer/MLDeformerFramework/Source/MLDeformerFrameworkEditor/Private/MLDeformerApplicationMode.cpp b/Engine/Plugins/Animation/MLDeformer/MLDeformerFramework/Source/MLDeformerFrameworkEditor/Private/MLDeformerApplicationMode.cpp
--- a/Engine/Plugins/Animation/MLDeformer/MLDeformerFramework/Source/MLDeformerFrameworkEditor/Private/MLDeformerApplicationMode.cpp
+++ b/Engine/Plugins/Animation/MLDeformer/MLDeformerFramework/Source/MLDeformerFrameworkEditor/Private/MLDeformerApplicationMode.cpp
@@ -118,7 +118,6 @@ namespace UE::MLDeformer
 		TSharedPtr<FMLDeformerEditorToolkit> Editor = EditorToolkit.Pin();
 		Editor->RegisterTabSpawners(InTabManager.ToSharedRef());
 		Editor->PushTabFactories(TabFactories);
-		FApplicationMode::RegisterTabFactories(InTabManager);
 	}
 
 }	// namespace UE::MLDeformer
diff --git a/Engine/Plugins/Experimental/CodeEditor/Source/CodeEditor/Private/CodeProjectEditor.cpp b/Engine/Plugins/Experimental/CodeEditor/Source/CodeEditor/Private/CodeProjectEditor.cpp
--- a/Engine/Plugins/Experimental/CodeEditor/Source/CodeEditor/Private/CodeProjectEditor.cpp
+++ b/Engine/Plugins/Experimental/CodeEditor/Source/CodeEditor/Private/CodeProjectEditor.cpp
@@ -174,8 +174,6 @@ void FBasicCodeEditorMode::RegisterTabFactories(TSharedPtr<FTabManager> InTabMan
 	Editor->RegisterToolbarTab(InTabManager.ToSharedRef());
 
 	Editor->PushTabFactories(TabFactories);
-
-	FApplicationMode::RegisterTabFactories(InTabManager);
 }
 
 FCodeProjectEditor::FCodeProjectEditor()
diff --git a/Engine/Plugins/Experimental/CommonConversation/Source/CommonConversationEditor/Private/ConversationEditorModes.cpp b/Engine/Plugins/Experimental/CommonConversation/Source/CommonConversationEditor/Private/ConversationEditorModes.cpp
--- a/Engine/Plugins/Experimental/CommonConversation/Source/CommonConversationEditor/Private/ConversationEditorModes.cpp
+++ b/Engine/Plugins/Experimental/CommonConversation/Source/CommonConversationEditor/Private/ConversationEditorModes.cpp
@@ -63,8 +63,6 @@ void FConversationEditorApplicationMode_GraphView::RegisterTabFactories(TSharedP
 
 	// Mode-specific setup
 	ConversationEditorPtr->PushTabFactories(ConversationEditorTabFactories);
-
-	FApplicationMode::RegisterTabFactories(InTabManager);
 }
 
 void FConversationEditorApplicationMode_GraphView::PreDeactivateMode()
diff --git a/Engine/Plugins/Experimental/PhysicsControl/Source/PhysicsControlEditor/Private/PhysicsControlAssetApplicationMode.cpp b/Engine/Plugins/Experimental/PhysicsControl/Source/PhysicsControlEditor/Private/PhysicsControlAssetApplicationMode.cpp
--- a/Engine/Plugins/Experimental/PhysicsControl/Source/PhysicsControlEditor/Private/PhysicsControlAssetApplicationMode.cpp
+++ b/Engine/Plugins/Experimental/PhysicsControl/Source/PhysicsControlEditor/Private/PhysicsControlAssetApplicationMode.cpp
@@ -167,7 +167,6 @@ void FPhysicsControlAssetApplicationMode::RegisterTabFactories(TSharedPtr<FTabMa
 	TSharedPtr<FPhysicsControlAssetEditor> Editor = PhysicsControlAssetEditor.Pin();
 	Editor->RegisterTabSpawners(InTabManager.ToSharedRef());
 	Editor->PushTabFactories(TabFactories);
-	FApplicationMode::RegisterTabFactories(InTabManager);
 }
 
 #undef LOCTEXT_NAMESPACE
diff --git a/Engine/Plugins/FX/Niagara/Source/NiagaraEditor/Private/Toolkits/SystemToolkitModes/NiagaraSystemToolkitModeBase.cpp b/Engine/Plugins/FX/Niagara/Source/NiagaraEditor/Private/Toolkits/SystemToolkitModes/NiagaraSystemToolkitModeBase.cpp
--- a/Engine/Plugins/FX/Niagara/Source/NiagaraEditor/Private/Toolkits/SystemToolkitModes/NiagaraSystemToolkitModeBase.cpp
+++ b/Engine/Plugins/FX/Niagara/Source/NiagaraEditor/Private/Toolkits/SystemToolkitModes/NiagaraSystemToolkitModeBase.cpp
@@ -233,7 +233,8 @@ void FNiagaraSystemToolkitModeBase::RegisterTabFactories(TSharedPtr<FTabManager>
 {
 	WorkspaceMenuCategory = InTabManager->AddLocalWorkspaceMenuCategory(LOCTEXT("WorkspaceMenu_NiagaraSystemEditor", "Niagara System"));
 
-	SystemToolkit.Pin()->RegisterToolbarTab(InTabManager.ToSharedRef());
+	TSharedPtr<FNiagaraSystemToolkit> PinnedSystemToolkit = SystemToolkit.Pin();
+	PinnedSystemToolkit->RegisterToolbarTab(InTabManager.ToSharedRef());
 	
 	InTabManager->RegisterTabSpawner(ViewportTabID, FOnSpawnTab::CreateSP(this, &FNiagaraSystemToolkitModeBase::SpawnTab_Viewport))
 		.SetDisplayName(LOCTEXT("Preview", "Preview"))
@@ -339,6 +340,8 @@ void FNiagaraSystemToolkitModeBase::RegisterTabFactories(TSharedPtr<FTabManager>
 		.SetDisplayName(LOCTEXT("ScratchPadHierarchyEditor", "Edit Scratch Pad Hierarchy"))
 		.SetGroup(WorkspaceMenuCategory.ToSharedRef())
 		.SetIcon(FSlateIcon(FNiagaraEditorStyle::Get().GetStyleSetName(), "Tab.UserParameterHierarchy"));
+
+	RegisterTabFactoriesWithAppAndManager(PinnedSystemToolkit.Get(), InTabManager.ToSharedRef());
 }
 
 int FNiagaraSystemToolkitModeBase::GetActiveSelectionDetailsIndex() const
diff --git a/Engine/Plugins/VirtualProduction/Avalanche/Source/AvalancheMediaEditor/Private/Broadcast/AppModes/AvaBroadcastAppMode.cpp b/Engine/Plugins/VirtualProduction/Avalanche/Source/AvalancheMediaEditor/Private/Broadcast/AppModes/AvaBroadcastAppMode.cpp
--- a/Engine/Plugins/VirtualProduction/Avalanche/Source/AvalancheMediaEditor/Private/Broadcast/AppModes/AvaBroadcastAppMode.cpp
+++ b/Engine/Plugins/VirtualProduction/Avalanche/Source/AvalancheMediaEditor/Private/Broadcast/AppModes/AvaBroadcastAppMode.cpp
@@ -18,7 +18,6 @@ void FAvaBroadcastAppMode::RegisterTabFactories(TSharedPtr<FTabManager> InTabMan
 {
 	TSharedPtr<FAvaBroadcastEditor> BroadcastEditor = BroadcastEditorWeak.Pin();
 	BroadcastEditor->PushTabFactories(TabFactories);
-	FApplicationMode::RegisterTabFactories(InTabManager);
 }
 
 FText FAvaBroadcastAppMode::GetLocalizedMode(const FName InMode)
diff --git a/Engine/Plugins/VirtualProduction/Avalanche/Source/AvalancheMediaEditor/Private/Playback/AppModes/AvaPlaybackAppMode.cpp b/Engine/Plugins/VirtualProduction/Avalanche/Source/AvalancheMediaEditor/Private/Playback/AppModes/AvaPlaybackAppMode.cpp
--- a/Engine/Plugins/VirtualProduction/Avalanche/Source/AvalancheMediaEditor/Private/Playback/AppModes/AvaPlaybackAppMode.cpp
+++ b/Engine/Plugins/VirtualProduction/Avalanche/Source/AvalancheMediaEditor/Private/Playback/AppModes/AvaPlaybackAppMode.cpp
@@ -18,7 +18,6 @@ void FAvaPlaybackAppMode::RegisterTabFactories(TSharedPtr<FTabManager> InTabMana
 {
 	TSharedPtr<FAvaPlaybackGraphEditor> PlaybackEditor = PlaybackEditorWeak.Pin();
 	PlaybackEditor->PushTabFactories(TabFactories);
-	FApplicationMode::RegisterTabFactories(InTabManager);
 }
 
 FText FAvaPlaybackAppMode::GetLocalizedMode(const FName InMode)
diff --git a/Engine/Plugins/VirtualProduction/Avalanche/Source/AvalancheMediaEditor/Private/Rundown/AppModes/AvaRundownAppMode.cpp b/Engine/Plugins/VirtualProduction/Avalanche/Source/AvalancheMediaEditor/Private/Rundown/AppModes/AvaRundownAppMode.cpp
--- a/Engine/Plugins/VirtualProduction/Avalanche/Source/AvalancheMediaEditor/Private/Rundown/AppModes/AvaRundownAppMode.cpp
+++ b/Engine/Plugins/VirtualProduction/Avalanche/Source/AvalancheMediaEditor/Private/Rundown/AppModes/AvaRundownAppMode.cpp
@@ -18,7 +18,6 @@ void FAvaRundownAppMode::RegisterTabFactories(TSharedPtr<FTabManager> InTabManag
 {
 	const TSharedPtr<FAvaRundownEditor> RundownEditor = RundownEditorWeak.Pin();
 	RundownEditor->PushTabFactories(TabFactories);
-	FApplicationMode::RegisterTabFactories(InTabManager);
 }
 
 TSharedPtr<FDocumentTabFactory> FAvaRundownAppMode::GetDocumentTabFactory(const FName& InName) const
diff --git a/Engine/Plugins/VirtualProduction/Avalanche/Source/AvalancheTransitionEditor/Private/AppModes/AvaTransitionAppMode.cpp b/Engine/Plugins/VirtualProduction/Avalanche/Source/AvalancheTransitionEditor/Private/AppModes/AvaTransitionAppMode.cpp
--- a/Engine/Plugins/VirtualProduction/Avalanche/Source/AvalancheTransitionEditor/Private/AppModes/AvaTransitionAppMode.cpp
+++ b/Engine/Plugins/VirtualProduction/Avalanche/Source/AvalancheTransitionEditor/Private/AppModes/AvaTransitionAppMode.cpp
@@ -101,7 +101,6 @@ void FAvaTransitionAppMode::RegisterTabFactories(TSharedPtr<FTabManager> InTabMa
 	check(Editor.IsValid());
 	Editor->SetEditorMode(EditorMode);
 	Editor->PushTabFactories(TabFactories);
-	FApplicationMode::RegisterTabFactories(InTabManager);
 }
 
 #undef LOCTEXT_NAMESPACE
diff --git a/Engine/Plugins/VirtualProduction/DMX/DMXEngine/Source/DMXEditor/Private/Modes/DMXEditorApplicationMode.cpp b/Engine/Plugins/VirtualProduction/DMX/DMXEngine/Source/DMXEditor/Private/Modes/DMXEditorApplicationMode.cpp
--- a/Engine/Plugins/VirtualProduction/DMX/DMXEngine/Source/DMXEditor/Private/Modes/DMXEditorApplicationMode.cpp
+++ b/Engine/Plugins/VirtualProduction/DMX/DMXEngine/Source/DMXEditor/Private/Modes/DMXEditorApplicationMode.cpp
@@ -62,8 +62,6 @@ void FDMXEditorDefaultApplicationMode::RegisterTabFactories(TSharedPtr<FTabManag
 
 		// Setup all tab factories
 		DMXEditorPtr->PushTabFactories(DefaultsTabFactories);
-
-		FApplicationMode::RegisterTabFactories(InTabManager);
 	}
 }
 
diff --git a/Engine/Source/Editor/AnimationEditor/Private/AnimationEditorMode.cpp b/Engine/Source/Editor/AnimationEditor/Private/AnimationEditorMode.cpp
--- a/Engine/Source/Editor/AnimationEditor/Private/AnimationEditorMode.cpp
+++ b/Engine/Source/Editor/AnimationEditor/Private/AnimationEditorMode.cpp
@@ -20,7 +20,7 @@
 FAnimationEditorMode::FAnimationEditorMode(TSharedRef<FWorkflowCentricApplication> InHostingApp, TSharedRef<ISkeletonTree> InSkeletonTree)
 	: FApplicationMode(AnimationEditorModes::AnimationEditorMode)
 {
-	HostingAppPtr = InHostingApp;
+	FWorkflowAllowedTabSet& TabFactories = ApplicationModeTabFactories;
 
 	TSharedRef<FAnimationEditor> AnimationEditor = StaticCastSharedRef<FAnimationEditor>(InHostingApp);
 
@@ -130,22 +130,10 @@ FAnimationEditorMode::FAnimationEditorMode(TSharedRef<FWorkflowCentricApplicatio
 
 void FAnimationEditorMode::RegisterTabFactories(TSharedPtr<FTabManager> InTabManager)
 {
-	TSharedPtr<FWorkflowCentricApplication> HostingApp = HostingAppPtr.Pin();
-	HostingApp->RegisterTabSpawners(InTabManager.ToSharedRef());
-	HostingApp->PushTabFactories(TabFactories);
-
-	FApplicationMode::RegisterTabFactories(InTabManager);
-}
-
-void FAnimationEditorMode::AddTabFactory(FCreateWorkflowTabFactory FactoryCreator)
-{
-	if (FactoryCreator.IsBound())
+	if (TSharedPtr<FWorkflowCentricApplication> HostingApp = GetHost())
 	{
-		TabFactories.RegisterFactory(FactoryCreator.Execute(HostingAppPtr.Pin()));		
+		TSharedRef<FTabManager> TabManager = InTabManager.ToSharedRef();
+		HostingApp->RegisterTabSpawners(TabManager);
+		RegisterTabFactoriesWithAppAndManager(HostingApp.Get(), TabManager);
 	}
 }
-
-void FAnimationEditorMode::RemoveTabFactory(FName TabFactoryID)
-{
-	TabFactories.UnregisterFactory(TabFactoryID);
-}
diff --git a/Engine/Source/Editor/AnimationEditor/Private/AnimationEditorMode.h b/Engine/Source/Editor/AnimationEditor/Private/AnimationEditorMode.h
--- a/Engine/Source/Editor/AnimationEditor/Private/AnimationEditorMode.h
+++ b/Engine/Source/Editor/AnimationEditor/Private/AnimationEditorMode.h
@@ -16,14 +16,4 @@ public:
 
 	/** FApplicationMode interface */
 	virtual void RegisterTabFactories(TSharedPtr<FTabManager> InTabManager) override;
-
-protected:
-	virtual void AddTabFactory(FCreateWorkflowTabFactory FactoryCreator) override;
-	virtual void RemoveTabFactory(FName TabFactoryID) override;
-protected:
-	/** The hosting app */
-	TWeakPtr<class FWorkflowCentricApplication> HostingAppPtr;
-
-	/** The tab factories we support */
-	FWorkflowAllowedTabSet TabFactories;
 };
diff --git a/Engine/Source/Editor/BehaviorTreeEditor/Private/BehaviorTreeEditorModes.cpp b/Engine/Source/Editor/BehaviorTreeEditor/Private/BehaviorTreeEditorModes.cpp
--- a/Engine/Source/Editor/BehaviorTreeEditor/Private/BehaviorTreeEditorModes.cpp
+++ b/Engine/Source/Editor/BehaviorTreeEditor/Private/BehaviorTreeEditorModes.cpp
@@ -72,8 +72,6 @@ void FBehaviorTreeEditorApplicationMode::RegisterTabFactories(TSharedPtr<FTabMan
 
 	// Mode-specific setup
 	BehaviorTreeEditorPtr->PushTabFactories(BehaviorTreeEditorTabFactories);
-
-	FApplicationMode::RegisterTabFactories(InTabManager);
 }
 
 void FBehaviorTreeEditorApplicationMode::PreDeactivateMode()
@@ -143,8 +141,6 @@ void FBlackboardEditorApplicationMode::RegisterTabFactories(TSharedPtr<FTabManag
 
 	// Mode-specific setup
 	BehaviorTreeEditorPtr->PushTabFactories(BlackboardTabFactories);
-
-	FApplicationMode::RegisterTabFactories(InTabManager);
 }
 
 void FBlackboardEditorApplicationMode::PostActivateMode()
diff --git a/Engine/Source/Editor/Kismet/Private/BlueprintEditorModes.cpp b/Engine/Source/Editor/Kismet/Private/BlueprintEditorModes.cpp
--- a/Engine/Source/Editor/Kismet/Private/BlueprintEditorModes.cpp
+++ b/Engine/Source/Editor/Kismet/Private/BlueprintEditorModes.cpp
@@ -129,8 +129,6 @@ void FBlueprintEditorApplicationMode::RegisterTabFactories(TSharedPtr<FTabManage
 	BP->PushTabFactories(CoreTabFactories);
 	BP->PushTabFactories(BlueprintEditorOnlyTabFactories);
 	BP->PushTabFactories(BlueprintEditorTabFactories);
-
-	FApplicationMode::RegisterTabFactories(InTabManager);
 }
 
 void FBlueprintEditorApplicationMode::PreDeactivateMode()
@@ -193,8 +191,6 @@ void FBlueprintDefaultsApplicationMode::RegisterTabFactories(TSharedPtr<FTabMana
 
 	// Mode-specific setup
 	BP->PushTabFactories(BlueprintDefaultsTabFactories);
-
-	FApplicationMode::RegisterTabFactories(InTabManager);
 }
 
 void FBlueprintDefaultsApplicationMode::PostActivateMode()
@@ -271,8 +267,6 @@ void FBlueprintComponentsApplicationMode::RegisterTabFactories(TSharedPtr<FTabMa
 
 	// Mode-specific setup
 	BP->PushTabFactories(BlueprintComponentsTabFactories);
-
-	FApplicationMode::RegisterTabFactories(InTabManager);
 }
 
 void FBlueprintComponentsApplicationMode::PreDeactivateMode()
@@ -402,8 +396,6 @@ void FBlueprintInterfaceApplicationMode::RegisterTabFactories(TSharedPtr<FTabMan
 
 	// Mode-specific setup
 	BP->PushTabFactories(BlueprintInterfaceTabFactories);
-
-	FApplicationMode::RegisterTabFactories(InTabManager);
 }
 
 void FBlueprintInterfaceApplicationMode::PreDeactivateMode()
@@ -499,8 +491,6 @@ void FBlueprintMacroApplicationMode::RegisterTabFactories(TSharedPtr<FTabManager
 
 	// Mode-specific setup
 	BP->PushTabFactories(BlueprintMacroTabFactories);
-
-	FApplicationMode::RegisterTabFactories(InTabManager);
 }
 
 void FBlueprintMacroApplicationMode::PreDeactivateMode()
@@ -687,8 +677,6 @@ void FBlueprintEditorUnifiedMode::RegisterTabFactories(TSharedPtr<FTabManager> I
 	BP->PushTabFactories(CoreTabFactories);
 	BP->PushTabFactories(BlueprintEditorOnlyTabFactories);
 	BP->PushTabFactories(BlueprintEditorTabFactories);
-
-	FApplicationMode::RegisterTabFactories(InTabManager);
 }
 
 void FBlueprintEditorUnifiedMode::PreDeactivateMode()
diff --git a/Engine/Source/Editor/MaterialEditor/Private/MaterialEditorModes.cpp b/Engine/Source/Editor/MaterialEditor/Private/MaterialEditorModes.cpp
--- a/Engine/Source/Editor/MaterialEditor/Private/MaterialEditorModes.cpp
+++ b/Engine/Source/Editor/MaterialEditor/Private/MaterialEditorModes.cpp
@@ -99,8 +99,6 @@ void FMaterialEditorApplicationMode::RegisterTabFactories(TSharedPtr<FTabManager
 
 	// Mode-specific setup
 	MaterialEditor->PushTabFactories(MaterialEditorTabFactories);
-
-	FApplicationMode::RegisterTabFactories(InTabManager);
 }
 
 #undef LOCTEXT_NAMESPACE
\ No newline at end of file
diff --git a/Engine/Source/Editor/PhysicsAssetEditor/Private/PhysicsAssetEditorMode.cpp b/Engine/Source/Editor/PhysicsAssetEditor/Private/PhysicsAssetEditorMode.cpp
--- a/Engine/Source/Editor/PhysicsAssetEditor/Private/PhysicsAssetEditorMode.cpp
+++ b/Engine/Source/Editor/PhysicsAssetEditor/Private/PhysicsAssetEditorMode.cpp
@@ -301,8 +301,6 @@ void FPhysicsAssetEditorMode::RegisterTabFactories(TSharedPtr<FTabManager> InTab
 	TSharedPtr<FPhysicsAssetEditor> PhysicsAssetEditor = PhysicsAssetEditorPtr.Pin();
 	PhysicsAssetEditor->RegisterTabSpawners(InTabManager.ToSharedRef());
 	PhysicsAssetEditor->PushTabFactories(TabFactories);
-
-	FApplicationMode::RegisterTabFactories(InTabManager);
 }
 
 #undef LOCTEXT_NAMESPACE
diff --git a/Engine/Source/Editor/SkeletalMeshEditor/Private/SkeletalMeshEditorMode.cpp b/Engine/Source/Editor/SkeletalMeshEditor/Private/SkeletalMeshEditorMode.cpp
--- a/Engine/Source/Editor/SkeletalMeshEditor/Private/SkeletalMeshEditorMode.cpp
+++ b/Engine/Source/Editor/SkeletalMeshEditor/Private/SkeletalMeshEditorMode.cpp
@@ -24,7 +24,7 @@ FSkeletalMeshEditorMode::FSkeletalMeshEditorMode(
 	)
 	: FApplicationMode(SkeletalMeshEditorModes::SkeletalMeshEditorMode)
 {
-	HostingAppPtr = InHostingApp;
+	FWorkflowAllowedTabSet& TabFactories = ApplicationModeTabFactories;
 
 	TSharedRef<FSkeletalMeshEditor> SkeletalMeshEditor = StaticCastSharedRef<FSkeletalMeshEditor>(InHostingApp);
 
@@ -131,27 +131,14 @@ FSkeletalMeshEditorMode::FSkeletalMeshEditorMode(
 
 void FSkeletalMeshEditorMode::RegisterTabFactories(TSharedPtr<FTabManager> InTabManager)
 {
-	TSharedPtr<FWorkflowCentricApplication> HostingApp = HostingAppPtr.Pin();
-	HostingApp->RegisterTabSpawners(InTabManager.ToSharedRef());
-	HostingApp->PushTabFactories(TabFactories);
-
-	FApplicationMode::RegisterTabFactories(InTabManager);
-}
-
-
-void FSkeletalMeshEditorMode::AddTabFactory(FCreateWorkflowTabFactory FactoryCreator)
-{
-	if (FactoryCreator.IsBound())
+	if (TSharedPtr<FWorkflowCentricApplication> HostingApp = GetHost())
 	{
-		TabFactories.RegisterFactory(FactoryCreator.Execute(HostingAppPtr.Pin()));
+		TSharedRef<FTabManager> TabManager = InTabManager.ToSharedRef();
+		HostingApp->RegisterTabSpawners(TabManager);
+		RegisterTabFactoriesWithAppAndManager(HostingApp.Get(), TabManager);
 	}
 }
 
-void FSkeletalMeshEditorMode::RemoveTabFactory(FName TabFactoryID)
-{
-	TabFactories.UnregisterFactory(TabFactoryID);
-}
-
 TSharedRef<class FWorkflowTabFactory> FSkeletalMeshEditorMode::CreateMeshControllerMappingTabFactory(const TSharedRef<class FWorkflowCentricApplication>& InHostingApp, const TWeakObjectPtr<class USkeletalMesh>& InEditingMesh, FSimpleMulticastDelegate& OnPostUndo) const
 {
 	return MakeShareable(new FMeshControllerMappingTabSummoner(InHostingApp, InEditingMesh, OnPostUndo));
diff --git a/Engine/Source/Editor/SkeletalMeshEditor/Private/SkeletalMeshEditorMode.h b/Engine/Source/Editor/SkeletalMeshEditor/Private/SkeletalMeshEditorMode.h
--- a/Engine/Source/Editor/SkeletalMeshEditor/Private/SkeletalMeshEditorMode.h
+++ b/Engine/Source/Editor/SkeletalMeshEditor/Private/SkeletalMeshEditorMode.h
@@ -26,17 +26,7 @@ public:
 	/** FApplicationMode interface */
 	virtual void RegisterTabFactories(TSharedPtr<FTabManager> InTabManager) override;
 
-
-	virtual void AddTabFactory(FCreateWorkflowTabFactory FactoryCreator) override;
-	virtual void RemoveTabFactory(FName TabFactoryID) override;
-
 protected:
-	/** The hosting app */
-	TWeakPtr<class FWorkflowCentricApplication> HostingAppPtr;
-
-	/** The tab factories we support */
-	FWorkflowAllowedTabSet TabFactories;
-
 	/** Mesh interface class */
 	TSharedRef<class FWorkflowTabFactory> CreateMeshControllerMappingTabFactory(const TSharedRef<class FWorkflowCentricApplication>& InHostingApp, const TWeakObjectPtr<class USkeletalMesh>& InEditingMesh, FSimpleMulticastDelegate& OnPostUndo) const;
 };
diff --git a/Engine/Source/Editor/SkeletonEditor/Private/SkeletonEditorMode.cpp b/Engine/Source/Editor/SkeletonEditor/Private/SkeletonEditorMode.cpp
--- a/Engine/Source/Editor/SkeletonEditor/Private/SkeletonEditorMode.cpp
+++ b/Engine/Source/Editor/SkeletonEditor/Private/SkeletonEditorMode.cpp
@@ -14,7 +14,7 @@
 FSkeletonEditorMode::FSkeletonEditorMode(TSharedRef<FWorkflowCentricApplication> InHostingApp, TSharedRef<ISkeletonTree> InSkeletonTree)
 	: FApplicationMode(SkeletonEditorModes::SkeletonEditorMode)
 {
-	HostingAppPtr = InHostingApp;
+	FWorkflowAllowedTabSet& TabFactories = ApplicationModeTabFactories;
 
 	TSharedRef<FSkeletonEditor> SkeletonEditor = StaticCastSharedRef<FSkeletonEditor>(InHostingApp);
 
@@ -116,22 +116,10 @@ FSkeletonEditorMode::FSkeletonEditorMode(TSharedRef<FWorkflowCentricApplication>
 
 void FSkeletonEditorMode::RegisterTabFactories(TSharedPtr<FTabManager> InTabManager)
 {
-	TSharedPtr<FWorkflowCentricApplication> HostingApp = HostingAppPtr.Pin();
-	HostingApp->RegisterTabSpawners(InTabManager.ToSharedRef());
-	HostingApp->PushTabFactories(TabFactories);
-
-	FApplicationMode::RegisterTabFactories(InTabManager);
-}
-
-void FSkeletonEditorMode::AddTabFactory(FCreateWorkflowTabFactory FactoryCreator)
-{
-	if (FactoryCreator.IsBound())
+	if (TSharedPtr<FWorkflowCentricApplication> HostingApp = GetHost())
 	{
-		TabFactories.RegisterFactory(FactoryCreator.Execute(HostingAppPtr.Pin()));
+		TSharedRef<FTabManager> TabManager = InTabManager.ToSharedRef();
+		HostingApp->RegisterTabSpawners(TabManager);
+		RegisterTabFactoriesWithAppAndManager(HostingApp.Get(), TabManager);
 	}
 }
-
-void FSkeletonEditorMode::RemoveTabFactory(FName TabFactoryID)
-{
-	TabFactories.UnregisterFactory(TabFactoryID);
-}
diff --git a/Engine/Source/Editor/SkeletonEditor/Private/SkeletonEditorMode.h b/Engine/Source/Editor/SkeletonEditor/Private/SkeletonEditorMode.h
--- a/Engine/Source/Editor/SkeletonEditor/Private/SkeletonEditorMode.h
+++ b/Engine/Source/Editor/SkeletonEditor/Private/SkeletonEditorMode.h
@@ -14,14 +14,4 @@ public:
 
 	/** FApplicationMode interface */
 	virtual void RegisterTabFactories(TSharedPtr<FTabManager> InTabManager) override;
-	
-protected:
-	virtual void AddTabFactory(FCreateWorkflowTabFactory FactoryCreator) override;
-	virtual void RemoveTabFactory(FName TabFactoryID) override;
-protected:
-	/** The hosting app */
-	TWeakPtr<class FWorkflowCentricApplication> HostingAppPtr;
-
-	/** The tab factories we support */
-	FWorkflowAllowedTabSet TabFactories;
 };
diff --git a/Engine/Source/Editor/UnrealEd/Private/ApplicationMode.cpp b/Engine/Source/Editor/UnrealEd/Private/ApplicationMode.cpp
--- a/Engine/Source/Editor/UnrealEd/Private/ApplicationMode.cpp
+++ b/Engine/Source/Editor/UnrealEd/Private/ApplicationMode.cpp
@@ -4,6 +4,7 @@
 #include "Framework/MultiBox/MultiBoxExtender.h"
 #include "Framework/Docking/LayoutService.h"
 #include "UnrealEdMisc.h"
+#include "WorkflowOrientedApp/WorkflowCentricApplication.h"
 
 DEFINE_LOG_CATEGORY_STATIC(LogApplicationMode, Log, All);
 
@@ -12,25 +13,28 @@ DEFINE_LOG_CATEGORY_STATIC(LogApplicationMode, Log, All);
 /////////////////////////////////////////////////////
 // FApplicationMode
 
-FApplicationMode::FApplicationMode(FName InModeName)
-	: ModeName(InModeName)
+FApplicationMode::FApplicationMode(FName InModeName) :
+	ModeName(InModeName),
+	ToolbarExtender(MakeShareable(new FExtender)),
+	WorkspaceMenuCategory(FWorkspaceItem::NewGroup(LOCTEXT("WorkspaceMenu_AssetEditor", "Asset Editor")))
 {
-	ToolbarExtender = MakeShareable(new FExtender);
-	WorkspaceMenuCategory = FWorkspaceItem::NewGroup(LOCTEXT("WorkspaceMenu_AssetEditor", "Asset Editor"));
 }
 
-FApplicationMode::FApplicationMode(FName InModeName, FText(*GetLocalizedMode)(const FName))
-	: ModeName(InModeName)
+FApplicationMode::FApplicationMode(FName InModeName, FText(*GetLocalizedMode)(const FName)) :
+	ModeName(InModeName),
+	ToolbarExtender(MakeShareable(new FExtender)),
+	WorkspaceMenuCategory(FWorkspaceItem::NewGroup(FText::Format(LOCTEXT("WorkspaceMenu_ApplicationMode", "{0} Editor"), GetLocalizedMode(InModeName))))
 {
-	ToolbarExtender = MakeShareable(new FExtender);
-	WorkspaceMenuCategory = FWorkspaceItem::NewGroup(FText::Format(LOCTEXT("WorkspaceMenu_ApplicationMode", "{0} Editor"), GetLocalizedMode(InModeName)));
 }
 
-void FApplicationMode::DeactivateMode(TSharedPtr<FTabManager> InTabManager)
+void FApplicationMode::DeactivateMode()
 {
 	// Save the layout to INI
-	check(InTabManager.IsValid());
-	FLayoutSaveRestore::SaveToConfig(LayoutIni, InTabManager->PersistLayout());
+	TSharedPtr<FWorkflowCentricApplication> Host = GetHost();
+	check(Host.IsValid());
+	TSharedPtr<FTabManager> TabManager = Host->GetAssociatedTabManager();
+	check(TabManager.IsValid());
+	FLayoutSaveRestore::SaveToConfig(LayoutIni, TabManager->PersistLayout());
 
 	// Unregister the tabs
 	/*
@@ -43,10 +47,9 @@ void FApplicationMode::DeactivateMode(TSharedPtr<FTabManager> InTabManager)
 	*/
 }
 
-TSharedRef<FTabManager::FLayout> FApplicationMode::ActivateMode(TSharedPtr<FTabManager> InTabManager)
+TSharedRef<FTabManager::FLayout> FApplicationMode::ActivateMode()
 {
-	check(InTabManager.IsValid());
-	RegisterTabFactories(InTabManager);
+	RegisterTabFactoriesWithHost();
 
 	/*
 	for (int32 Index = 0; Index < AllowableTabs.Num(); ++Index)
@@ -62,4 +65,67 @@ TSharedRef<FTabManager::FLayout> FApplicationMode::ActivateMode(TSharedPtr<FTabM
 	return FLayoutSaveRestore::LoadFromConfig(LayoutIni, TabLayout.ToSharedRef());
 }
 
+void FApplicationMode::AddTabFactory(FCreateWorkflowTabFactory FactoryCreator)
+{
+	if (FactoryCreator.IsBound())
+	{
+		TSharedPtr<FWorkflowCentricApplication> Host = GetHost();
+		check(Host.IsValid());
+		ApplicationModeTabFactories.RegisterFactory(FactoryCreator.Execute(Host));
+	}
+}
+
+void FApplicationMode::RemoveTabFactory(FName TabFactoryID)
+{
+	ApplicationModeTabFactories.UnregisterFactory(TabFactoryID);
+}
+
+void FApplicationMode::RegisterTabFactoriesWithHost()
+{
+	TSharedPtr<FWorkflowCentricApplication> Host = GetHost();
+	check(Host.IsValid());
+	RegisterTabFactoriesWithManager(Host->GetAssociatedTabManager());
+}
+
+void FApplicationMode::RegisterTabFactoriesWithManager(const TSharedPtr<FTabManager>& InTabManager)
+{
+	if (InTabManager)
+	{
+		RegisterTabFactories(InTabManager);
+	}
+}
+
+void FApplicationMode::RegisterTabFactoriesWithAppAndManager(FWorkflowCentricApplication* InApp, const TSharedRef<FTabManager>& InTabManager)
+{
+	if (InApp)
+	{
+		InApp->PushTabFactories(ApplicationModeTabFactories);
+	}
+}
+
+void FApplicationMode::RegisterTabFactories(TSharedPtr<FTabManager> InTabManager)
+{
+	check(InTabManager.IsValid());
+	RegisterTabFactoriesWithAppAndManager(GetHost().Get(), InTabManager.ToSharedRef());
+}
+
+void FApplicationMode::ShareHost(const FApplicationMode* Mode)
+{
+	if (Mode)
+	{
+		HostWeakPtr = Mode->HostWeakPtr;
+	}
+	else
+	{
+		if (HostWeakPtr.IsUnique())
+		{
+			SetHost(nullptr);
+		}
+		else
+		{
+			HostWeakPtr = MakeShared<TWeakPtr<FWorkflowCentricApplication>>();
+		}
+	}
+}
+
 #undef LOCTEXT_NAMESPACE
diff --git a/Engine/Source/Editor/UnrealEd/Private/WorkflowCentricApplication.cpp b/Engine/Source/Editor/UnrealEd/Private/WorkflowCentricApplication.cpp
--- a/Engine/Source/Editor/UnrealEd/Private/WorkflowCentricApplication.cpp
+++ b/Engine/Source/Editor/UnrealEd/Private/WorkflowCentricApplication.cpp
@@ -20,7 +20,7 @@ void FWorkflowCentricApplication::RegisterTabSpawners(const TSharedRef<class FTa
 {
 	if (CurrentAppModePtr.IsValid())
 	{
-		CurrentAppModePtr->RegisterTabFactories(InTabManager);
+		CurrentAppModePtr->RegisterTabFactoriesWithManager(InTabManager);
 	}
 }
 
@@ -87,7 +87,7 @@ void FWorkflowCentricApplication::SetCurrentMode(FName NewMode)
 			{
 				check(TabManager.IsValid());
 				CurrentAppModePtr->PreDeactivateMode();
-				CurrentAppModePtr->DeactivateMode(TabManager);
+				CurrentAppModePtr->DeactivateMode();
 				RemoveToolbarExtender(CurrentAppModePtr->GetToolbarExtender());
 				RemoveAllToolbarWidgets();
 			}
@@ -103,7 +103,7 @@ void FWorkflowCentricApplication::SetCurrentMode(FName NewMode)
 			TabManager->AddLocalWorkspaceMenuItem( CurrentAppModePtr->GetWorkspaceMenuCategory() );
 
 			// Activate the new layout
-			const TSharedRef<FTabManager::FLayout> NewLayout = CurrentAppModePtr->ActivateMode(TabManager);
+			const TSharedRef<FTabManager::FLayout> NewLayout = CurrentAppModePtr->ActivateMode();
 			RestoreFromLayout(NewLayout);
 
 			// Give the new mode a chance to do init
@@ -137,7 +137,7 @@ void FWorkflowCentricApplication::OnClose()
 
 		// Deactivate the old mode
 		CurrentAppModePtr->PreDeactivateMode();
-		CurrentAppModePtr->DeactivateMode(TabManager);
+		CurrentAppModePtr->DeactivateMode();
 		RemoveToolbarExtender(CurrentAppModePtr->GetToolbarExtender());
 		RemoveAllToolbarWidgets();
 
@@ -148,11 +148,22 @@ void FWorkflowCentricApplication::OnClose()
 
 void FWorkflowCentricApplication::AddApplicationMode(FName ModeName, TSharedRef<FApplicationMode> Mode)
 {
+	// The given mode must not belong to another host.
+	check(!Mode->GetHost() || Mode->GetHost().Get() == this);
+
+	// Ensure the mode has a unique reference to this host.
+	Mode->ShareHost(nullptr);
+	Mode->SetHost(StaticCastWeakPtr<FWorkflowCentricApplication>(AsWeak()));
+
 	for (int32 Index = 0; Index < ModeExtenderList.Num(); ++Index)
 	{
-		Mode = ModeExtenderList[Index].Execute(ModeName, Mode);
-	}
+		TSharedRef<FApplicationMode> ExtendedMode = ModeExtenderList[Index].Execute(ModeName, Mode);
+
+		// Share the mode's host with the extended mode.
+		ExtendedMode->ShareHost(&Mode.Get());
 
+		Swap(ExtendedMode, Mode);
+	}
 	ApplicationModeList.Add(ModeName, Mode);
 }
 
@@ -160,7 +171,17 @@ bool FWorkflowCentricApplication::RemoveApplicationMode(FName ModeName)
 {
 	if (GetCurrentMode() != ModeName)
 	{
-		return ApplicationModeList.Remove(ModeName) > 0;
+		TSharedPtr<FApplicationMode> Mode;
+		if (ApplicationModeList.RemoveAndCopyValue(ModeName, Mode))
+		{
+			if (Mode)
+			{
+				check(this == Mode->GetHost().Get());
+				Mode->SetHost(nullptr);
+			}
+			return true;
+		}
+		return false;
 	}
 	else
 	{
@@ -171,6 +192,14 @@ bool FWorkflowCentricApplication::RemoveApplicationMode(FName ModeName)
 
 void FWorkflowCentricApplication::RemoveAllApplicationModes()
 {
+	for (auto& Entry : ApplicationModeList)
+	{
+		if (Entry.Value)
+		{
+			check(this == Entry.Value->GetHost().Get());
+			Entry.Value->SetHost(nullptr);
+		}
+	}
 	ApplicationModeList.Reset();
 }
 
diff --git a/Engine/Source/Editor/UnrealEd/Public/WorkflowOrientedApp/ApplicationMode.h b/Engine/Source/Editor/UnrealEd/Public/WorkflowOrientedApp/ApplicationMode.h
--- a/Engine/Source/Editor/UnrealEd/Public/WorkflowOrientedApp/ApplicationMode.h
+++ b/Engine/Source/Editor/UnrealEd/Public/WorkflowOrientedApp/ApplicationMode.h
@@ -6,6 +6,7 @@
 #include "Framework/Docking/WorkspaceItem.h"
 #include "Framework/Docking/TabManager.h"
 #include "Framework/Docking/LayoutExtender.h"
+#include "WorkflowOrientedApp/WorkflowTabManager.h"
 
 DECLARE_DELEGATE_RetVal_OneParam(TSharedPtr<class FWorkflowTabFactory>, FCreateWorkflowTabFactory, TSharedPtr<class FAssetEditorToolkit>)
 
@@ -17,7 +18,17 @@ class FExtender;
 
 class FApplicationMode : public TSharedFromThis<FApplicationMode>
 {
+private:
+	friend class FWorkflowCentricApplication;
+
+	// A shared weak reference to the host.  If a mode extender creates a new mode that references its given mode,
+	// the host will use ShareHost to ensure that the new mode and the given mode share the same reference to the host.
+	TSharedRef<TWeakPtr<FWorkflowCentricApplication>> HostWeakPtr{ MakeShared<TWeakPtr<FWorkflowCentricApplication>>() };
+
 protected:
+	/** The tab factories we support */
+	FWorkflowAllowedTabSet ApplicationModeTabFactories;
+
 	// The layout to use in this mode
 	TSharedPtr<FTabManager::FLayout> TabLayout;
 
@@ -40,6 +51,14 @@ protected:
 	/** What ini file should be used to persist and restore the user's layout. Defaults to GEditorLayoutIni. */
 	FString LayoutIni = GEditorLayoutIni;
 
+	// ShareHost allows the host to set the host reference for potentially related modes.
+	UNREALED_API void ShareHost(const FApplicationMode* Mode);
+
+	void SetHost(TWeakPtr<FWorkflowCentricApplication> Host)
+	{
+		*HostWeakPtr = Host;
+	}
+
 public:
 	UNREALED_API FApplicationMode(FName InModeName);
 
@@ -47,13 +66,37 @@ public:
 
 	virtual ~FApplicationMode() = default;
 
-	UNREALED_API void DeactivateMode(TSharedPtr<FTabManager> InTabManager);
-	UNREALED_API TSharedRef<FTabManager::FLayout> ActivateMode(TSharedPtr<FTabManager> InTabManager);
-	virtual void RegisterTabFactories(TSharedPtr<FTabManager> InTabManager) {}
+	TSharedPtr<FWorkflowCentricApplication> GetHost() const
+	{
+		return HostWeakPtr->Pin();
+	}
 
-	virtual void AddTabFactory(FCreateWorkflowTabFactory FactoryCreator) {}
-	virtual void RemoveTabFactory(FName TabFactoryID) {}
-	
+	UNREALED_API void DeactivateMode();
+	UNREALED_API TSharedRef<FTabManager::FLayout> ActivateMode();
+
+	UNREALED_API void AddTabFactory(FCreateWorkflowTabFactory FactoryCreator);
+	UNREALED_API void RemoveTabFactory(FName TabFactoryID);
+
+	// Register this mode's tab factories with its host's tab manager.
+	UNREALED_API void RegisterTabFactoriesWithHost();
+
+	// Interface to register this mode's tab factories with an arbitrary tab manager.
+	UNREALED_API void RegisterTabFactoriesWithManager(const TSharedPtr<FTabManager>& InTabManager);
+
+protected:
+	// Register this mode's tab factories with an arbitrary app and tab manager.
+	// Derived classes should call this if they override RegisterTabFactories,
+	// but want the base class functionality.
+	UNREALED_API void RegisterTabFactoriesWithAppAndManager(FWorkflowCentricApplication* InApp, const TSharedRef<FTabManager>& InTabManager);
+
+private:
+	// Implementation of registering this mode's tab factories with an arbitrary tab manager.
+	// This is private to prevent derived classes from calling it.  It didn't used to do anything.
+	// Now it calls RegisterTabFactoriesWithAppAndManager with the host and the provided tab manager.
+	// If derived classes want this behaviour, they can call RegisterTabFactoriesWithAppAndManager directly.
+	UNREALED_API virtual void RegisterTabFactories(TSharedPtr<FTabManager> InTabManager);
+
+public:
 	virtual void PreDeactivateMode() {}
 	virtual void PostActivateMode() {}
 
