Add GetPreviewPlatformName Blueprint function

#Synced-CL 48952469

diff --git a/Engine/Source/Editor/UnrealEd/Private/Subsystems/EditorSubsystemBlueprintLibrary.cpp b/Engine/Source/Editor/UnrealEd/Private/Subsystems/EditorSubsystemBlueprintLibrary.cpp
--- a/Engine/Source/Editor/UnrealEd/Private/Subsystems/EditorSubsystemBlueprintLibrary.cpp
+++ b/Engine/Source/Editor/UnrealEd/Private/Subsystems/EditorSubsystemBlueprintLibrary.cpp
@@ -11,19 +11,19 @@
 	return GEditor->GetEditorSubsystemBase(Class);
 }
 
-FName GetPreviewPlatformDropdownName(const FPreviewPlatformMenuItem& Item)
+FName GetPreviewPlatformDropdownName(const FName& PreviewShaderPlatformName, const FName& DeviceProfileName)
 {
 	// We cannot use the localized name, because it will get cached and not work for other Languages.
 	// Instead we take the PreviewShaderPlatformName and remove the postfix and suffix if they are present
 	// and concatenate it with the DeviceProfileName if present.
 
-	FString DropdownName = Item.PreviewShaderPlatformName.ToString();
+	FString DropdownName = PreviewShaderPlatformName.ToString();
 	DropdownName.RemoveFromEnd(TEXT("_Preview"));
 	DropdownName.RemoveFromStart(TEXT("SP_"));
 
-	if (Item.DeviceProfileName != NAME_None)
+	if (DeviceProfileName != NAME_None)
 	{
-		DropdownName += TEXT(" ") + Item.DeviceProfileName.ToString();
+		DropdownName += TEXT(" ") + DeviceProfileName.ToString();
 	}
 
 	return FName(DropdownName);
@@ -36,7 +36,7 @@ TArray<FName> UEditorSubsystemBlueprintLibrary::GetPreviewPlatformOptions() cons
 	TArray<FName> Outputs;
 	for (const FPreviewPlatformMenuItem& Item : FDataDrivenPlatformInfoRegistry::GetAllPreviewPlatformMenuItems())
 	{
-		FName PreviewNameAndDeviceProfileName = GetPreviewPlatformDropdownName(Item);
+		FName PreviewNameAndDeviceProfileName = GetPreviewPlatformDropdownName(Item.PreviewShaderPlatformName, Item.DeviceProfileName);
 		Outputs.Add(PreviewNameAndDeviceProfileName);
 	}
 
@@ -60,7 +60,7 @@ void UEditorSubsystemBlueprintLibrary::SetPreviewPlatform(FName PreviewShaderPla
 	{
 		// We check to see if we can find the matching PreviewShaderPlatform DeviceProfile combination in the Preview Platform Menu.
 
-		FName NameToCheck = GetPreviewPlatformDropdownName(Item);
+		FName NameToCheck = GetPreviewPlatformDropdownName(Item.PreviewShaderPlatformName, Item.DeviceProfileName);
 		if (NameToCheck == PreviewShaderPlatformName)
 		{
 			bFoundMatch = true;
@@ -95,6 +95,20 @@ void UEditorSubsystemBlueprintLibrary::SetPreviewPlatform(FName PreviewShaderPla
 	}
 }
 
+FName UEditorSubsystemBlueprintLibrary::GetPreviewPlatformName()
+{
+	FName PreviewPlatformName = NAME_None;
+	if (GEditor)
+	{
+		FPreviewPlatformInfo PreviewPlatform = GEditor->PreviewPlatform;
+		if (PreviewPlatform.bPreviewFeatureLevelActive)
+		{
+			PreviewPlatformName = GetPreviewPlatformDropdownName(PreviewPlatform.PreviewShaderPlatformName, PreviewPlatform.DeviceProfileName);
+		}
+	}
+	return PreviewPlatformName;
+}
+
 void UEditorSubsystemBlueprintLibrary::DisablePreviewPlatform()
 {
 	FPreviewPlatformInfo DisablePreviewPlatformInfo(GMaxRHIFeatureLevel, GMaxRHIShaderPlatform, NAME_None, NAME_None, NAME_None, false, NAME_None);
diff --git a/Engine/Source/Editor/UnrealEd/Public/Subsystems/EditorSubsystemBlueprintLibrary.h b/Engine/Source/Editor/UnrealEd/Public/Subsystems/EditorSubsystemBlueprintLibrary.h
--- a/Engine/Source/Editor/UnrealEd/Public/Subsystems/EditorSubsystemBlueprintLibrary.h
+++ b/Engine/Source/Editor/UnrealEd/Public/Subsystems/EditorSubsystemBlueprintLibrary.h
@@ -34,6 +34,12 @@ public:
 	UFUNCTION(BlueprintCallable, Category = "Editor")
 	static UNREALED_API void SetPreviewPlatform(UPARAM(meta = (GetOptions = "GetPreviewPlatformOptions")) FName PreviewShaderPlatformName);
 
+	/**
+	 * Gets the Preview Platform Name
+	 */
+	UFUNCTION(BlueprintCallable, Category = "Editor")
+	static UNREALED_API FName GetPreviewPlatformName();
+
 	/**
 	 * Disable the Preview Platform.
 	 */
