Apply the static mesh usage for preview material at the UObject level and not in the FMaterial.
This allows the auto set usage flag to behave correctly since that acts at the UObject level.
#rnx
#Synced-CL 49065038

diff --git a/Engine/Source/Editor/UnrealEd/Classes/MaterialEditor/PreviewMaterial.h b/Engine/Source/Editor/UnrealEd/Classes/MaterialEditor/PreviewMaterial.h
--- a/Engine/Source/Editor/UnrealEd/Classes/MaterialEditor/PreviewMaterial.h
+++ b/Engine/Source/Editor/UnrealEd/Classes/MaterialEditor/PreviewMaterial.h
@@ -13,10 +13,11 @@ class UPreviewMaterial : public UMaterial
 {
 	GENERATED_UCLASS_BODY()
 
-
 	//~ Begin UMaterial Interface.
 	UNREALED_API virtual FMaterialResource* AllocateResource() override;
 	virtual bool IsAsset()  const override  { return false; }
+	virtual uint32 GetUsageFlags() const override;
+	virtual bool GetUsageByFlag(EMaterialUsage Usage) const override;
 	//~ End UMaterial Interface.
 };
 
diff --git a/Engine/Source/Editor/UnrealEd/Private/PreviewMaterial.cpp b/Engine/Source/Editor/UnrealEd/Private/PreviewMaterial.cpp
--- a/Engine/Source/Editor/UnrealEd/Private/PreviewMaterial.cpp
+++ b/Engine/Source/Editor/UnrealEd/Private/PreviewMaterial.cpp
@@ -51,9 +51,6 @@ public:
 	{
 	}
 
-	/** Always allow to be used with static mesh so that we can use in preview view. */
-	virtual bool IsUsedWithStaticMesh() const override { return true; }
-
 	/**
 	 * Should the shader for this material with the given platform, shader type and vertex 
 	 * factory type combination be compiled
@@ -259,6 +256,18 @@ FMaterialResource* UPreviewMaterial::AllocateResource()
 	return new FPreviewMaterial();
 }
 
+uint32 UPreviewMaterial::GetUsageFlags() const
+{
+	// Allows enable static mesh usage in preview.
+	return Super::GetUsageFlags() | (1u << MATUSAGE_StaticMesh);
+}
+
+bool UPreviewMaterial::GetUsageByFlag(EMaterialUsage Usage) const
+{
+	// Allows enable static mesh usage in preview.
+	return (Usage == MATUSAGE_StaticMesh) || Super::GetUsageByFlag(Usage);
+}
+
 void UMaterialEditorPreviewParameters::PostEditChangeProperty(FPropertyChangedEvent& PropertyChangedEvent)
 {
 	if (PreviewMaterial && PropertyChangedEvent.ChangeType != EPropertyChangeType::Interactive)
