Fix usage flag override logic to apply for full usage flag mask, instead of casting to bool.
#rnx
#Synced-CL 49065038

diff --git a/Engine/Source/Editor/UnrealEd/Private/PreviewMaterial.cpp b/Engine/Source/Editor/UnrealEd/Private/PreviewMaterial.cpp
--- a/Engine/Source/Editor/UnrealEd/Private/PreviewMaterial.cpp
+++ b/Engine/Source/Editor/UnrealEd/Private/PreviewMaterial.cpp
@@ -1052,8 +1052,8 @@ void UMaterialEditorInstanceConstant::SetSourceInstance(UMaterialInstanceConstan
 
 	PreviewInstance = NewObject<UMaterialInstanceConstant>();
 	PreviewInstance->SetParentEditorOnly(MaterialInterface);
-	PreviewInstance->BasePropertyOverrides.bOverride_UsageFlags = (1 << MATUSAGE_StaticMesh);
-	PreviewInstance->BasePropertyOverrides.UsageFlags = (1 << MATUSAGE_StaticMesh);
+	PreviewInstance->BasePropertyOverrides.bOverride_UsageFlags = (1u << MATUSAGE_StaticMesh);
+	PreviewInstance->BasePropertyOverrides.UsageFlags = (1u << MATUSAGE_StaticMesh);
 
 	Parent = SourceInstance->Parent;
 	PhysMaterial = SourceInstance->PhysMaterial;
@@ -1184,10 +1184,8 @@ void UMaterialEditorInstanceConstant::CopyBasePropertiesFromParent()
 	{
 		BasePropertyOverrides.bCompatibleWithLumenCardSharing = SourceInstance->IsCompatibleWithLumenCardSharing();
 	}
-	if (!BasePropertyOverrides.bOverride_UsageFlags)
-	{
-		BasePropertyOverrides.UsageFlags = SourceInstance->GetUsageFlags();
-	}
+	// Usage flags and overrides are a bitmask.	
+	BasePropertyOverrides.UsageFlags = (BasePropertyOverrides.bOverride_UsageFlags & BasePropertyOverrides.UsageFlags) | (~BasePropertyOverrides.bOverride_UsageFlags & SourceInstance->GetUsageFlags());
 
 	// Copy the Lightmass settings...
 	// The lightmass functions (GetCastShadowAsMasked, etc.) check if the value is overridden and returns the current value if so, otherwise returns the parent value
