[ITF Viewport Interactions] Fix a bug where Shift-Translating objects in unlit orthographic viewports caused extreme drag position errors.

#rnx

#Synced-CL 49033042

diff --git a/Engine/Source/Editor/UnrealEd/Private/LevelEditorViewport.cpp b/Engine/Source/Editor/UnrealEd/Private/LevelEditorViewport.cpp
--- a/Engine/Source/Editor/UnrealEd/Private/LevelEditorViewport.cpp
+++ b/Engine/Source/Editor/UnrealEd/Private/LevelEditorViewport.cpp
@@ -3380,6 +3380,9 @@ bool FLevelEditorViewportClient::InputWidgetDelta(FViewport* InViewport, EAxisLi
 											FSceneViewFamily::ConstructionValues(Viewport, GetScene(), EngineShowFlags)
 												.SetRealtimeUpdate(IsRealtime())
 										);
+										// ViewMode defaults to VMI_Lit. Orthographic views can end up with different view matrices depending
+										// on the view mode, so the view mode must be consistent. 
+										ViewFamily.ViewMode = GetViewMode();
  
 										FSceneView* View = CalcSceneView(&ViewFamily);
 										GizmoViewContext->ResetFromSceneView(*View);
