UE: Bump incremental cook version. Tag change

diff --git a/Engine/Source/Editor/UnrealEd/Private/Cooker/CookIncrementalVersion.cpp b/Engine/Source/Editor/UnrealEd/Private/Cooker/CookIncrementalVersion.cpp
--- a/Engine/Source/Editor/UnrealEd/Private/Cooker/CookIncrementalVersion.cpp
+++ b/Engine/Source/Editor/UnrealEd/Private/Cooker/CookIncrementalVersion.cpp
@@ -8,6 +8,6 @@ namespace UE::Cook
 {
 
 // Change CookIncrementalVersion to a new guid when all packages in an incremental cook need to be invalidated
-FGuid CookIncrementalVersion(0x890DC5D0, 0x4FC046B9, 0xA6F1D755, 0x09E0E84D);
+FGuid CookIncrementalVersion(0x74F40C2C, 0x03FD4D06, 0x9FD3B32C, 0x46AE93F7);
 
 }
