Added RF_ReInst, this the set of flags we transfer from old to new instance during reinstantiation, this avoids the various redundant declarations
Added RF_HasDynamicImport to flags we need to carry along during reinstantiation

#Synced-CL 49106026

diff --git a/Engine/Source/Editor/Kismet/Private/BlueprintCompilationManager.cpp b/Engine/Source/Editor/Kismet/Private/BlueprintCompilationManager.cpp
--- a/Engine/Source/Editor/Kismet/Private/BlueprintCompilationManager.cpp
+++ b/Engine/Source/Editor/Kismet/Private/BlueprintCompilationManager.cpp
@@ -2885,6 +2885,7 @@ void FBlueprintCompilationManagerImpl::ReinstanceBatch(TArray<FReinstancingJob>&
 					NewClass->SetDefaultObject(nullptr);
 				}
 				UObject* NewCDO = NewClass->GetDefaultObject(true);
+				NewCDO->SetFlags(OldCDO->GetFlags() & RF_ReInst);
 
 				UBlueprint* CompiledBlueprint = UBlueprint::GetBlueprintFromClass(OldClass);
 				if (CompiledBlueprint && CompiledBlueprint->bIsRegeneratingOnLoad)
@@ -3227,9 +3228,8 @@ void FBlueprintCompilationManagerImpl::ReinstanceBatch(TArray<FReinstancingJob>&
 
 				// reconstruct
 				FMakeClassSpawnableOnScope TemporarilySpawnable(NewClass);
-				const EObjectFlags FlagMask = RF_Public | RF_ArchetypeObject | RF_Transactional | RF_Transient | RF_TextExportTransient | RF_InheritableComponentTemplate | RF_Standalone; //TODO: what about RF_RootSet?
 
-				UObject* NewArchetype = NewObject<UObject>(OriginalOuter, NewClass, OriginalName, OriginalFlags & FlagMask, Template, false, bUsePrePopulatedInstancingGraph ? &PrePopulatedInstancingGraph : nullptr);
+				UObject* NewArchetype = NewObject<UObject>(OriginalOuter, NewClass, OriginalName, OriginalFlags & RF_ReInst, Template, false, bUsePrePopulatedInstancingGraph ? &PrePopulatedInstancingGraph : nullptr);
 
 				OldArchetypeToNewArchetype.Add(OldInstance, NewArchetype);
 				OldToNewTemplatesForClass->Add(OldInstance, NewArchetype);
diff --git a/Engine/Source/Editor/UnrealEd/Private/Kismet2/KismetReinstanceUtilities.cpp b/Engine/Source/Editor/UnrealEd/Private/Kismet2/KismetReinstanceUtilities.cpp
--- a/Engine/Source/Editor/UnrealEd/Private/Kismet2/KismetReinstanceUtilities.cpp
+++ b/Engine/Source/Editor/UnrealEd/Private/Kismet2/KismetReinstanceUtilities.cpp
@@ -77,11 +77,6 @@ static FAutoConsoleVariableRef CVarUseLegacyRootComponentPreservation(
 	TEXT("Use the legacy re-instancing behavior to preserve root components, ignoring data in previous instances")
 );
 
-namespace UE::ReinstanceUtils
-{
-	const EObjectFlags FlagMask = RF_Public | RF_ArchetypeObject | RF_Transactional | RF_Transient | RF_TextExportTransient | RF_InheritableComponentTemplate | RF_Standalone; //TODO: what about RF_RootSet?
-}
-
 namespace UE::Private
 {
 
@@ -2195,7 +2190,7 @@ void FBlueprintCompileReinstancer::ReplaceObjectHelper(UObject*& OldObject, UCla
 			// Otherwise, just use the old object's current name.
 			NewName = OldObject->GetFName();
 		}
-		NewArchetype = UObject::GetArchetypeFromRequiredInfo(NewClass, OldObject->GetOuter(), NewName, OldObject->GetFlags() & UE::ReinstanceUtils::FlagMask);
+		NewArchetype = UObject::GetArchetypeFromRequiredInfo(NewClass, OldObject->GetOuter(), NewName, OldObject->GetFlags() & RF_ReInst);
 	}
 	else
 	{
@@ -2305,7 +2300,7 @@ void FBlueprintCompileReinstancer::ReplaceObjectHelper(UObject*& OldObject, UCla
 
 	check(NewUObject != nullptr);
 
-	NewUObject->SetFlags(OldFlags & UE::ReinstanceUtils::FlagMask);
+	NewUObject->SetFlags(OldFlags & RF_ReInst);
 
 	TMap<UObject*, UObject*> CreatedInstanceMap;
 	TArray< TTuple<UObject*, UObject*>> OrderedListOfObjectToCopy;
@@ -2452,7 +2447,7 @@ static void ReplaceActorHelper(AActor* OldActor, UClass* OldClass, UObject*& New
 	SpawnInfo.SpawnCollisionHandlingOverride = ESpawnActorCollisionHandlingMethod::AlwaysSpawn;
 	SpawnInfo.bDeferConstruction = true;
 	SpawnInfo.Name = OldActor->GetFName();
-	SpawnInfo.ObjectFlags |= OldActor->GetFlags() & UE::ReinstanceUtils::FlagMask;
+	SpawnInfo.ObjectFlags |= OldActor->GetFlags() & RF_ReInst;
 
 	if (!OldActor->IsListedInSceneOutliner())
 	{
@@ -3412,7 +3407,7 @@ void FBlueprintCompileReinstancer::PreCreateSubObjectsForReinstantiation_Inner(
 		UClass* OldSubObjectClass = OldSubObject->GetClass();
 		UObject* OldSubObjectOuter = OldSubObject->GetOuter();
 		FName SubObjectName = OldSubObject->GetFName();
-		EObjectFlags SubObjectFlags = OldSubObject->GetFlags() & UE::ReinstanceUtils::FlagMask;
+		EObjectFlags SubObjectFlags = OldSubObject->GetFlags() & RF_ReInst;
 
 		// Only re-create objects that are created from a default object,
 		// skip any default subobjects are they will be created by their parent object CDO.
diff --git a/Engine/Source/Runtime/CoreUObject/Public/UObject/ObjectMacros.h b/Engine/Source/Runtime/CoreUObject/Public/UObject/ObjectMacros.h
--- a/Engine/Source/Runtime/CoreUObject/Public/UObject/ObjectMacros.h
+++ b/Engine/Source/Runtime/CoreUObject/Public/UObject/ObjectMacros.h
@@ -648,6 +648,9 @@ inline constexpr EObjectFlags RF_HasPlaceholderType = EObjectFlags(0x20000000);
 /** Flags to load from unreal asset files */
 #define RF_Load						((EObjectFlags)(RF_Public | RF_Standalone | RF_Transactional | RF_ClassDefaultObject | RF_ArchetypeObject | RF_DefaultSubObject | RF_TextExportTransient | RF_InheritableComponentTemplate | RF_DuplicateTransient | RF_NonPIEDuplicateTransient | RF_MigratingAsset | RF_HasDynamicImports)) 
 
+/** Flags to copy during reinstantiation of UObjects */
+#define RF_ReInst					((EObjectFlags)(RF_Public | RF_ArchetypeObject | RF_Transactional | RF_Transient | RF_TextExportTransient | RF_InheritableComponentTemplate | RF_Standalone | RF_HasDynamicImports)) 
+
 /** Sub-objects will inherit these flags from their SuperObject */
 #define RF_PropagateToSubObjects	((EObjectFlags)(RF_Public | RF_ArchetypeObject | RF_Transactional | RF_Transient))
 
