[Engine][SColorGradientEditor]

Fixing an export dll linker error that started happening for unknown reasons.

-- Code --
SColorGradientEditor:
Dtor() - Replaced the export with inline.

#Synced-CL 49061972

diff --git a/Engine/Source/Editor/UnrealEd/Public/SColorGradientEditor.h b/Engine/Source/Editor/UnrealEd/Public/SColorGradientEditor.h
--- a/Engine/Source/Editor/UnrealEd/Public/SColorGradientEditor.h
+++ b/Engine/Source/Editor/UnrealEd/Public/SColorGradientEditor.h
@@ -63,7 +63,7 @@ public:
 	SLATE_END_ARGS()
 
 	UNREALED_API SColorGradientEditor();
-	UNREALED_API virtual ~SColorGradientEditor() = default;
+	inline virtual ~SColorGradientEditor() = default;
 
 	UNREALED_API void Construct( const FArguments& InArgs );
 
