Remove RF_Reinst from ObjectMacros.h to avoid suggestion that it is stable and public

#rnx
#Synced-CL 49181000

diff --git a/Engine/Source/Editor/Kismet/Private/BlueprintCompilationManager.cpp b/Engine/Source/Editor/Kismet/Private/BlueprintCompilationManager.cpp
--- a/Engine/Source/Editor/Kismet/Private/BlueprintCompilationManager.cpp
+++ b/Engine/Source/Editor/Kismet/Private/BlueprintCompilationManager.cpp
@@ -51,6 +51,7 @@
 #include "BlueprintEditorModule.h"
 #include "Algo/TopologicalSort.h"
 #include "Animation/AnimBlueprint.h"
+#include "Reinstancing/ReinstancingInternal.h"
 #include "Stats/StatsHierarchical.h"
 #include "UObject/PropertyBagRepository.h"
 #include "UObject/OverridableManager.h"
@@ -2885,7 +2886,7 @@ void FBlueprintCompilationManagerImpl::ReinstanceBatch(TArray<FReinstancingJob>&
 					NewClass->SetDefaultObject(nullptr);
 				}
 				UObject* NewCDO = NewClass->GetDefaultObject(true);
-				NewCDO->SetFlags(OldCDO->GetFlags() & RF_ReInst);
+				NewCDO->SetFlags(OldCDO->GetFlags() & RF_Reinst);
 
 				UBlueprint* CompiledBlueprint = UBlueprint::GetBlueprintFromClass(OldClass);
 				if (CompiledBlueprint && CompiledBlueprint->bIsRegeneratingOnLoad)
@@ -3229,7 +3230,7 @@ void FBlueprintCompilationManagerImpl::ReinstanceBatch(TArray<FReinstancingJob>&
 				// reconstruct
 				FMakeClassSpawnableOnScope TemporarilySpawnable(NewClass);
 
-				UObject* NewArchetype = NewObject<UObject>(OriginalOuter, NewClass, OriginalName, OriginalFlags & RF_ReInst, Template, false, bUsePrePopulatedInstancingGraph ? &PrePopulatedInstancingGraph : nullptr);
+				UObject* NewArchetype = NewObject<UObject>(OriginalOuter, NewClass, OriginalName, OriginalFlags & RF_Reinst, Template, false, bUsePrePopulatedInstancingGraph ? &PrePopulatedInstancingGraph : nullptr);
 
 				OldArchetypeToNewArchetype.Add(OldInstance, NewArchetype);
 				OldToNewTemplatesForClass->Add(OldInstance, NewArchetype);
diff --git a/Engine/Source/Editor/UnrealEd/Internal/Reinstancing/ReinstancingInternal.h b/Engine/Source/Editor/UnrealEd/Internal/Reinstancing/ReinstancingInternal.h
new file mode 100644
--- /dev/null
+++ b/Engine/Source/Editor/UnrealEd/Internal/Reinstancing/ReinstancingInternal.h
@@ -0,0 +1,7 @@
+// Copyright Epic Games, Inc. All Rights Reserved.
+
+#pragma once
+
+/** Flags to copy during reinstantiation of UObjects */
+#define RF_Reinst					((EObjectFlags)(RF_Public | RF_ArchetypeObject | RF_Transactional | \
+	RF_Transient | RF_TextExportTransient | RF_InheritableComponentTemplate | RF_Standalone | RF_HasDynamicImports)) 
diff --git a/Engine/Source/Editor/UnrealEd/Private/Kismet2/KismetReinstanceUtilities.cpp b/Engine/Source/Editor/UnrealEd/Private/Kismet2/KismetReinstanceUtilities.cpp
--- a/Engine/Source/Editor/UnrealEd/Private/Kismet2/KismetReinstanceUtilities.cpp
+++ b/Engine/Source/Editor/UnrealEd/Private/Kismet2/KismetReinstanceUtilities.cpp
@@ -53,6 +53,7 @@
 #include "UObject/PropertyBagRepository.h"
 #include "UObject/UObjectArchetypeHelper.h"
 #include "ProfilingDebugging/LoadTimeTracker.h"
+#include "Reinstancing/ReinstancingInternal.h"
 
 DECLARE_CYCLE_STAT(TEXT("Replace Instances"), EKismetReinstancerStats_ReplaceInstancesOfClass, STATGROUP_KismetReinstancer );
 DECLARE_CYCLE_STAT(TEXT("Find Referencers"), EKismetReinstancerStats_FindReferencers, STATGROUP_KismetReinstancer );
@@ -2190,7 +2191,7 @@ void FBlueprintCompileReinstancer::ReplaceObjectHelper(UObject*& OldObject, UCla
 			// Otherwise, just use the old object's current name.
 			NewName = OldObject->GetFName();
 		}
-		NewArchetype = UObject::GetArchetypeFromRequiredInfo(NewClass, OldObject->GetOuter(), NewName, OldObject->GetFlags() & RF_ReInst);
+		NewArchetype = UObject::GetArchetypeFromRequiredInfo(NewClass, OldObject->GetOuter(), NewName, OldObject->GetFlags() & RF_Reinst);
 	}
 	else
 	{
@@ -2300,7 +2301,7 @@ void FBlueprintCompileReinstancer::ReplaceObjectHelper(UObject*& OldObject, UCla
 
 	check(NewUObject != nullptr);
 
-	NewUObject->SetFlags(OldFlags & RF_ReInst);
+	NewUObject->SetFlags(OldFlags & RF_Reinst);
 
 	TMap<UObject*, UObject*> CreatedInstanceMap;
 	TArray< TTuple<UObject*, UObject*>> OrderedListOfObjectToCopy;
@@ -2447,7 +2448,7 @@ static void ReplaceActorHelper(AActor* OldActor, UClass* OldClass, UObject*& New
 	SpawnInfo.SpawnCollisionHandlingOverride = ESpawnActorCollisionHandlingMethod::AlwaysSpawn;
 	SpawnInfo.bDeferConstruction = true;
 	SpawnInfo.Name = OldActor->GetFName();
-	SpawnInfo.ObjectFlags |= OldActor->GetFlags() & RF_ReInst;
+	SpawnInfo.ObjectFlags |= OldActor->GetFlags() & RF_Reinst;
 
 	if (!OldActor->IsListedInSceneOutliner())
 	{
@@ -3407,7 +3408,7 @@ void FBlueprintCompileReinstancer::PreCreateSubObjectsForReinstantiation_Inner(
 		UClass* OldSubObjectClass = OldSubObject->GetClass();
 		UObject* OldSubObjectOuter = OldSubObject->GetOuter();
 		FName SubObjectName = OldSubObject->GetFName();
-		EObjectFlags SubObjectFlags = OldSubObject->GetFlags() & RF_ReInst;
+		EObjectFlags SubObjectFlags = OldSubObject->GetFlags() & RF_Reinst;
 
 		// Only re-create objects that are created from a default object,
 		// skip any default subobjects are they will be created by their parent object CDO.
diff --git a/Engine/Source/Runtime/CoreUObject/Public/UObject/ObjectMacros.h b/Engine/Source/Runtime/CoreUObject/Public/UObject/ObjectMacros.h
--- a/Engine/Source/Runtime/CoreUObject/Public/UObject/ObjectMacros.h
+++ b/Engine/Source/Runtime/CoreUObject/Public/UObject/ObjectMacros.h
@@ -648,9 +648,6 @@ inline constexpr EObjectFlags RF_HasPlaceholderType = EObjectFlags(0x20000000);
 /** Flags to load from unreal asset files */
 #define RF_Load						((EObjectFlags)(RF_Public | RF_Standalone | RF_Transactional | RF_ClassDefaultObject | RF_ArchetypeObject | RF_DefaultSubObject | RF_TextExportTransient | RF_InheritableComponentTemplate | RF_DuplicateTransient | RF_NonPIEDuplicateTransient | RF_MigratingAsset | RF_HasDynamicImports)) 
 
-/** Flags to copy during reinstantiation of UObjects */
-#define RF_ReInst					((EObjectFlags)(RF_Public | RF_ArchetypeObject | RF_Transactional | RF_Transient | RF_TextExportTransient | RF_InheritableComponentTemplate | RF_Standalone | RF_HasDynamicImports)) 
-
 /** Sub-objects will inherit these flags from their SuperObject */
 #define RF_PropagateToSubObjects	((EObjectFlags)(RF_Public | RF_ArchetypeObject | RF_Transactional | RF_Transient))
 
