ExternalActorsCommandlet:
- Make sure to initialize the world before saving external actors, instead of doing it just when saving the map package.
- This makes sure that actors are initialized when creating their descriptors and ensure we have the good streaming bounds.
- Added the option to enable OFPA on maps without OFPA, before this change we had to do this manually in the editor.
- This allows taking advantage of the commandlet map list batching, so we can automate the process if required.
- When disabling OFPA, make sure to also delete invalid actors, if any, to avoid having lingering actors on disk.
- Fix support for working without source control: don't try to checkout or mark files as writable if they don't exist on disk.

#Synced-CL 49104972

diff --git a/Engine/Source/Editor/UnrealEd/Classes/Commandlets/ExternalActorsCommandlet.h b/Engine/Source/Editor/UnrealEd/Classes/Commandlets/ExternalActorsCommandlet.h
--- a/Engine/Source/Editor/UnrealEd/Classes/Commandlets/ExternalActorsCommandlet.h
+++ b/Engine/Source/Editor/UnrealEd/Classes/Commandlets/ExternalActorsCommandlet.h
@@ -27,6 +27,9 @@ protected:
 	UPROPERTY()
 	bool bDisable;
 
+	UPROPERTY()
+	bool bEnable;
+
 	UPROPERTY()
 	bool bForce;
 
diff --git a/Engine/Source/Editor/UnrealEd/Private/Commandlets/ExternalActorsCommandlet.cpp b/Engine/Source/Editor/UnrealEd/Private/Commandlets/ExternalActorsCommandlet.cpp
--- a/Engine/Source/Editor/UnrealEd/Private/Commandlets/ExternalActorsCommandlet.cpp
+++ b/Engine/Source/Editor/UnrealEd/Private/Commandlets/ExternalActorsCommandlet.cpp
@@ -4,6 +4,7 @@
 #include "AssetRegistry/AssetData.h"
 #include "AssetRegistry/AssetRegistryModule.h"
 #include "HAL/FileManager.h"
+#include "HAL/PlatformFileManager.h"
 #include "Misc/FileHelper.h"
 #include "Misc/Paths.h"
 #include "Editor.h"
@@ -57,6 +58,12 @@ int32 UExternalActorsCommandlet::Main(const FString& Params)
 			return 1;
 		}
 
+		if (bEnable)
+		{
+			UE_LOG(LogExternalActorsCommandlet, Error, TEXT("Bad parameters: -enable cannot be used with -disable"));
+			return 1;
+		}
+
 		bListMaps |= !DumpCSVFile.IsEmpty();
 	}
 	else if (bForce || bReport || bListMaps || DumpCSVFile.Len())
@@ -194,7 +201,7 @@ int32 UExternalActorsCommandlet::Main(const FString& Params)
 		{
 			if (!MainWorld->PersistentLevel->IsUsingExternalActors())
 			{
-				UE_LOG(LogExternalActorsCommandlet, Error, TEXT("Cannot disable external actors for partitioned level '%s' (already disabled)"), *FullMapName);
+				UE_LOG(LogExternalActorsCommandlet, Error, TEXT("Cannot disable external actors for level '%s' (already disabled)"), *FullMapName);
 				Result = 0;
 				continue;
 			}
@@ -241,14 +248,25 @@ int32 UExternalActorsCommandlet::Main(const FString& Params)
 				{
 					WorldPartition->Initialize(MainWorld, FTransform::Identity);
 
+					TArray<FString> LocalPackagesToDelete;
+					WorldPartition->ForEachActorDescContainerInstance([&LocalPackagesToDelete](UActorDescContainerInstance* ActorDescContainerInstance)
+					{
+						for (const FAssetData& InvalidActor : ActorDescContainerInstance->GetContainer()->GetInvalidActors())
+						{
+							LocalPackagesToDelete.Add(InvalidActor.ToSoftObjectPath().ToString());
+						}
+					});
+
+					// This will uninitialize the world partition object 
 					if (!UWorldPartition::RemoveWorldPartition(WorldPartition->GetWorld()->GetWorldSettings()))
 					{
+						WorldPartition->Uninitialize();
 						UE_LOG(LogExternalActorsCommandlet, Error, TEXT("Error disabling external actors for partitioned level '%s'"), *FullMapName);
 						Result = 0;
 						continue;
 					}
 
-					WorldPartition->Uninitialize();
+					PackagesToDelete.Append(LocalPackagesToDelete);
 				}
 			}
 
@@ -285,6 +303,34 @@ int32 UExternalActorsCommandlet::Main(const FString& Params)
 				MainWorld->PersistentLevel->SetUseExternalActors(false);
 			}
 		}
+		else if (bEnable)
+		{
+			if (MainWorld->PersistentLevel->IsUsingExternalActors())
+			{
+				UE_LOG(LogExternalActorsCommandlet, Error, TEXT("Cannot enable external actors for level '%s' (already enabled)"), *FullMapName);
+				Result = 0;
+				continue;
+			}
+
+			if (!bReport)
+			{
+				for (AActor* Actor : MainWorld->PersistentLevel->Actors)
+				{
+					if (IsValid(Actor) && Actor->SupportsExternalPackaging())
+					{
+						Actor->SetPackageExternal(true);
+						PackagesToSave.Add(Actor->GetPackage());
+					}
+				}
+
+				if (PackagesToSave.Num())
+				{
+					PackagesToSave.Add(MainWorld->GetPackage());
+				}
+
+				MainWorld->PersistentLevel->SetUseExternalActors(true);
+			}
+		}
 		else
 		{
 			// Validate external actors
@@ -389,18 +435,29 @@ int32 UExternalActorsCommandlet::Main(const FString& Params)
 		FSavePackageArgs SaveArgs;
 		SaveArgs.TopLevelFlags = RF_Standalone;
 
+		bool bForceInitializedWorld;
+		const bool bInitializedPhysicsSceneForSave = GEditor->InitializePhysicsSceneForSaveIfNecessary(MainWorld, bForceInitializedWorld);
+
 		for (UPackage* PackageToSave : PackagesToSave)
 		{
 			const FString PackageFileName = SourceControlHelpers::PackageFilename(PackageToSave);
 			UE_LOG(LogExternalActorsCommandlet, Display, TEXT("Saving package '%s'"), *PackageFileName);
 
-			if (PackageHelper.Checkout(PackageToSave))
-			{				
-				if (UPackage::SavePackage(PackageToSave, nullptr, *PackageFileName, SaveArgs))
-				{
-					PackageHelper.AddToSourceControl(PackageToSave);
-				}
+			bool bCanSavePackage = true;
+			if (FPlatformFileManager::Get().GetPlatformFile().FileExists(*PackageFileName))
+			{
+				bCanSavePackage = PackageHelper.Checkout(PackageToSave);
 			}
+
+			if (bCanSavePackage && UPackage::SavePackage(PackageToSave, nullptr, *PackageFileName, SaveArgs))
+			{
+				PackageHelper.AddToSourceControl(PackageToSave);
+			}
+		}
+
+		if (bInitializedPhysicsSceneForSave)
+		{
+			GEditor->CleanupPhysicsSceneThatWasInitializedForSave(MainWorld, bForceInitializedWorld);
 		}
 
 		CollectGarbage(RF_NoFlags);
