Skein / URC2: Fixed localization discrepancies.

#rnx
#Synced-CL 49147668

diff --git a/Engine/Source/Editor/UnrealEd/Private/FileHelpers.cpp b/Engine/Source/Editor/UnrealEd/Private/FileHelpers.cpp
--- a/Engine/Source/Editor/UnrealEd/Private/FileHelpers.cpp
+++ b/Engine/Source/Editor/UnrealEd/Private/FileHelpers.cpp
@@ -5193,10 +5193,10 @@ bool FEditorFileUtils::AutomaticCheckoutOrPromptToRevertPackages(const TArray<UP
 		// No, SourceControl is offline.
 		// Warn the user that they're working in offline mode.
 
-		FText OfflineTitle = NSLOCTEXT("FileHelper", "OfflineDialogTitle", "Warning - Offline, Conflicts may occur");
+		FText OfflineTitle = NSLOCTEXT("FileHelper", "OfflineDialogTitle", "Warning: Changes While Offline");
 		FText OfflineMessage = NSLOCTEXT("FileHelper", "OfflineDialogMessage",
-			"You've made changes while offline. These and any further changes made offline could conflict with your teammates' work when you reconnect.\r\n\r\n"
-			"Reconnect as soon as possible to minimize conflicts and continue making changes.");
+			"Changes you make while offline may conflict with your teammates' work when you reconnect.\r\n\r\n"
+			"To minimize conflicts, reconnect as soon as possible.");
 
 		FSuppressableWarningDialog::FSetupInfo Info(OfflineMessage, OfflineTitle, TEXT("ShowOfflineModeWarning"), GEditorPerProjectIni);
 		Info.ConfirmText = NSLOCTEXT("FileHelper", "OfflineDialog_ConfirmText", "Ok");
