Significantly optimize PreCreateSubObjectsForReinstantiation reinstancing by caching the archetype trees instead of recomputing them for every subobject.
Refactor some of the reinstance helpers to cut down on function params.

#Synced-CL 49126341

diff --git a/Engine/Source/Editor/Kismet/Private/BlueprintCompilationManager.cpp b/Engine/Source/Editor/Kismet/Private/BlueprintCompilationManager.cpp
--- a/Engine/Source/Editor/Kismet/Private/BlueprintCompilationManager.cpp
+++ b/Engine/Source/Editor/Kismet/Private/BlueprintCompilationManager.cpp
@@ -3160,6 +3160,8 @@ void FBlueprintCompilationManagerImpl::ReinstanceBatch(TArray<FReinstancingJob>&
 	{
 		SCOPED_LOADTIMER_ASSET_TEXT(TEXT("FinishReinstancing Archetypes"));
 
+		TMap<UObject*, FObjectInstancingGraph> GraphCache;
+
 		for (const FArchetypeInfo& OldArchetypeInfo : OldArchetypeObjects)
 		{
 			checkf(OldArchetypeInfo.ReinstancingJob, TEXT("Expecting the job ptr to be set."));
@@ -3197,7 +3199,7 @@ void FBlueprintCompilationManagerImpl::ReinstanceBatch(TArray<FReinstancingJob>&
 
 				TArray<UClass*> ClassesOfArchetypeOuterChain;
 				FObjectInstancingGraph PrePopulatedInstancingGraph;
-				const bool bUsePrePopulatedInstancingGraph = FBlueprintCompileReinstancer::PrePopulateInstancingGraphForArchetype(PrePopulatedInstancingGraph, OldInstance, Template, &OldArchetypeToNewArchetype);
+				const bool bUsePrePopulatedInstancingGraph = FBlueprintCompileReinstancer::PrePopulateInstancingGraphForArchetype(GraphCache, PrePopulatedInstancingGraph, OldInstance, Template, &OldArchetypeToNewArchetype);
 				if (bUsePrePopulatedInstancingGraph)
 				{
 					// Make sure the old archetype has all it references patched up before using copying
diff --git a/Engine/Source/Editor/UnrealEd/Private/Kismet2/KismetReinstanceUtilities.cpp b/Engine/Source/Editor/UnrealEd/Private/Kismet2/KismetReinstanceUtilities.cpp
--- a/Engine/Source/Editor/UnrealEd/Private/Kismet2/KismetReinstanceUtilities.cpp
+++ b/Engine/Source/Editor/UnrealEd/Private/Kismet2/KismetReinstanceUtilities.cpp
@@ -2178,14 +2178,14 @@ bool FBlueprintCompileReinstancer::IsReinstClass(const UClass* Class)
 	return Class && Class->GetFName().ToString().StartsWith(ReinstPrefix);
 }
 
-void FBlueprintCompileReinstancer::ReplaceObjectHelper(UObject*& OldObject, UClass* OldClass, UObject*& NewUObject, UClass* NewClass, TMap<UObject*, UObject*>& OldToNewInstanceMap, const TMap<UClass*, UClass*>& OldToNewClassMap, TMap<UObject*, FName>& OldToNewNameMap, int32 OldObjIndex, TArray<UObject*>& ObjectsToReplace, TArray<UObject*>& PotentialEditorsForRefreshing, TSet<AActor*>& OwnersToRerunConstructionScript, TFunctionRef<TArray<TObjectPtr<USceneComponent>>&(USceneComponent*)> GetAttachChildrenArray, bool bIsComponent, bool bArchetypesAreUpToDate)
+void FBlueprintCompileReinstancer::ReplaceObjectHelper(const FReinstancingState& State, UObject*& OldObject, UClass* OldClass, UObject*& NewUObject, UClass* NewClass, int32 OldObjIndex, TArray<UObject*>& ObjectsToReplace, TFunctionRef<TArray<TObjectPtr<USceneComponent>>& (USceneComponent*)> GetAttachChildrenArray, bool bIsComponent, bool bArchetypesAreUpToDate)
 {
 	SCOPED_LOADTIMER_ASSET_TEXT(*WriteToString<256>(TEXT("ReplaceObjectHelper "), *GetPathNameSafe(OldObject)));
 	// If the old object was spawned from an archetype (i.e. not the CDO), we must use the new version of that archetype as the template object when constructing the new instance.
 	UObject* NewArchetype = nullptr;
 	if(bArchetypesAreUpToDate)
 	{
-		FName NewName = OldToNewNameMap.FindRef(OldObject);
+		FName NewName = State.OldToNewNameMap.FindRef(OldObject);
 		if (NewName == NAME_None)
 		{
 			// Otherwise, just use the old object's current name.
@@ -2196,7 +2196,7 @@ void FBlueprintCompileReinstancer::ReplaceObjectHelper(UObject*& OldObject, UCla
 	else
 	{
 		UObject* OldArchetype = OldObject->GetArchetype();
-		NewArchetype = OldToNewInstanceMap.FindRef(OldArchetype);
+		NewArchetype = State.OldToNewInstanceMap.FindRef(OldArchetype);
 
 		bool bArchetypeReinstanced = (OldArchetype == OldClass->GetDefaultObject()) || (NewArchetype != nullptr);
 		// if we don't have a updated archetype to spawn from, we need to update/reinstance it
@@ -2215,7 +2215,7 @@ void FBlueprintCompileReinstancer::ReplaceObjectHelper(UObject*& OldObject, UCla
 					check(OldObject == OldArchetype);
 
 					OldArchetype = OldObject->GetArchetype();
-					NewArchetype = OldToNewInstanceMap.FindRef(OldArchetype);
+					NewArchetype = State.OldToNewInstanceMap.FindRef(OldArchetype);
 					bArchetypeReinstanced = (OldArchetype == OldClass->GetDefaultObject()) || (NewArchetype != nullptr);
 				}
 				else
@@ -2238,7 +2238,7 @@ void FBlueprintCompileReinstancer::ReplaceObjectHelper(UObject*& OldObject, UCla
 	FName OldName(OldObject->GetFName());
 
 	// If the old object is in this table, we've already renamed it away in a previous iteration. Don't rename it again!
-	if (!OldToNewNameMap.Contains(OldObject))
+	if (!State.OldToNewNameMap.Contains(OldObject))
 	{
 		// If we're reinstancing a component template, we also need to rename any inherited templates that are found to be based on it, in order to preserve archetype paths.
 		if (bIsComponent && OldObject->HasAllFlags(RF_ArchetypeObject) && OldObject->GetOuter()->IsA<UBlueprintGeneratedClass>())
@@ -2252,7 +2252,7 @@ void FBlueprintCompileReinstancer::ReplaceObjectHelper(UObject*& OldObject, UCla
 
 			for (UObject* OldArchetypeObject : OldArchetypeObjects)
 			{
-				OldToNewNameMap.Add(OldArchetypeObject, OldName);
+				State.OldToNewNameMap.Add(OldArchetypeObject, OldName);
 				OldArchetypeObject->Rename(*OldArchetypeName, OldArchetypeObject->GetOuter(), REN_DoNotDirty | REN_DontCreateRedirectors);
 			}
 		}
@@ -2264,7 +2264,7 @@ void FBlueprintCompileReinstancer::ReplaceObjectHelper(UObject*& OldObject, UCla
 						
 	{
 		// We may have already renamed this object to temp space if it was an inherited archetype in a previous iteration; check for that here.
-		FName NewName = OldToNewNameMap.FindRef(OldObject);
+		FName NewName = State.OldToNewNameMap.FindRef(OldObject);
 		if (NewName == NAME_None)
 		{
 			// Otherwise, just use the old object's current name.
@@ -2273,7 +2273,7 @@ void FBlueprintCompileReinstancer::ReplaceObjectHelper(UObject*& OldObject, UCla
 
 		UObject* DestinationOuter = OldObject->GetOuter();
 		// Check to make sure our original outer hasn't already been reinstanced:
-		if (UObject* const* ReinstancedOuter = OldToNewInstanceMap.Find(DestinationOuter))
+		if (UObject* const* ReinstancedOuter = State.OldToNewInstanceMap.Find(DestinationOuter))
 		{
 			// Our outer has been replaced, use the newer object:
 			DestinationOuter = *ReinstancedOuter;
@@ -2293,7 +2293,7 @@ void FBlueprintCompileReinstancer::ReplaceObjectHelper(UObject*& OldObject, UCla
 		}
 
 		FObjectInstancingGraph PrePopulatedInstancingGraph;
-		const bool bUsePrePopulatedInstancingGraph = PrePopulateInstancingGraphForArchetype(PrePopulatedInstancingGraph, OldObject, NewArchetype, &OldToNewInstanceMap);
+		const bool bUsePrePopulatedInstancingGraph = PrePopulateInstancingGraphForArchetype(State.GraphCache, PrePopulatedInstancingGraph, OldObject, NewArchetype, &State.OldToNewInstanceMap);
 		
 		FMakeClassSpawnableOnScope TemporarilySpawnable(NewClass);
 		NewUObject = NewObject<UObject>(DestinationOuter, NewClass, NewName, RF_NoFlags, NewArchetype, false, bUsePrePopulatedInstancingGraph ? &PrePopulatedInstancingGraph : nullptr, OldExternalPackage);
@@ -2305,19 +2305,19 @@ void FBlueprintCompileReinstancer::ReplaceObjectHelper(UObject*& OldObject, UCla
 
 	TMap<UObject*, UObject*> CreatedInstanceMap;
 	TArray< TTuple<UObject*, UObject*>> OrderedListOfObjectToCopy;
-	PreCreateSubObjectsForReinstantiation(OldToNewClassMap, OldObject, NewUObject, CreatedInstanceMap, &OldToNewInstanceMap, &OrderedListOfObjectToCopy);
-	OldToNewInstanceMap.Append(CreatedInstanceMap);
+	PreCreateSubObjectsForReinstantiation(State.InOldToNewClassMap, OldObject, NewUObject, CreatedInstanceMap, &State.OldToNewInstanceMap, &OrderedListOfObjectToCopy);
+	State.OldToNewInstanceMap.Append(CreatedInstanceMap);
 
 	// Copy property values
 	UEngine::FCopyPropertiesForUnrelatedObjectsParams Options;
 	Options.bNotifyObjectReplacement = true;
 	Options.bSkipCompilerGeneratedDefaults = true;
 	Options.bOnlyHandleDirectSubObjects = true;
-	Options.OptionalReplacementMappings = &OldToNewInstanceMap;
+	Options.OptionalReplacementMappings = &State.OldToNewInstanceMap;
 	if (FOverridableManager::Get().IsEnabled(OldObject))
 	{
 		Options.bReplaceInternalReferenceUponRead = true;
-		Options.OptionalOldToNewClassMappings = &OldToNewClassMap;
+		Options.OptionalOldToNewClassMappings = &State.InOldToNewClassMap;
 	}
 	// this currently happens because of some misguided logic in UBlueprintGeneratedClass::FindArchetype that
 	// points us to a mismatched archetype, in which case delta serialization becomes unsafe.. without
@@ -2358,7 +2358,7 @@ void FBlueprintCompileReinstancer::ReplaceObjectHelper(UObject*& OldObject, UCla
 		}
 		, true, RF_NoFlags, EInternalObjectFlags::Garbage);
 
-	OldToNewInstanceMap.Add(OldObject, NewUObject);
+	State.OldToNewInstanceMap.Add(OldObject, NewUObject);
 
 	if (bIsComponent)
 	{
@@ -2371,13 +2371,13 @@ void FBlueprintCompileReinstancer::ReplaceObjectHelper(UObject*& OldObject, UCla
 			// Check to see if they have an editor that potentially needs to be refreshed
 			if (OwningActor->GetClass()->ClassGeneratedBy)
 			{
-				PotentialEditorsForRefreshing.AddUnique(OwningActor->GetClass()->ClassGeneratedBy);
+				State.PotentialEditorsForRefreshing.AddUnique(OwningActor->GetClass()->ClassGeneratedBy);
 			}
 
 			// we need to keep track of actor instances that need 
 			// their construction scripts re-ran (since we've just 
 			// replaced a component they own)
-			OwnersToRerunConstructionScript.Add(OwningActor);
+			State.OwnersToRerunConstructionScript.Add(OwningActor);
 		}
 
 		if (bWasRegistered)
@@ -2599,6 +2599,8 @@ void FBlueprintCompileReinstancer::ReplaceInstancesOfClass_Inner(const TMap<UCla
 
 	TMap<FSoftObjectPath, UObject*> ReinstancedObjectsWeakReferenceMap;
 
+	TMap<UObject*, FObjectInstancingGraph> GraphCache;
+
 	// actors being replace
 	TArray<FActorReplacementHelper> ReplacementActors;
 
@@ -2678,30 +2680,32 @@ void FBlueprintCompileReinstancer::ReplaceInstancesOfClass_Inner(const TMap<UCla
 			}
 		}) : FDelegateHandle();
 
+		FReinstancingState ReinstancingState
+		{
+			.ObjectsToReplaceScratch = ObjectsToReplace,
+			.CachedArchetypeObjects = CachedArchetypeObjects,
+			.OldToNewInstanceMap = OldToNewInstanceMap,
+			.InOldToNewClassMap = InOldToNewClassMap,
+			.OldToNewNameMap = OldToNewNameMap,
+			.OldToNewTemplates = Params.OldToNewTemplates,
+			.PotentialEditorsForRefreshing = PotentialEditorsForRefreshing,
+			.OwnersToRerunConstructionScript = OwnersToRerunConstructionScript,
+			.ReplacementActors = ReplacementActors,
+			.ReinstancedObjectsWeakReferenceMap = ReinstancedObjectsWeakReferenceMap,
+			.GraphCache = GraphCache,
+			.ConstructionVersioningData = *Params.ConstructionVersioningData,
+			.ObjectsReplaced = ObjectsReplaced,
+			.CleanPackageList = CleanPackageList,
+			.bSelectionChanged = bSelectionChanged,
+			.bFixupSCS = bFixupSCS
+		};
+
 		if(Params.ConstructionVersioningData)
 		{
 			check(bArchetypesAreUpToDate); // not supported, can't be
 			check(Params.InstancesThatShouldUseOldClass == nullptr); // not supported, could be.. but why?
 
-			ReinstanceWithVersioningData(
-				{
-					.ObjectsToReplaceScratch = ObjectsToReplace,
-					.CachedArchetypeObjects = CachedArchetypeObjects,
-					.OldToNewInstanceMap = OldToNewInstanceMap,
-					.InOldToNewClassMap = InOldToNewClassMap,
-					.OldToNewNameMap = OldToNewNameMap,
-					.OldToNewTemplates = Params.OldToNewTemplates,
-					.PotentialEditorsForRefreshing = PotentialEditorsForRefreshing,
-					.OwnersToRerunConstructionScript = OwnersToRerunConstructionScript,
-					.ReplacementActors = ReplacementActors,
-					.ReinstancedObjectsWeakReferenceMap = ReinstancedObjectsWeakReferenceMap,
-					.ConstructionVersioningData = *Params.ConstructionVersioningData,
-					.ObjectsReplaced = ObjectsReplaced,
-					.CleanPackageList = CleanPackageList,
-					.bSelectionChanged = bSelectionChanged,
-					.bFixupSCS = bFixupSCS
-				}
-			);
+			ReinstanceWithVersioningData(ReinstancingState);
 		}
 		else
 		{
@@ -2807,7 +2811,7 @@ void FBlueprintCompileReinstancer::ReplaceInstancesOfClass_Inner(const TMap<UCla
 						CheckAndSaveOuterPackageToCleanList(CleanPackageList, OldObject);
 
 						UObject* NewUObject = nullptr;
-						ReplaceObjectHelper(OldObject, OldClass, NewUObject, NewClass, OldToNewInstanceMap, InOldToNewClassMap, OldToNewNameMap, OldObjIndex, ObjectsToReplace, PotentialEditorsForRefreshing, OwnersToRerunConstructionScript, &FDirectAttachChildrenAccessor::Get, bIsComponent, bArchetypesAreUpToDate);
+						ReplaceObjectHelper(ReinstancingState, OldObject, OldClass, NewUObject, NewClass, OldObjIndex, ObjectsToReplace, &FDirectAttachChildrenAccessor::Get, bIsComponent, bArchetypesAreUpToDate);
 						UpdateObjectBeingDebugged(OldObject, NewUObject);
 						ObjectsReplaced.Add(OldObject);
 
@@ -2887,7 +2891,7 @@ void FBlueprintCompileReinstancer::ReplaceInstancesOfClass_Inner(const TMap<UCla
 						else
 						{
 							// Actors that are not in a level cannot be reconstructed, sequencer team decided to reinstance these as normal objects:
-							ReplaceObjectHelper(OldObject, OldClass, NewUObject, NewClass, OldToNewInstanceMap, InOldToNewClassMap, OldToNewNameMap, OldObjIndex, ObjectsToReplace, PotentialEditorsForRefreshing, OwnersToRerunConstructionScript, &FDirectAttachChildrenAccessor::Get, false, bArchetypesAreUpToDate);
+							ReplaceObjectHelper(ReinstancingState, OldObject, OldClass, NewUObject, NewClass, OldObjIndex, ObjectsToReplace, &FDirectAttachChildrenAccessor::Get, false, bArchetypesAreUpToDate);
 						}
 						UpdateObjectBeingDebugged(OldObject, NewUObject);
 						ObjectsReplaced.Add(OldObject);
@@ -3305,7 +3309,7 @@ void FBlueprintCompileReinstancer::PreCreateSubObjectsForReinstantiation(
 	}
 }
 
-bool FBlueprintCompileReinstancer::PrePopulateInstancingGraphForArchetype(FObjectInstancingGraph& InstancingGraph, TNotNull<UObject*> OldInstance, UObject* InArchetype, TMap<UObject*, UObject*>* OldToNewInstanceMap)
+bool FBlueprintCompileReinstancer::PrePopulateInstancingGraphForArchetype(TMap<UObject*, FObjectInstancingGraph>& GraphCache, FObjectInstancingGraph& InstancingGraph, TNotNull<UObject*> OldInstance, UObject* InArchetype, TMap<UObject*, UObject*>* OldToNewInstanceMap)
 {
 	checkf(!InstancingGraph.HasDestinationRoot(), TEXT("Can only pre-populate an empty InstancingGraph"));
 
@@ -3327,53 +3331,59 @@ bool FBlueprintCompileReinstancer::PrePopulateInstancingGraphForArchetype(FObjec
 	} while (InstancingGraphRoot && !InstancingGraphRoot->GetArchetype()->HasAnyFlags(RF_ClassDefaultObject));
 
 	// Sanitization of the found instancing root
-	if (!InstancingGraphRoot || InstancingGraphRoot->IsA<UPackage>())
+	if (!InstancingGraphRoot || Cast<UPackage>(InstancingGraphRoot) != nullptr)
 	{
 		UE_LOG(LogClass, Warning, TEXT("Object %s was created from an archetype we could not find the root of the instancing graph."), *GetPathNameSafe(OldInstance));
 		return false;
 	}
 
-	auto GatherMatchingSubObjects = [&InstancingGraph, OldInstance, OldToNewInstanceMap](UObject* Instance, UObject* Template, auto GatherMatchingSubObjectsRecurse) -> void
-	{
-		InstancingGraph.PrepopulateWithInstance(Instance, Template);
+	// Look for an existing graph
+	FObjectInstancingGraph* SourceGraph = GraphCache.Find(InstancingGraphRoot);
 
-		ForEachObjectWithOuter(Instance, [OldInstance, Template, GatherMatchingSubObjectsRecurse, OldToNewInstanceMap](UObject* SubObject)
+	if (!SourceGraph)
+	{
+		SourceGraph = &GraphCache.Add(InstancingGraphRoot);
+		
+		// Compute the graph for a specific top level object
+		auto GatherMatchingSubObjects = [SourceGraph, InstancingGraphRoot, OldToNewInstanceMap](UObject* Instance, UObject* Template, auto GatherMatchingSubObjectsRecurse) -> void
 		{
-			// Skip the object that is about to be re-instantiated
-			if (SubObject == OldInstance)
-			{
-				return;
-			}
-
-			// Just handle subobject that are created from archetype, skip the others
-			UObject* SubObjectArchetype = SubObject->GetArchetype();
-			if (SubObjectArchetype->HasAnyFlags(RF_ClassDefaultObject))
-			{
-				return;
-			}
+			SourceGraph->PrepopulateWithInstance(Instance, Template);
 
-			// Sanitization of the archetype
-			if (SubObjectArchetype->GetOuter() != Template)
+			ForEachObjectWithOuter(Instance, [Template, InstancingGraphRoot, GatherMatchingSubObjectsRecurse, OldToNewInstanceMap](UObject* SubObject)
 			{
-				// Handle the fact that subobject will be re-instantiated later by pointing its archetype to the new one.
-				if (UObject** NewSubObjectArchetype = OldToNewInstanceMap ? OldToNewInstanceMap->Find(SubObjectArchetype) : nullptr)
+				// Just handle subobjects that are created from archetype, skip the others
+				UObject* SubObjectArchetype = SubObject->GetArchetype();
+				if (SubObjectArchetype->HasAnyFlags(RF_ClassDefaultObject))
 				{
-					SubObjectArchetype = *NewSubObjectArchetype;
+					return;
 				}
-				else
+
+				// Sanitization of the archetype
+				if (SubObjectArchetype->GetOuter() != Template)
 				{
-					UE_LOG(LogClass, VeryVerbose, TEXT("Missmatched while trying to rebuild instancing graph for '%s', expecting the outer of archetype %s to be %s but was %s"), 
-						*GetPathNameSafe(OldInstance), *GetPathNameSafe(SubObjectArchetype), *GetPathNameSafe(Template), *GetPathNameSafe(SubObjectArchetype->GetOuter()));
-					return;
+					// Handle the fact that subobject will be re-instantiated later by pointing its archetype to the new one.
+					if (UObject** NewSubObjectArchetype = OldToNewInstanceMap ? OldToNewInstanceMap->Find(SubObjectArchetype) : nullptr)
+					{
+						SubObjectArchetype = *NewSubObjectArchetype;
+					}
+					else
+					{
+						UE_LOG(LogClass, VeryVerbose, TEXT("Missmatched while trying to rebuild instancing graph for '%s', expecting the outer of archetype %s to be %s but was %s"), 
+							*GetPathNameSafe(InstancingGraphRoot), *GetPathNameSafe(SubObjectArchetype), *GetPathNameSafe(Template), *GetPathNameSafe(SubObjectArchetype->GetOuter()));
+						return;
+					}
 				}
-			}
 
-			GatherMatchingSubObjectsRecurse(SubObject, SubObjectArchetype, GatherMatchingSubObjectsRecurse);
+				GatherMatchingSubObjectsRecurse(SubObject, SubObjectArchetype, GatherMatchingSubObjectsRecurse);
 
-		}, /*bIncludeNestedObjects*/false);
-	};
+			}, /*bIncludeNestedObjects*/false);
+		};
 
-	GatherMatchingSubObjects(InstancingGraphRoot, InstancingGraphRoot->GetArchetype(), GatherMatchingSubObjects );
+		GatherMatchingSubObjects(InstancingGraphRoot, InstancingGraphRoot->GetArchetype(), GatherMatchingSubObjects);
+	}
+
+	// Copy from existing graph
+	InstancingGraph = *SourceGraph;
 
 	return true;
 }
@@ -3756,17 +3766,13 @@ void FBlueprintCompileReinstancer::ReinstanceNormalObjects(
 			CheckAndSaveOuterPackageToCleanList(ReinstancingState.CleanPackageList, OldObject);
 			UObject* NewUObject = nullptr;
 			ReplaceObjectHelper(
+				ReinstancingState,
 				OldObject, 
 				OldClass, 
 				NewUObject, 
 				NewClass, 
-				ReinstancingState.OldToNewInstanceMap, 
-				ReinstancingState.InOldToNewClassMap, 
-				ReinstancingState.OldToNewNameMap, 
 				OldObjIndex, 
 				EmptyList, 
-				ReinstancingState.PotentialEditorsForRefreshing, 
-				ReinstancingState.OwnersToRerunConstructionScript, 
 				&FDirectAttachChildrenAccessor::Get, 
 				bIsComponent, 
 				bArchetypesAreUpToDate);
@@ -3836,17 +3842,13 @@ void FBlueprintCompileReinstancer::ReinstanceActors(
 			{
 				// Actors that are not in a level cannot be reconstructed, reinstance them as normal objects:
 				ReplaceObjectHelper(
+					ReinstancingState,
 					OldObject, 
 					OldClass,
 					NewUObject,
 					NewClass,
-					ReinstancingState.OldToNewInstanceMap,
-					ReinstancingState.InOldToNewClassMap,
-					ReinstancingState.OldToNewNameMap,
 					OldObjIndex,
 					EmptyList,
-					ReinstancingState.PotentialEditorsForRefreshing,
-					ReinstancingState.OwnersToRerunConstructionScript,
 					&FDirectAttachChildrenAccessor::Get,
 					bIsComponent,
 					bArchetypesAreUpToDate);
@@ -3942,4 +3944,4 @@ bool FBlueprintCompileReinstancer::ShouldReinstanceActor(const AActor* OldActor)
 		return false;
 	}
 	return true;
-}
+}
\ No newline at end of file
diff --git a/Engine/Source/Editor/UnrealEd/Public/Kismet2/KismetReinstanceUtilities.h b/Engine/Source/Editor/UnrealEd/Public/Kismet2/KismetReinstanceUtilities.h
--- a/Engine/Source/Editor/UnrealEd/Public/Kismet2/KismetReinstanceUtilities.h
+++ b/Engine/Source/Editor/UnrealEd/Public/Kismet2/KismetReinstanceUtilities.h
@@ -329,20 +329,40 @@ public:
 	static UNREALED_API void PreCreateSubObjectsForReinstantiation(const TMap<UClass*, UClass*>& OldToNewClassMap, UObject* OldObject, UObject* NewUObject, TMap<UObject*, UObject*>& CreatedInstanceMap, const TMap<UObject*, UObject*>* OldToNewInstanceMap = nullptr, TArray< TTuple<UObject*, UObject*>>* OrderedListOfObjectToCopy = nullptr);
 
 private:
+	struct FReinstancingState
+	{
+		TArray<UObject*>& ObjectsToReplaceScratch;
+		TSet<UObject*>& CachedArchetypeObjects;
+		TMap<UObject*, UObject*>& OldToNewInstanceMap;
+		const TMap<UClass*, UClass*>& InOldToNewClassMap;
+		TMap<UObject*, FName>& OldToNewNameMap;
+		TMap<UClass*, TMap<UObject*, UObject*>>* OldToNewTemplates = nullptr;
+		TArray<UObject*>& PotentialEditorsForRefreshing;
+		TSet<AActor*>& OwnersToRerunConstructionScript;
+		TArray<struct FActorReplacementHelper>& ReplacementActors;
+		TMap<FSoftObjectPath, UObject*>& ReinstancedObjectsWeakReferenceMap;
+		TMap<UObject*, FObjectInstancingGraph>& GraphCache;
+		TMap<TObjectPtr<const UClass>, FBlake3Hash>& ConstructionVersioningData;
+		TArray<UObject*>& ObjectsReplaced;
+		TSet<UPackage*>& CleanPackageList;
+		bool& bSelectionChanged;
+		bool& bFixupSCS;
+	};
 
 	/** Helper to replace an instance of an object*/
-	static UNREALED_API void ReplaceObjectHelper(UObject*& OldObject, UClass* OldClass, UObject*& NewUObject, UClass* NewClass, TMap<UObject*, UObject*>& OldToNewInstanceMap, const TMap<UClass*, UClass*>& OldToNewClassMap, TMap<UObject*, FName>& OldToNewNameMap, int32 OldObjIndex, TArray<UObject*>& ObjectsToReplace, TArray<UObject*>& PotentialEditorsForRefreshing, TSet<AActor*>& OwnersToRerunConstructionScript, TFunctionRef<TArray<TObjectPtr<USceneComponent>>&(USceneComponent*)> GetAttachChildrenArray, bool bIsComponent, bool bArchetypesAreUpToDate);
+	static UNREALED_API void ReplaceObjectHelper(const FReinstancingState& State, UObject*& OldObject, UClass* OldClass, UObject*& NewUObject, UClass* NewClass, int32 OldObjIndex, TArray<UObject*>& ObjectsToReplace, TFunctionRef<TArray<TObjectPtr<USceneComponent>>&(USceneComponent*)> GetAttachChildrenArray, bool bIsComponent, bool bArchetypesAreUpToDate);
 
 	/**
 	 * Pre-populate instancing graph for a specified instance if needed
 	 *
+	 * @param GraphCache Cache of precomputed instance graphs for each top level object
 	 * @param InstancingGraph The instancing graph to populate
 	 * @param OldInstance Object that is about to be reallocated
 	 * @param InArchetype Archetype of ObjectInstance
 	 * @param OldToNewInstanceMap Possible instance mapping of instances
 	 * @return false if the instancing graph is not needed
 	 */
-	static UNREALED_API bool PrePopulateInstancingGraphForArchetype(FObjectInstancingGraph& InstancingGraph, TNotNull<UObject*> OldInstance, UObject* InArchetype, TMap<UObject*, UObject*>* OldToNewInstanceMap = nullptr);
+	static UNREALED_API bool PrePopulateInstancingGraphForArchetype(TMap<UObject*,FObjectInstancingGraph>& GraphCache, FObjectInstancingGraph& InstancingGraph, TNotNull<UObject*> OldInstance, UObject* InArchetype, TMap<UObject*, UObject*>* OldToNewInstanceMap = nullptr);
 
 
 	/** Handles the sub object pre-creation recursively */
@@ -364,25 +384,6 @@ private:
 	/** Returns the above, but only if there are instances of the class that would require reinstancing */
 	static UNREALED_API FBlake3Hash CalculateConstructionHashIfNecessary(const UClass* ForClass);
 
-	struct FReinstancingState
-	{
-		TArray<UObject*>& ObjectsToReplaceScratch;
-		TSet<UObject*>& CachedArchetypeObjects;
-		TMap<UObject*, UObject*>& OldToNewInstanceMap;
-		const TMap<UClass*, UClass*>& InOldToNewClassMap;
-		TMap<UObject*, FName>& OldToNewNameMap;
-		TMap<UClass*, TMap<UObject*, UObject*>>* OldToNewTemplates = nullptr;
-		TArray<UObject*>& PotentialEditorsForRefreshing;
-		TSet<AActor*>& OwnersToRerunConstructionScript;
-		TArray<struct FActorReplacementHelper>& ReplacementActors;
-		TMap<FSoftObjectPath, UObject*>& ReinstancedObjectsWeakReferenceMap;
-		TMap<TObjectPtr<const UClass>, FBlake3Hash>& ConstructionVersioningData;
-		TArray<UObject*>& ObjectsReplaced;
-		TSet<UPackage*>& CleanPackageList;
-		bool& bSelectionChanged;
-		bool& bFixupSCS;
-	};
-
 	static void ReinstanceWithVersioningData(const FReinstancingState& ReinstancingState);
 	static void CategorizeClassesAndInstances(
 		const TMap<UClass*, UClass*>& InOldToNewClassMap
