Fix PlayLevel check for uncompiled blueprints to skip garbage assets instead of crashing before the check.
#Synced-CL 49126341

diff --git a/Engine/Source/Editor/UnrealEd/Private/PlayLevel.cpp b/Engine/Source/Editor/UnrealEd/Private/PlayLevel.cpp
--- a/Engine/Source/Editor/UnrealEd/Private/PlayLevel.cpp
+++ b/Engine/Source/Editor/UnrealEd/Private/PlayLevel.cpp
@@ -1267,8 +1267,8 @@ int32 FInternalPlayLevelUtils::ResolveDirtyBlueprints(const bool bPromptForCompi
 	{
 		UBlueprint* Blueprint = *BlueprintIt;
 
-		// ignore up-to-date BPs
-		if (Blueprint->IsUpToDate())
+		// ignore garbage and up-to-date BPs
+		if (!IsValid(Blueprint) || Blueprint->IsUpToDate())
 		{
 			continue;
 		}
@@ -1277,8 +1277,7 @@ int32 FInternalPlayLevelUtils::ResolveDirtyBlueprints(const bool bPromptForCompi
 		const bool bIsDirtyAndShouldBeRecompiled = Blueprint->IsPossiblyDirty();
 		if (!FBlueprintEditorUtils::IsDataOnlyBlueprint(Blueprint)
 			&& (bIsDirtyAndShouldBeRecompiled || (FBlueprintEditorUtils::IsLevelScriptBlueprint(Blueprint) && bForceLevelScriptRecompile))
-			&& (Blueprint->Status != BS_Unknown)
-			&& IsValid(Blueprint))
+			&& (Blueprint->Status != BS_Unknown))
 		{
 			InNeedOfRecompile.Add(Blueprint);
 
