Fix issue noticed by static analysis where it could assign a reference to a resolved null pointer
#Synced-CL 49126341

diff --git a/Engine/Source/Editor/UnrealEd/Private/Kismet2/KismetReinstanceUtilities.cpp b/Engine/Source/Editor/UnrealEd/Private/Kismet2/KismetReinstanceUtilities.cpp
--- a/Engine/Source/Editor/UnrealEd/Private/Kismet2/KismetReinstanceUtilities.cpp
+++ b/Engine/Source/Editor/UnrealEd/Private/Kismet2/KismetReinstanceUtilities.cpp
@@ -2693,7 +2693,7 @@ void FBlueprintCompileReinstancer::ReplaceInstancesOfClass_Inner(const TMap<UCla
 			.ReplacementActors = ReplacementActors,
 			.ReinstancedObjectsWeakReferenceMap = ReinstancedObjectsWeakReferenceMap,
 			.GraphCache = GraphCache,
-			.ConstructionVersioningData = *Params.ConstructionVersioningData,
+			.ConstructionVersioningData = Params.ConstructionVersioningData,
 			.ObjectsReplaced = ObjectsReplaced,
 			.CleanPackageList = CleanPackageList,
 			.bSelectionChanged = bSelectionChanged,
@@ -3642,7 +3642,7 @@ void FBlueprintCompileReinstancer::CategorizeClassesAndInstances(
 		bool bUseActorReinstancing = false;
 		// if we have a script constructor we would have to check it for equivalence here
 		// but there are many obstacles to script constructors:
-		const FBlake3Hash* ConstructionVersion = ReinstancingState.ConstructionVersioningData.Find(NewClass);
+		const FBlake3Hash* ConstructionVersion = ReinstancingState.ConstructionVersioningData ? ReinstancingState.ConstructionVersioningData->Find(NewClass) : nullptr;
 		// anim blueprints have uncontrolled backpointers to CDO data and therefore instances must
 		// be reconstructed after compilation:
 		if( ConstructionVersion && 
diff --git a/Engine/Source/Editor/UnrealEd/Public/Kismet2/KismetReinstanceUtilities.h b/Engine/Source/Editor/UnrealEd/Public/Kismet2/KismetReinstanceUtilities.h
--- a/Engine/Source/Editor/UnrealEd/Public/Kismet2/KismetReinstanceUtilities.h
+++ b/Engine/Source/Editor/UnrealEd/Public/Kismet2/KismetReinstanceUtilities.h
@@ -342,7 +342,7 @@ private:
 		TArray<struct FActorReplacementHelper>& ReplacementActors;
 		TMap<FSoftObjectPath, UObject*>& ReinstancedObjectsWeakReferenceMap;
 		TMap<UObject*, FObjectInstancingGraph>& GraphCache;
-		TMap<TObjectPtr<const UClass>, FBlake3Hash>& ConstructionVersioningData;
+		TMap<TObjectPtr<const UClass>, FBlake3Hash>* ConstructionVersioningData;
 		TArray<UObject*>& ObjectsReplaced;
 		TSet<UPackage*>& CleanPackageList;
 		bool& bSelectionChanged;
