[Cooker] Bump up cook incremental version number as an attempted fix for issue 1170838

diff --git a/Engine/Source/Editor/UnrealEd/Private/Cooker/CookIncrementalVersion.cpp b/Engine/Source/Editor/UnrealEd/Private/Cooker/CookIncrementalVersion.cpp
--- a/Engine/Source/Editor/UnrealEd/Private/Cooker/CookIncrementalVersion.cpp
+++ b/Engine/Source/Editor/UnrealEd/Private/Cooker/CookIncrementalVersion.cpp
@@ -8,6 +8,6 @@ namespace UE::Cook
 {
 
 // Change CookIncrementalVersion to a new guid when all packages in an incremental cook need to be invalidated
-FGuid CookIncrementalVersion(0x74F40C2C, 0x03FD4D06, 0x9FD3B32C, 0x46AE93F7);
+FGuid CookIncrementalVersion(0xA3F19B7E, 0x4D0C2A11, 0x8B72E4C9, 0xF15D09A3);
 
 }
