UE: Bump incremental cook due to atlas issue

diff --git a/Engine/Source/Editor/UnrealEd/Private/Cooker/CookIncrementalVersion.cpp b/Engine/Source/Editor/UnrealEd/Private/Cooker/CookIncrementalVersion.cpp
--- a/Engine/Source/Editor/UnrealEd/Private/Cooker/CookIncrementalVersion.cpp
+++ b/Engine/Source/Editor/UnrealEd/Private/Cooker/CookIncrementalVersion.cpp
@@ -8,6 +8,6 @@ namespace UE::Cook
 {
 
 // Change CookIncrementalVersion to a new guid when all packages in an incremental cook need to be invalidated
-FGuid CookIncrementalVersion(0xA3F19B7E, 0x4D0C2A11, 0x8B72E4C9, 0xF15D09A3);
+FGuid CookIncrementalVersion(0x15887AC8, 0xA1AD43BD, 0x8A60BB6A, 0x724880FC);
 
 }
