UE: Bump incremental cook version. Serialization schema change

diff --git a/Engine/Source/Editor/UnrealEd/Private/Cooker/CookIncrementalVersion.cpp b/Engine/Source/Editor/UnrealEd/Private/Cooker/CookIncrementalVersion.cpp
--- a/Engine/Source/Editor/UnrealEd/Private/Cooker/CookIncrementalVersion.cpp
+++ b/Engine/Source/Editor/UnrealEd/Private/Cooker/CookIncrementalVersion.cpp
@@ -8,6 +8,6 @@ namespace UE::Cook
 {
 
 // Change CookIncrementalVersion to a new guid when all packages in an incremental cook need to be invalidated
-FGuid CookIncrementalVersion(0x15897AC8, 0xA1AA43BD, 0x8A60CB6A, 0x721880FC);
+FGuid CookIncrementalVersion(0x75897AC8, 0xA1AA13BD, 0x8A60CB2A, 0xA21880FC);
 
 }
