ExternalActorsCommandlet:
- Make sure to delete actors that WP specific such as WorldDataLayer and WorldPartitionMinimap when disabling external actors.
- Added the option to report mismatched external actor filenames as warnings, for duplicated maps that keep the original filenames.
- Also added more verbose errors when detecting duplicated actor guids.

#Synced-CL 49330242

diff --git a/Engine/Source/Editor/UnrealEd/Private/Commandlets/ExternalActorsCommandlet.cpp b/Engine/Source/Editor/UnrealEd/Private/Commandlets/ExternalActorsCommandlet.cpp
--- a/Engine/Source/Editor/UnrealEd/Private/Commandlets/ExternalActorsCommandlet.cpp
+++ b/Engine/Source/Editor/UnrealEd/Private/Commandlets/ExternalActorsCommandlet.cpp
@@ -242,7 +242,7 @@ int32 UExternalActorsCommandlet::Main(const FString& Params)
 				{
 					Result = 0;
 					continue;
-				}			
+				}
 
 				if (!bReport)
 				{
@@ -272,27 +272,26 @@ int32 UExternalActorsCommandlet::Main(const FString& Params)
 
 			if (!bReport)
 			{
-				for (AActor* Actor : MainWorld->PersistentLevel->Actors)
+				ForEachObjectWithOuter(MainWorld->PersistentLevel, [MainWorld, &PackagesToDelete](UObject* Object)
 				{
-					if (IsValid(Actor) && Actor->IsMainPackageActor())
+					if (Object->IsPackageExternal())
 					{
-						PackagesToDelete.Add(Actor->GetPackage()->GetLoadedPath().GetLocalFullPath());
-						Actor->SetPackageExternal(false);
-					}
-				}
+						PackagesToDelete.Add(Object->GetPackage()->GetLoadedPath().GetLocalFullPath());
 
-				if (MainWorld->PersistentLevel->IsUsingExternalObjects())
-				{
-					ForEachObjectWithOuter(MainWorld->PersistentLevel, [MainWorld, &PackagesToDelete](UObject* Object)
-					{
-						if (Object->IsPackageExternal() && !Object->IsA<AActor>())
+						if (IsValid(Object))
 						{
-							PackagesToDelete.Add(Object->GetPackage()->GetLoadedPath().GetLocalFullPath());
-							FExternalPackageHelper::SetPackagingMode(Object, MainWorld->PersistentLevel, false);
+							if (AActor* Actor = Cast<AActor>(Object))
+							{
+								Actor->SetPackageExternal(false);
+							}
+							else
+							{
+								FExternalPackageHelper::SetPackagingMode(Object, MainWorld->PersistentLevel, false);
+							}
 						}
-						return true;
-					}, true, RF_MirroredGarbage);
-				}
+					}
+					return true;
+				});
 
 				if (PackagesToDelete.Num())
 				{
@@ -399,9 +398,11 @@ int32 UExternalActorsCommandlet::Main(const FString& Params)
 
 							if (MainPackageActor)
 							{
-								if (ActorGuids.Contains(MainPackageActor->GetActorGuid()))
+								if (AActor** ExistingActor = ActorGuids.Find(MainPackageActor->GetActorGuid()))
 								{
-									UE_LOG(LogExternalActorsCommandlet, Error, TEXT("Duplicated actor guid for file '%s'"), *Filename);
+									UE_LOG(LogExternalActorsCommandlet, Error, TEXT("Duplicated actor guid '%s' for file'%s':"), *MainPackageActor->GetActorGuid().ToString(), *Filename);
+									UE_LOG(LogExternalActorsCommandlet, Error, TEXT("\tActor 1: %s (ignore)"), *(*ExistingActor)->GetName())
+									UE_LOG(LogExternalActorsCommandlet, Error, TEXT("\tActor 2: %s (keep)"), *MainPackageActor->GetName());									
 
 									if (bRepair)
 									{
@@ -413,6 +414,16 @@ int32 UExternalActorsCommandlet::Main(const FString& Params)
 								}
 								else
 								{
+									const FString ActorPath = MainPackageActor->GetPathName();
+									const FString ActorPackageName = ULevel::GetActorPackageName(MainPackageActor->GetLevel()->GetPackage(), MainPackageActor->GetLevel()->GetActorPackagingScheme(), ActorPath);
+
+									if (MainPackageActor->GetPackage()->GetName() != ActorPackageName)
+									{
+										UE_LOG(LogExternalActorsCommandlet, Warning, TEXT("Mismatched actor filename for actor '%s':"), *MainPackageActor->GetName());
+										UE_LOG(LogExternalActorsCommandlet, Warning, TEXT("\t Current filename: %s"), *MainPackageActor->GetPackage()->GetName());
+										UE_LOG(LogExternalActorsCommandlet, Warning, TEXT("\tExpected filename: %s"), *ActorPackageName);										
+									}
+
 									ActorGuids.Add(MainPackageActor->GetActorGuid(), MainPackageActor);
 								}
 							}
