Fix an issue where deleting the current EditorWorld leaves a GC???d reference inside ShownObjectsToDelete.
It needs to be replaced with the new object array generated in ReloadEditorWorldForReferenceReplacementIfNecessary.

#Synced-CL 49157442

diff --git a/Engine/Source/Editor/UnrealEd/Private/ObjectTools.cpp b/Engine/Source/Editor/UnrealEd/Private/ObjectTools.cpp
--- a/Engine/Source/Editor/UnrealEd/Private/ObjectTools.cpp
+++ b/Engine/Source/Editor/UnrealEd/Private/ObjectTools.cpp
@@ -3893,7 +3893,11 @@ namespace ObjectTools
 			}
 			
 			// Give systems opportunity to clean up references to the objects being deleted
-			FEditorDelegates::OnAssetsPreDelete.Broadcast(ShownObjectsToDelete);
+			{
+				TArray<UObject*> ObjectsToPreDelete;
+				Algo::Transform(ObjectsToDelete, ObjectsToPreDelete, [](const TWeakObjectPtr<UObject>& ObjectWeakPtr) { return ObjectWeakPtr.Get(); });
+				FEditorDelegates::OnAssetsPreDelete.Broadcast(ObjectsToPreDelete);
+			}
 
 			// Load the asset tools module to get access to the browser type maps
 			FAssetToolsModule& AssetToolsModule = FModuleManager::LoadModuleChecked<FAssetToolsModule>(TEXT("AssetTools"));
