EditorUsability : Preferences
- Added Keyboard Shortcut menu entry
- Added keybinds to access Editor, Project, Keyboard preferences
- Set icons for commands

#Synced-CL 49408319

diff --git a/Engine/Source/Editor/EditorStyle/Private/StarshipStyle.cpp b/Engine/Source/Editor/EditorStyle/Private/StarshipStyle.cpp
--- a/Engine/Source/Editor/EditorStyle/Private/StarshipStyle.cpp
+++ b/Engine/Source/Editor/EditorStyle/Private/StarshipStyle.cpp
@@ -1394,11 +1394,14 @@ void FStarshipEditorStyle::FStyle::SetupGeneralStyles()
 	// Editor preferences
 	{
 		Set("EditorPreferences.TabIcon", new CORE_IMAGE_BRUSH_SVG("Starship/Common/Preferences", Icon16x16));
+		Set("SystemWideCommands.OpenEditorPreferences", new CORE_IMAGE_BRUSH_SVG("Starship/Common/Preferences", Icon16x16));
+		Set("SystemWideCommands.OpenKeyboardShortcuts", new CORE_IMAGE_BRUSH_SVG("Starship/Common/Preferences", Icon16x16));
 	}
 
 	// Project settings
 	{
 		Set("ProjectSettings.TabIcon", new IMAGE_BRUSH_SVG("Starship/Common/ProjectSettings", Icon16x16));
+		Set("SystemWideCommands.OpenProjectSettings", new IMAGE_BRUSH_SVG("Starship/Common/ProjectSettings", Icon16x16));
 	}
 
 	// Main frame
diff --git a/Engine/Source/Editor/MainFrame/Private/Menus/MainMenu.cpp b/Engine/Source/Editor/MainFrame/Private/Menus/MainMenu.cpp
--- a/Engine/Source/Editor/MainFrame/Private/Menus/MainMenu.cpp
+++ b/Engine/Source/Editor/MainFrame/Private/Menus/MainMenu.cpp
@@ -244,39 +244,10 @@ void FMainMenu::RegisterEditMenu()
 		}
 		else
 		{
-#if !PLATFORM_MAC // Handled by app's menu in menu bar
-			{
-				IMainFrameModule& MainFrame = FModuleManager::LoadModuleChecked<IMainFrameModule>("MainFrame");
-				FName CategoryName;
-				FName SectionName;
-				MainFrame.GetEditorSettingsDefaultSelectionOverride(CategoryName, SectionName);
-
-				if (CategoryName.IsNone())
-				{
-					CategoryName = FName("General");
-				}
-				if (SectionName.IsNone())
-				{
-					SectionName = FName("Appearance");
-				}
-
-				Section.AddMenuEntry(
-					"EditorPreferencesMenu",
-					LOCTEXT("EditorPreferencesMenuLabel", "Editor Preferences..."),
-					LOCTEXT("EditorPreferencesMenuToolTip", "Configure the behavior and features of the Unreal Editor."),
-					FSlateIcon(FAppStyle::GetAppStyleSetName(), "EditorPreferences.TabIcon"),
-					FUIAction(FExecuteAction::CreateStatic(&FSettingsMenu::OpenSettings, FName("Editor"), CategoryName, SectionName))
-				);
-			}
-#endif //if !PLATFORM_MAC
-
-			Section.AddMenuEntry(
-				"ProjectSettingsMenu",
-				LOCTEXT("ProjectSettingsMenuLabel", "Project Settings..."),
-				LOCTEXT("ProjectSettingsMenuToolTip", "Change the settings of the currently loaded project."),
-				FSlateIcon(FAppStyle::GetAppStyleSetName(), "ProjectSettings.TabIcon"),
-				FUIAction(FExecuteAction::CreateStatic(&FSettingsMenu::OpenSettings, FName("Project"), FName("Project"), FName("General")))
-			);
+			const FGlobalEditorCommonCommands& GlobalEditorCommands = FGlobalEditorCommonCommands::Get();
+			Section.AddMenuEntry(GlobalEditorCommands.OpenEditorPreferences);
+			Section.AddMenuEntry(GlobalEditorCommands.OpenKeyboardShortcuts);
+			Section.AddMenuEntry(GlobalEditorCommands.OpenProjectSettings);
 		}
 
 		Section.AddDynamicEntry("PluginsEditor", FNewToolMenuSectionDelegate::CreateLambda([](FToolMenuSection& InSection)
diff --git a/Engine/Source/Editor/UnrealEd/Private/Toolkits/GlobalEditorCommonCommands.cpp b/Engine/Source/Editor/UnrealEd/Private/Toolkits/GlobalEditorCommonCommands.cpp
--- a/Engine/Source/Editor/UnrealEd/Private/Toolkits/GlobalEditorCommonCommands.cpp
+++ b/Engine/Source/Editor/UnrealEd/Private/Toolkits/GlobalEditorCommonCommands.cpp
@@ -3,21 +3,23 @@
 #include "Toolkits/GlobalEditorCommonCommands.h"
 
 #include "FindInBlueprints.h"
-#include "Modules/ModuleManager.h"
-#include "Widgets/DeclarativeSyntaxSupport.h"
-#include "Widgets/SWindow.h"
-#include "Layout/WidgetPath.h"
 #include "Framework/Application/MenuStack.h"
 #include "Framework/Application/SlateApplication.h"
 #include "Framework/MultiBox/MultiBoxBuilder.h"
-#include "Styling/AppStyle.h"
-#include "OutputLogModule.h"
+#include "Interfaces/IMainFrameModule.h"
+#include "ISettingsModule.h"
+#include "Layout/WidgetPath.h"
 #include "LevelEditor.h"
-#include "Widgets/Docking/SDockTab.h"
+#include "LevelEditorActions.h"
+#include "Modules/ModuleManager.h"
+#include "OutputLogModule.h"
+#include "StatusBarSubsystem.h"
+#include "Styling/AppStyle.h"
 #include "Toolkits/SGlobalOpenAssetDialog.h"
 #include "Toolkits/SGlobalTabSwitchingDialog.h"
-#include "StatusBarSubsystem.h"
-#include "LevelEditorActions.h"
+#include "Widgets/DeclarativeSyntaxSupport.h"
+#include "Widgets/SWindow.h"
+#include "Widgets/Docking/SDockTab.h"
 
 #define LOCTEXT_NAMESPACE "GlobalEditorCommonCommands"
 
@@ -70,56 +72,98 @@ void FGlobalEditorCommonCommands::RegisterCommands()
 #endif
 
 	UI_COMMAND(OpenFindInAllBlueprints, "Open Find in all Blueprints", "Opens the Find in all Blueprints dialogue", EUserInterfaceActionType::Button, FInputChord());
+
+#if PLATFORM_MAC
+	EModifierKey::Type ModifierKey = EModifierKey::Command;
+#else
+	EModifierKey::Type ModifierKey = EModifierKey::Control;
+#endif
+
+	UI_COMMAND(OpenEditorPreferences, "Editor Preferences...", "Configure the behavior and features of the Unreal Editor.", EUserInterfaceActionType::Button, FInputChord(EModifierKey::Alt | ModifierKey, EKeys::E));
+	UI_COMMAND(OpenKeyboardShortcuts, "Keyboard Shortcuts...", "Configure the keyboard shortcuts of the Unreal Editor.", EUserInterfaceActionType::Button, FInputChord(EModifierKey::Alt | ModifierKey, EKeys::K));
+	UI_COMMAND(OpenProjectSettings, "Project Settings...", "Change the settings of the currently loaded project.", EUserInterfaceActionType::Button, FInputChord(EModifierKey::Alt | ModifierKey, EKeys::P));
 }
 
 void FGlobalEditorCommonCommands::MapActions(TSharedRef<FUICommandList>& ToolkitCommands)
 {
 	Register();
+	
+	FGlobalEditorCommonCommands& GlobalEditorCommands = Get();
 
 	ToolkitCommands->MapAction(
-		Get().SummonControlTabNavigation,
-		FExecuteAction::CreateStatic(&FGlobalEditorCommonCommands::OnPressedCtrlTab, Get().SummonControlTabNavigation));
+		GlobalEditorCommands.SummonControlTabNavigation,
+		FExecuteAction::CreateStatic(&FGlobalEditorCommonCommands::OnPressedCtrlTab, GlobalEditorCommands.SummonControlTabNavigation));
 
 	ToolkitCommands->MapAction(
-		Get().SummonControlTabNavigationAlternate,
-		FExecuteAction::CreateStatic(&FGlobalEditorCommonCommands::OnPressedCtrlTab, Get().SummonControlTabNavigationAlternate));
+		GlobalEditorCommands.SummonControlTabNavigationAlternate,
+		FExecuteAction::CreateStatic(&FGlobalEditorCommonCommands::OnPressedCtrlTab, GlobalEditorCommands.SummonControlTabNavigationAlternate));
 
 	ToolkitCommands->MapAction(
-		Get().SummonControlTabNavigationBackwards,
-		FExecuteAction::CreateStatic(&FGlobalEditorCommonCommands::OnPressedCtrlTab, Get().SummonControlTabNavigationBackwards));
+		GlobalEditorCommands.SummonControlTabNavigationBackwards,
+		FExecuteAction::CreateStatic(&FGlobalEditorCommonCommands::OnPressedCtrlTab, GlobalEditorCommands.SummonControlTabNavigationBackwards));
 
 	ToolkitCommands->MapAction(
-		Get().SummonControlTabNavigationBackwardsAlternate,
-		FExecuteAction::CreateStatic(&FGlobalEditorCommonCommands::OnPressedCtrlTab, Get().SummonControlTabNavigationBackwardsAlternate));
+		GlobalEditorCommands.SummonControlTabNavigationBackwardsAlternate,
+		FExecuteAction::CreateStatic(&FGlobalEditorCommonCommands::OnPressedCtrlTab, GlobalEditorCommands.SummonControlTabNavigationBackwardsAlternate));
 
 	ToolkitCommands->MapAction(
-		Get().SummonOpenAssetDialog,
+		GlobalEditorCommands.SummonOpenAssetDialog,
 		FExecuteAction::CreateStatic(&FGlobalEditorCommonCommands::OnSummonedAssetPicker));
 
 	ToolkitCommands->MapAction(
-		Get().SummonOpenAssetDialogAlternate,
+		GlobalEditorCommands.SummonOpenAssetDialogAlternate,
 		FExecuteAction::CreateStatic(&FGlobalEditorCommonCommands::OnSummonedAssetPicker));
 
 	ToolkitCommands->MapAction(
-		Get().OpenConsoleCommandBox,
+		GlobalEditorCommands.OpenConsoleCommandBox,
 		FExecuteAction::CreateStatic(&FGlobalEditorCommonCommands::OnSummonedConsoleCommandBox));
 
 	ToolkitCommands->MapAction(
-		Get().OpenContentBrowserDrawer,
+		GlobalEditorCommands.OpenContentBrowserDrawer,
 		FExecuteAction::CreateStatic(&FGlobalEditorCommonCommands::OnOpenContentBrowserDrawer));
 
 	ToolkitCommands->MapAction(
-		Get().OpenOutputLogDrawer,
+		GlobalEditorCommands.OpenOutputLogDrawer,
 		FExecuteAction::CreateStatic(&FGlobalEditorCommonCommands::OnOpenOutputLogDrawer));
 
 	ToolkitCommands->MapAction(
-		Get().OpenLevel,
+		GlobalEditorCommands.OpenLevel,
 		FExecuteAction::CreateStatic(&FLevelEditorActionCallbacks::OpenLevel),
 		FCanExecuteAction::CreateStatic(&FLevelEditorActionCallbacks::OpenLevel_CanExecute));
 
 	ToolkitCommands->MapAction(
-		Get().OpenFindInAllBlueprints,
+		GlobalEditorCommands.OpenFindInAllBlueprints,
 		FExecuteAction::CreateStatic(&FGlobalEditorCommonCommands::OnOpenFindInAllBlueprints));
+	
+	// Get editor default category and section
+	FName CategoryName;
+	FName SectionName;
+	if (IMainFrameModule* MainFrameModule = FModuleManager::GetModulePtr<IMainFrameModule>("MainFrame"))
+	{
+		MainFrameModule->GetEditorSettingsDefaultSelectionOverride(CategoryName, SectionName);
+		
+		if (CategoryName.IsNone())
+		{
+			CategoryName = FName("General");
+		}
+		
+		if (SectionName.IsNone())
+		{
+			SectionName = FName("Appearance");
+		}
+	}
+
+	ToolkitCommands->MapAction(
+		GlobalEditorCommands.OpenEditorPreferences,
+		FExecuteAction::CreateStatic(&FGlobalEditorCommonCommands::OpenSettings, FName("Editor"), CategoryName, SectionName));
+	
+	ToolkitCommands->MapAction(
+		GlobalEditorCommands.OpenKeyboardShortcuts,
+		FExecuteAction::CreateStatic(&FGlobalEditorCommonCommands::OpenSettings, FName("Editor"), FName("General"), FName("InputBindings")));
+	
+	ToolkitCommands->MapAction(
+		GlobalEditorCommands.OpenProjectSettings,
+		FExecuteAction::CreateStatic(&FGlobalEditorCommonCommands::OpenSettings, FName("Project"), FName("Project"), FName("General")));
 }
 
 void FGlobalEditorCommonCommands::OnPressedCtrlTab(TSharedPtr<FUICommandInfo> TriggeringCommand)
@@ -196,6 +240,11 @@ TSharedPtr<IMenu> FGlobalEditorCommonCommands::OpenPopupMenu(TSharedRef<SWidget>
 	return TSharedPtr<IMenu>();
 }
 
+void FGlobalEditorCommonCommands::OpenSettings(FName InContainerName, FName InCategoryName, FName InSectionName)
+{
+	FModuleManager::LoadModuleChecked<ISettingsModule>("Settings").ShowViewer(InContainerName, InCategoryName, InSectionName);
+}
+
 static void CloseDebugConsole()
 {
 	FOutputLogModule& OutputLogModule = FModuleManager::LoadModuleChecked< FOutputLogModule >(TEXT("OutputLog"));
diff --git a/Engine/Source/Editor/UnrealEd/Public/Toolkits/GlobalEditorCommonCommands.h b/Engine/Source/Editor/UnrealEd/Public/Toolkits/GlobalEditorCommonCommands.h
--- a/Engine/Source/Editor/UnrealEd/Public/Toolkits/GlobalEditorCommonCommands.h
+++ b/Engine/Source/Editor/UnrealEd/Public/Toolkits/GlobalEditorCommonCommands.h
@@ -38,6 +38,9 @@ protected:
 	static UE_API TSharedRef<SDockTab> SpawnAssetPicker(const FSpawnTabArgs& InArgs);
 
 	static UE_API TSharedPtr<IMenu> OpenPopupMenu(TSharedRef<SWidget> WindowContents, const FVector2D& PopupDesiredSize);
+	
+	static void OpenSettings(FName InContainerName, FName InCategoryName, FName InSectionName);
+	
 public:
 	TSharedPtr<FUICommandInfo> FindInContentBrowser;
 	TSharedPtr<FUICommandInfo> SummonControlTabNavigation;
@@ -55,6 +58,11 @@ public:
 
 	/** Level file commands */
 	TSharedPtr<FUICommandInfo> OpenLevel;
+	
+	/** Edit/Configuration commands */
+	TSharedPtr<FUICommandInfo> OpenEditorPreferences;
+	TSharedPtr<FUICommandInfo> OpenKeyboardShortcuts;
+	TSharedPtr<FUICommandInfo> OpenProjectSettings;
 };
 
 #undef UE_API
