[EntityFramework] Entities now render selected if their prefab root is selected. Introduced more SceneGraphUtil helpers.

#Synced-CL 49208729

diff --git a/Engine/Source/Editor/UnrealEd/Private/Elements/Actor/ActorElementEditorSelectionInterface.cpp b/Engine/Source/Editor/UnrealEd/Private/Elements/Actor/ActorElementEditorSelectionInterface.cpp
--- a/Engine/Source/Editor/UnrealEd/Private/Elements/Actor/ActorElementEditorSelectionInterface.cpp
+++ b/Engine/Source/Editor/UnrealEd/Private/Elements/Actor/ActorElementEditorSelectionInterface.cpp
@@ -63,14 +63,7 @@ bool UActorElementEditorSelectionInterface::SelectElement(const FTypedElementHan
 		{
 			if (ICoreProvider* DataStorage = GetMutableDataStorageFeature<ICoreProvider>(StorageFeatureName))
 			{
-				DataStorage->AddColumn<FTypedElementSelectionColumn>(Row);
-
-				// TODO: If the column add is deferred because the row has not been assigned yet, we do not have a way to add the selection set currently
-				// as GetColumn will return null here
-				if(FTypedElementSelectionColumn* SelectionColumn = DataStorage->GetColumn<FTypedElementSelectionColumn>(Row))
-				{
-					SelectionColumn->SelectionSet = InSelectionOptions.GetNameForTEDSIntegration();
-				}
+				DataStorage->AddColumn(Row, FTypedElementSelectionColumn{ .SelectionSet = InSelectionOptions.GetNameForTEDSIntegration() });
 			}
 		}
 	}
