UE: Expose "Soft" property to asset reference filter so filters can handle soft references different than hard references

#Synced-CL 49565632

diff --git a/Engine/Source/Editor/UnrealEd/Classes/Editor/EditorEngine.h b/Engine/Source/Editor/UnrealEd/Classes/Editor/EditorEngine.h
--- a/Engine/Source/Editor/UnrealEd/Classes/Editor/EditorEngine.h
+++ b/Engine/Source/Editor/UnrealEd/Classes/Editor/EditorEngine.h
@@ -318,7 +318,8 @@ private:
 enum class EAssetReferenceFilterProperties : uint8
 {
 	None = 0,
-	EditorOnly = 1
+	EditorOnly = 1,
+	Soft = 2
 };
 ENUM_CLASS_FLAGS(EAssetReferenceFilterProperties);
 
diff --git a/Engine/Source/Editor/UnrealEd/Private/EditorEngine.cpp b/Engine/Source/Editor/UnrealEd/Private/EditorEngine.cpp
--- a/Engine/Source/Editor/UnrealEd/Private/EditorEngine.cpp
+++ b/Engine/Source/Editor/UnrealEd/Private/EditorEngine.cpp
@@ -527,6 +527,12 @@ void FAssetReferenceFilterContext::AddReferencingAssetsFromPropertyHandle(const
 
 	EAssetReferenceFilterProperties ReferencerProperties = EAssetReferenceFilterProperties::None;
 
+	// Determine if this is a soft object reference property
+	if (const FSoftObjectProperty* SoftProperty = CastField<FSoftObjectProperty>(PropertyHandle->GetProperty()))
+	{
+		ReferencerProperties |= EAssetReferenceFilterProperties::Soft;
+	}
+
 	// Determine if this is an editor only property by following the parent property chain and checking if any are editor only
 	{
 		TSharedPtr<IPropertyHandle> CurrentPropertyHandle = PropertyHandle;
