DerivedDataCacheCommandlet: Add commandline argument to process primary assets (and the secondary assets they reference) rather than processing all assets in the depot. This is fairly close to the list of assets cooked by the cook commandlet when building the project.
#rnx
#Synced-CL 49551581

diff --git a/Engine/Source/Editor/UnrealEd/Private/Commandlets/DerivedDataCacheCommandlet.cpp b/Engine/Source/Editor/UnrealEd/Private/Commandlets/DerivedDataCacheCommandlet.cpp
--- a/Engine/Source/Editor/UnrealEd/Private/Commandlets/DerivedDataCacheCommandlet.cpp
+++ b/Engine/Source/Editor/UnrealEd/Private/Commandlets/DerivedDataCacheCommandlet.cpp
@@ -16,6 +16,7 @@ DerivedDataCacheCommandlet.cpp: Commandlet for DDC maintenence
 #include "DistanceFieldAtlas.h"
 #include "Editor.h"
 #include "EditorWorldUtils.h"
+#include "Engine/AssetManager.h"
 #include "Engine/Texture.h"
 #include "GlobalShader.h"
 #include "ICollectionContainer.h"
@@ -582,6 +583,31 @@ int32 UDerivedDataCacheCommandlet::Main( const FString& Params )
 			}
 		}
 
+		// check to see if we should use the primary asset as the starting set
+		if (Switches.Contains("USEPRIMARYASSETS"))
+		{
+			// SearchAllAssets is required by ModifyCook
+			IAssetRegistry& AssetRegistry = FModuleManager::Get().LoadModuleChecked<FAssetRegistryModule>(AssetRegistryConstants::ModuleName).Get();
+			AssetRegistry.SearchAllAssets(true /* bSynchronousSearch */);
+
+			// Call ModifyCook with one TargetPlatform at a time and take the union of PrimaryAssets. If we have
+			// multiple target platforms and they disagree about development objects, ModifyCook produces an error
+			// because the cook output is invalid. We're not producing cooked output, so we don't care about that, so
+			// call them one at a time to avoid the error.
+			TSet<FName> PrimaryAssetNames;
+			for (const ITargetPlatform* TargetPlatform : Platforms)
+			{
+				TArray<FName> PlatformPrimaryAssets;
+				TArray<FName> Unused;
+				UAssetManager::Get().ModifyCook({ TargetPlatform }, PlatformPrimaryAssets, Unused);
+				PrimaryAssetNames.Append(PlatformPrimaryAssets);
+			}
+			for (FName PrimaryAsset : PrimaryAssetNames)
+			{
+				CommandLinePackageNames.Add(PrimaryAsset.ToString());
+			}
+		}
+
 		// Add defaults if we haven't specifically found anything on the command line 
 		if (Tokens.IsEmpty() && CommandLinePackageNames.IsEmpty())
 		{
@@ -619,7 +645,7 @@ int32 UDerivedDataCacheCommandlet::Main( const FString& Params )
 		for ( int32 TokenIndex = 0; TokenIndex < Tokens.Num(); TokenIndex++ )
 		{
 			TokenFiles.Reset();
-			if ( !NormalizePackageNames( Unused, TokenFiles, Tokens[TokenIndex], PackageFilter) )
+			if (!NormalizePackageNames(Unused, TokenFiles, Tokens[TokenIndex], PackageFilter))
 			{
 				UE_LOG(LogDerivedDataCacheCommandlet, Display, TEXT("No packages found for parameter %i: '%s'"), TokenIndex, *Tokens[TokenIndex]);
 				continue;
