[HLOD] Fix LOCTEXT macro text conflict

#Synced-CL 49640739

diff --git a/Engine/Source/Editor/UnrealEd/Private/WorldPartition/WorldPartitionHLODsBuilder.cpp b/Engine/Source/Editor/UnrealEd/Private/WorldPartition/WorldPartitionHLODsBuilder.cpp
--- a/Engine/Source/Editor/UnrealEd/Private/WorldPartition/WorldPartitionHLODsBuilder.cpp
+++ b/Engine/Source/Editor/UnrealEd/Private/WorldPartition/WorldPartitionHLODsBuilder.cpp
@@ -823,7 +823,7 @@ bool UWorldPartitionHLODsBuilder::BuildHLODActors()
 				}
 				check(HLODActor);
 
-				SlowTask.EnterProgressFrame(1.f, FText::Format(LOCTEXT("BuildingHLODActor", "Building HLOD actor: {0}"), FText::FromString(HLODActor->GetActorLabel())));
+				SlowTask.EnterProgressFrame(1.f, FText::Format(LOCTEXT("BuildingHLODActorWithLabel", "Building HLOD actor: {0}"), FText::FromString(HLODActor->GetActorLabel())));
 
 				HLODActorPackage = HLODActor->GetPackage();
 
