[HomeScreen] Remove HomeScreen from the tab switching dialog and remove the tab flash when switching

#Synced-CL 49539534

diff --git a/Engine/Source/Editor/MainFrame/Private/MainFrameModule.cpp b/Engine/Source/Editor/MainFrame/Private/MainFrameModule.cpp
--- a/Engine/Source/Editor/MainFrame/Private/MainFrameModule.cpp
+++ b/Engine/Source/Editor/MainFrame/Private/MainFrameModule.cpp
@@ -1142,8 +1142,8 @@ TSharedRef<SDockTab> FMainFrameModule::SpawnHomeScreenWindow(const FSpawnTabArgs
 {
 	TSharedRef<SDockTab> HomeScreenTab = SNew(SDockTab)
 		.TabRole(ETabRole::MajorTab)
-		// Assign a display name that will be used when switching tabs with the ctrl + tab feature
-		.Label(LOCTEXT("HomePanelDisplayName", "Home Panel"))
+		// Assign a display name that will be used when switching tabs with the ctrl + tab feature to remove this tab from the list
+		.Label(IMainFrameModule::GetHomeScreenTabLabel())
 		.IconColor(FStyleColors::White.GetSpecifiedColor())
 		.ContentPadding(FMargin(0));
 
diff --git a/Engine/Source/Editor/MainFrame/Public/Interfaces/IMainFrameModule.h b/Engine/Source/Editor/MainFrame/Public/Interfaces/IMainFrameModule.h
--- a/Engine/Source/Editor/MainFrame/Public/Interfaces/IMainFrameModule.h
+++ b/Engine/Source/Editor/MainFrame/Public/Interfaces/IMainFrameModule.h
@@ -280,6 +280,15 @@ public:
 	 */
 	virtual void UnregisterCreateProjectBrowserOverride(FDelegateHandle InHandle) = 0;
 
+	/**
+	 * Return the HomeScreen tab label
+	 */
+	static FText& GetHomeScreenTabLabel()
+	{
+		static FText HomeScreenTabLabel = NSLOCTEXT("MainFrameModule", "HomePanelDisplayName", "Home Panel");
+		return HomeScreenTabLabel;
+	}
+
 public:
 
 	/**
diff --git a/Engine/Source/Editor/UnrealEd/Private/Toolkits/SGlobalTabSwitchingDialog.cpp b/Engine/Source/Editor/UnrealEd/Private/Toolkits/SGlobalTabSwitchingDialog.cpp
--- a/Engine/Source/Editor/UnrealEd/Private/Toolkits/SGlobalTabSwitchingDialog.cpp
+++ b/Engine/Source/Editor/UnrealEd/Private/Toolkits/SGlobalTabSwitchingDialog.cpp
@@ -9,6 +9,7 @@
 #include "Framework/Docking/TabManager.h"
 #include "Framework/MultiBox/MultiBox.h"
 #include "Framework/MultiBox/MultiBoxBuilder.h"
+#include "Interfaces/IMainFrameModule.h"
 #include "LevelEditor.h"
 #include "Modules/ModuleManager.h"
 #include "Styling/AppStyle.h"
@@ -151,7 +152,6 @@ public:
 			{
 				// Don't use DrawAttention since it steals focus
 				DockTab->ActivateInParent(ETabActivationCause::SetDirectly);
-				DockTab->FlashTab();
 			}
 		}
 	}
@@ -206,7 +206,6 @@ public:
 			{
 				// Don't use DrawAttention since it steals focus
 				DockTab->ActivateInParent(ETabActivationCause::SetDirectly);
-				DockTab->FlashTab();
 			}
 		}
 	}
@@ -288,7 +287,6 @@ public:
 		{
 			// Don't use DrawAttention since it steals focus
 			DockTab->ActivateInParent(ETabActivationCause::SetDirectly);
-			DockTab->FlashTab();
 		}
 	}
 
@@ -502,15 +500,27 @@ void SGlobalTabSwitchingDialog::Construct(const FArguments& InArgs, FVector2D In
 				// In this mode, active tab is the first to be added
 				const int32 IndexActiveTab = TabsInfo.DockTabs.Find(TabsInfo.ActiveTab);
 
+				const FText HomeScreenTabName = IMainFrameModule::GetHomeScreenTabLabel();
+
 				// Add all right side tabs
 				for (int32 Index = IndexActiveTab; Index < TabsInfo.DockTabs.Num(); Index++)
 				{
+					// HomeScreen tab should not appear in the switching tab dialog
+					if (TabsInfo.DockTabs[Index]->GetTabLabel().EqualToCaseIgnored(HomeScreenTabName))
+					{
+						continue;
+					}
 					MainTabsListDataSource.Add(MakeShared<FTabSwitchingListItem_DockTab>(TabsInfo.DockTabs[Index]));
 				}
 
 				// Add all left side tabs
 				for (int32 Index = 0; Index < IndexActiveTab; Index++)
 				{
+					// HomeScreen tab should not appear in the switching tab dialog
+					if (TabsInfo.DockTabs[Index]->GetTabLabel().EqualToCaseIgnored(HomeScreenTabName))
+					{
+						continue;
+					}
 					MainTabsListDataSource.Add(MakeShared<FTabSwitchingListItem_DockTab>(TabsInfo.DockTabs[Index]));
 				}
 			}
