Make sure the map package name is known before building a WP navmesh in editor to be able to reload the map.
#Synced-CL 49590991

diff --git a/Engine/Source/Editor/UnrealEd/Private/EditorBuildUtils.cpp b/Engine/Source/Editor/UnrealEd/Private/EditorBuildUtils.cpp
--- a/Engine/Source/Editor/UnrealEd/Private/EditorBuildUtils.cpp
+++ b/Engine/Source/Editor/UnrealEd/Private/EditorBuildUtils.cpp
@@ -970,9 +970,19 @@ void FEditorBuildUtils::TriggerNavigationBuilder(UWorld*& InOutWorld, FName Id)
 
 		if (UE::EditorBuildUtils::bNavmeshAllowPartitionedBuildingFromEditor && InOutWorld->IsPartitionedWorld())
 		{
-			const FString& LongPackageName = GetNameSafe(InOutWorld->GetPackage());
-			WorldPartitionBuildNavigation(LongPackageName);
-			InOutWorld = GEditor->GetEditorWorldContext().World();
+			// Ask user to save dirty packages before getting the package name since it might change (for untitled worlds).
+			if (FEditorFileUtils::SaveDirtyPackages(
+					/*bPromptUserToSave=*/true,
+					/*bSaveMapPackages=*/true,
+					/*bSaveContentPackages=*/false,
+					/*bFastSave*/false,
+					/*bNotifyNoPackagesSaved=*/false,
+					/*bCanBeDeclined=*/false))
+			{
+				const FString& LongPackageName = GetNameSafe(InOutWorld->GetPackage());
+				WorldPartitionBuildNavigation(LongPackageName);
+				InOutWorld = GEditor->GetEditorWorldContext().World();
+			}
 		}
 		else
 		{
@@ -1021,11 +1031,13 @@ bool FEditorBuildUtils::WorldPartitionBuildNavigation(const FString& InLongPacka
 				DefaultBuildNavigationOptions->bVerbose ? TEXT("-Verbose") : TEXT(""),
 				DefaultBuildNavigationOptions->bCleanPackages ? TEXT("-CleanPackages") : TEXT(""));
 		
-		RunWorldPartitionBuilder(InLongPackageName,
+		const bool bSuccess = RunWorldPartitionBuilder(InLongPackageName,
 			LOCTEXT("WorldPartitionBuildNavigationProgress", "Building navigation..."),
 			LOCTEXT("WorldPartitionBuildNavigationCancelled", "Building navigation cancelled!"),
 			LOCTEXT("WorldPartitionBuildNavigationFailed", "Errors occured during the build process, please refer to the logs ('WPNavigationBuilderLog.txt')."),
 			Arguments);
+
+		return bSuccess;
 	}
 	
 	return false;
