Resolve Compute PSO hitches from dynamically loaded components.

UE::PipelineCacheUtilities::LoadStableKeysFile - Disable processing for the old csv format export because they were corrupting the paths for compute state objects.
Added a new function FString FCompactFullName::GetPackageNameOnly() const - This allows us to retrieve the proper package name while collecting compute PSOs.
Modification to SaveBinaryPipelineCacheFile - In the CreateAndDispatchWhenReady, modified to use new GetPackageNameOnly() function instead of the old ToStringPathOnly() function so that it can accurately locate the compute object.

#Synced-CL 49564828

diff --git a/Engine/Source/Editor/UnrealEd/Private/Commandlets/ShaderPipelineCacheToolsCommandlet.cpp b/Engine/Source/Editor/UnrealEd/Private/Commandlets/ShaderPipelineCacheToolsCommandlet.cpp
--- a/Engine/Source/Editor/UnrealEd/Private/Commandlets/ShaderPipelineCacheToolsCommandlet.cpp
+++ b/Engine/Source/Editor/UnrealEd/Private/Commandlets/ShaderPipelineCacheToolsCommandlet.cpp
@@ -2296,7 +2296,9 @@ int32 SaveBinaryPipelineCacheFile(const FString& OutputFilename, const EShaderPl
 			{
 				for (const TUniquePtr<FStableShaderKeyAndValue>& ShaderKeyPtr : StableMap)	// could be parallelized (skip first N*ThreadIdx iterations on each thread?)
 				{
-					FName PackageName(*ShaderKeyPtr->ClassNameAndObjectPath.ToStringPathOnly());
+					//Ensure we extract only the package name from ClassNameAndObjectPath, Compute PSOs may include subobject paths (e.g. “.../Core_0.GPUComputeScript”), which previously
+					//caused TSet<FName> Packages lookups to fail because they were using normal package paths (see ShadersInChunk)
+					FName PackageName(*ShaderKeyPtr->ClassNameAndObjectPath.GetPackageNameOnly());
 					StableNameMap.Add(PackageName, ShaderKeyPtr->OutputHash);
 				}
 			}, TStatId());
diff --git a/Engine/Source/Runtime/RenderCore/Private/PipelineCacheUtilities.cpp b/Engine/Source/Runtime/RenderCore/Private/PipelineCacheUtilities.cpp
--- a/Engine/Source/Runtime/RenderCore/Private/PipelineCacheUtilities.cpp
+++ b/Engine/Source/Runtime/RenderCore/Private/PipelineCacheUtilities.cpp
@@ -379,10 +379,6 @@ bool UE::PipelineCacheUtilities::LoadStableKeysFile(const FStringView& Filename,
 		}
 		Item.OutputHash = Hashes[static_cast<int32>(HashIdx)];
 
-		// Standardize on all CompactNames being parsed from string. This is a temporary hack until the names are parsed from CSV when reading StablePC
-		FString StringRep = Item.ClassNameAndObjectPath.ToString();
-		Item.ClassNameAndObjectPath.ParseFromString(StringRep);
-
 		Item.ComputeKeyHash();
 		InOutArray.Add(Item);
 	}
diff --git a/Engine/Source/Runtime/RenderCore/Private/ShaderLibrary/ShaderCodeLibrary.cpp b/Engine/Source/Runtime/RenderCore/Private/ShaderLibrary/ShaderCodeLibrary.cpp
--- a/Engine/Source/Runtime/RenderCore/Private/ShaderLibrary/ShaderCodeLibrary.cpp
+++ b/Engine/Source/Runtime/RenderCore/Private/ShaderLibrary/ShaderCodeLibrary.cpp
@@ -393,6 +393,11 @@ void FCompactFullName::SetCompactFullNameFromObject(UObject* InDepObject)
 		ObjectClassAndPath.Add(FName("null"));
 	}
 }
+
+FString FCompactFullName::GetPackageNameOnly() const
+{
+	return ObjectClassAndPath.Num() >= 2 ? ObjectClassAndPath[1].ToString() : ToStringPathOnly();
+}
 #endif
 
 uint32 GetTypeHash(const FCompactFullName& A)
@@ -2877,4 +2882,4 @@ FAutoConsoleCommandWithArgsAndOutputDevice GListShaderLibrariesCmd(
 			{
 				UE_LOG(LogShaderLibrary, Warning, TEXT("ShaderLibrariesCollection is not available."));
 			}
-		}));
\ No newline at end of file
+		}));
diff --git a/Engine/Source/Runtime/RenderCore/Public/ShaderCodeLibrary.h b/Engine/Source/Runtime/RenderCore/Public/ShaderCodeLibrary.h
--- a/Engine/Source/Runtime/RenderCore/Public/ShaderCodeLibrary.h
+++ b/Engine/Source/Runtime/RenderCore/Public/ShaderCodeLibrary.h
@@ -111,6 +111,7 @@ struct FCompactFullName
 #if WITH_EDITOR
 	/** Used to set up some compact FName paths for the FCompactFullName */
 	RENDERCORE_API void SetCompactFullNameFromObject(UObject* InDepObject);
+	RENDERCORE_API FString GetPackageNameOnly() const;
 #endif
 };
 
