- Go through all the imports of a package using the LinkerLoad and create an object redirector dependency for each import.
- In FRedirectionSummary, put the object redirects in their own list and not in the same list as core redirects.
- New function dependency to track object redirects.

Here is an example of why we need this : Plugin 1 from DLC1 has package A, plugin 2 from DLC2 has package B. Package A is an object redirector to package B. Package C from DLC2 has an import (a pointer) to an object in package B. When doing a multiproduct cook of package C, plugin 1 is not loaded because it is not needed for DLC2. The dependencies of package C will contain a redirector dependency for package B. This dependency will track any redirector where B is a source or target of a redirection. In this case there is no redirector so the hash will be equivalent to 0. Now if I run a regular cook, I have all the plugins loaded. So when I cook package C, the redirector dependency for package B will find the object redirect package A and use it to compute the hash. The dependencies are now different from the multiproduct cook.

The issues is the redirect dependency tracking package B is too wide and contains too many things. It contains object redirects where B is a source or target. In this case package C has a pointer to B. The fact package A points to B is irrelevant and should not be tracked.

The fix :
- The object redirectors are not part of the redirect dependencies anymore. Redirect dependencies are tracking a package whether it's the source or target of a redirection. As explained above I don't want that.
- I made a new function dependency to track object redirects. It tracks redirection for a source package only. In this case an object redirect dependency for package B would find no redirection because package B is not the source of any object redirect. This will give the same result in a regular cook and multiproduct cook.
- I need to create a function dependency for each import. The LinkerLoad can give me the list of imports as saved in the package (before any object redirector is applied). If the import has an object redirector, the dependency will track it and any change to it will trigger a recook. If there is no object redirect, a dependency will still be created so in the future, if an object redirector is created for the import, the dependency will change and trigger a recook.

#rnx
#Synced-CL 49609904

diff --git a/Engine/Source/Editor/UnrealEd/Private/Cooker/CookIncrementalVersion.cpp b/Engine/Source/Editor/UnrealEd/Private/Cooker/CookIncrementalVersion.cpp
--- a/Engine/Source/Editor/UnrealEd/Private/Cooker/CookIncrementalVersion.cpp
+++ b/Engine/Source/Editor/UnrealEd/Private/Cooker/CookIncrementalVersion.cpp
@@ -8,6 +8,6 @@ namespace UE::Cook
 {
 
 // Change CookIncrementalVersion to a new guid when all packages in an incremental cook need to be invalidated
-FGuid CookIncrementalVersion(0x75897AC8, 0xA1AB13BD, 0x8A60CB2A, 0xA21880FC);
+FGuid CookIncrementalVersion(0X48D958B3, 0X565444EE, 0X96DDFB64, 0X4F776273);
 
 }
diff --git a/Engine/Source/Editor/UnrealEd/Private/Cooker/CookPackageArtifacts.cpp b/Engine/Source/Editor/UnrealEd/Private/Cooker/CookPackageArtifacts.cpp
--- a/Engine/Source/Editor/UnrealEd/Private/Cooker/CookPackageArtifacts.cpp
+++ b/Engine/Source/Editor/UnrealEd/Private/Cooker/CookPackageArtifacts.cpp
@@ -12,6 +12,7 @@
 #include "Cooker/CookEvents.h"
 #include "Cooker/CookGenerationHelper.h"
 #include "Cooker/CookGlobalDependencies.h"
+#include "Cooker/ObjectRedirectCookDependency.h"
 #include "Cooker/PackageBuildDependencyTracker.h"
 #include "CookOnTheSide/CookLog.h"
 #include "CookPackageSplitter.h"
@@ -30,6 +31,7 @@
 #include "TargetDomain/TargetDomainUtils.h"
 #include "UObject/CoreRedirects.h"
 #include "UObject/ICookInfo.h"
+#include "UObject/LinkerLoad.h"
 
 namespace UE::Cook
 {
@@ -513,6 +515,68 @@ static void JoinMessagesIntoErrorReason(FStringBuilderBase& OutText,
 	}
 }
 
+/**
+ * Go through all the imports of a package using the LinkerLoad and create an object redirector dependency for each import.
+ * 
+ * Here is an example of why we need this : Plugin 1 from DLC1 has package A, plugin 2 from DLC2 has package B. Package A contains an object 
+ * redirector A.A -> B.B. Package C from DLC2 has an import (a pointer) to B.B.
+ * When doing a multiproduct cook of package C, plugin 1 is not loaded because it is not needed for DLC2. The dependencies of package C will
+ * contain a redirector dependency for package B. This dependency will track any redirector where B is a source or target of a redirection. In this
+ * case there is no redirector so the hash will be equivalent to 0.
+ * Now if I run a regular cook, I have all the plugins loaded. So when I cook package C, the redirector dependency for package B will find the object
+ * redirect package A and use it to compute the hash. The dependencies are now different from the multiproduct cook.
+ * 
+ */
+static void CollectObjectRedirectDependencies(const UPackage* Package, TArray<FCookDependency>& Dependencies)
+{
+	const FLinkerLoad* LinkerLoad = Package->GetLinker();
+	if (!LinkerLoad)
+	{
+		return;
+	}
+	
+	TSet<FName> UniqueImportNames;
+	for(int ArrayIndex = 0; ArrayIndex < LinkerLoad->ImportMap.Num(); ++ArrayIndex)
+	{
+		FPackageIndex Index = FPackageIndex::FromImport(ArrayIndex);
+
+		FString PathName;
+		LinkerLoad->BuildPathName(PathName, Index);
+
+		FSoftObjectPath ObjectPath(PathName);
+		FName ImportPackageName = ObjectPath.GetLongPackageFName();
+		
+		if (ImportPackageName.IsNone())
+		{
+			continue;
+		}
+
+		bool bAlreadyInSet = false;
+		UniqueImportNames.Add(ImportPackageName, &bAlreadyInSet);
+
+		if (bAlreadyInSet)
+		{
+			continue;
+		}
+
+		EPackageMountPoint MountPoint = GetPackageMountPoint(ImportPackageName);
+		switch (MountPoint)
+		{
+
+		case EPackageMountPoint::Content:
+			Dependencies.Add(UE::Cook::CreateObjectRedirectDependency(ImportPackageName));
+			break;
+
+
+		case EPackageMountPoint::Transient:
+		case EPackageMountPoint::Script:
+		case EPackageMountPoint::GeneratedContent:
+		default:
+			break;
+		}
+	}
+}
+
 bool FBuildDependencySet::TryCollectInternal(FBuildResultDependenciesMap& InOutResultDependencies,
 	TArray<FName>& InOutRuntimeDependencies, TArray<TPair<ELogVerbosity::Type, FString>>* OutMessages,
 	FName DefaultBuildResult, const UPackage* Package, const ITargetPlatform* TargetPlatform,
@@ -668,6 +732,8 @@ bool FBuildDependencySet::TryCollectInternal(FBuildResultDependenciesMap& InOutR
 		}
 	}
 
+	CollectObjectRedirectDependencies(Package, DefaultResultDependencies);
+
 	// Put the Dependencies we have collected onto the requested DefaultBuildResult
 	InOutResultDependencies.FindOrAdd(DefaultBuildResult).Append(MoveTemp(DefaultResultDependencies));
 
@@ -857,6 +923,7 @@ bool FBuildDependencySet::TryCollectInternal(FBuildResultDependenciesMap& InOutR
 		for (FName RedirectionTarget : RedirectionTargets)
 		{
 			ResultDependencies.Add(FCookDependency::RedirectionTarget(RedirectionTarget));
+			ResultDependencies.Add(UE::Cook::CreateObjectRedirectDependency(RedirectionTarget));
 		}
 	}
 	return true;
diff --git a/Engine/Source/Editor/UnrealEd/Private/Cooker/ObjectRedirectCookDependency.cpp b/Engine/Source/Editor/UnrealEd/Private/Cooker/ObjectRedirectCookDependency.cpp
new file mode 100644
--- /dev/null
+++ b/Engine/Source/Editor/UnrealEd/Private/Cooker/ObjectRedirectCookDependency.cpp
@@ -0,0 +1,35 @@
+// Copyright Epic Games, Inc. All Rights Reserved.
+
+#include "Cooker/ObjectRedirectCookDependency.h"
+
+#include "Cooker/CookDependency.h"
+#include "Cooker/CookDependencyContext.h"
+#include "Serialization/CompactBinaryWriter.h"
+#include "UObject/CoreRedirects.h"
+
+namespace UE::Cook
+{
+	UE_COOK_DEPENDENCY_FUNCTION(ObjectRedirectDependency, ValidateObjectRedirectDependency);
+
+	FCookDependency CreateObjectRedirectDependency(const FName& PackageName)
+	{
+		FCbWriter Writer;
+		Writer.AddString(PackageName.ToString());
+		
+		FCbFieldIterator Fields = Writer.Save();
+
+		return FCookDependency::Function(UE_COOK_DEPENDENCY_FUNCTION_CALL(ObjectRedirectDependency), MoveTemp(Fields));
+	}
+
+	void ValidateObjectRedirectDependency(FCbFieldViewIterator Args, FCookDependencyContext& Context)
+	{
+		FUtf8StringView PackageNameAsString = Args.AsString();
+		FName PackageName(PackageNameAsString);
+
+		FBlake3 Hasher;
+		FCoreRedirects::GetObjectRedirectHashAffectingPackage(PackageName, Hasher);
+
+		FBlake3Hash Hash = Hasher.Finalize();
+		Context.Update(&Hash, sizeof(Hash));
+	}
+}
\ No newline at end of file
diff --git a/Engine/Source/Editor/UnrealEd/Private/Cooker/ObjectRedirectCookDependency.h b/Engine/Source/Editor/UnrealEd/Private/Cooker/ObjectRedirectCookDependency.h
new file mode 100644
--- /dev/null
+++ b/Engine/Source/Editor/UnrealEd/Private/Cooker/ObjectRedirectCookDependency.h
@@ -0,0 +1,18 @@
+// Copyright Epic Games, Inc. All Rights Reserved.
+
+#pragma once
+
+#include "Serialization/CompactBinary.h"
+#include "UObject/NameTypes.h"
+
+namespace UE::Cook { class FCookDependency; }
+namespace UE::Cook { struct FCookDependencyContext; }
+
+namespace UE::Cook
+{
+	/** Create a function dependency tracking the object redirect for a given package. */
+	FCookDependency CreateObjectRedirectDependency(const FName& PackageName);
+
+	/** Deserialize the function dependency parameters and recompute the hash value. */
+	void ValidateObjectRedirectDependency(FCbFieldViewIterator Args, FCookDependencyContext& Context);
+}
diff --git a/Engine/Source/Runtime/CoreUObject/Internal/UObject/CoreRedirects/RedirectionSummary.h b/Engine/Source/Runtime/CoreUObject/Internal/UObject/CoreRedirects/RedirectionSummary.h
--- a/Engine/Source/Runtime/CoreUObject/Internal/UObject/CoreRedirects/RedirectionSummary.h
+++ b/Engine/Source/Runtime/CoreUObject/Internal/UObject/CoreRedirects/RedirectionSummary.h
@@ -27,9 +27,13 @@ public:
 	void Add(const FCoreRedirect& CoreRedirect);
 	void Remove(const FCoreRedirect& CoreRedirect);
 
+	void AddObjectRedirector(const FCoreRedirect& ObjectRedirect);
+	void RemoveObjectRedirector(const FCoreRedirect& ObjectRedirect);
+
 	UE_DEPRECATED(5.6, "Use GetHashAffectingPackages instead.")
 	void AppendHashAffectingPackages(FBlake3& Hasher, TConstArrayView<FName> PackageNames);
 	
+	void GetObjectRedirectHashAffectingPackage(const FName& PackageName, FBlake3& Hasher);
 	void GetHashAffectingPackages(const TConstArrayView<FName>& PackageNames, TArray<FBlake3Hash>& HashArray);
 	void AppendHashGlobal(FBlake3& Hasher);
 
@@ -65,9 +69,17 @@ private:
 	static TArray<FName, TInlineAllocator<2>> GetAffectedPackages(const FCoreRedirect& Redirect);
 
 private:
+	/** Core Redirects specific to a package.*/
 	TMap<FName, FRedirectContainer> RedirectsForPackage;
+
+	/** Global core redirects. Those can affect any package.*/
 	FRedirectContainer GlobalRedirects;
+
+	/** Prefixed core redirects that can affect multiple packages.*/
 	TSortedMap<FCoreRedirect, bool, FDefaultAllocator, FCompareRedirect> PrefixedRedirects;
+	
+	/** Object redirectors. The key is the source package name. Redirectors in this map are not in RedirectsForPackage. */
+	TMap<FName, FRedirectContainer> ObjectRedirectsForPackage;
 
 	/**
 	 * CoreRedirects are written when the engine is single threaded, but we do not write the hashes for the
diff --git a/Engine/Source/Runtime/CoreUObject/Private/UObject/CoreRedirects.cpp b/Engine/Source/Runtime/CoreUObject/Private/UObject/CoreRedirects.cpp
--- a/Engine/Source/Runtime/CoreUObject/Private/UObject/CoreRedirects.cpp
+++ b/Engine/Source/Runtime/CoreUObject/Private/UObject/CoreRedirects.cpp
@@ -2080,6 +2080,16 @@ void FCoreRedirects::AppendHashOfGlobalRedirects(FBlake3& Hasher)
 	}
 }
 
+void FCoreRedirects::GetObjectRedirectHashAffectingPackage(const FName& PackageName, FBlake3& Hasher)
+{
+	FCoreRedirectsContext& Context = FCoreRedirectsContext::GetThreadContext();
+	if (!!(Context.GetFlags() & FCoreRedirectsContext::EFlags::UseRedirectionSummary))
+	{
+		FRedirectionSummary& RedirectionSummary = Context.GetRedirectionSummary();
+		RedirectionSummary.GetObjectRedirectHashAffectingPackage(PackageName, Hasher);
+	}
+}
+
 void FCoreRedirects::RecordAddedObjectRedirector(const FSoftObjectPath& Source, const FSoftObjectPath& Dest)
 {
 	FCoreRedirectsContext& Context = FCoreRedirectsContext::GetThreadContext();
@@ -2087,7 +2097,7 @@ void FCoreRedirects::RecordAddedObjectRedirector(const FSoftObjectPath& Source,
 	{
 		FCoreRedirect ConvertedToCoreRedirect(ECoreRedirectFlags::Type_Object,
 			FCoreRedirectObjectName(Source), FCoreRedirectObjectName(Dest));
-		Context.GetRedirectionSummary().Add(ConvertedToCoreRedirect);
+		Context.GetRedirectionSummary().AddObjectRedirector(ConvertedToCoreRedirect);
 	}
 }
 
@@ -2098,7 +2108,7 @@ void FCoreRedirects::RecordRemovedObjectRedirector(const FSoftObjectPath& Source
 	{
 		FCoreRedirect ConvertedToCoreRedirect(ECoreRedirectFlags::Type_Object,
 			FCoreRedirectObjectName(Source), FCoreRedirectObjectName(Dest));
-		Context.GetRedirectionSummary().Remove(ConvertedToCoreRedirect);
+		Context.GetRedirectionSummary().RemoveObjectRedirector(ConvertedToCoreRedirect);
 	}
 }
 
diff --git a/Engine/Source/Runtime/CoreUObject/Private/UObject/RedirectionSummary.cpp b/Engine/Source/Runtime/CoreUObject/Private/UObject/RedirectionSummary.cpp
--- a/Engine/Source/Runtime/CoreUObject/Private/UObject/RedirectionSummary.cpp
+++ b/Engine/Source/Runtime/CoreUObject/Private/UObject/RedirectionSummary.cpp
@@ -96,6 +96,29 @@ void FRedirectionSummary::Remove(const FCoreRedirect& CoreRedirect)
 	}
 }
 
+void FRedirectionSummary::AddObjectRedirector(const FCoreRedirect& ObjectRedirect)
+{
+	FWriteScopeLock ScopeLock(Lock);
+	ObjectRedirectsForPackage.FindOrAdd(ObjectRedirect.OldName.PackageName).Add(FCoreRedirect(ObjectRedirect));
+}
+
+void FRedirectionSummary::RemoveObjectRedirector(const FCoreRedirect& ObjectRedirect)
+{
+	FWriteScopeLock ScopeLock(Lock);
+
+	FRedirectContainer* Container = ObjectRedirectsForPackage.Find(ObjectRedirect.OldName.PackageName);
+	if (!Container)
+	{
+		return;
+	}
+
+	Container->Remove(ObjectRedirect);
+	if (Container->IsEmpty())
+	{
+		ObjectRedirectsForPackage.Remove(ObjectRedirect.OldName.PackageName);
+	}
+}
+
 void FRedirectionSummary::AppendHashAffectingPackages(FBlake3& Hasher, TConstArrayView<FName> PackageNames)
 {
 	TArray<FBlake3Hash> Hashes;
@@ -113,6 +136,33 @@ void FRedirectionSummary::AppendHashAffectingPackages(FBlake3& Hasher, TConstArr
 	}
 }
 
+void FRedirectionSummary::GetObjectRedirectHashAffectingPackage(const FName& PackageName, FBlake3& Hasher)
+{
+	{
+		FReadScopeLock ScopeLock(Lock);
+
+		const FRedirectContainer* Container = ObjectRedirectsForPackage.Find(PackageName);
+		if (!Container)
+		{
+			return;
+		}
+
+		bool bSuccess = Container->TryAppendHashInReadLock(Hasher);
+		if (bSuccess)
+		{
+			return;
+		}
+	}
+
+	{
+		FWriteScopeLock WriteLock(Lock);
+		if (FRedirectContainer* Container = ObjectRedirectsForPackage.Find(PackageName))
+		{
+			Container->AppendHashInWriteLock(Hasher);
+		}
+	}
+}
+
 void FRedirectionSummary::GetHashAffectingPackages(const TConstArrayView<FName>& PackageNames, TArray<FBlake3Hash>& Hashes)
 {
 	check(Hashes.Num() == PackageNames.Num());
diff --git a/Engine/Source/Runtime/CoreUObject/Public/UObject/CoreRedirects.h b/Engine/Source/Runtime/CoreUObject/Public/UObject/CoreRedirects.h
--- a/Engine/Source/Runtime/CoreUObject/Public/UObject/CoreRedirects.h
+++ b/Engine/Source/Runtime/CoreUObject/Public/UObject/CoreRedirects.h
@@ -409,7 +409,7 @@ struct FCoreRedirects
 
 	/**
 	 * For each package in PackageNames, compute the hash of all redirects that affect the package, either
-	 * redirecting from or to the package. It doesn't include the GlobalRedirects.
+	 * redirecting from or to the package. It doesn't include the GlobalRedirects and doesn't include the object redirects.
 	 * Used in Incremental cooking to invalidate the cooked version of packages when CoreRedirects change.
 	 */
 	static COREUOBJECT_API void GetHashOfRedirectsAffectingPackages(const TConstArrayView<FName>& PackageNames, TArray<FBlake3Hash>& Hashes);
@@ -420,9 +420,13 @@ struct FCoreRedirects
 	 */
 	static COREUOBJECT_API void AppendHashOfGlobalRedirects(FBlake3& Hasher);
 
-	/** Add the given Source->Path redirector to the summary used for AppendHashOfRedirectsAffectingPackages. */
+	/** Calculate the hash of the object redirectors affecting a package, if any. */
+	static COREUOBJECT_API void GetObjectRedirectHashAffectingPackage(const FName& PackageName, FBlake3& Hasher);
+
+	/** Add the given Source->Path redirector to the summary used for GetObjectRedirectHashAffectingPackage. */
 	static COREUOBJECT_API void RecordAddedObjectRedirector(const FSoftObjectPath& Source, const FSoftObjectPath& Dest);
-	/** Remove the given Source->Path redirector to the summary used for AppendHashOfRedirectsAffectingPackages. */
+
+	/** Remove the given Source->Path redirector to the summary used for GetObjectRedirectHashAffectingPackage. */
 	static COREUOBJECT_API void RecordRemovedObjectRedirector(const FSoftObjectPath& Source, const FSoftObjectPath& Dest);
 
 	/** Thread-safe delegate broadcasted when a global core redirect is added to FRedirectionSummary. It can be called from any thread. */
diff --git a/Engine/Source/Runtime/CoreUObject/Public/UObject/LinkerLoad.h b/Engine/Source/Runtime/CoreUObject/Public/UObject/LinkerLoad.h
--- a/Engine/Source/Runtime/CoreUObject/Public/UObject/LinkerLoad.h
+++ b/Engine/Source/Runtime/CoreUObject/Public/UObject/LinkerLoad.h
@@ -753,7 +753,7 @@ public:
 	 * @param OutPathName		[out] Will contain the full path for the resource
 	 * @param ResourceIndex		Index of a resource in the export table
 	 */
-	void BuildPathName( FString& OutPathName, FPackageIndex ExportIndex ) const;
+	void COREUOBJECT_API BuildPathName( FString& OutPathName, FPackageIndex ExportIndex ) const;
 
 	/**
 	 * Checks if the specified export should be loaded or not.
