[UObject]
When renaming an object, we will detach objects from their linker (if any). As such, if the rename is being done temporarily, or we want to maintain linker references with the old name to the same object pointer, we provide REN_AllowPackageLinkerMismatch flag. This is similar to the REN_ForceNoResetLoaders flag however the meaning is less specific to the implementation as we no longer call ResetLoaders and the REN_AllowPackageLinkerMismatch is an explicit opt-in to breaking the linker's mapping from name to loaded object (for whatever purpose this user has)

- Renames all existing usage of REN_ForceNoResetLoaders to be REN_AllowPackageLinkerMismatch
- Adds the REN_AllowPackageLinkerMismatch to all callsites that originally passed in REN_ForceNoResetLoaders (removed in 34424025) with the exception of any callsite where we are explicitly renaming the object to a new package -- in those cases we always clear the linker unless REN_AllowPackageLinkerMismatch is provided. All cases where we are moving to a new package and want a linker mismatch _already_ provide the REN_AllowPackageLinkerMismatch flag (done in 34304926)

#rnx
#Synced-CL 49684429

diff --git a/Engine/Plugins/2D/Paper2D/Source/Paper2D/Private/PaperTileMapComponent.cpp b/Engine/Plugins/2D/Paper2D/Source/Paper2D/Private/PaperTileMapComponent.cpp
--- a/Engine/Plugins/2D/Paper2D/Source/Paper2D/Private/PaperTileMapComponent.cpp
+++ b/Engine/Plugins/2D/Paper2D/Source/Paper2D/Private/PaperTileMapComponent.cpp
@@ -142,7 +142,7 @@ void UPaperTileMapComponent::PostLoad()
 		// Convert the layers
 		for (UPaperTileLayer* Layer : TileMap->TileLayers)
 		{
-			Layer->Rename(nullptr, TileMap, REN_DontCreateRedirectors);
+			Layer->Rename(nullptr, TileMap, REN_DontCreateRedirectors | REN_AllowPackageLinkerMismatch);
 			Layer->ConvertToTileSetPerCell();
 		}
 
diff --git a/Engine/Plugins/Animation/ControlRig/Source/ControlRig/Private/Sequencer/MovieSceneControlRigParameterSection.cpp b/Engine/Plugins/Animation/ControlRig/Source/ControlRig/Private/Sequencer/MovieSceneControlRigParameterSection.cpp
--- a/Engine/Plugins/Animation/ControlRig/Source/ControlRig/Private/Sequencer/MovieSceneControlRigParameterSection.cpp
+++ b/Engine/Plugins/Animation/ControlRig/Source/ControlRig/Private/Sequencer/MovieSceneControlRigParameterSection.cpp
@@ -1898,7 +1898,7 @@ void UMovieSceneControlRigParameterSection::AddConstraintChannel(UTickableConstr
 		
 		//make copy that we can spawn if it doesn't exist
 		//the rename changes the outer to this section (from any actor manager)
-		InConstraint->Rename(nullptr, this, REN_DontCreateRedirectors);
+		InConstraint->Rename(nullptr, this, REN_DontCreateRedirectors | REN_AllowPackageLinkerMismatch);
 		
 		if (OnConstraintChannelAdded.IsBound())
 		{
diff --git a/Engine/Plugins/Animation/DeformerGraph/Source/OptimusCore/Private/OptimusHelpers.cpp b/Engine/Plugins/Animation/DeformerGraph/Source/OptimusCore/Private/OptimusHelpers.cpp
--- a/Engine/Plugins/Animation/DeformerGraph/Source/OptimusCore/Private/OptimusHelpers.cpp
+++ b/Engine/Plugins/Animation/DeformerGraph/Source/OptimusCore/Private/OptimusHelpers.cpp
@@ -111,7 +111,7 @@ const void Optimus::ConvertFTransformToFMatrix3x4(const FTransform& InTransform,
 
 bool Optimus::RenameObject(UObject* InObjectToRename, const TCHAR* InNewName, UObject* InNewOuter)
 {
-	return InObjectToRename->Rename(InNewName, InNewOuter, REN_DoNotDirty | REN_DontCreateRedirectors | REN_NonTransactional);
+	return InObjectToRename->Rename(InNewName, InNewOuter, REN_DoNotDirty | REN_DontCreateRedirectors | REN_NonTransactional | REN_AllowPackageLinkerMismatch);
 }
 
 void Optimus::RemoveObject(UObject* InObjectToRemove)
diff --git a/Engine/Plugins/ChaosVD/Source/ChaosVD/Private/ChaosVDObjectPool.h b/Engine/Plugins/ChaosVD/Source/ChaosVD/Private/ChaosVDObjectPool.h
--- a/Engine/Plugins/ChaosVD/Source/ChaosVD/Private/ChaosVDObjectPool.h
+++ b/Engine/Plugins/ChaosVD/Source/ChaosVD/Private/ChaosVDObjectPool.h
@@ -151,7 +151,7 @@ ObjectType* TChaosVDObjectPool<ObjectType>::AcquireObject(UObject* Outer, FName
 		if (Object)
 		{
 			const FName NewName = MakeUniqueObjectName(Outer, ObjectType::StaticClass(), Name);
-			Object->Rename(*NewName.ToString() , Outer, REN_NonTransactional | REN_DoNotDirty | REN_SkipGeneratedClasses | REN_DontCreateRedirectors);
+			Object->Rename(*NewName.ToString() , Outer, REN_NonTransactional | REN_DoNotDirty | REN_SkipGeneratedClasses | REN_DontCreateRedirectors | REN_AllowPackageLinkerMismatch);
 		
 			if (IChaosVDPooledObject* AsPooledObject = Cast<IChaosVDPooledObject>(Object))
 			{
diff --git a/Engine/Plugins/Editor/WorldPartitionHLODUtilities/Source/Private/WorldPartition/HLOD/Builders/HLODBuilderMeshApproximate.cpp b/Engine/Plugins/Editor/WorldPartitionHLODUtilities/Source/Private/WorldPartition/HLOD/Builders/HLODBuilderMeshApproximate.cpp
--- a/Engine/Plugins/Editor/WorldPartitionHLODUtilities/Source/Private/WorldPartition/HLOD/Builders/HLODBuilderMeshApproximate.cpp
+++ b/Engine/Plugins/Editor/WorldPartitionHLODUtilities/Source/Private/WorldPartition/HLOD/Builders/HLODBuilderMeshApproximate.cpp
@@ -137,7 +137,7 @@ TArray<UActorComponent*> UHLODBuilderMeshApproximate::Build(const FHLODBuildCont
 				UPackage* TempPackage = NewAsset->GetPackage();
 
 				// Rename the asset to its final destination
-				NewAsset->Rename(*AssetName, InHLODBuildContext.AssetsOuter, REN_DontCreateRedirectors | REN_NonTransactional);
+				NewAsset->Rename(*AssetName, InHLODBuildContext.AssetsOuter, REN_DontCreateRedirectors | REN_NonTransactional | REN_AllowPackageLinkerMismatch);
 				NewAsset->ClearFlags(RF_Public | RF_Standalone);
 				NewAsset->SetFlags(InHLODBuildContext.AssetsObjectFlags);
 
diff --git a/Engine/Plugins/Enterprise/DatasmithContent/Source/DatasmithContent/Private/ObjectTemplates/DatasmithObjectTemplate.cpp b/Engine/Plugins/Enterprise/DatasmithContent/Source/DatasmithContent/Private/ObjectTemplates/DatasmithObjectTemplate.cpp
--- a/Engine/Plugins/Enterprise/DatasmithContent/Source/DatasmithContent/Private/ObjectTemplates/DatasmithObjectTemplate.cpp
+++ b/Engine/Plugins/Enterprise/DatasmithContent/Source/DatasmithContent/Private/ObjectTemplates/DatasmithObjectTemplate.cpp
@@ -126,7 +126,7 @@ void FDatasmithObjectTemplateUtils::SetObjectTemplate(UObject* Outer, UDatasmith
 		{
 			// The outer chain is important for most of the engine functionality.
 			// If it's not set properly the deep copy of object won't work properly
-			ObjectTemplate->Rename(nullptr, UserData, REN_DontCreateRedirectors);
+			ObjectTemplate->Rename(nullptr, UserData, REN_DontCreateRedirectors | REN_AllowPackageLinkerMismatch);
 		}
 	}
 #endif // #if WITH_EDITORONLY_DATA
diff --git a/Engine/Plugins/Experimental/Dataflow/Source/DataflowEditor/Private/Dataflow/DataflowSimulationScene.cpp b/Engine/Plugins/Experimental/Dataflow/Source/DataflowEditor/Private/Dataflow/DataflowSimulationScene.cpp
--- a/Engine/Plugins/Experimental/Dataflow/Source/DataflowEditor/Private/Dataflow/DataflowSimulationScene.cpp
+++ b/Engine/Plugins/Experimental/Dataflow/Source/DataflowEditor/Private/Dataflow/DataflowSimulationScene.cpp
@@ -178,7 +178,7 @@ void FDataflowSimulationScene::ResetSimulationScene()
 			
 		// Since deletion can be delayed, rename to avoid future name collision
 		// Call UObject::Rename directly on actor to avoid AActor::Rename which unnecessarily sunregister and re-register components
-		PreviewActor->UObject::Rename(nullptr, GetTransientPackage(), REN_DontCreateRedirectors | REN_ForceNoResetLoaders);
+		PreviewActor->UObject::Rename(nullptr, GetTransientPackage(), REN_DontCreateRedirectors | REN_AllowPackageLinkerMismatch);
 	}
 	RemoveSceneObject(RootSceneActor);
 
diff --git a/Engine/Plugins/Experimental/UAF/UAF/Source/UAFUncookedOnly/Private/AnimNextRigVMAssetEditorData.cpp b/Engine/Plugins/Experimental/UAF/UAF/Source/UAFUncookedOnly/Private/AnimNextRigVMAssetEditorData.cpp
--- a/Engine/Plugins/Experimental/UAF/UAF/Source/UAFUncookedOnly/Private/AnimNextRigVMAssetEditorData.cpp
+++ b/Engine/Plugins/Experimental/UAF/UAF/Source/UAFUncookedOnly/Private/AnimNextRigVMAssetEditorData.cpp
@@ -1304,7 +1304,7 @@ void UAnimNextRigVMAssetEditorData::RemoveProgrammaticGraphs(TArrayView<URigVMGr
 	for(URigVMGraph* Graph : InGraphs)
 	{
 		VMClient->RemoveController(Graph);
-		Graph->Rename(nullptr, GetTransientPackage(), REN_ForceNoResetLoaders | REN_DoNotDirty | REN_DontCreateRedirectors | REN_NonTransactional);
+		Graph->Rename(nullptr, GetTransientPackage(), REN_AllowPackageLinkerMismatch | REN_DoNotDirty | REN_DontCreateRedirectors | REN_NonTransactional);
 	}
 }
 
@@ -1317,7 +1317,7 @@ void UAnimNextRigVMAssetEditorData::RemoveTransientGraphs(TArrayView<URigVMGraph
 		if(Graph->HasAnyFlags(RF_Transient))
 		{
 			VMClient->RemoveController(Graph);
-			Graph->Rename(nullptr, GetTransientPackage(), REN_ForceNoResetLoaders | REN_DoNotDirty | REN_DontCreateRedirectors | REN_NonTransactional);
+			Graph->Rename(nullptr, GetTransientPackage(), REN_AllowPackageLinkerMismatch | REN_DoNotDirty | REN_DontCreateRedirectors | REN_NonTransactional);
 		}
 	}
 }
@@ -2212,7 +2212,7 @@ bool UAnimNextRigVMAssetEditorData::RemoveEdGraphForCollapseNode(URigVMCollapseN
 						}
 
 						FunctionEdGraphs.Remove(RigFunctionGraph);
-						RigFunctionGraph->Rename(nullptr, GetTransientPackage(), REN_ForceNoResetLoaders | REN_DontCreateRedirectors);
+						RigFunctionGraph->Rename(nullptr, GetTransientPackage(), REN_AllowPackageLinkerMismatch | REN_DontCreateRedirectors);
 						RigFunctionGraph->MarkAsGarbage();
 						return bNotify;
 					}
@@ -2241,7 +2241,7 @@ bool UAnimNextRigVMAssetEditorData::RemoveEdGraphForCollapseNode(URigVMCollapseN
 						}
 
 						RigGraph->SubGraphs.Remove(SubRigGraph);
-						SubRigGraph->Rename(nullptr, GetTransientPackage(), REN_ForceNoResetLoaders | REN_DontCreateRedirectors);
+						SubRigGraph->Rename(nullptr, GetTransientPackage(), REN_AllowPackageLinkerMismatch | REN_DontCreateRedirectors);
 						SubRigGraph->MarkAsGarbage();
 						return bNotify;
 					}
diff --git a/Engine/Plugins/Experimental/UAF/UAF/Source/UAFUncookedOnly/Private/UncookedOnlyUtils.cpp b/Engine/Plugins/Experimental/UAF/UAF/Source/UAFUncookedOnly/Private/UncookedOnlyUtils.cpp
--- a/Engine/Plugins/Experimental/UAF/UAF/Source/UAFUncookedOnly/Private/UncookedOnlyUtils.cpp
+++ b/Engine/Plugins/Experimental/UAF/UAF/Source/UAFUncookedOnly/Private/UncookedOnlyUtils.cpp
@@ -1536,7 +1536,7 @@ void FUtils::OpenProgrammaticGraphs(UAnimNextRigVMAssetEditorData* EditorData, c
 			// Some explanation needed here!
 			// URigVMEdGraph caches its underlying model internally in GetModel depending on its outer if it is no attached to a RigVMClient
 			// So here we rename the graph into the transient package so we dont get any notifications
-			ProgrammaticGraph->Rename(nullptr, GetTransientPackage(), REN_ForceNoResetLoaders | REN_DoNotDirty | REN_DontCreateRedirectors | REN_NonTransactional);
+			ProgrammaticGraph->Rename(nullptr, GetTransientPackage(), REN_AllowPackageLinkerMismatch | REN_DoNotDirty | REN_DontCreateRedirectors | REN_NonTransactional);
 
 			// then create the graph (transient so it outers to the RigVMGraph)
 			URigVMEdGraph* EdGraph = CastChecked<URigVMEdGraph>(EditorData->CreateEdGraph(ProgrammaticGraph, true));
@@ -1546,7 +1546,7 @@ void FUtils::OpenProgrammaticGraphs(UAnimNextRigVMAssetEditorData* EditorData, c
 			Graphs.Add(EdGraph);
 
 			// Now rename into this asset again to be able to correctly create a controller (needed to view the graph and interact with it)
-			ProgrammaticGraph->Rename(nullptr, EditorData, REN_ForceNoResetLoaders | REN_DoNotDirty | REN_DontCreateRedirectors | REN_NonTransactional);
+			ProgrammaticGraph->Rename(nullptr, EditorData, REN_AllowPackageLinkerMismatch | REN_DoNotDirty | REN_DontCreateRedirectors | REN_NonTransactional);
 			URigVMController* ProgrammaticController = EditorData->GetOrCreateController(ProgrammaticGraph);
 
 			// Resend notifications to rebuild the EdGraph
diff --git a/Engine/Plugins/Experimental/Water/Source/Runtime/Private/WaterBodyActor.cpp b/Engine/Plugins/Experimental/Water/Source/Runtime/Private/WaterBodyActor.cpp
--- a/Engine/Plugins/Experimental/Water/Source/Runtime/Private/WaterBodyActor.cpp
+++ b/Engine/Plugins/Experimental/Water/Source/Runtime/Private/WaterBodyActor.cpp
@@ -438,7 +438,7 @@ PRAGMA_ENABLE_DEPRECATION_WARNINGS
 
 				OldLakeCollision->DestroyComponent();
 				// Rename it so we can use the name
-				OldLakeCollision->Rename(TEXT("LakeCollision_Old"), this, REN_DoNotDirty | REN_DontCreateRedirectors | REN_NonTransactional);
+				OldLakeCollision->Rename(TEXT("LakeCollision_Old"), this, REN_DoNotDirty | REN_DontCreateRedirectors | REN_NonTransactional | REN_AllowPackageLinkerMismatch);
 				It.RemoveCurrent();
 			}
 		}
@@ -576,7 +576,7 @@ PRAGMA_ENABLE_DEPRECATION_WARNINGS
 		{
 			WaterWaves->ClearFlags(RF_Public);
 			// At one point, WaterWaves's outer was the level. We need them to be outered by the water body : 
-			WaterWaves->Rename(nullptr, this, REN_DoNotDirty | REN_DontCreateRedirectors | REN_NonTransactional);
+			WaterWaves->Rename(nullptr, this, REN_DoNotDirty | REN_DontCreateRedirectors | REN_NonTransactional | REN_AllowPackageLinkerMismatch);
 		}
 	}
 	
diff --git a/Engine/Plugins/Importers/USDImporter/Source/USDSchemas/Private/Custom/MaterialXUSDShadeMaterialTranslator.cpp b/Engine/Plugins/Importers/USDImporter/Source/USDSchemas/Private/Custom/MaterialXUSDShadeMaterialTranslator.cpp
--- a/Engine/Plugins/Importers/USDImporter/Source/USDSchemas/Private/Custom/MaterialXUSDShadeMaterialTranslator.cpp
+++ b/Engine/Plugins/Importers/USDImporter/Source/USDSchemas/Private/Custom/MaterialXUSDShadeMaterialTranslator.cpp
@@ -233,7 +233,7 @@ namespace UE::USDMaterialXTranslator::Private
 							// Rename the UMaterialInterface into the target UPackage the asset cache created for us.
 							// SanitizedName will already match it.
 							const bool bRenamed = ImportedObject
-													  ->Rename(*SanitizedName.ToString(), Outer, REN_NonTransactional | REN_DontCreateRedirectors);
+													  ->Rename(*SanitizedName.ToString(), Outer, REN_NonTransactional | REN_DontCreateRedirectors | REN_AllowPackageLinkerMismatch);
 							ensure(bRenamed);
 
 							ImportedObject->ClearFlags(ImportedObject->GetFlags());
diff --git a/Engine/Plugins/Importers/USDImporter/Source/USDStageImporter/Private/USDStageImporter.cpp b/Engine/Plugins/Importers/USDImporter/Source/USDStageImporter/Private/USDStageImporter.cpp
--- a/Engine/Plugins/Importers/USDImporter/Source/USDStageImporter/Private/USDStageImporter.cpp
+++ b/Engine/Plugins/Importers/USDImporter/Source/USDStageImporter/Private/USDStageImporter.cpp
@@ -691,7 +691,7 @@ namespace UE::USDStageImporter::Private
 			// rename our asset to its target package too
 			else
 			{
-				ensure(Asset->Rename(*TargetAssetName, Package, REN_DontCreateRedirectors | REN_NonTransactional | REN_DoNotDirty));
+				ensure(Asset->Rename(*TargetAssetName, Package, REN_DontCreateRedirectors | REN_NonTransactional | REN_DoNotDirty | REN_AllowPackageLinkerMismatch));
 				MovedAsset = Asset;
 			}
 
@@ -709,7 +709,7 @@ namespace UE::USDStageImporter::Private
 		else
 		{
 			// We can't dirty the package here. Read the comment around MarkPackageDirty, below
-			ensure(Asset->Rename(*TargetAssetName, Package, REN_DontCreateRedirectors | REN_NonTransactional | REN_DoNotDirty));
+			ensure(Asset->Rename(*TargetAssetName, Package, REN_DontCreateRedirectors | REN_NonTransactional | REN_DoNotDirty | REN_AllowPackageLinkerMismatch));
 			MovedAsset = Asset;
 		}
 
diff --git a/Engine/Plugins/Importers/USDImporterMDL/Source/USDImporterMDL/Private/MDLUSDShadeMaterialTranslator.cpp b/Engine/Plugins/Importers/USDImporterMDL/Source/USDImporterMDL/Private/MDLUSDShadeMaterialTranslator.cpp
--- a/Engine/Plugins/Importers/USDImporterMDL/Source/USDImporterMDL/Private/MDLUSDShadeMaterialTranslator.cpp
+++ b/Engine/Plugins/Importers/USDImporterMDL/Source/USDImporterMDL/Private/MDLUSDShadeMaterialTranslator.cpp
@@ -198,7 +198,7 @@ void FMdlUsdShadeMaterialTranslator::CreateAssets()
 
 				// Rename the UMaterialInterface into the target UPackage the asset cache created for us.
 				// SanitizedName will already match it.
-				const bool bRenamed = ReferenceMaterial->Rename(*SanitizedName.ToString(), Outer, REN_NonTransactional | REN_DontCreateRedirectors);
+				const bool bRenamed = ReferenceMaterial->Rename(*SanitizedName.ToString(), Outer, REN_NonTransactional | REN_DontCreateRedirectors | REN_AllowPackageLinkerMismatch);
 				ensure(bRenamed);
 
 				// Let's not trust the flags the MDLImporter used and just use our own instead
@@ -305,7 +305,7 @@ void FMdlUsdShadeMaterialTranslator::CreateAssets()
 					{
 						// Rename the asset into the target UPackage the asset cache created for us.
 						// SanitizedName will already match it.
-						const bool bRenamed = Texture->Rename(*SanitizedName.ToString(), Outer, REN_NonTransactional | REN_DontCreateRedirectors);
+						const bool bRenamed = Texture->Rename(*SanitizedName.ToString(), Outer, REN_NonTransactional | REN_DontCreateRedirectors | REN_AllowPackageLinkerMismatch);
 						ensure(bRenamed);
 
 						// Let's not trust the flags the MDLImporter used and just use our own instead
diff --git a/Engine/Plugins/Media/MediaViewer/Source/MediaViewer/Private/Widgets/SMediaViewer.cpp b/Engine/Plugins/Media/MediaViewer/Source/MediaViewer/Private/Widgets/SMediaViewer.cpp
--- a/Engine/Plugins/Media/MediaViewer/Source/MediaViewer/Private/Widgets/SMediaViewer.cpp
+++ b/Engine/Plugins/Media/MediaViewer/Source/MediaViewer/Private/Widgets/SMediaViewer.cpp
@@ -1266,7 +1266,7 @@ void SMediaViewer::TakeSnapshot()
 		return;
 	}
 
-	Texture->Rename(*AssetName, Package, REN_DontCreateRedirectors | REN_DoNotDirty | REN_ForceNoResetLoaders | REN_NonTransactional);
+	Texture->Rename(*AssetName, Package, REN_DontCreateRedirectors | REN_DoNotDirty | REN_AllowPackageLinkerMismatch | REN_NonTransactional);
 
 	FAssetRegistryModule::AssetCreated(Package);
 
diff --git a/Engine/Plugins/PCG/Source/PCG/Private/Elements/PCGAddComponent.cpp b/Engine/Plugins/PCG/Source/PCG/Private/Elements/PCGAddComponent.cpp
--- a/Engine/Plugins/PCG/Source/PCG/Private/Elements/PCGAddComponent.cpp
+++ b/Engine/Plugins/PCG/Source/PCG/Private/Elements/PCGAddComponent.cpp
@@ -173,7 +173,7 @@ void UPCGAddComponentSettings::RefreshTemplateComponent()
 				Options.bNotifyObjectReplacement = true;
 				UEngine::CopyPropertiesForUnrelatedObjects(TemplateComponent, NewTemplateComponent, Options);
 
-				TemplateComponent->Rename(nullptr, GetTransientPackage(), REN_DontCreateRedirectors | REN_ForceNoResetLoaders);
+				TemplateComponent->Rename(nullptr, GetTransientPackage(), REN_DontCreateRedirectors | REN_AllowPackageLinkerMismatch);
 
 				TMap<UObject*, UObject*> OldToNew;
 				OldToNew.Emplace(TemplateComponent, NewTemplateComponent);
@@ -188,7 +188,7 @@ void UPCGAddComponentSettings::RefreshTemplateComponent()
 			TemplateComponent->Modify();
 
 			// Outer to this object
-			TemplateComponent->Rename(nullptr, this, REN_DoNotDirty | REN_DontCreateRedirectors | REN_ForceNoResetLoaders);
+			TemplateComponent->Rename(nullptr, this, REN_DoNotDirty | REN_DontCreateRedirectors | REN_AllowPackageLinkerMismatch);
 		}
 	}
 	else
diff --git a/Engine/Plugins/PCG/Source/PCG/Private/Elements/PCGCreateTargetActor.cpp b/Engine/Plugins/PCG/Source/PCG/Private/Elements/PCGCreateTargetActor.cpp
--- a/Engine/Plugins/PCG/Source/PCG/Private/Elements/PCGCreateTargetActor.cpp
+++ b/Engine/Plugins/PCG/Source/PCG/Private/Elements/PCGCreateTargetActor.cpp
@@ -254,7 +254,7 @@ void UPCGCreateTargetActor::RefreshTemplateActor()
 				Options.bNotifyObjectReplacement = true;
 				UEngine::CopyPropertiesForUnrelatedObjects(TemplateActor, NewTemplateActor, Options);
 
-				TemplateActor->Rename(nullptr, GetTransientPackage(), REN_DontCreateRedirectors);
+				TemplateActor->Rename(nullptr, GetTransientPackage(), REN_DontCreateRedirectors | REN_AllowPackageLinkerMismatch);
 
 				TMap<UObject*, UObject*> OldToNew;
 				OldToNew.Emplace(TemplateActor, NewTemplateActor);
@@ -269,7 +269,7 @@ void UPCGCreateTargetActor::RefreshTemplateActor()
 			TemplateActor->Modify();
 
 			// Outer to this object
-			TemplateActor->Rename(nullptr, this, REN_DoNotDirty | REN_DontCreateRedirectors);
+			TemplateActor->Rename(nullptr, this, REN_DoNotDirty | REN_DontCreateRedirectors | REN_AllowPackageLinkerMismatch);
 		}
 	}
 	else
diff --git a/Engine/Plugins/PCG/Source/PCG/Private/Elements/PCGSpawnActor.cpp b/Engine/Plugins/PCG/Source/PCG/Private/Elements/PCGSpawnActor.cpp
--- a/Engine/Plugins/PCG/Source/PCG/Private/Elements/PCGSpawnActor.cpp
+++ b/Engine/Plugins/PCG/Source/PCG/Private/Elements/PCGSpawnActor.cpp
@@ -398,7 +398,7 @@ void UPCGSpawnActorSettings::RefreshTemplateActor()
 				Options.bNotifyObjectReplacement = true;
 				UEngine::CopyPropertiesForUnrelatedObjects(TemplateActor, NewTemplateActor, Options);
 
-				TemplateActor->Rename(nullptr, GetTransientPackage(), REN_DontCreateRedirectors);
+				TemplateActor->Rename(nullptr, GetTransientPackage(), REN_DontCreateRedirectors | REN_AllowPackageLinkerMismatch);
 
 				TMap<UObject*, UObject*> OldToNew;
 				OldToNew.Emplace(TemplateActor, NewTemplateActor);
@@ -413,7 +413,7 @@ void UPCGSpawnActorSettings::RefreshTemplateActor()
 			TemplateActor->Modify();
 
 			// Outer to this object
-			TemplateActor->Rename(nullptr, this, REN_DoNotDirty | REN_DontCreateRedirectors);
+			TemplateActor->Rename(nullptr, this, REN_DoNotDirty | REN_DontCreateRedirectors | REN_AllowPackageLinkerMismatch);
 		}
 	}
 	else
diff --git a/Engine/Plugins/PCG/Source/PCG/Private/PCGGraph.cpp b/Engine/Plugins/PCG/Source/PCG/Private/PCGGraph.cpp
--- a/Engine/Plugins/PCG/Source/PCG/Private/PCGGraph.cpp
+++ b/Engine/Plugins/PCG/Source/PCG/Private/PCGGraph.cpp
@@ -988,7 +988,7 @@ UPCGNode* UPCGGraph::AddNodeOfType(TSubclassOf<class UPCGSettings> InSettingsCla
 
 	if (Node)
 	{
-		Settings->Rename(nullptr, Node, REN_DontCreateRedirectors);
+		Settings->Rename(nullptr, Node, REN_DontCreateRedirectors | REN_AllowPackageLinkerMismatch);
 	}
 
 	OutDefaultNodeSettings = Settings;
@@ -1015,7 +1015,7 @@ UPCGNode* UPCGGraph::AddNode(UPCGSettingsInterface* InSettingsInterface)
 		Node->SetSettingsInterface(InSettingsInterface);
 
 		// Reparent node to this graph
-		Node->Rename(nullptr, this, REN_DontCreateRedirectors);
+		Node->Rename(nullptr, this, REN_DontCreateRedirectors | REN_AllowPackageLinkerMismatch);
 
 #if WITH_EDITOR
 		const FName DefaultNodeName = InSettingsInterface->GetSettings()->GetDefaultNodeName();
@@ -1023,7 +1023,7 @@ UPCGNode* UPCGGraph::AddNode(UPCGSettingsInterface* InSettingsInterface)
 		{
 			const FName NodeName = MakeUniqueObjectName(this, UPCGNode::StaticClass(), DefaultNodeName);
 			// Flags added because default flags favor tick/interactive, not load-time renaming.
-			Node->Rename(*NodeName.ToString(), nullptr, REN_DontCreateRedirectors);
+			Node->Rename(*NodeName.ToString(), nullptr, REN_DontCreateRedirectors | REN_AllowPackageLinkerMismatch);
 		}
 
 		// Give the chance to the graph/settings to update the node title.
@@ -1051,7 +1051,7 @@ UPCGNode* UPCGGraph::AddNodeInstance(UPCGSettings* InSettings)
 
 	if (Node)
 	{
-		SettingsInstance->Rename(nullptr, Node, REN_DontCreateRedirectors);
+		SettingsInstance->Rename(nullptr, Node, REN_DontCreateRedirectors | REN_AllowPackageLinkerMismatch);
 		SettingsInstance->SetFlags(RF_Transactional);
 	}
 
@@ -1070,7 +1070,7 @@ UPCGNode* UPCGGraph::AddNodeCopy(const UPCGSettings* InSettings, UPCGSettings*&
 
 	if (SettingsCopy)
 	{
-		SettingsCopy->Rename(nullptr, NewNode, REN_DontCreateRedirectors);
+		SettingsCopy->Rename(nullptr, NewNode, REN_DontCreateRedirectors | REN_AllowPackageLinkerMismatch);
 	}
 
 	DefaultNodeSettings = SettingsCopy;
@@ -1253,14 +1253,14 @@ void UPCGGraph::AddNodes_Internal(TArrayView<UPCGNode*> InNodes)
 	for (UPCGNode* Node : InNodes)
 	{
 		check(Node);
-		Node->Rename(nullptr, this, REN_DontCreateRedirectors);
+		Node->Rename(nullptr, this, REN_DontCreateRedirectors | REN_AllowPackageLinkerMismatch);
 
 #if WITH_EDITOR
 		const FName DefaultNodeName = Node->GetSettings()->GetDefaultNodeName();
 		if (DefaultNodeName != NAME_None)
 		{
 			FName NodeName = MakeUniqueObjectName(this, UPCGNode::StaticClass(), DefaultNodeName);
-			Node->Rename(*NodeName.ToString(), nullptr, REN_DontCreateRedirectors);
+			Node->Rename(*NodeName.ToString(), nullptr, REN_DontCreateRedirectors | REN_AllowPackageLinkerMismatch);
 		}
 #endif
 
diff --git a/Engine/Plugins/PCGInterops/PCGExternalDataInterop/Source/PCGExternalDataInteropEditor/Private/PCGLoadAlembic.cpp b/Engine/Plugins/PCGInterops/PCGExternalDataInterop/Source/PCGExternalDataInteropEditor/Private/PCGLoadAlembic.cpp
--- a/Engine/Plugins/PCGInterops/PCGExternalDataInterop/Source/PCGExternalDataInteropEditor/Private/PCGLoadAlembic.cpp
+++ b/Engine/Plugins/PCGInterops/PCGExternalDataInterop/Source/PCGExternalDataInteropEditor/Private/PCGLoadAlembic.cpp
@@ -86,7 +86,7 @@ void UPCGLoadAlembicFunctionLibrary::LoadAlembicFileToPCGInternal(const FPCGLoad
 		{
 			if (TaggedData.Data)
 			{
-				(const_cast<UPCGData*>(TaggedData.Data.Get()))->Rename(nullptr, TargetOuter, REN_DoNotDirty | REN_DontCreateRedirectors);
+				(const_cast<UPCGData*>(TaggedData.Data.Get()))->Rename(nullptr, TargetOuter, REN_DoNotDirty | REN_DontCreateRedirectors | REN_AllowPackageLinkerMismatch);
 			}
 		}
 	}
diff --git a/Engine/Plugins/Runtime/RigVM/Source/RigVMDeveloper/Private/EdGraph/RigVMEdGraph.cpp b/Engine/Plugins/Runtime/RigVM/Source/RigVMDeveloper/Private/EdGraph/RigVMEdGraph.cpp
--- a/Engine/Plugins/Runtime/RigVM/Source/RigVMDeveloper/Private/EdGraph/RigVMEdGraph.cpp
+++ b/Engine/Plugins/Runtime/RigVM/Source/RigVMDeveloper/Private/EdGraph/RigVMEdGraph.cpp
@@ -81,7 +81,7 @@ void URigVMEdGraph::HandleRigVMGraphRenamed(const FString& InOldNodePath, const
 		}
 		GraphName = FRigVMClient::GetUniqueName(GetOuter(), *GraphName).ToString();
 
-		Rename(*GraphName, nullptr, REN_DontCreateRedirectors);
+		Rename(*GraphName, nullptr, REN_DontCreateRedirectors | REN_AllowPackageLinkerMismatch);
 	}
 	else if(ModelNodePath.StartsWith(OldPrefix))
 	{
diff --git a/Engine/Plugins/Runtime/RigVM/Source/RigVMDeveloper/Private/RigVMEditorAsset.cpp b/Engine/Plugins/Runtime/RigVM/Source/RigVMDeveloper/Private/RigVMEditorAsset.cpp
--- a/Engine/Plugins/Runtime/RigVM/Source/RigVMDeveloper/Private/RigVMEditorAsset.cpp
+++ b/Engine/Plugins/Runtime/RigVM/Source/RigVMDeveloper/Private/RigVMEditorAsset.cpp
@@ -592,7 +592,7 @@ void IRigVMEditorAssetInterface::PostRename(UObject* OldOuter, const FName OldNa
 	{
 		if (URigVMMemoryStorageGeneratorClass* MemoryClass = Cast<URigVMMemoryStorageGeneratorClass>(ClassObject))
 		{
-			MemoryClass->Rename(nullptr, GetObject()->GetPackage(), REN_DoNotDirty | REN_DontCreateRedirectors | REN_NonTransactional);
+			MemoryClass->Rename(nullptr, GetObject()->GetPackage(), REN_DoNotDirty | REN_DontCreateRedirectors | REN_NonTransactional | REN_AllowPackageLinkerMismatch);
 		}
 	}
 
diff --git a/Engine/Plugins/Runtime/RigVM/Source/RigVMDeveloper/Private/RigVMModel/RigVMClient.cpp b/Engine/Plugins/Runtime/RigVM/Source/RigVMDeveloper/Private/RigVMModel/RigVMClient.cpp
--- a/Engine/Plugins/Runtime/RigVM/Source/RigVMDeveloper/Private/RigVMModel/RigVMClient.cpp
+++ b/Engine/Plugins/Runtime/RigVM/Source/RigVMDeveloper/Private/RigVMModel/RigVMClient.cpp
@@ -977,7 +977,7 @@ FName FRigVMClient::RenameModel(const FString& InNodePathOrName, const FName& In
 		{
 			UE::TScopeLock Lock(ControllersLock);
 			TObjectPtr<URigVMController>* Controller = Controllers.Find(Model);
-			Model->Rename(*SafeNewName.ToString(), nullptr, REN_DontCreateRedirectors);
+			Model->Rename(*SafeNewName.ToString(), nullptr, REN_DontCreateRedirectors | REN_AllowPackageLinkerMismatch);
 			NewNodePath = Model->GetNodePath();
 			if (Controller)
 			{
diff --git a/Engine/Plugins/Runtime/RigVM/Source/RigVMDeveloper/Private/RigVMModel/RigVMController.cpp b/Engine/Plugins/Runtime/RigVM/Source/RigVMDeveloper/Private/RigVMModel/RigVMController.cpp
--- a/Engine/Plugins/Runtime/RigVM/Source/RigVMDeveloper/Private/RigVMModel/RigVMController.cpp
+++ b/Engine/Plugins/Runtime/RigVM/Source/RigVMDeveloper/Private/RigVMModel/RigVMController.cpp
@@ -22075,7 +22075,9 @@ void URigVMController::RewireLinks(URigVMPin* InOldPin, URigVMPin* InNewPin, boo
 
 bool URigVMController::RenameObject(UObject* InObjectToRename, const TCHAR* InNewName, UObject* InNewOuter, ERenameFlags InFlags) const
 {
-	const bool bSuccess = InObjectToRename->Rename(InNewName, InNewOuter, InFlags | REN_DoNotDirty | REN_DontCreateRedirectors | REN_NonTransactional);
+	// Renames are commonly done to temporarily move objects to a new outer and then move back (often during loading). 
+	// The extra flags are used to support that flow avoiding the creation of redirectors, marking objects for save or affecting the linker state.
+	const bool bSuccess = InObjectToRename->Rename(InNewName, InNewOuter, InFlags | REN_AllowPackageLinkerMismatch | REN_DoNotDirty | REN_DontCreateRedirectors | REN_NonTransactional);
 	if(bSuccess)
 	{
 		if(const URigVMCollapseNode* CollapseNode = Cast<URigVMCollapseNode>(InObjectToRename))
diff --git a/Engine/Plugins/Runtime/StateTree/Source/StateTreeEditorModule/Private/StateTreeEditorData.cpp b/Engine/Plugins/Runtime/StateTree/Source/StateTreeEditorModule/Private/StateTreeEditorData.cpp
--- a/Engine/Plugins/Runtime/StateTree/Source/StateTreeEditorModule/Private/StateTreeEditorData.cpp
+++ b/Engine/Plugins/Runtime/StateTree/Source/StateTreeEditorModule/Private/StateTreeEditorData.cpp
@@ -739,7 +739,7 @@ void UStateTreeEditorData::ReparentStates()
 		if (State.GetOuter() != ExpectedOuter)
 		{
 			UE_LOG(LogStateTreeEditor, Log, TEXT("%s: Fixing outer on state %s."), *TreeData->GetFullName(), *GetNameSafe(&State));
-			State.Rename(nullptr, ExpectedOuter, REN_DontCreateRedirectors | REN_DoNotDirty);
+			State.Rename(nullptr, ExpectedOuter, REN_DontCreateRedirectors | REN_DoNotDirty | REN_AllowPackageLinkerMismatch);
 		}
 		
 		State.Parent = ParentState;
@@ -812,7 +812,7 @@ void UStateTreeEditorData::FixObjectInstances(TSet<UObject*>& SeenObjects, UObje
 					if (InstanceObject->GetOuter() != &Outer)
 					{
 						UE_LOG(LogStateTreeEditor, Log, TEXT("%s: Fixing outer on node instance %s."), *GetFullName(), *GetNameSafe(InstanceObject));
-						InstanceObject->Rename(nullptr, &Outer, REN_DontCreateRedirectors | REN_DoNotDirty);
+						InstanceObject->Rename(nullptr, &Outer, REN_DontCreateRedirectors | REN_DoNotDirty | REN_AllowPackageLinkerMismatch);
 					}
 				}
 				SeenObjects.Add(InstanceObject);
diff --git a/Engine/Plugins/Runtime/USDCore/Source/USDClasses/Private/USDAssetCache2.cpp b/Engine/Plugins/Runtime/USDCore/Source/USDClasses/Private/USDAssetCache2.cpp
--- a/Engine/Plugins/Runtime/USDCore/Source/USDClasses/Private/USDAssetCache2.cpp
+++ b/Engine/Plugins/Runtime/USDCore/Source/USDClasses/Private/USDAssetCache2.cpp
@@ -488,13 +488,13 @@ namespace UE::AssetCache::Private
 				if (NewOuter != SubObject->GetOuter())
 				{
 					const TCHAR* NewName = nullptr;
-					SubObject->Rename(NewName, NewOuter, REN_NonTransactional | REN_DontCreateRedirectors);
+					SubObject->Rename(NewName, NewOuter, REN_NonTransactional | REN_DontCreateRedirectors | REN_AllowPackageLinkerMismatch);
 				}
 
 				UObject* NoNewOuter = nullptr;
 				if (SubObjectName != SubObject->GetName() && SubObject->Rename(*SubObjectName, NoNewOuter, REN_Test))
 				{
-					SubObject->Rename(*SubObjectName, NoNewOuter, REN_NonTransactional | REN_DontCreateRedirectors);
+					SubObject->Rename(*SubObjectName, NoNewOuter, REN_NonTransactional | REN_DontCreateRedirectors | REN_AllowPackageLinkerMismatch);
 				}
 			}
 		}
@@ -752,7 +752,7 @@ void UUsdAssetCache2::CacheAsset(const FString& Hash, UObject* Asset, const UObj
 	{
 		UObject* NewOuter = this;
 		const FName NewName = MakeUniqueObjectName(NewOuter, Asset->GetClass(), Asset->GetFName());
-		Asset->Rename(*NewName.ToString(), NewOuter, REN_DontCreateRedirectors | REN_NonTransactional);
+		Asset->Rename(*NewName.ToString(), NewOuter, REN_DontCreateRedirectors | REN_NonTransactional | REN_AllowPackageLinkerMismatch);
 	}
 
 	FCachedAssetInfo NewInfo;
@@ -1590,7 +1590,7 @@ bool UUsdAssetCache2::TryUnloadAsset(FCachedAssetInfo& InOutInfo)
 
 			const TCHAR* NewName = nullptr;
 			UObject* NewOuter = GetTransientPackage();
-			Asset->Rename(NewName, NewOuter, REN_DontCreateRedirectors | REN_NonTransactional);
+			Asset->Rename(NewName, NewOuter, REN_DontCreateRedirectors | REN_NonTransactional | REN_AllowPackageLinkerMismatch);
 		}
 
 		return InOutInfo.BulkData.UnloadBulkData();
diff --git a/Engine/Plugins/Runtime/nDisplay/Source/DisplayClusterConfiguration/Private/DisplayClusterConfigurationTypes.cpp b/Engine/Plugins/Runtime/nDisplay/Source/DisplayClusterConfiguration/Private/DisplayClusterConfigurationTypes.cpp
--- a/Engine/Plugins/Runtime/nDisplay/Source/DisplayClusterConfiguration/Private/DisplayClusterConfigurationTypes.cpp
+++ b/Engine/Plugins/Runtime/nDisplay/Source/DisplayClusterConfiguration/Private/DisplayClusterConfigurationTypes.cpp
@@ -446,7 +446,7 @@ void UDisplayClusterConfigurationData_Base::Serialize(FArchive& Ar)
 			}
 			if (Object->GetOuter() != this)
 			{
-				Object->Rename(nullptr, this, REN_DoNotDirty | REN_DontCreateRedirectors);
+				Object->Rename(nullptr, this, REN_DoNotDirty | REN_DontCreateRedirectors | REN_AllowPackageLinkerMismatch);
 			}
 			Object->SetFlags(RF_Public);
 			Object->ClearFlags(RF_Transient);
diff --git a/Engine/Plugins/Runtime/nDisplay/Source/DisplayClusterConfigurator/Private/ClusterConfiguration/DisplayClusterConfiguratorClusterUtils.cpp b/Engine/Plugins/Runtime/nDisplay/Source/DisplayClusterConfigurator/Private/ClusterConfiguration/DisplayClusterConfiguratorClusterUtils.cpp
--- a/Engine/Plugins/Runtime/nDisplay/Source/DisplayClusterConfigurator/Private/ClusterConfiguration/DisplayClusterConfiguratorClusterUtils.cpp
+++ b/Engine/Plugins/Runtime/nDisplay/Source/DisplayClusterConfigurator/Private/ClusterConfiguration/DisplayClusterConfiguratorClusterUtils.cpp
@@ -469,7 +469,7 @@ UDisplayClusterConfigurationViewport* UE::DisplayClusterConfiguratorClusterUtils
 
 	ClusterNode->Modify();
 	Viewport->Modify();
-	Viewport->Rename(*ViewportName, ClusterNode, REN_DontCreateRedirectors | REN_DoNotDirty);
+	Viewport->Rename(*ViewportName, ClusterNode, REN_DontCreateRedirectors | REN_DoNotDirty | REN_AllowPackageLinkerMismatch);
 
 	const FName FieldName = GET_MEMBER_NAME_CHECKED(UDisplayClusterConfigurationClusterNode, Viewports);
 	UDisplayClusterConfigurationViewport* NewViewport = CastChecked<UDisplayClusterConfigurationViewport>(AddKeyWithInstancedValueToMap(ClusterNode, FieldName, ViewportName, Viewport));
diff --git a/Engine/Plugins/Runtime/nDisplay/Source/DisplayClusterConfigurator/Private/DisplayClusterConfiguratorCompiler.cpp b/Engine/Plugins/Runtime/nDisplay/Source/DisplayClusterConfigurator/Private/DisplayClusterConfiguratorCompiler.cpp
--- a/Engine/Plugins/Runtime/nDisplay/Source/DisplayClusterConfigurator/Private/DisplayClusterConfiguratorCompiler.cpp
+++ b/Engine/Plugins/Runtime/nDisplay/Source/DisplayClusterConfigurator/Private/DisplayClusterConfiguratorCompiler.cpp
@@ -172,7 +172,7 @@ void FDisplayClusterConfiguratorKismetCompilerContext::CopyTermDefaultsToDefault
 			});
 		}
 		
-		const ERenameFlags RenFlags = REN_DontCreateRedirectors | REN_NonTransactional | REN_DoNotDirty;
+		const ERenameFlags RenFlags = REN_DontCreateRedirectors | REN_NonTransactional | REN_DoNotDirty | REN_AllowPackageLinkerMismatch;
 
 		// Rename our new config data (along with all new sub-objects) to the transient package.
 		{
diff --git a/Engine/Plugins/Tests/InterchangeTests/Source/InterchangeTests/Private/InterchangeImportTestStepBase.cpp b/Engine/Plugins/Tests/InterchangeTests/Source/InterchangeTests/Private/InterchangeImportTestStepBase.cpp
--- a/Engine/Plugins/Tests/InterchangeTests/Source/InterchangeTests/Private/InterchangeImportTestStepBase.cpp
+++ b/Engine/Plugins/Tests/InterchangeTests/Source/InterchangeTests/Private/InterchangeImportTestStepBase.cpp
@@ -113,7 +113,7 @@ void UInterchangeImportTestStepBase::SaveReloadAssets(FInterchangeImportTestData
 		// Renaming the original objects avoids having to do a GC sweep here (this is done at the end of each test step)
 		// Any existing references to them will be retained but irrelevant.
 		// Then the new object can be loaded in their place, as if it were being loaded for the first time.
-		const ERenameFlags RenameFlags = REN_DontCreateRedirectors | REN_NonTransactional | REN_DoNotDirty;
+		const ERenameFlags RenameFlags = REN_DontCreateRedirectors | REN_NonTransactional | REN_DoNotDirty | REN_AllowPackageLinkerMismatch;
 		PackageObject->Rename(*(PackageObject->GetName() + TEXT("_TRASH")), nullptr, RenameFlags);
 		PackageObject->RemoveFromRoot();
 		PackageObject->MarkAsGarbage();
diff --git a/Engine/Plugins/Tests/InterchangeTests/Source/InterchangeTests/Private/InterchangeImportTestStepImport.cpp b/Engine/Plugins/Tests/InterchangeTests/Source/InterchangeTests/Private/InterchangeImportTestStepImport.cpp
--- a/Engine/Plugins/Tests/InterchangeTests/Source/InterchangeTests/Private/InterchangeImportTestStepImport.cpp
+++ b/Engine/Plugins/Tests/InterchangeTests/Source/InterchangeTests/Private/InterchangeImportTestStepImport.cpp
@@ -45,7 +45,7 @@ TTuple<UE::Interchange::FAssetImportResultPtr, UE::Interchange::FSceneImportResu
 				//Do not rename actors
 				if (!ObjectInPackage->IsA<AActor>())
 				{
-					const ERenameFlags RenameFlags = REN_DontCreateRedirectors | REN_NonTransactional | REN_DoNotDirty;
+					const ERenameFlags RenameFlags = REN_DontCreateRedirectors | REN_NonTransactional | REN_DoNotDirty | REN_AllowPackageLinkerMismatch;
 					ObjectInPackage->Rename(nullptr, nullptr, RenameFlags);
 				}
 				ObjectInPackage->ClearFlags(RF_Standalone | RF_Public);
diff --git a/Engine/Plugins/VirtualProduction/ActorModifierCore/Source/ActorModifierCore/Private/Modifiers/ActorModifierCoreBase.cpp b/Engine/Plugins/VirtualProduction/ActorModifierCore/Source/ActorModifierCore/Private/Modifiers/ActorModifierCoreBase.cpp
--- a/Engine/Plugins/VirtualProduction/ActorModifierCore/Source/ActorModifierCore/Private/Modifiers/ActorModifierCoreBase.cpp
+++ b/Engine/Plugins/VirtualProduction/ActorModifierCore/Source/ActorModifierCore/Private/Modifiers/ActorModifierCoreBase.cpp
@@ -500,7 +500,7 @@ void UActorModifierCoreBase::PostLoad()
 		&& Outer
 		&& Outer->IsA<AActor>())
 	{
-		constexpr int32 RenameFlags = REN_DontCreateRedirectors | REN_DoNotDirty | REN_NonTransactional;
+		constexpr int32 RenameFlags = REN_DontCreateRedirectors | REN_DoNotDirty | REN_NonTransactional | REN_AllowPackageLinkerMismatch;
 
 		// Migrate modifiers to component stack instead and discard this stack
 		if (const UActorModifierCoreStack* ThisStack = Cast<UActorModifierCoreStack>(this))
diff --git a/Engine/Plugins/VirtualProduction/ClonerEffector/Source/ClonerEffector/Private/Cloner/Layouts/CEClonerLayoutBase.cpp b/Engine/Plugins/VirtualProduction/ClonerEffector/Source/ClonerEffector/Private/Cloner/Layouts/CEClonerLayoutBase.cpp
--- a/Engine/Plugins/VirtualProduction/ClonerEffector/Source/ClonerEffector/Private/Cloner/Layouts/CEClonerLayoutBase.cpp
+++ b/Engine/Plugins/VirtualProduction/ClonerEffector/Source/ClonerEffector/Private/Cloner/Layouts/CEClonerLayoutBase.cpp
@@ -366,7 +366,7 @@ void UCEClonerLayoutBase::OnSystemPackageLoaded(const FName& InName, UPackage* I
 		NiagaraSystem->RemoveFromRoot();
 		NiagaraSystem->ClearFlags(RF_Standalone | RF_Public | RF_Transient | RF_Transactional);
 
-		constexpr ERenameFlags RenameFlags = REN_NonTransactional | REN_DontCreateRedirectors;
+		constexpr ERenameFlags RenameFlags = REN_NonTransactional | REN_DontCreateRedirectors | REN_AllowPackageLinkerMismatch;
 		if (NiagaraSystem->Rename(nullptr, this, RenameFlags))
 		{
 #if WITH_EDITOR
diff --git a/Engine/Plugins/VirtualProduction/DynamicMaterial/Source/DynamicMaterial/Public/DMDefs.h b/Engine/Plugins/VirtualProduction/DynamicMaterial/Source/DynamicMaterial/Public/DMDefs.h
--- a/Engine/Plugins/VirtualProduction/DynamicMaterial/Source/DynamicMaterial/Public/DMDefs.h
+++ b/Engine/Plugins/VirtualProduction/DynamicMaterial/Source/DynamicMaterial/Public/DMDefs.h
@@ -48,7 +48,7 @@ enum class EDMIterationResult : uint8
 
 namespace UE::DynamicMaterial
 {
-	constexpr int32 RenameFlags = REN_DontCreateRedirectors | REN_DoNotDirty | REN_ForceNoResetLoaders | REN_NonTransactional;
+	constexpr int32 RenameFlags = REN_DontCreateRedirectors | REN_DoNotDirty | REN_AllowPackageLinkerMismatch | REN_NonTransactional;
 	constexpr EClassFlags InvalidClassFlags = CLASS_Abstract | CLASS_Hidden | CLASS_Deprecated | CLASS_NewerVersionExists;
 
 	DYNAMICMATERIAL_API void ForEachMaterialPropertyType(TFunctionRef<EDMIterationResult(EDMMaterialPropertyType InType)> InCallable,
diff --git a/Engine/Plugins/VirtualProduction/RemoteControl/Source/RemoteControlLogic/Private/Action/Path/RCSetAssetByPathActionNew.cpp b/Engine/Plugins/VirtualProduction/RemoteControl/Source/RemoteControlLogic/Private/Action/Path/RCSetAssetByPathActionNew.cpp
--- a/Engine/Plugins/VirtualProduction/RemoteControl/Source/RemoteControlLogic/Private/Action/Path/RCSetAssetByPathActionNew.cpp
+++ b/Engine/Plugins/VirtualProduction/RemoteControl/Source/RemoteControlLogic/Private/Action/Path/RCSetAssetByPathActionNew.cpp
@@ -236,7 +236,7 @@ bool URCSetAssetByPathActionNew::SetInternalAsset(UObject* InSetterObject)
 			SpawnParams.Name = OldActorName;
 
 			const FName OldActorReplacedNamed = MakeUniqueObjectName(OldActor->GetOuter(), OldActor->GetClass(), *FString::Printf(TEXT("%s_REPLACED"), *OldActor->GetFName().ToString()));
-			OldActor->Rename(*OldActorReplacedNamed.ToString(), OldActor->GetOuter(), REN_DoNotDirty | REN_DontCreateRedirectors);
+			OldActor->Rename(*OldActorReplacedNamed.ToString(), OldActor->GetOuter(), REN_DoNotDirty | REN_DontCreateRedirectors | REN_AllowPackageLinkerMismatch);
 
 			if (AActor* NewActor = World->SpawnActor(Cast<UBlueprint>(InSetterObject)->GeneratedClass, &OldLocation, &OldRotation, SpawnParams))
 			{
@@ -247,7 +247,7 @@ bool URCSetAssetByPathActionNew::SetInternalAsset(UObject* InSetterObject)
 			}
 
 			UE_LOG(LogRemoteControl, Error, TEXT("Path Behavior unable to delete old Actor %s"), *OldActor->GetName());
-			OldActor->UObject::Rename(*OldActorName.ToString(), OldActor->GetOuter(), REN_DoNotDirty | REN_DontCreateRedirectors);
+			OldActor->UObject::Rename(*OldActorName.ToString(), OldActor->GetOuter(), REN_DoNotDirty | REN_DontCreateRedirectors | REN_AllowPackageLinkerMismatch);
 			OldActor->RegisterAllComponents();
 
 			return false;
diff --git a/Engine/Plugins/VirtualProduction/RemoteControl/Source/RemoteControlLogic/Private/Behaviour/Builtin/Path/RCSetAssetByPathBehaviour.cpp b/Engine/Plugins/VirtualProduction/RemoteControl/Source/RemoteControlLogic/Private/Behaviour/Builtin/Path/RCSetAssetByPathBehaviour.cpp
--- a/Engine/Plugins/VirtualProduction/RemoteControl/Source/RemoteControlLogic/Private/Behaviour/Builtin/Path/RCSetAssetByPathBehaviour.cpp
+++ b/Engine/Plugins/VirtualProduction/RemoteControl/Source/RemoteControlLogic/Private/Behaviour/Builtin/Path/RCSetAssetByPathBehaviour.cpp
@@ -290,7 +290,7 @@ bool URCSetAssetByPathBehaviour::SetInternalAsset(UObject* SetterObject)
 			SpawnParams.Name = OldActorName;
 			
 			const FName OldActorReplacedNamed = MakeUniqueObjectName(OldActor->GetOuter(), OldActor->GetClass(), *FString::Printf(TEXT("%s_REPLACED"), *OldActor->GetFName().ToString()));
-			OldActor->Rename(*OldActorReplacedNamed.ToString(), OldActor->GetOuter(), REN_DoNotDirty | REN_DontCreateRedirectors);
+			OldActor->Rename(*OldActorReplacedNamed.ToString(), OldActor->GetOuter(), REN_DoNotDirty | REN_DontCreateRedirectors | REN_AllowPackageLinkerMismatch);
 
 			if (AActor* NewActor = World->SpawnActor(Cast<UBlueprint>(SetterObject)->GeneratedClass, &OldLocation, &OldRotation, SpawnParams))
 			{
@@ -301,7 +301,7 @@ bool URCSetAssetByPathBehaviour::SetInternalAsset(UObject* SetterObject)
 			}
 
 			UE_LOG(LogRemoteControl, Error, TEXT("Path Behaviour unable to delete old Actor %s"), *OldActor->GetName());
-			OldActor->UObject::Rename(*OldActorName.ToString(), OldActor->GetOuter(), REN_DoNotDirty | REN_DontCreateRedirectors);
+			OldActor->UObject::Rename(*OldActorName.ToString(), OldActor->GetOuter(), REN_DoNotDirty | REN_DontCreateRedirectors | REN_AllowPackageLinkerMismatch);
 			OldActor->RegisterAllComponents();
 		}
 	}
diff --git a/Engine/Source/Developer/AssetTools/Private/AssetTools.cpp b/Engine/Source/Developer/AssetTools/Private/AssetTools.cpp
--- a/Engine/Source/Developer/AssetTools/Private/AssetTools.cpp
+++ b/Engine/Source/Developer/AssetTools/Private/AssetTools.cpp
@@ -1221,7 +1221,7 @@ namespace UE::AssetTools::Private
 
 			for (UPackage* Package : ExistingPackageToClean)
 			{
-				const ERenameFlags PkgRenameFlags = REN_ForceNoResetLoaders | REN_DoNotDirty | REN_DontCreateRedirectors | REN_NonTransactional | REN_SkipGeneratedClasses;
+				const ERenameFlags PkgRenameFlags = REN_AllowPackageLinkerMismatch | REN_DoNotDirty | REN_DontCreateRedirectors | REN_NonTransactional | REN_SkipGeneratedClasses;
 				check(Package->Rename(*MakeUniqueObjectName(nullptr, UPackage::StaticClass(), *FString::Printf(TEXT("%s_DEADFROMMIGRATION"), *Package->GetName())).ToString(), nullptr, PkgRenameFlags));
 				(*PurgeObject)(Package);
 				ForEachObjectWithOuter(Package, *PurgeObject);
diff --git a/Engine/Source/Developer/AssetTools/Private/PackageMigrationContext.cpp b/Engine/Source/Developer/AssetTools/Private/PackageMigrationContext.cpp
--- a/Engine/Source/Developer/AssetTools/Private/PackageMigrationContext.cpp
+++ b/Engine/Source/Developer/AssetTools/Private/PackageMigrationContext.cpp
@@ -157,13 +157,13 @@ FPackageMigrationContext::FScopedTemporalyMovedPackage::FScopedTemporalyMovedPac
 	: PackageToMove(InPackageToMove)
 	, OriginalName(PackageToMove->GetName())
 {
-	const ERenameFlags PkgRenameFlags = REN_DoNotDirty | REN_DontCreateRedirectors | REN_NonTransactional | REN_SkipGeneratedClasses;
+	const ERenameFlags PkgRenameFlags = REN_DoNotDirty | REN_DontCreateRedirectors | REN_NonTransactional | REN_SkipGeneratedClasses | REN_AllowPackageLinkerMismatch;
 	check(PackageToMove->Rename(*MakeUniqueObjectName(nullptr, UPackage::StaticClass(), *FString::Printf(TEXT("%s_MOVEDFORMIGRATION"), *PackageToMove->GetName())).ToString(), nullptr, PkgRenameFlags));
 }
 
 FPackageMigrationContext::FScopedTemporalyMovedPackage::~FScopedTemporalyMovedPackage()
 {
-	const ERenameFlags PkgRenameFlags = REN_DoNotDirty | REN_DontCreateRedirectors | REN_NonTransactional | REN_SkipGeneratedClasses;
+	const ERenameFlags PkgRenameFlags = REN_DoNotDirty | REN_DontCreateRedirectors | REN_NonTransactional | REN_SkipGeneratedClasses | REN_AllowPackageLinkerMismatch;
 	check(PackageToMove->Rename(*OriginalName, nullptr, PkgRenameFlags));
 }
 
diff --git a/Engine/Source/Developer/HierarchicalLODUtilities/Private/HierarchicalLODUtilities.cpp b/Engine/Source/Developer/HierarchicalLODUtilities/Private/HierarchicalLODUtilities.cpp
--- a/Engine/Source/Developer/HierarchicalLODUtilities/Private/HierarchicalLODUtilities.cpp
+++ b/Engine/Source/Developer/HierarchicalLODUtilities/Private/HierarchicalLODUtilities.cpp
@@ -579,7 +579,7 @@ FHLODBuildResults GenerateHLODMesh_Approximate(const FHLODBuildParams& InBuildPa
 			AssetToReplace->Rename(nullptr, GetTransientPackage(), REN_DontCreateRedirectors | REN_NonTransactional);
 		}
 			
-		NewAsset->Rename(*AssetName, InBuildParams.AssetsOuter, REN_NonTransactional | REN_DontCreateRedirectors);
+		NewAsset->Rename(*AssetName, InBuildParams.AssetsOuter, REN_NonTransactional | REN_DontCreateRedirectors | REN_AllowPackageLinkerMismatch);
 	};
 
 	Algo::ForEach(Results.NewMeshAssets, RenameNewAsset);
diff --git a/Engine/Source/Editor/BlueprintGraph/Private/EdGraphSchema_K2.cpp b/Engine/Source/Editor/BlueprintGraph/Private/EdGraphSchema_K2.cpp
--- a/Engine/Source/Editor/BlueprintGraph/Private/EdGraphSchema_K2.cpp
+++ b/Engine/Source/Editor/BlueprintGraph/Private/EdGraphSchema_K2.cpp
@@ -6711,12 +6711,14 @@ UEdGraphNode* UEdGraphSchema_K2::CreateSubstituteNode(UEdGraphNode* Node, const
 			FBlueprintEditorUtils::GetFunctionNameList(Blueprint, ExistingNamesInUse);
 			FBlueprintEditorUtils::GetClassVariableList(Blueprint, ExistingNamesInUse);
 
+			const ERenameFlags RenameFlags = (Blueprint->bIsRegeneratingOnLoad ? REN_AllowPackageLinkerMismatch : 0);
+
 			// Allow the old object name to be used in the graph
 			FName ObjName = EventNode->GetFName();
 			UObject* Found = FindObject<UObject>(EventNode->GetOuter(), *ObjName.ToString());
 			if(Found)
 			{
-				Found->Rename(NULL, NULL, REN_DontCreateRedirectors);
+				Found->Rename(NULL, NULL, REN_DontCreateRedirectors | RenameFlags);
 			}
 
 			// Create a custom event node to replace the original event node imported from text
diff --git a/Engine/Source/Editor/BlueprintGraph/Private/K2Node_AddComponent.cpp b/Engine/Source/Editor/BlueprintGraph/Private/K2Node_AddComponent.cpp
--- a/Engine/Source/Editor/BlueprintGraph/Private/K2Node_AddComponent.cpp
+++ b/Engine/Source/Editor/BlueprintGraph/Private/K2Node_AddComponent.cpp
@@ -89,7 +89,7 @@ void UK2Node_AddComponent::Serialize(FArchive& Ar)
 					GetBlueprint()->OldToNewComponentTemplateNames.Add(Template->GetFName()) = NewTemplateName;
 
 					// Rename the component template to conform to the new name.
-					Template->Rename(*NewTemplateName.ToString(), Template->GetOuter(), REN_DontCreateRedirectors);
+					Template->Rename(*NewTemplateName.ToString(), Template->GetOuter(), REN_DontCreateRedirectors | REN_AllowPackageLinkerMismatch);
 
 					// Update the template name pin's default value to match the template name.
 					GetTemplateNamePinChecked()->DefaultValue = Template->GetName();
diff --git a/Engine/Source/Editor/BlueprintGraph/Private/K2Node_Composite.cpp b/Engine/Source/Editor/BlueprintGraph/Private/K2Node_Composite.cpp
--- a/Engine/Source/Editor/BlueprintGraph/Private/K2Node_Composite.cpp
+++ b/Engine/Source/Editor/BlueprintGraph/Private/K2Node_Composite.cpp
@@ -274,7 +274,7 @@ void UK2Node_Composite::RenameBoundGraphCloseToName(const FString& Name)
 		{
 			//Name is available
 			UBlueprint* BP = FBlueprintEditorUtils::FindBlueprintForGraphChecked(BoundGraph);
-			BoundGraph->Rename(*NewName, BoundGraph->GetOuter(), REN_DontCreateRedirectors);
+			BoundGraph->Rename(*NewName, BoundGraph->GetOuter(), (BP->bIsRegeneratingOnLoad ? REN_AllowPackageLinkerMismatch : 0) | REN_DontCreateRedirectors);
 			bFoundName = true;
 		}
 	}
diff --git a/Engine/Source/Editor/BlueprintGraph/Private/K2Node_CustomEvent.cpp b/Engine/Source/Editor/BlueprintGraph/Private/K2Node_CustomEvent.cpp
--- a/Engine/Source/Editor/BlueprintGraph/Private/K2Node_CustomEvent.cpp
+++ b/Engine/Source/Editor/BlueprintGraph/Private/K2Node_CustomEvent.cpp
@@ -292,7 +292,7 @@ void UK2Node_CustomEvent::RenameCustomEventCloseToName(int32 StartIndex)
 		{
 			UBlueprint* Blueprint = GetBlueprint();
 			CustomFunctionName = FName(NewName.GetCharArray().GetData());
-			Rename(*NewName, GetOuter(), REN_DontCreateRedirectors);
+			Rename(*NewName, GetOuter(), (Blueprint->bIsRegeneratingOnLoad ? REN_AllowPackageLinkerMismatch : 0) | REN_DontCreateRedirectors);
 			bFoundName = true;
 		}
 	}
diff --git a/Engine/Source/Editor/Blutility/Private/EditorUtilityLibrary.cpp b/Engine/Source/Editor/Blutility/Private/EditorUtilityLibrary.cpp
--- a/Engine/Source/Editor/Blutility/Private/EditorUtilityLibrary.cpp
+++ b/Engine/Source/Editor/Blutility/Private/EditorUtilityLibrary.cpp
@@ -388,7 +388,7 @@ void UEditorUtilityLibrary::ConvertToEditorUtilityWidget(UWidgetBlueprint* Widge
 	}
 	for (UObject* Child : Children)
 	{
-		Child->Rename(nullptr, EWBP, REN_DontCreateRedirectors | REN_SkipGeneratedClasses);
+		Child->Rename(nullptr, EWBP, REN_DontCreateRedirectors | REN_SkipGeneratedClasses | REN_AllowPackageLinkerMismatch);
 	}
 
 	UEngine::FCopyPropertiesForUnrelatedObjectsParams Params;
diff --git a/Engine/Source/Editor/KismetCompiler/Private/KismetCompiler.cpp b/Engine/Source/Editor/KismetCompiler/Private/KismetCompiler.cpp
--- a/Engine/Source/Editor/KismetCompiler/Private/KismetCompiler.cpp
+++ b/Engine/Source/Editor/KismetCompiler/Private/KismetCompiler.cpp
@@ -4537,7 +4537,7 @@ void FKismetCompilerContext::ProcessOneFunctionGraph(UEdGraph* SourceGraph, bool
 		FString FunctionGraphName = SourceGraph->GetName() + TEXT("_MERGED");
 
 		ERenameFlags RenameFlags =
-			(REN_NonTransactional | REN_DoNotDirty | REN_DontCreateRedirectors);
+			(REN_NonTransactional | REN_DoNotDirty | REN_DontCreateRedirectors | REN_AllowPackageLinkerMismatch);
 
 		if (UEdGraph* ExistingGraph = FindObject<UEdGraph>(Blueprint, *FunctionGraphName))
 		{
diff --git a/Engine/Source/Editor/MergeActors/Private/MeshApproximationTool/MeshApproximationTool.cpp b/Engine/Source/Editor/MergeActors/Private/MeshApproximationTool/MeshApproximationTool.cpp
--- a/Engine/Source/Editor/MergeActors/Private/MeshApproximationTool/MeshApproximationTool.cpp
+++ b/Engine/Source/Editor/MergeActors/Private/MeshApproximationTool/MeshApproximationTool.cpp
@@ -248,7 +248,7 @@ bool FMeshApproximationTool::RunMerge(const FString& PackageName, const TArray<T
 			UPackage* TempPackage = NewAsset->GetPackage();
 
 			// Rename the asset to its final destination
-			NewAsset->Rename(*AssetName, TargetPackage, REN_DontCreateRedirectors | REN_NonTransactional);
+			NewAsset->Rename(*AssetName, TargetPackage, REN_DontCreateRedirectors | REN_NonTransactional | REN_AllowPackageLinkerMismatch);
 			check(NewAsset->HasAllFlags(RF_Public | RF_Standalone));
 
 			// Clean up flags on the temp package. It is not useful anymore.
diff --git a/Engine/Source/Editor/MergeActors/Private/MeshMergingTool/MeshMergingTool.cpp b/Engine/Source/Editor/MergeActors/Private/MeshMergingTool/MeshMergingTool.cpp
--- a/Engine/Source/Editor/MergeActors/Private/MeshMergingTool/MeshMergingTool.cpp
+++ b/Engine/Source/Editor/MergeActors/Private/MeshMergingTool/MeshMergingTool.cpp
@@ -181,7 +181,7 @@ bool FMeshMergingTool::RunMerge(const FString& PackageName, const TArray<TShared
 					ObjectTools::ConsolidateObjects(AssetToSync, { OldDataAssetArray }, bShowDeleteConfirmation);
 				}
 
-				AssetToSync->Rename(*AssetName, AssetPackage, REN_DontCreateRedirectors);
+				AssetToSync->Rename(*AssetName, AssetPackage, REN_DontCreateRedirectors | REN_AllowPackageLinkerMismatch);
 				AssetToSync->SetFlags(RF_Public | RF_Standalone);
 			}
 
diff --git a/Engine/Source/Editor/SparseVolumeTexture/Private/SparseVolumeTextureFactory.cpp b/Engine/Source/Editor/SparseVolumeTexture/Private/SparseVolumeTextureFactory.cpp
--- a/Engine/Source/Editor/SparseVolumeTexture/Private/SparseVolumeTextureFactory.cpp
+++ b/Engine/Source/Editor/SparseVolumeTexture/Private/SparseVolumeTextureFactory.cpp
@@ -759,7 +759,7 @@ UObject* USparseVolumeTextureFactory::ImportInternal(UClass* InClass, UObject* I
 				UPackage* ExistingPackage = FindPackage(InParent->GetOuter(), *NewPackageName);
 				if (!ExistingPackage)
 				{
-					InParent->Rename(*NewPackageName, nullptr, REN_DontCreateRedirectors);
+					InParent->Rename(*NewPackageName, nullptr, REN_DontCreateRedirectors | REN_AllowPackageLinkerMismatch);
 					NewObjectName = *NewFileName;
 				}
 			}
diff --git a/Engine/Source/Editor/UMGEditor/Private/TabFactory/AnimationTabSummoner.cpp b/Engine/Source/Editor/UMGEditor/Private/TabFactory/AnimationTabSummoner.cpp
--- a/Engine/Source/Editor/UMGEditor/Private/TabFactory/AnimationTabSummoner.cpp
+++ b/Engine/Source/Editor/UMGEditor/Private/TabFactory/AnimationTabSummoner.cpp
@@ -859,7 +859,7 @@ private:
 				MakeUniqueObjectName( WidgetBlueprint, UWidgetAnimation::StaticClass(), SelectedAnimation->Animation->GetFName() )
 			);
 	
-		NewAnimation->MovieScene->Rename(*NewAnimation->GetName(), nullptr, REN_DontCreateRedirectors);
+		NewAnimation->MovieScene->Rename(*NewAnimation->GetName(), nullptr, REN_DontCreateRedirectors | REN_AllowPackageLinkerMismatch);
 		NewAnimation->SetDisplayLabel(NewAnimation->GetName());
 
 		bool bRenameRequestPending = true;
diff --git a/Engine/Source/Editor/UMGEditor/Private/WidgetBlueprintCompiler.cpp b/Engine/Source/Editor/UMGEditor/Private/WidgetBlueprintCompiler.cpp
--- a/Engine/Source/Editor/UMGEditor/Private/WidgetBlueprintCompiler.cpp
+++ b/Engine/Source/Editor/UMGEditor/Private/WidgetBlueprintCompiler.cpp
@@ -851,7 +851,7 @@ void FWidgetBlueprintCompilerContext::FixAbandonedWidgetTree(UWidgetBlueprint* W
 				AbandonedWidgetTree->Rename(nullptr, GetTransientPackage(), REN_DontCreateRedirectors | REN_NonTransactional | REN_DoNotDirty);
 			}
 
-			WidgetTree->Rename(TEXT("WidgetTree"), nullptr, REN_DontCreateRedirectors | REN_NonTransactional | REN_DoNotDirty);
+			WidgetTree->Rename(TEXT("WidgetTree"), nullptr, REN_DontCreateRedirectors | REN_NonTransactional | REN_DoNotDirty | REN_AllowPackageLinkerMismatch);
 			WidgetTree->SetFlags(RF_DefaultSubObject);
 		}
 	}
diff --git a/Engine/Source/Editor/UnrealEd/Private/Commandlets/WorldPartitionConvertCommandlet.cpp b/Engine/Source/Editor/UnrealEd/Private/Commandlets/WorldPartitionConvertCommandlet.cpp
--- a/Engine/Source/Editor/UnrealEd/Private/Commandlets/WorldPartitionConvertCommandlet.cpp
+++ b/Engine/Source/Editor/UnrealEd/Private/Commandlets/WorldPartitionConvertCommandlet.cpp
@@ -563,7 +563,7 @@ bool UWorldPartitionConvertCommandlet::RenameWorldPackageWithSuffix(UWorld* Worl
 	FString OldPackageName = Package->GetName();
 	FString NewPackageName = OldPackageName + ConversionSuffix;
 	FString NewPackageResourceName = Package->GetLoadedPath().GetPackageName().Replace(*OldPackageName, *NewPackageName);
-	bRenamedSuccess = Package->Rename(*NewPackageName, nullptr, REN_NonTransactional | REN_DontCreateRedirectors);
+	bRenamedSuccess = Package->Rename(*NewPackageName, nullptr, REN_NonTransactional | REN_DontCreateRedirectors | REN_AllowPackageLinkerMismatch);
 	if (!bRenamedSuccess)
 	{
 		UE_LOG(LogWorldPartitionConvertCommandlet, Error, TEXT("Unable to rename package to %s"), *NewPackageName);
@@ -574,7 +574,7 @@ bool UWorldPartitionConvertCommandlet::RenameWorldPackageWithSuffix(UWorld* Worl
 	// Make sure the world object has the same name as its package.
 	FString OldWorldName = World->GetName();
 	FString NewWorldName = OldWorldName + ConversionSuffix;
-	bRenamedSuccess = World->Rename(*NewWorldName, nullptr, REN_NonTransactional | REN_DontCreateRedirectors);
+	bRenamedSuccess = World->Rename(*NewWorldName, nullptr, REN_NonTransactional | REN_DontCreateRedirectors | REN_AllowPackageLinkerMismatch);
 	if (!bRenamedSuccess)
 	{
 		UE_LOG(LogWorldPartitionConvertCommandlet, Error, TEXT("Unable to rename world to %s"), *NewWorldName);
diff --git a/Engine/Source/Editor/UnrealEd/Private/EdGraphUtilities.cpp b/Engine/Source/Editor/UnrealEd/Private/EdGraphUtilities.cpp
--- a/Engine/Source/Editor/UnrealEd/Private/EdGraphUtilities.cpp
+++ b/Engine/Source/Editor/UnrealEd/Private/EdGraphUtilities.cpp
@@ -432,7 +432,7 @@ void FEdGraphUtilities::RenameGraphCloseToName(UEdGraph* Graph, const FString& B
 		if (Graph->Rename(*NewName, Graph->GetOuter(), REN_Test))
 		{
 			UBlueprint* BP = FBlueprintEditorUtils::FindBlueprintForGraphChecked(Graph);
-			Graph->Rename(*NewName, Graph->GetOuter(), REN_DontCreateRedirectors);
+			Graph->Rename(*NewName, Graph->GetOuter(), (BP->bIsRegeneratingOnLoad ? REN_AllowPackageLinkerMismatch : 0) | REN_DontCreateRedirectors);
 			return;
 		}
 
@@ -446,7 +446,7 @@ void FEdGraphUtilities::RenameGraphToNameOrCloseToName(UEdGraph* Graph, const FS
 	if (Graph->Rename(*DesiredName, Graph->GetOuter(), REN_Test))
 	{
 		UBlueprint* BP = FBlueprintEditorUtils::FindBlueprintForGraphChecked(Graph);
-		Graph->Rename(*DesiredName, Graph->GetOuter(), REN_DontCreateRedirectors);
+		Graph->Rename(*DesiredName, Graph->GetOuter(), (BP->bIsRegeneratingOnLoad ? REN_AllowPackageLinkerMismatch : 0) | REN_DontCreateRedirectors);
 	}
 	else
 	{
diff --git a/Engine/Source/Editor/UnrealEd/Private/Factories/EditorFactories.cpp b/Engine/Source/Editor/UnrealEd/Private/Factories/EditorFactories.cpp
--- a/Engine/Source/Editor/UnrealEd/Private/Factories/EditorFactories.cpp
+++ b/Engine/Source/Editor/UnrealEd/Private/Factories/EditorFactories.cpp
@@ -822,10 +822,10 @@ UObject* ULevelFactory::LevelFactoryCreateText
 				Buffer += FCString::Strlen(TEXT("Name="));
 				Buffer += MapName.Len();
 				// Check to make sure that there are no naming conflicts
-				if( RootMapPackage->Rename(*MapName, nullptr, REN_Test) )
+				if( RootMapPackage->Rename(*MapName, nullptr, REN_Test | REN_AllowPackageLinkerMismatch) )
 				{
 					// Rename it!
-					RootMapPackage->Rename(*MapName, nullptr);
+					RootMapPackage->Rename(*MapName, nullptr, REN_AllowPackageLinkerMismatch);
 				}
 				else
 				{
@@ -3822,7 +3822,7 @@ UObject* UTextureFactory::FactoryCreateBinary
 							}
 							else
 							{
-								verify(InParent->Rename(*PackageUDIMName, nullptr, REN_DontCreateRedirectors));
+								verify(InParent->Rename(*PackageUDIMName, nullptr, REN_DontCreateRedirectors | REN_AllowPackageLinkerMismatch));
 							}
 						}
 					}
diff --git a/Engine/Source/Editor/UnrealEd/Private/FileHelpers.cpp b/Engine/Source/Editor/UnrealEd/Private/FileHelpers.cpp
--- a/Engine/Source/Editor/UnrealEd/Private/FileHelpers.cpp
+++ b/Engine/Source/Editor/UnrealEd/Private/FileHelpers.cpp
@@ -1155,7 +1155,7 @@ static bool SaveWorld(UWorld* World,
 					// Explict Reset Loaders of Package here because we want to avoid resetting of all loaders which is the current behavior of UObject::Rename when passing in a UPackage
 					ResetLoaders(Package);
 					// Duplicate failed or not needed. Just do a rename.
-					Package->Rename(*NewPackageName, NULL, REN_NonTransactional | REN_DontCreateRedirectors);
+					Package->Rename(*NewPackageName, NULL, REN_NonTransactional | REN_DontCreateRedirectors | REN_AllowPackageLinkerMismatch);
 					
 					if (bWorldNeedsRename)
 					{
@@ -1172,7 +1172,7 @@ static bool SaveWorld(UWorld* World,
 							}
 						}
 
-						World->Rename(*NewWorldAssetName, NULL, REN_NonTransactional | REN_DontCreateRedirectors);
+						World->Rename(*NewWorldAssetName, NULL, REN_NonTransactional | REN_DontCreateRedirectors | REN_AllowPackageLinkerMismatch);
 					}
 
 					// We're changing the world path, add a path redirector so that soft object paths get fixed on save
diff --git a/Engine/Source/Editor/UnrealEd/Private/Kismet2/BlueprintEditorUtils.cpp b/Engine/Source/Editor/UnrealEd/Private/Kismet2/BlueprintEditorUtils.cpp
--- a/Engine/Source/Editor/UnrealEd/Private/Kismet2/BlueprintEditorUtils.cpp
+++ b/Engine/Source/Editor/UnrealEd/Private/Kismet2/BlueprintEditorUtils.cpp
@@ -1295,7 +1295,7 @@ void FBlueprintEditorUtils::RemoveStaleFunctions(UBlueprintGeneratedClass* Class
 		OrphanedClass->ClassFlags |= CLASS_CompiledFromBlueprint;
 		OrphanedClass->ClassGeneratedBy = Class->ClassGeneratedBy;
 
-		const ERenameFlags RenFlags = REN_DontCreateRedirectors | REN_NonTransactional | REN_DoNotDirty;
+		const ERenameFlags RenFlags = REN_DontCreateRedirectors | REN_NonTransactional | REN_DoNotDirty | (Blueprint->bIsRegeneratingOnLoad ? REN_AllowPackageLinkerMismatch : 0);
 
 		while (Fn)
 		{
@@ -2234,7 +2234,7 @@ UEdGraph* FBlueprintEditorUtils::CreateNewGraph(UObject* ParentScope, const FNam
 				// Rename the old graph out of the way - this may confuse the user somewhat - and even
 				// break their logic. But name collisions are not avoidable e.g. someone can add
 				// a function to an interface that conflicts with something in a class hierarchy
-				ExistingObject->Rename(nullptr, ExistingObject->GetOuter(), REN_DoNotDirty);
+				ExistingObject->Rename(nullptr, ExistingObject->GetOuter(), REN_DoNotDirty | REN_AllowPackageLinkerMismatch);
 			}
 			else if (ExistingObject->IsA<UObjectRedirector>())
 			{
@@ -2263,7 +2263,7 @@ UEdGraph* FBlueprintEditorUtils::CreateNewGraph(UObject* ParentScope, const FNam
 	// Now move to where we want it to. Workaround to ensure transaction buffer is correctly utilized
 	if (bRename)
 	{
-		NewGraph->Rename(*(GraphName.ToString()), ParentScope, REN_DoNotDirty);
+		NewGraph->Rename(*(GraphName.ToString()), ParentScope, REN_DoNotDirty | REN_AllowPackageLinkerMismatch);
 	}
 	return NewGraph;
 }
@@ -2718,6 +2718,10 @@ void FBlueprintEditorUtils::RenameGraph(UEdGraph* Graph, const FString& NewNameS
 		};
 
 		ERenameFlags RenameFlagsToApply = REN_None;
+		if (Blueprint->bIsRegeneratingOnLoad)
+		{
+			RenameFlagsToApply |= REN_AllowPackageLinkerMismatch;
+		}
 
 		// Macro library graphs are referenced indirectly and resolved at edit/compile time via GUID (see FGraphReference).
 		// However, they will be exported by name at save time, so renaming a macro library graph implies we should also
@@ -2743,7 +2747,7 @@ void FBlueprintEditorUtils::RenameGraph(UEdGraph* Graph, const FString& NewNameS
 			{
 				if (FunctionGraph->GetFName() == OldGraphName)
 				{
-					RenameGraphLambda(FunctionGraph, OldGraphName, NewGraphName, REN_DontCreateRedirectors);
+					RenameGraphLambda(FunctionGraph, OldGraphName, NewGraphName, (InChildBP->bIsRegeneratingOnLoad ? REN_AllowPackageLinkerMismatch : 0) | REN_DontCreateRedirectors);
 				}
 			}
 
@@ -2823,7 +2827,7 @@ void FBlueprintEditorUtils::RenameGraphWithSuggestion(class UEdGraph* Graph, TSh
 	FString NewName = DesiredName;
 	NameValidator->FindValidString(NewName);
 	UBlueprint* BP = FBlueprintEditorUtils::FindBlueprintForGraphChecked(Graph);
-	Graph->Rename(*NewName, Graph->GetOuter(), REN_DontCreateRedirectors);
+	Graph->Rename(*NewName, Graph->GetOuter(), (BP->bIsRegeneratingOnLoad ? REN_AllowPackageLinkerMismatch : 0) | REN_DontCreateRedirectors);
 }
 
 /** 
@@ -7486,10 +7490,11 @@ static void ConformInterfaceByName(UBlueprint* Blueprint, FBPInterfaceDescriptio
 
 					// rename to free up the graph's name.. which may be needed by an inherited function
 					// alternatively we could move this into the functions list?
+					// REN_AllowPackageLinkerMismatch is passed to not have to force load completion before renaming
 					CurrentGraph->Rename(
 						nullptr,
 						CurrentGraph->GetOuter(),
-						REN_DoNotDirty | REN_DontCreateRedirectors);
+						REN_DoNotDirty | REN_DontCreateRedirectors | (Blueprint->bIsRegeneratingOnLoad ? REN_AllowPackageLinkerMismatch : 0));
 					// removing from root, standalone, and public is defensive to make sure it is not saved:
 					CurrentGraph->ClearFlags(RF_Standalone | RF_Public);
 					CurrentGraph->RemoveFromRoot();
@@ -7667,7 +7672,7 @@ void FBlueprintEditorUtils::UpdateOutOfDateCompositeWithOuter(UBlueprint* Bluepr
 				if (BoundGraph->GetOuter() != Node)
 				{
 					// change the outer of the BoundGraph to be the composite node instead of the OuterGraph
-					if (false == BoundGraph->Rename(*BoundGraph->GetName(), Node, REN_DontCreateRedirectors))
+					if (false == BoundGraph->Rename(*BoundGraph->GetName(), Node, (BoundGraph->HasAnyFlags(RF_NeedLoad | RF_NeedPostLoad) ? REN_AllowPackageLinkerMismatch : 0) | REN_DontCreateRedirectors))
 					{
 						UE_LOG(LogBlueprintDebug, Log, TEXT("CompositeNode: On Blueprint '%s' could not fix Outer() for BoundGraph of composite node '%s'"), *Blueprint->GetPathName(), *Node->GetName());
 					}
@@ -8152,7 +8157,7 @@ bool FBlueprintEditorUtils::RenameTimeline(UBlueprint* Blueprint, const FName Ol
 			{
 				ExistingObject->Rename(*MakeUniqueObjectName(ExistingObject->GetOuter(), ExistingObject->GetClass(), ExistingObject->GetFName()).ToString());
 			}
-			Template->Rename(*NewTemplateName, Template->GetOuter(), REN_None);
+			Template->Rename(*NewTemplateName, Template->GetOuter(), (Blueprint->bIsRegeneratingOnLoad ? REN_AllowPackageLinkerMismatch : REN_None));
 			Blueprint->Timelines.Add(Template);
 
 			// Validate child blueprints and adjust variable names to avoid a potential name collision
diff --git a/Engine/Source/Editor/UnrealEd/Private/Kismet2/Kismet2.cpp b/Engine/Source/Editor/UnrealEd/Private/Kismet2/Kismet2.cpp
--- a/Engine/Source/Editor/UnrealEd/Private/Kismet2/Kismet2.cpp
+++ b/Engine/Source/Editor/UnrealEd/Private/Kismet2/Kismet2.cpp
@@ -1884,7 +1884,7 @@ void CreateChildActorComponentsForActors(const FBlueprintAssemblyProps& Assembly
 	// unless we rename these temporary components out of the way
 	for (UActorComponent* CAC : ChildActorComponents)
 	{
-		CAC->Rename(nullptr, nullptr, REN_DoNotDirty | REN_DontCreateRedirectors);
+		CAC->Rename(nullptr, nullptr, REN_DoNotDirty | REN_DontCreateRedirectors | REN_AllowPackageLinkerMismatch);
 	}
 }
 
@@ -2533,7 +2533,7 @@ void FKismetEditorUtilities::UpgradeCosmeticallyStaleBlueprint(UBlueprint* Bluep
 		if ((OldStateGraph != NULL) && (CollidingObject == NULL))
 		{
 			check(!OldStateGraph->HasAnyFlags(RF_Public));
-			OldStateGraph->Rename(*(UEdGraphSchema_K2::GN_EventGraph.ToString()), OldStateGraph->GetOuter(), REN_DoNotDirty);
+			OldStateGraph->Rename(*(UEdGraphSchema_K2::GN_EventGraph.ToString()), OldStateGraph->GetOuter(), REN_DoNotDirty | REN_AllowPackageLinkerMismatch);
 			Blueprint->Status = BS_Dirty;
 		}
 	}
diff --git a/Engine/Source/Editor/UnrealEd/Private/Kismet2/KismetReinstanceUtilities.cpp b/Engine/Source/Editor/UnrealEd/Private/Kismet2/KismetReinstanceUtilities.cpp
--- a/Engine/Source/Editor/UnrealEd/Private/Kismet2/KismetReinstanceUtilities.cpp
+++ b/Engine/Source/Editor/UnrealEd/Private/Kismet2/KismetReinstanceUtilities.cpp
@@ -1333,7 +1333,7 @@ void FBlueprintCompileReinstancer::PropagateSparseClassDataToNewClass(UClass* Ne
 
 	if (SparseClassDataStruct == OriginalSCDStruct && SparseClassDataStruct->GetOuter() == DuplicatedClass)
 	{
-		SparseClassDataStruct->Rename(nullptr, NewClass, REN_DoNotDirty | REN_DontCreateRedirectors | REN_NonTransactional);
+		SparseClassDataStruct->Rename(nullptr, NewClass, REN_DoNotDirty | REN_DontCreateRedirectors | REN_NonTransactional | REN_AllowPackageLinkerMismatch);
 	}
 	NewClass->SparseClassDataStruct = SparseClassDataStruct;
 	NewClass->CreateSparseClassData();
@@ -2061,7 +2061,7 @@ void FBlueprintCompileReinstancer::MoveDependentSkelToReinst(UClass* const Owner
 		if (OldCDO)
 		{
 			CurClass->SetDefaultObject(nullptr);
-			OldCDO->Rename(nullptr, ReinstClass->GetOuter(), REN_DoNotDirty | REN_DontCreateRedirectors | REN_NonTransactional);
+			OldCDO->Rename(nullptr, ReinstClass->GetOuter(), REN_DoNotDirty | REN_DontCreateRedirectors | REN_NonTransactional | REN_AllowPackageLinkerMismatch);
 			ReinstClass->SetDefaultObject(OldCDO);
 			OldCDO->SetClass(ReinstClass);
 		}
@@ -2141,7 +2141,7 @@ UClass* FBlueprintCompileReinstancer::MoveCDOToNewClass(UClass* OwnerClass, cons
 		UObject* OwnedArchetype = OwnedObjects[I];
 		if(OwnedArchetype->HasAnyFlags(RF_ArchetypeObject))
 		{
-			OwnedArchetype->Rename(*OriginalNames[I].ToString(), OwnerClass, REN_DoNotDirty | REN_DontCreateRedirectors | REN_AllowPackageLinkerMismatch | REN_NonTransactional);
+			OwnedArchetype->Rename(*OriginalNames[I].ToString(), OwnerClass, REN_DoNotDirty | REN_DontCreateRedirectors | REN_NonTransactional | REN_AllowPackageLinkerMismatch);
 		}
 	}
 
@@ -2253,7 +2253,7 @@ void FBlueprintCompileReinstancer::ReplaceObjectHelper(UObject*& OldObject, UCla
 			for (UObject* OldArchetypeObject : OldArchetypeObjects)
 			{
 				OldToNewNameMap.Add(OldArchetypeObject, OldName);
-				OldArchetypeObject->Rename(*OldArchetypeName, OldArchetypeObject->GetOuter(), REN_DoNotDirty | REN_DontCreateRedirectors);
+				OldArchetypeObject->Rename(*OldArchetypeName, OldArchetypeObject->GetOuter(), REN_DoNotDirty | REN_DontCreateRedirectors | REN_AllowPackageLinkerMismatch);
 			}
 		}
 		else
@@ -2464,7 +2464,7 @@ static void ReplaceActorHelper(AActor* OldActor, UClass* OldClass, UObject*& New
 	// Don't go through AActor::Rename here because we aren't changing outers (the actor's level) and we also don't want to reset loaders
 	// if the actor is using an external package. We really just want to rename that actor out of the way so we can spawn the new one in
 	// the exact same package, keeping the package name intact.
-	OldActor->UObject::Rename(nullptr, OldActor->GetOuter(), REN_DoNotDirty | REN_DontCreateRedirectors);
+	OldActor->UObject::Rename(nullptr, OldActor->GetOuter(), REN_DoNotDirty | REN_DontCreateRedirectors | REN_AllowPackageLinkerMismatch);
 
 	const bool bPackageNewlyCreated = OldActor->GetExternalPackage() && OldActor->GetExternalPackage()->HasAnyPackageFlags(PKG_NewlyCreated);
 
@@ -3428,12 +3428,12 @@ void FBlueprintCompileReinstancer::PreCreateSubObjectsForReinstantiation_Inner(
 					// This happens on object that were reinstantiated and were not yet outered to the right object because it was not yet created
 					checkf(OldInstancedSubObjects && !OldInstancedSubObjects->Contains(*AlreadyCreatedSubObject), TEXT("The object must have already been created but should have been skipped by the fact that the old object was not referencing it"));
 				}
-				(*AlreadyCreatedSubObject)->Rename(nullptr, NewUObject, REN_DoNotDirty | REN_DontCreateRedirectors | REN_NonTransactional);
+				(*AlreadyCreatedSubObject)->Rename(nullptr, NewUObject, REN_DoNotDirty | REN_DontCreateRedirectors | REN_NonTransactional | REN_AllowPackageLinkerMismatch);
 			}
 			// Could still be already re-instantiated but be also replaced
 			else if (UObject*const* AlreadyReplacedSubObject = OldToNewInstanceMap && OldToNewInstanceMap->FindKey(OldSubObject) ? &OldSubObject : nullptr)
 			{
-				(*AlreadyReplacedSubObject)->Rename(nullptr, NewUObject, REN_DoNotDirty | REN_DontCreateRedirectors | REN_NonTransactional);
+				(*AlreadyReplacedSubObject)->Rename(nullptr, NewUObject, REN_DoNotDirty | REN_DontCreateRedirectors | REN_NonTransactional | REN_AllowPackageLinkerMismatch);
 			}
 			else
 			{
diff --git a/Engine/Source/Editor/UnrealEd/Private/ObjectTools.cpp b/Engine/Source/Editor/UnrealEd/Private/ObjectTools.cpp
--- a/Engine/Source/Editor/UnrealEd/Private/ObjectTools.cpp
+++ b/Engine/Source/Editor/UnrealEd/Private/ObjectTools.cpp
@@ -1927,7 +1927,7 @@ namespace ObjectTools
 
 				if ( Redirector->Rename(*ObjName.ToString(), NULL, REN_Test) )
 				{
-					Redirector->Rename(*ObjName.ToString(), NULL, REN_DontCreateRedirectors | REN_NonTransactional);
+					Redirector->Rename(*ObjName.ToString(), NULL, REN_DontCreateRedirectors | REN_NonTransactional | REN_AllowPackageLinkerMismatch);
 					FAssetRegistryModule::AssetCreated(Redirector);
 				}
 				else
diff --git a/Engine/Source/Editor/UnrealEd/Private/PlayLevel.cpp b/Engine/Source/Editor/UnrealEd/Private/PlayLevel.cpp
--- a/Engine/Source/Editor/UnrealEd/Private/PlayLevel.cpp
+++ b/Engine/Source/Editor/UnrealEd/Private/PlayLevel.cpp
@@ -567,7 +567,7 @@ void UEditorEngine::EndPlayMap()
 			FName NewName = MakeUniqueObjectName(nullptr, UPackage::StaticClass());
 			UE_LOG(LogTemp, Log, TEXT("Renaming PIE package from '%s' to '%s' to prevent future name collisions."),
 				*ObjectPackage->GetName(), *NewName.ToString());
-			ObjectPackage->Rename(*NewName.ToString(), nullptr, REN_DontCreateRedirectors | REN_DoNotDirty | REN_NonTransactional);
+			ObjectPackage->Rename(*NewName.ToString(), nullptr, REN_DontCreateRedirectors | REN_DoNotDirty | REN_NonTransactional | REN_AllowPackageLinkerMismatch);
 		}
 	}
 
diff --git a/Engine/Source/Editor/UnrealEd/Private/WorldPartition/WorldPartitionRenameDuplicateBuilder.cpp b/Engine/Source/Editor/UnrealEd/Private/WorldPartition/WorldPartitionRenameDuplicateBuilder.cpp
--- a/Engine/Source/Editor/UnrealEd/Private/WorldPartition/WorldPartitionRenameDuplicateBuilder.cpp
+++ b/Engine/Source/Editor/UnrealEd/Private/WorldPartition/WorldPartitionRenameDuplicateBuilder.cpp
@@ -284,7 +284,7 @@ bool UWorldPartitionRenameDuplicateBuilder::RunInternal(UWorld* World, const FCe
 					UPackage* PreviousActorPackage = Actor->GetExternalPackage();
 
 					// Rename Actor first so new package gets created
-					Actor->Rename(nullptr, NewWorld->PersistentLevel, REN_NonTransactional | REN_DontCreateRedirectors | REN_DoNotDirty);
+					Actor->Rename(nullptr, NewWorld->PersistentLevel, REN_NonTransactional | REN_DontCreateRedirectors | REN_DoNotDirty | REN_AllowPackageLinkerMismatch);
 
 					TArray<UObject*> DependantObjects;
 					ForEachObjectWithPackage(PreviousActorPackage, [&DependantObjects](UObject* Object)
@@ -296,7 +296,7 @@ bool UWorldPartitionRenameDuplicateBuilder::RunInternal(UWorld* World, const FCe
 					// Move dependant objects into the new actor package
 					for (UObject* DependantObject : DependantObjects) //-V1078
 					{
-						DependantObject->Rename(nullptr, Actor->GetExternalPackage(), REN_NonTransactional | REN_DontCreateRedirectors | REN_DoNotDirty);
+						DependantObject->Rename(nullptr, Actor->GetExternalPackage(), REN_NonTransactional | REN_DontCreateRedirectors | REN_DoNotDirty | REN_AllowPackageLinkerMismatch);
 					}
 
 					// Releases file handle so it can be deleted
@@ -332,7 +332,7 @@ bool UWorldPartitionRenameDuplicateBuilder::RunInternal(UWorld* World, const FCe
 					AActor* Actor = ActorReference.GetActor();
 					UPackage* NewActorPackage = Actor->GetExternalPackage();
 					check(Actor);
-					Actor->Rename(nullptr, World->PersistentLevel, REN_NonTransactional | REN_DontCreateRedirectors | REN_DoNotDirty);
+					Actor->Rename(nullptr, World->PersistentLevel, REN_NonTransactional | REN_DontCreateRedirectors | REN_DoNotDirty | REN_AllowPackageLinkerMismatch);
 
 					TArray<UObject*> DependantObjects;
 					ForEachObjectWithPackage(NewActorPackage, [&DependantObjects](UObject* Object)
@@ -344,7 +344,7 @@ bool UWorldPartitionRenameDuplicateBuilder::RunInternal(UWorld* World, const FCe
 					// Move back dependant objects into the previous actor package
 					for (UObject* DependantObject : DependantObjects) //-V1078
 					{
-						DependantObject->Rename(nullptr, Actor->GetExternalPackage(), REN_NonTransactional | REN_DontCreateRedirectors | REN_DoNotDirty);
+						DependantObject->Rename(nullptr, Actor->GetExternalPackage(), REN_NonTransactional | REN_DontCreateRedirectors | REN_DoNotDirty | REN_AllowPackageLinkerMismatch);
 					}
 				}
 				ActorReferences.Empty();
diff --git a/Engine/Source/Editor/UnrealEd/Private/WorldPartition/WorldPartitionResaveActorsBuilder.cpp b/Engine/Source/Editor/UnrealEd/Private/WorldPartition/WorldPartitionResaveActorsBuilder.cpp
--- a/Engine/Source/Editor/UnrealEd/Private/WorldPartition/WorldPartitionResaveActorsBuilder.cpp
+++ b/Engine/Source/Editor/UnrealEd/Private/WorldPartition/WorldPartitionResaveActorsBuilder.cpp
@@ -325,7 +325,7 @@ bool UWorldPartitionResaveActorsBuilder::RunInternal(UWorld* World, const FCellI
 							// Move dependant objects into the new world package temporarily
 							for (UObject* DependantObject : DependantObjects)
 							{
-								DependantObject->Rename(nullptr, WorldPackage, REN_NonTransactional | REN_DontCreateRedirectors);
+								DependantObject->Rename(nullptr, WorldPackage, REN_NonTransactional | REN_DontCreateRedirectors | REN_AllowPackageLinkerMismatch);
 							}
 
 							// Move actor in its new package
@@ -335,7 +335,7 @@ bool UWorldPartitionResaveActorsBuilder::RunInternal(UWorld* World, const FCellI
 							UPackage* NewActorPackage = Actor->GetExternalPackage();
 							for (UObject* DependantObject : DependantObjects)
 							{
-								DependantObject->Rename(nullptr, NewActorPackage, REN_NonTransactional | REN_DontCreateRedirectors);
+								DependantObject->Rename(nullptr, NewActorPackage, REN_NonTransactional | REN_DontCreateRedirectors | REN_AllowPackageLinkerMismatch);
 							}
 
 							PackagesToSave.Add(NewActorPackage);
diff --git a/Engine/Source/Runtime/ClothingSystemRuntimeCommon/Private/ClothingAsset.cpp b/Engine/Source/Runtime/ClothingSystemRuntimeCommon/Private/ClothingAsset.cpp
--- a/Engine/Source/Runtime/ClothingSystemRuntimeCommon/Private/ClothingAsset.cpp
+++ b/Engine/Source/Runtime/ClothingSystemRuntimeCommon/Private/ClothingAsset.cpp
@@ -1269,17 +1269,17 @@ void UClothingAssetCommon::PostLoad()
 	if (ClothSimConfig_DEPRECATED)
 	{
 		ClothSimConfig_DEPRECATED->ConditionalPostLoad();  // PostLoad old configs before replacing them
-		ClothSimConfig_DEPRECATED->Rename(nullptr, nullptr, REN_DoNotDirty | REN_DontCreateRedirectors | REN_NonTransactional);  // Rename the config so that the name doesn't collide with the new config map name
+		ClothSimConfig_DEPRECATED->Rename(nullptr, nullptr, REN_DoNotDirty | REN_DontCreateRedirectors | REN_NonTransactional | REN_AllowPackageLinkerMismatch);  // Rename the config so that the name doesn't collide with the new config map name
 	}
 	if (ChaosClothSimConfig_DEPRECATED)
 	{
 		ChaosClothSimConfig_DEPRECATED->ConditionalPostLoad();  // PostLoad old configs before replacing them
-		ChaosClothSimConfig_DEPRECATED->Rename(nullptr, nullptr, REN_DoNotDirty | REN_DontCreateRedirectors | REN_NonTransactional);  // Rename the config so that the name doesn't collide with the new config map name
+		ChaosClothSimConfig_DEPRECATED->Rename(nullptr, nullptr, REN_DoNotDirty | REN_DontCreateRedirectors | REN_NonTransactional | REN_AllowPackageLinkerMismatch);  // Rename the config so that the name doesn't collide with the new config map name
 	}
 	if (ClothSharedSimConfig_DEPRECATED)
 	{
 		ClothSharedSimConfig_DEPRECATED->ConditionalPostLoad();  // PostLoad old configs before replacing them
-		ClothSharedSimConfig_DEPRECATED->Rename(nullptr, nullptr, REN_DoNotDirty | REN_DontCreateRedirectors | REN_NonTransactional);  // Rename the config so that the name doesn't collide with the new config map name
+		ClothSharedSimConfig_DEPRECATED->Rename(nullptr, nullptr, REN_DoNotDirty | REN_DontCreateRedirectors | REN_NonTransactional | REN_AllowPackageLinkerMismatch);  // Rename the config so that the name doesn't collide with the new config map name
 	}
 
 	if (AddClothConfigs())
@@ -1554,7 +1554,7 @@ void UClothingAssetCommon::PropagateSharedConfigs(bool bMigrateSharedConfigToCon
 				// Fix the shared config outer if it is still a UClothingAssetCommon (the config must belong to the skeletal mesh, as it is shared between assets)
 				if (Cast<UClothingAssetCommon>(ClothSharedConfig->GetOuter()))
 				{
-					ClothSharedConfig->Rename(nullptr, SkeletalMesh, REN_DontCreateRedirectors | REN_NonTransactional);
+					ClothSharedConfig->Rename(nullptr, SkeletalMesh, REN_DontCreateRedirectors | REN_NonTransactional | REN_AllowPackageLinkerMismatch);
 				}
 
 				// Fix the shared config ownership, asset might have been copied and the shared config could still point to a different skeletal mesh
diff --git a/Engine/Source/Runtime/CoreUObject/Private/InstancedReferenceSubobjectHelper.cpp b/Engine/Source/Runtime/CoreUObject/Private/InstancedReferenceSubobjectHelper.cpp
--- a/Engine/Source/Runtime/CoreUObject/Private/InstancedReferenceSubobjectHelper.cpp
+++ b/Engine/Source/Runtime/CoreUObject/Private/InstancedReferenceSubobjectHelper.cpp
@@ -341,7 +341,7 @@ void FFindInstancedReferenceSubobjectHelper::Duplicate(
 								RealNewSubobject->Rename(
 									*(RealNewSubobject->GetName()),
 									NewObject,
-									REN_DoNotDirty | REN_DontCreateRedirectors | REN_NonTransactional);
+									REN_DoNotDirty | REN_DontCreateRedirectors | REN_NonTransactional | REN_AllowPackageLinkerMismatch);
 							}
 						}
 					}
diff --git a/Engine/Source/Runtime/CoreUObject/Private/Serialization/TransactionallySafeAsyncLoading.cpp b/Engine/Source/Runtime/CoreUObject/Private/Serialization/TransactionallySafeAsyncLoading.cpp
--- a/Engine/Source/Runtime/CoreUObject/Private/Serialization/TransactionallySafeAsyncLoading.cpp
+++ b/Engine/Source/Runtime/CoreUObject/Private/Serialization/TransactionallySafeAsyncLoading.cpp
@@ -222,7 +222,7 @@ public:
 						// in `FindObject` calls.
 						if (Package)
 						{
-							Package->Rename(*CustomPackageName.ToString(), nullptr, REN_DontCreateRedirectors | REN_NonTransactional | REN_DoNotDirty);
+							Package->Rename(*CustomPackageName.ToString(), nullptr, REN_DontCreateRedirectors | REN_NonTransactional | REN_DoNotDirty | REN_AllowPackageLinkerMismatch);
 						}
 
 						InCompletionDelegate.ExecuteIfBound(CustomPackageName, Package, MapPayload->Result);
diff --git a/Engine/Source/Runtime/CoreUObject/Private/UObject/Class.cpp b/Engine/Source/Runtime/CoreUObject/Private/UObject/Class.cpp
--- a/Engine/Source/Runtime/CoreUObject/Private/UObject/Class.cpp
+++ b/Engine/Source/Runtime/CoreUObject/Private/UObject/Class.cpp
@@ -1175,7 +1175,7 @@ void UStruct::Link(FArchive& Ar, bool bRelinkExistingProperties)
 	// Discard old wrapper objects used by property grids
 	for (UPropertyWrapper* Wrapper : PropertyWrappers)
 	{
-		Wrapper->Rename(nullptr, GetTransientPackage(), REN_DoNotDirty | REN_DontCreateRedirectors | REN_NonTransactional);
+		Wrapper->Rename(nullptr, GetTransientPackage(), REN_DoNotDirty | REN_DontCreateRedirectors | REN_NonTransactional| REN_AllowPackageLinkerMismatch);
 		Wrapper->RemoveFromRoot();
 	}
 	PropertyWrappers.Empty();
@@ -5784,7 +5784,7 @@ PRAGMA_DISABLE_DEPRECATION_WARNINGS
 	{
 		ImmutableDefaultObject->ClearFlags(RF_ClassDefaultObject | RF_ImmutableDefaultObject | RF_ArchetypeObject);
 		ImmutableDefaultObject->RemoveFromRoot();
-		ImmutableDefaultObject->Rename(*MakeUniqueObjectName(GetTransientPackage(), this).ToString(), GetTransientPackage(), REN_DoNotDirty | REN_NonTransactional | REN_DontCreateRedirectors | REN_ForceNoResetLoaders);
+		ImmutableDefaultObject->Rename(*MakeUniqueObjectName(GetTransientPackage(), this).ToString(), GetTransientPackage(), REN_DoNotDirty | REN_NonTransactional | REN_DontCreateRedirectors | REN_AllowPackageLinkerMismatch);
 		ImmutableDefaultObject = nullptr;
 	}
 #endif
diff --git a/Engine/Source/Runtime/CoreUObject/Private/UObject/CookedMetaData.cpp b/Engine/Source/Runtime/CoreUObject/Private/UObject/CookedMetaData.cpp
--- a/Engine/Source/Runtime/CoreUObject/Private/UObject/CookedMetaData.cpp
+++ b/Engine/Source/Runtime/CoreUObject/Private/UObject/CookedMetaData.cpp
@@ -18,7 +18,7 @@ void PrepareCookedMetaDataForPurge(UObject* CookedMetaDataPtr)
 	{
 		FNameBuilder BaseMetaDataName(CookedMetaDataPtr->GetFName());
 		BaseMetaDataName << TEXT("_PURGED");
-		CookedMetaDataPtr->Rename(FNameBuilder(MakeUniqueObjectName(CookedMetaDataPtr->GetOuter(), CookedMetaDataPtr->GetClass(), FName(BaseMetaDataName))).ToString(), nullptr, REN_DoNotDirty | REN_DontCreateRedirectors | REN_NonTransactional);
+		CookedMetaDataPtr->Rename(FNameBuilder(MakeUniqueObjectName(CookedMetaDataPtr->GetOuter(), CookedMetaDataPtr->GetClass(), FName(BaseMetaDataName))).ToString(), nullptr, REN_DoNotDirty | REN_DontCreateRedirectors | REN_NonTransactional | REN_AllowPackageLinkerMismatch);
 	}
 
 	CookedMetaDataPtr->ClearFlags(RF_Standalone | RF_Public);
diff --git a/Engine/Source/Runtime/CoreUObject/Private/UObject/ObjectRename.cpp b/Engine/Source/Runtime/CoreUObject/Private/UObject/ObjectRename.cpp
--- a/Engine/Source/Runtime/CoreUObject/Private/UObject/ObjectRename.cpp
+++ b/Engine/Source/Runtime/CoreUObject/Private/UObject/ObjectRename.cpp
@@ -16,6 +16,6 @@ namespace UE::Object
 		TStringBuilder<FName::StringBufferSize> NewNameString;
 		NewNameString << NewName;
 		UE_LOG(LogObj, Log, TEXT("Renaming leaked package %s to %s"), *Package->GetName(), *NewNameString);
-		Package->Rename(*NewNameString, nullptr, REN_DontCreateRedirectors | REN_NonTransactional);
+		Package->Rename(*NewNameString, nullptr, REN_DontCreateRedirectors | REN_NonTransactional | REN_AllowPackageLinkerMismatch);
 	}
 }
diff --git a/Engine/Source/Runtime/CoreUObject/Private/UObject/PackageReload.cpp b/Engine/Source/Runtime/CoreUObject/Private/UObject/PackageReload.cpp
--- a/Engine/Source/Runtime/CoreUObject/Private/UObject/PackageReload.cpp
+++ b/Engine/Source/Runtime/CoreUObject/Private/UObject/PackageReload.cpp
@@ -463,7 +463,7 @@ UPackage* LoadReplacementPackage(UPackage* InExistingPackage, const uint32 InLoa
 	const FString ExistingPackageName = InExistingPackage->GetName();
 
 	// Rename the old package, and then load the new one in its place
-	const ERenameFlags PkgRenameFlags = REN_DoNotDirty | REN_DontCreateRedirectors | REN_NonTransactional | REN_SkipGeneratedClasses;
+	const ERenameFlags PkgRenameFlags = REN_DoNotDirty | REN_DontCreateRedirectors | REN_NonTransactional | REN_SkipGeneratedClasses | REN_AllowPackageLinkerMismatch;
 	InExistingPackage->Rename(*MakeUniqueObjectName(Cast<UPackage>(InExistingPackage->GetOuter()), UPackage::StaticClass(), *FString::Printf(TEXT("%s_DEADPACKAGE"), *InExistingPackage->GetName())).ToString(), nullptr, PkgRenameFlags);
 	MarkPackageReplaced(InExistingPackage);
 	UPackage* NewPackage = LoadPackage(Cast<UPackage>(InExistingPackage->GetOuter()), *ExistingPackageName, InLoadFlags);
diff --git a/Engine/Source/Runtime/CoreUObject/Private/UObject/PropertyBaseObject.cpp b/Engine/Source/Runtime/CoreUObject/Private/UObject/PropertyBaseObject.cpp
--- a/Engine/Source/Runtime/CoreUObject/Private/UObject/PropertyBaseObject.cpp
+++ b/Engine/Source/Runtime/CoreUObject/Private/UObject/PropertyBaseObject.cpp
@@ -434,7 +434,7 @@ const TCHAR* FObjectPropertyBase::ImportText_Internal( const TCHAR* InBuffer, vo
 		// 
 		if (UObject* ExistingObject = static_cast<UObject*>(FindObjectWithOuter(Parent, nullptr, DesiredName)))
 		{
-			ExistingObject->Rename(nullptr, nullptr, REN_DontCreateRedirectors);
+			ExistingObject->Rename(nullptr, nullptr, REN_DontCreateRedirectors | REN_AllowPackageLinkerMismatch);
 		}
 		
 		const EObjectFlags MaskedOuterFlags = Parent ? Parent->GetMaskedFlags(RF_PropagateToSubObjects) : RF_AllFlags;
diff --git a/Engine/Source/Runtime/Engine/Private/ActorPartition/ActorPartitionSubsystem.cpp b/Engine/Source/Runtime/Engine/Private/ActorPartition/ActorPartitionSubsystem.cpp
--- a/Engine/Source/Runtime/Engine/Private/ActorPartition/ActorPartitionSubsystem.cpp
+++ b/Engine/Source/Runtime/Engine/Private/ActorPartition/ActorPartitionSubsystem.cpp
@@ -281,7 +281,7 @@ public:
 				ExistingActor->Modify();
 				// Don't go through AActor::Rename here because we aren't changing outers (the actor's level). We just want to rename that actor 
 				// out of the way so we can spawn the new one in the exact same package, keeping the package name intact.
-				ExistingActor->UObject::Rename(nullptr, nullptr, REN_DontCreateRedirectors | REN_DoNotDirty | REN_NonTransactional);
+				ExistingActor->UObject::Rename(nullptr, nullptr, REN_DontCreateRedirectors | REN_DoNotDirty | REN_NonTransactional | REN_AllowPackageLinkerMismatch);
 				
 				// Reuse ActorGuid so that ActorDesc can be updated on save
 				SpawnParams.OverrideActorGuid = ExistingActor->GetActorGuid();
diff --git a/Engine/Source/Runtime/Engine/Private/Blueprint.cpp b/Engine/Source/Runtime/Engine/Private/Blueprint.cpp
--- a/Engine/Source/Runtime/Engine/Private/Blueprint.cpp
+++ b/Engine/Source/Runtime/Engine/Private/Blueprint.cpp
@@ -531,7 +531,7 @@ bool UBlueprint::RenameGeneratedClasses(const TCHAR* NewName, UObject* NewOuter,
 				if (UObject* ExistingCDO = StaticFindObjectFast(UObject::StaticClass(), NewOuter, *NewCDOName))
 				{
 					const FName NewName = MakeUniqueObjectName(NewOuter, ExistingCDO->GetClass(), *NewCDOName);
-					ExistingCDO->Rename(*(NewName.ToString()), NewOuter, InFlags | REN_DontCreateRedirectors);
+					ExistingCDO->Rename(*(NewName.ToString()), NewOuter, InFlags | REN_DontCreateRedirectors | REN_AllowPackageLinkerMismatch);
 				}
 			};
 
@@ -547,7 +547,7 @@ bool UBlueprint::RenameGeneratedClasses(const TCHAR* NewName, UObject* NewOuter,
 
 					if (bFoundRedirector)
 					{
-						Redirector->Rename(nullptr, GetTransientPackage(), REN_DontCreateRedirectors);
+						Redirector->Rename(nullptr, GetTransientPackage(), REN_DontCreateRedirectors | REN_AllowPackageLinkerMismatch);
 					}
 				}
 			};
@@ -2107,7 +2107,7 @@ bool UBlueprint::ChangeOwnerOfTemplates()
 			{
 				if (Component->GetOuter() == this)
 				{
-					const bool bRenamed = Component->Rename(*Component->GetName(), BPGClass, REN_DoNotDirty);
+					const bool bRenamed = Component->Rename(*Component->GetName(), BPGClass, REN_DoNotDirty | REN_AllowPackageLinkerMismatch);
 					ensure(bRenamed);
 					bIsStillStale |= !bRenamed;
 					bMigratedOwner = true;
@@ -2128,7 +2128,7 @@ bool UBlueprint::ChangeOwnerOfTemplates()
 				{
 					const FName OldTemplateName = Template->GetFName();
 					ensure(!OldTemplateName.ToString().EndsWith(UTimelineTemplate::TemplatePostfix));
-					const bool bRenamed = Template->Rename(*UTimelineTemplate::TimelineVariableNameToTemplateName(Template->GetFName()), BPGClass, REN_DoNotDirty);
+					const bool bRenamed = Template->Rename(*UTimelineTemplate::TimelineVariableNameToTemplateName(Template->GetFName()), BPGClass, REN_DoNotDirty | REN_AllowPackageLinkerMismatch);
 					ensure(bRenamed);
 					bIsStillStale |= !bRenamed;
 					ensure(OldTemplateName == Template->GetVariableName());
@@ -2141,7 +2141,7 @@ bool UBlueprint::ChangeOwnerOfTemplates()
 		{
 			if (Curve && (Curve->GetOuter() == this))
 			{
-				const bool bRenamed = Curve->Rename(FUniqueNewNameHelper(Curve->GetName(), BPGClass).Get(), BPGClass, REN_DoNotDirty);
+				const bool bRenamed = Curve->Rename(FUniqueNewNameHelper(Curve->GetName(), BPGClass).Get(), BPGClass, REN_DoNotDirty | REN_AllowPackageLinkerMismatch);
 				ensure(bRenamed);
 				bIsStillStale |= !bRenamed;
 			}
@@ -2151,7 +2151,7 @@ bool UBlueprint::ChangeOwnerOfTemplates()
 		{
 			if(SCS->GetOuter() == this)
 			{
-				const bool bRenamed = SCS->Rename(FUniqueNewNameHelper(SCS->GetName(), BPGClass).Get(), BPGClass, REN_DoNotDirty);
+				const bool bRenamed = SCS->Rename(FUniqueNewNameHelper(SCS->GetName(), BPGClass).Get(), BPGClass, REN_DoNotDirty | REN_AllowPackageLinkerMismatch);
 				ensure(bRenamed);
 				bIsStillStale |= !bRenamed;
 				bMigratedOwner = true;
@@ -2162,7 +2162,7 @@ bool UBlueprint::ChangeOwnerOfTemplates()
 				UActorComponent* Component = SCSNode ? ToRawPtr(SCSNode->ComponentTemplate) : NULL;
 				if (Component && Component->GetOuter() == this)
 				{
-					const bool bRenamed = Component->Rename(FUniqueNewNameHelper(Component->GetName(), BPGClass).Get(), BPGClass, REN_DoNotDirty);
+					const bool bRenamed = Component->Rename(FUniqueNewNameHelper(Component->GetName(), BPGClass).Get(), BPGClass, REN_DoNotDirty | REN_AllowPackageLinkerMismatch);
 					ensure(bRenamed);
 					bIsStillStale |= !bRenamed;
 					bMigratedOwner = true;
diff --git a/Engine/Source/Runtime/Engine/Private/Brush.cpp b/Engine/Source/Runtime/Engine/Private/Brush.cpp
--- a/Engine/Source/Runtime/Engine/Private/Brush.cpp
+++ b/Engine/Source/Runtime/Engine/Private/Brush.cpp
@@ -404,6 +404,8 @@ void ABrush::PostLoad()
 		// UModelFactory::FactoryCreateText was passing in the ABrush as the Outer instead of the UModel
 		if (Brush->Polys->GetOuter() == this)
 		{
+			// Complete loading of the polys before renaming
+			Brush->Polys->ConditionalPostLoad();
 			Brush->Polys->Rename(*Brush->Polys->GetName(), Brush);
 		}
 	}
diff --git a/Engine/Source/Runtime/Engine/Private/ComponentInstanceDataCache.cpp b/Engine/Source/Runtime/Engine/Private/ComponentInstanceDataCache.cpp
--- a/Engine/Source/Runtime/Engine/Private/ComponentInstanceDataCache.cpp
+++ b/Engine/Source/Runtime/Engine/Private/ComponentInstanceDataCache.cpp
@@ -546,7 +546,7 @@ void FActorComponentInstanceData::ApplyToComponent(UActorComponent* Component, c
 				{
 					OtherObject->Rename(nullptr, GetTransientPackage(), REN_DontCreateRedirectors);
 				}
-				DuplicatedObjectData.DuplicatedObject->Rename(nullptr, Component, REN_DontCreateRedirectors);
+				DuplicatedObjectData.DuplicatedObject->Rename(nullptr, Component, REN_DontCreateRedirectors | REN_AllowPackageLinkerMismatch);
 			}
 		}
 
diff --git a/Engine/Source/Runtime/Engine/Private/Components/ChildActorComponent.cpp b/Engine/Source/Runtime/Engine/Private/Components/ChildActorComponent.cpp
--- a/Engine/Source/Runtime/Engine/Private/Components/ChildActorComponent.cpp
+++ b/Engine/Source/Runtime/Engine/Private/Components/ChildActorComponent.cpp
@@ -500,9 +500,9 @@ void UChildActorComponent::ApplyComponentInstanceData(FChildActorComponentInstan
 		   ChildActor->GetFName() != ChildActorName)
 		{
 			const FString ChildActorNameString = ChildActorName.ToString();
-			if (ChildActor->Rename(*ChildActorNameString, nullptr, REN_Test))
+			if (ChildActor->Rename(*ChildActorNameString, nullptr, REN_Test | REN_AllowPackageLinkerMismatch))
 			{
-				ChildActor->Rename(*ChildActorNameString, nullptr, REN_DoNotDirty);
+				ChildActor->Rename(*ChildActorNameString, nullptr, REN_DoNotDirty | REN_AllowPackageLinkerMismatch);
 #if WITH_EDITOR
 				ChildActor->ClearActorLabel();
 #endif
@@ -584,7 +584,7 @@ void UChildActorComponent::SetChildActorClass(TSubclassOf<AActor> Class, AActor*
 					ExistingChildActorTemplate->Rename(nullptr, GetTransientPackage(), REN_DontCreateRedirectors);
 				}
 
-				ChildActorTemplate->Rename(*TemplateName, this, REN_DoNotDirty | REN_DontCreateRedirectors);
+				ChildActorTemplate->Rename(*TemplateName, this, REN_DoNotDirty | REN_DontCreateRedirectors | REN_AllowPackageLinkerMismatch);
 			}
 		}
 		else if (ChildActorTemplate)
@@ -719,7 +719,7 @@ void UChildActorComponent::OnChildActorDestroyed(AActor* Actor)
 			if (!IsGarbageCollecting())
 			{
 				const FString ObjectBaseName = FString::Printf(TEXT("DESTROYED_%s_CHILDACTOR"), *ChildClass->GetName());
-				Actor->Rename(*MakeUniqueObjectName(Actor->GetOuter(), ChildClass, *ObjectBaseName).ToString(), nullptr, REN_DoNotDirty);
+				Actor->Rename(*MakeUniqueObjectName(Actor->GetOuter(), ChildClass, *ObjectBaseName).ToString(), nullptr, REN_DoNotDirty | REN_AllowPackageLinkerMismatch);
 			}
 			else
 			{
diff --git a/Engine/Source/Runtime/Engine/Private/EdGraph/EdGraph.cpp b/Engine/Source/Runtime/Engine/Private/EdGraph/EdGraph.cpp
--- a/Engine/Source/Runtime/Engine/Private/EdGraph/EdGraph.cpp
+++ b/Engine/Source/Runtime/Engine/Private/EdGraph/EdGraph.cpp
@@ -308,7 +308,7 @@ void UEdGraph::MoveNodesToAnotherGraph(UEdGraph* DestinationGraph, bool bIsLoadi
 		{
 			// Let the name be autogenerated, to automatically avoid naming conflicts
 			// Since this graph is always going to come from a cloned source graph, user readable names can come from the remap stored in a MessageLog.
-			Node->Rename(/*NewName=*/ nullptr, /*NewOuter=*/ DestinationGraph, REN_DontCreateRedirectors | (bInIsCompiling ? REN_NonTransactional : 0));
+			Node->Rename(/*NewName=*/ nullptr, /*NewOuter=*/ DestinationGraph, REN_DontCreateRedirectors | (bInIsCompiling ? REN_NonTransactional : 0) | (bIsLoading ? REN_AllowPackageLinkerMismatch : 0));
 	
 			DestinationGraph->Nodes.Add(Node);
 		}
diff --git a/Engine/Source/Runtime/Engine/Private/InheritableComponentHandler.cpp b/Engine/Source/Runtime/Engine/Private/InheritableComponentHandler.cpp
--- a/Engine/Source/Runtime/Engine/Private/InheritableComponentHandler.cpp
+++ b/Engine/Source/Runtime/Engine/Private/InheritableComponentHandler.cpp
@@ -555,11 +555,11 @@ void UInheritableComponentHandler::FixComponentTemplateName(UActorComponent* Com
 	// PostLoad() validation and see that it still doesn't match its original template name.
 	if (UObject* ExistingObject = (UObject*)FindObjectWithOuter(ComponentTemplate->GetOuter(), nullptr, NewName))
 	{
-		ExistingObject->Rename(nullptr, nullptr, REN_DontCreateRedirectors);
+		ExistingObject->Rename(nullptr, nullptr, REN_AllowPackageLinkerMismatch | REN_DontCreateRedirectors);
 	}
 
 	// Now that we're sure there are no collisions with other records, we can safely rename this one to its new name.
-	ComponentTemplate->Rename(*NewName.ToString(), nullptr, REN_DontCreateRedirectors);
+	ComponentTemplate->Rename(*NewName.ToString(), nullptr, REN_AllowPackageLinkerMismatch | REN_DontCreateRedirectors);
 }
 
 // FComponentOverrideRecord
diff --git a/Engine/Source/Runtime/Engine/Private/Level.cpp b/Engine/Source/Runtime/Engine/Private/Level.cpp
--- a/Engine/Source/Runtime/Engine/Private/Level.cpp
+++ b/Engine/Source/Runtime/Engine/Private/Level.cpp
@@ -1490,11 +1490,11 @@ void ULevel::PostLoad()
 			// If there is already something there with that name, rename it to something else.
 			if (UObject* ExistingObject = StaticFindObject(nullptr, LevelScriptBlueprint->GetOuter(), *OuterWorld->GetName()))
 			{
-				ExistingObject->Rename(nullptr, nullptr, REN_DoNotDirty | REN_DontCreateRedirectors | REN_NonTransactional);
+				ExistingObject->Rename(nullptr, nullptr, REN_DoNotDirty | REN_DontCreateRedirectors | REN_NonTransactional | REN_AllowPackageLinkerMismatch);
 			}
 
 			// Use LevelScriptBlueprint->GetOuter() instead of NULL to make sure the generated top level objects are moved appropriately
-			LevelScriptBlueprint->Rename(*OuterWorld->GetName(), LevelScriptBlueprint->GetOuter(), REN_DoNotDirty | REN_DontCreateRedirectors | REN_NonTransactional | REN_SkipGeneratedClasses);
+			LevelScriptBlueprint->Rename(*OuterWorld->GetName(), LevelScriptBlueprint->GetOuter(), REN_DoNotDirty | REN_DontCreateRedirectors | REN_NonTransactional | REN_SkipGeneratedClasses | REN_AllowPackageLinkerMismatch);
 		}
 	}
 
@@ -4720,7 +4720,7 @@ ULevelScriptBlueprint* ULevel::GetLevelScriptBlueprint(bool bDontCreate)
 		// If there is already something there with that name, rename it to something else.
 		if (UObject* ExistingObject = StaticFindObject(nullptr, this, *LevelScriptName))
 		{
-			ExistingObject->Rename(nullptr, nullptr, REN_DoNotDirty | REN_DontCreateRedirectors | REN_NonTransactional);
+			ExistingObject->Rename(nullptr, nullptr, REN_DoNotDirty | REN_DontCreateRedirectors | REN_NonTransactional | REN_AllowPackageLinkerMismatch);
 		}
 
 		// If no blueprint is found, create one. 
diff --git a/Engine/Source/Runtime/Engine/Private/Materials/MaterialInterface.cpp b/Engine/Source/Runtime/Engine/Private/Materials/MaterialInterface.cpp
--- a/Engine/Source/Runtime/Engine/Private/Materials/MaterialInterface.cpp
+++ b/Engine/Source/Runtime/Engine/Private/Materials/MaterialInterface.cpp
@@ -129,7 +129,7 @@ static bool SafeRenameEditorOnlyData(UMaterialInterface* MaterialInterface, cons
 			return false;
 		}
 		FName UniqueDummyName = MakeUniqueObjectName(GetTransientPackage(), UMaterialInterfaceEditorOnlyData::StaticClass());
-		EditorOnlyWithIncorrectName->Rename(*UniqueDummyName.ToString(), GetTransientPackage(), REN_NonTransactional | REN_DontCreateRedirectors);
+		EditorOnlyWithIncorrectName->Rename(*UniqueDummyName.ToString(), GetTransientPackage(), REN_NonTransactional | REN_DontCreateRedirectors | REN_AllowPackageLinkerMismatch);
 		EditorOnlyWithIncorrectName->MarkAsGarbage();
 	}
 
diff --git a/Engine/Source/Runtime/Engine/Private/SCS_Node.cpp b/Engine/Source/Runtime/Engine/Private/SCS_Node.cpp
--- a/Engine/Source/Runtime/Engine/Private/SCS_Node.cpp
+++ b/Engine/Source/Runtime/Engine/Private/SCS_Node.cpp
@@ -384,7 +384,7 @@ void USCS_Node::RenameComponentTemplate(UActorComponent* ComponentTemplate, cons
 
 		// Rename the component template (archetype) - note that this can be called during compile-on-load, so we include the flag not to reset the BPGC's package loader.
 		const FString NewComponentName = NewName.ToString();
-		ComponentTemplate->Rename(*(NewComponentName + USimpleConstructionScript::ComponentTemplateNameSuffix), nullptr, REN_DontCreateRedirectors);
+		ComponentTemplate->Rename(*(NewComponentName + USimpleConstructionScript::ComponentTemplateNameSuffix), nullptr, REN_DontCreateRedirectors | REN_AllowPackageLinkerMismatch);
 
 		// Rename all component instances to match the updated variable name
 		for (UObject* ArchetypeInstance : ArchetypeInstances)
@@ -402,7 +402,7 @@ void USCS_Node::RenameComponentTemplate(UActorComponent* ComponentTemplate, cons
 					Actor->CheckComponentInstanceName(NewName);
 				}
 
-				ArchetypeInstance->Rename(*NewComponentName, nullptr, REN_DontCreateRedirectors);
+				ArchetypeInstance->Rename(*NewComponentName, nullptr, REN_DontCreateRedirectors | REN_AllowPackageLinkerMismatch);
 			}
 		}
 	}
diff --git a/Engine/Source/Runtime/Engine/Private/SimpleConstructionScript.cpp b/Engine/Source/Runtime/Engine/Private/SimpleConstructionScript.cpp
--- a/Engine/Source/Runtime/Engine/Private/SimpleConstructionScript.cpp
+++ b/Engine/Source/Runtime/Engine/Private/SimpleConstructionScript.cpp
@@ -1404,12 +1404,12 @@ USCS_Node* USimpleConstructionScript::CreateNode(UClass* NewComponentClass, FNam
 	UObject* Collision = FindObject<UObject>(Blueprint->GeneratedClass, *Name);
 	while(Collision)
 	{
-		Collision->Rename(nullptr, GetTransientPackage(), REN_DoNotDirty | REN_DontCreateRedirectors);
+		Collision->Rename(nullptr, GetTransientPackage(), REN_DoNotDirty | REN_DontCreateRedirectors | REN_AllowPackageLinkerMismatch);
 		Collision = FindObject<UObject>(Blueprint->GeneratedClass, *Name);
 	}
 
 	// Now set the actual name and outer to the BPGC.
-	NewComponentTemplate->Rename(*Name, Blueprint->GeneratedClass, REN_DoNotDirty | REN_DontCreateRedirectors);
+	NewComponentTemplate->Rename(*Name, Blueprint->GeneratedClass, REN_DoNotDirty | REN_DontCreateRedirectors | REN_AllowPackageLinkerMismatch);
 
 	return CreateNodeImpl(NewComponentTemplate, NewComponentVariableName);
 }
diff --git a/Engine/Source/Runtime/Engine/Private/SkeletalMesh.cpp b/Engine/Source/Runtime/Engine/Private/SkeletalMesh.cpp
--- a/Engine/Source/Runtime/Engine/Private/SkeletalMesh.cpp
+++ b/Engine/Source/Runtime/Engine/Private/SkeletalMesh.cpp
@@ -4926,7 +4926,7 @@ bool USkeletalMesh::RenameMorphTarget(FName InOldName, FName InNewName)
 	}
 
 	// Rename the morph target itself
-	MorphTarget->Rename(*InNewName.ToString(), nullptr, REN_DontCreateRedirectors);
+	MorphTarget->Rename(*InNewName.ToString(), nullptr, REN_DontCreateRedirectors | REN_AllowPackageLinkerMismatch);
 
 	//Clean up the LodInfo Imported morph target source filename we must also rename the entry
 	for (int32 InternalLodIndex = 0; InternalLodIndex < GetLODNum(); ++InternalLodIndex)
diff --git a/Engine/Source/Runtime/Engine/Private/UnrealEngine.cpp b/Engine/Source/Runtime/Engine/Private/UnrealEngine.cpp
--- a/Engine/Source/Runtime/Engine/Private/UnrealEngine.cpp
+++ b/Engine/Source/Runtime/Engine/Private/UnrealEngine.cpp
@@ -17106,7 +17106,7 @@ void UEngine::CheckAndHandleStaleWorldObjectReferences(FWorldContext* WorldConte
 				Builder << TEXT("/Leaked");
 				Builder << BaseName;
 				FName NewName = MakeUniqueObjectName(nullptr, Pkg->GetClass(), FName(*Builder));
-				Pkg->Rename(*NewName.ToString(), nullptr, REN_DontCreateRedirectors | REN_DoNotDirty);
+				Pkg->Rename(*NewName.ToString(), nullptr, REN_DontCreateRedirectors | REN_DoNotDirty | REN_AllowPackageLinkerMismatch);
 			}
 		}
 	}
@@ -17452,7 +17452,7 @@ bool UEngine::CommitMapChange( FWorldContext &Context )
 		for (int32 Index = FakeWorld->GetStreamingLevels().Num() - 1; Index >= 0; --Index)
 		{
 			ULevelStreaming* const FakeWorldStreamingLevel = FakeWorld->GetStreamingLevels()[Index];
-			FakeWorldStreamingLevel->Rename(nullptr, Context.World(), REN_DontCreateRedirectors);
+			FakeWorldStreamingLevel->Rename(nullptr, Context.World(), REN_DontCreateRedirectors | REN_AllowPackageLinkerMismatch);
 			FakeWorld->RemoveStreamingLevelAt(Index);
 			StreamingLevelsToMove[Index] = FakeWorldStreamingLevel;
 		}
diff --git a/Engine/Source/Runtime/Engine/Private/World.cpp b/Engine/Source/Runtime/Engine/Private/World.cpp
--- a/Engine/Source/Runtime/Engine/Private/World.cpp
+++ b/Engine/Source/Runtime/Engine/Private/World.cpp
@@ -1253,7 +1253,7 @@ bool UWorld::Rename(const TCHAR* InName, UObject* NewOuter, ERenameFlags Flags)
 				// If there is already something there with that name, rename it to something else.
 				if (UObject* ExistingObject = StaticFindObject(nullptr, LevelScriptBlueprint->GetOuter(), InName))
 				{
-					ExistingObject->Rename(nullptr, nullptr, REN_DoNotDirty | REN_DontCreateRedirectors | REN_NonTransactional);
+					ExistingObject->Rename(nullptr, nullptr, REN_DoNotDirty | REN_DontCreateRedirectors | REN_NonTransactional | REN_AllowPackageLinkerMismatch);
 				}
 
 				// This is a normal rename. Use LevelScriptBlueprint->GetOuter() instead of NULL to make sure the generated top level objects are moved appropriately
@@ -1742,7 +1742,7 @@ void UWorld::PostLoad()
 			if (GetName() != ShortPackageName)
 			{
 				// Do not go through UWorld::Rename as we do not want to go through map build data/external actors or hlod renaming in post load
-				UObject::Rename(*ShortPackageName, NULL, REN_NonTransactional | REN_DontCreateRedirectors);
+				UObject::Rename(*ShortPackageName, NULL, REN_NonTransactional | REN_DontCreateRedirectors | REN_AllowPackageLinkerMismatch);
 			}
 
 			// Worlds are assets so they need RF_Public and RF_Standalone (for the editor)
@@ -1756,7 +1756,7 @@ void UWorld::PostLoad()
 			{
 				if (Model->GetOuter() != DefaultBrush->GetOuter())
 				{
-					Model->Rename(TEXT("Brush"), DefaultBrush->GetOuter(), REN_DoNotDirty | REN_DontCreateRedirectors | REN_NonTransactional);
+					Model->Rename(TEXT("Brush"), DefaultBrush->GetOuter(), REN_DoNotDirty | REN_DontCreateRedirectors | REN_NonTransactional | REN_AllowPackageLinkerMismatch);
 				}
 			}
 		}
@@ -2207,7 +2207,7 @@ void UWorld::RepairWorldSettings()
 		// Rename invalid WorldSettings to avoid name collisions
 		if (ExistingWorldSettings)
 		{
-			ExistingWorldSettings->Rename(nullptr, PersistentLevel);
+			ExistingWorldSettings->Rename(nullptr, PersistentLevel, REN_AllowPackageLinkerMismatch);
 		}
 		
 		bool bClearOwningWorld = false;
@@ -2456,7 +2456,7 @@ void UWorld::InitWorld(const InitializationValues IVS)
 	{
 		// Move persistent level into world so the world object won't get garbage collected in the multi- level
 		// case as it is still referenced via the level's outer. This is required for multi- level editing to work.
-		PersistentLevel->Rename(*PersistentLevel->GetName(), this);
+		PersistentLevel->Rename(*PersistentLevel->GetName(), this, REN_AllowPackageLinkerMismatch);
 	}
 
 	Levels.Empty(1);
@@ -4322,7 +4322,7 @@ void UWorld::RenameToPIEWorld(int32 PIEInstanceID)
 	WorldPackage->SetPackageFlags(PKG_PlayInEditor);
 
 	const FString PIEPackageName = *UWorld::ConvertToPIEPackageName(WorldPackage->GetName(), PIEInstanceID);
-	WorldPackage->Rename(*PIEPackageName, nullptr);
+	WorldPackage->Rename(*PIEPackageName, nullptr, REN_AllowPackageLinkerMismatch);
 	FSoftObjectPath::AddPIEPackageName(FName(*PIEPackageName));
 
 	for (UPackage* ExternalPackage : WorldPackage->GetExternalPackages())
diff --git a/Engine/Source/Runtime/Engine/Private/WorldPartition/DataLayer/ExternalDataLayerHelper.cpp b/Engine/Source/Runtime/Engine/Private/WorldPartition/DataLayer/ExternalDataLayerHelper.cpp
--- a/Engine/Source/Runtime/Engine/Private/WorldPartition/DataLayer/ExternalDataLayerHelper.cpp
+++ b/Engine/Source/Runtime/Engine/Private/WorldPartition/DataLayer/ExternalDataLayerHelper.cpp
@@ -356,7 +356,7 @@ bool FExternalDataLayerHelper::MoveActorsToExternalDataLayer(const TArray<AActor
 		UPackage* NewActorPackage = InActor->GetExternalPackage();
 		for (UObject* DependantObject : DependantObjects) //-V1078
 		{
-			DependantObject->Rename(nullptr, NewActorPackage, REN_NonTransactional | REN_DontCreateRedirectors | REN_DoNotDirty);
+			DependantObject->Rename(nullptr, NewActorPackage, REN_NonTransactional | REN_DontCreateRedirectors | REN_DoNotDirty | REN_AllowPackageLinkerMismatch);
 		}
 
 		if (!bLevelPackageWasDirty)
diff --git a/Engine/Source/Runtime/Engine/Private/WorldPartition/DataLayer/ExternalDataLayerManager.cpp b/Engine/Source/Runtime/Engine/Private/WorldPartition/DataLayer/ExternalDataLayerManager.cpp
--- a/Engine/Source/Runtime/Engine/Private/WorldPartition/DataLayer/ExternalDataLayerManager.cpp
+++ b/Engine/Source/Runtime/Engine/Private/WorldPartition/DataLayer/ExternalDataLayerManager.cpp
@@ -924,7 +924,7 @@ URuntimeHashExternalStreamingObjectBase* UExternalDataLayerManager::CreateExtern
 	// Outer DataLayerInstances to the ExternalStreamingObject and set its RootExternalDataLayerInstance
 	DuplicatedEDLWorldDataLayers->ForEachDataLayerInstance([ExternalStreamingObject](UDataLayerInstance* DataLayerInstance)
 	{
-		DataLayerInstance->Rename(nullptr, ExternalStreamingObject, REN_DontCreateRedirectors | REN_DoNotDirty | REN_NonTransactional);
+		DataLayerInstance->Rename(nullptr, ExternalStreamingObject, REN_DontCreateRedirectors | REN_DoNotDirty | REN_NonTransactional | REN_AllowPackageLinkerMismatch);
 		ExternalStreamingObject->DataLayerInstances.Add(DataLayerInstance);
 		if (UExternalDataLayerInstance* ExternalDataLayerInstance = Cast<UExternalDataLayerInstance>(DataLayerInstance))
 		{
diff --git a/Engine/Source/Runtime/Engine/Private/WorldPartition/DataLayer/WorldDataLayers.cpp b/Engine/Source/Runtime/Engine/Private/WorldPartition/DataLayer/WorldDataLayers.cpp
--- a/Engine/Source/Runtime/Engine/Private/WorldPartition/DataLayer/WorldDataLayers.cpp
+++ b/Engine/Source/Runtime/Engine/Private/WorldPartition/DataLayer/WorldDataLayers.cpp
@@ -616,7 +616,7 @@ AWorldDataLayers* AWorldDataLayers::Create(const FActorSpawnParameters& SpawnPar
 		if (!IsValidChecked(WorldDataLayers))
 		{
 			// Handle the case where the actor already exists, but it's pending kill
-			WorldDataLayers->Rename(nullptr, nullptr, REN_DontCreateRedirectors | REN_DoNotDirty | REN_NonTransactional);
+			WorldDataLayers->Rename(nullptr, nullptr, REN_DontCreateRedirectors | REN_DoNotDirty | REN_NonTransactional | REN_AllowPackageLinkerMismatch);
 			WorldDataLayers = nullptr;
 		}
 	}
diff --git a/Engine/Source/Runtime/Engine/Private/WorldPartition/WorldPartitionConverter.cpp b/Engine/Source/Runtime/Engine/Private/WorldPartition/WorldPartitionConverter.cpp
--- a/Engine/Source/Runtime/Engine/Private/WorldPartition/WorldPartitionConverter.cpp
+++ b/Engine/Source/Runtime/Engine/Private/WorldPartition/WorldPartitionConverter.cpp
@@ -206,7 +206,7 @@ bool FWorldPartitionConverter::ShouldDeleteActor(AActor* InActor, bool bIsMainLe
 void FWorldPartitionConverter::ChangeObjectOuter(UObject* InObject, UObject* InNewOuter)
 {
 	FString OldPath = FSoftObjectPath(InObject).ToString();
-	InObject->Rename(nullptr, InNewOuter, REN_NonTransactional | REN_DontCreateRedirectors | REN_DoNotDirty);
+	InObject->Rename(nullptr, InNewOuter, REN_NonTransactional | REN_DontCreateRedirectors | REN_DoNotDirty | REN_AllowPackageLinkerMismatch);
 	FString NewPath = FSoftObjectPath(InObject).ToString();
 	RemapSoftObjectPaths.Add(OldPath, NewPath);
 }
diff --git a/Engine/Source/Runtime/Engine/Private/WorldPartition/WorldPartitionLevelHelper.cpp b/Engine/Source/Runtime/Engine/Private/WorldPartition/WorldPartitionLevelHelper.cpp
--- a/Engine/Source/Runtime/Engine/Private/WorldPartition/WorldPartitionLevelHelper.cpp
+++ b/Engine/Source/Runtime/Engine/Private/WorldPartition/WorldPartitionLevelHelper.cpp
@@ -262,7 +262,7 @@ void FWorldPartitionLevelHelper::MoveExternalActorsToLevel(const TArray<FWorldPa
 			// (UObject::Rename doesn't check if Rename is being called with existing outer and assigns new name)
 			if (!bSameOuter)
 			{
-				Actor->Rename(nullptr, InLevel, REN_NonTransactional | REN_DoNotDirty | REN_DontCreateRedirectors);
+				Actor->Rename(nullptr, InLevel, REN_NonTransactional | REN_DoNotDirty | REN_DontCreateRedirectors | REN_AllowPackageLinkerMismatch);
 
 				// AActor::Rename will register components but doesn't call RerunConstructionScripts like AddLoadedActors does.
 				// If bIsWorldInitialized is false. RerunConstructionScripts will get called as part of UEditorEngine::InitializePhysicsSceneForSaveIfNecessary during Cell package save
@@ -299,7 +299,7 @@ void FWorldPartitionLevelHelper::MoveExternalActorsToLevel(const TArray<FWorldPa
 							AActor* NestedActor = Cast<AActor>(Object);
 							if (InLevel != Object->GetOuter())
 							{
-								Object->Rename(nullptr, InLevel, REN_NonTransactional | REN_DoNotDirty | REN_DontCreateRedirectors);
+								Object->Rename(nullptr, InLevel, REN_NonTransactional | REN_DoNotDirty | REN_DontCreateRedirectors | REN_AllowPackageLinkerMismatch);
 								MovedObjects.Add(Object);
 							}
 							else if (NestedActor && !InLevel->Actors.Contains(NestedActor))
@@ -323,7 +323,7 @@ void FWorldPartitionLevelHelper::MoveExternalActorsToLevel(const TArray<FWorldPa
 								NewName = MakeUniqueObjectName(LevelPackage, Object->GetClass(), *Object->GetName());
 							}
 
-							Object->Rename(*NewName.ToString(), LevelPackage, REN_NonTransactional | REN_DoNotDirty | REN_DontCreateRedirectors);
+							Object->Rename(*NewName.ToString(), LevelPackage, REN_NonTransactional | REN_DoNotDirty | REN_DontCreateRedirectors | REN_AllowPackageLinkerMismatch);
 							MovedObjects.Add(Object);
 						}
 					}
@@ -895,14 +895,14 @@ bool FWorldPartitionLevelHelper::LoadActorsInternal(FLoadActorsParams&& InParams
 					OuterWorld->GetSoftObjectPathMapping(SourceOuterWorldPath, DummyUnusedPath);
 
 					// Rename through UObject to avoid changing Actor's external packaging and folder properties
-					Actor->UObject::Rename(*AddActorContainerID(PackageObjectMapping->ContainerID, Actor->GetName()), DestLevel, REN_NonTransactional | REN_DoNotDirty | REN_DontCreateRedirectors);
+					Actor->UObject::Rename(*AddActorContainerID(PackageObjectMapping->ContainerID, Actor->GetName()), DestLevel, REN_NonTransactional | REN_DoNotDirty | REN_DontCreateRedirectors | REN_AllowPackageLinkerMismatch);
 
 					// Handle child actors
 					Actor->ForEachComponent<UChildActorComponent>(true, [DestLevel = DestLevel, PackageObjectMapping](UChildActorComponent* ChildActorComponent)
 					{
 						if (AActor* ChildActor = ChildActorComponent->GetChildActor())
 						{
-							ChildActor->UObject::Rename(*AddActorContainerID(PackageObjectMapping->ContainerID, ChildActor->GetName()), DestLevel, REN_NonTransactional | REN_DoNotDirty | REN_DontCreateRedirectors);
+							ChildActor->UObject::Rename(*AddActorContainerID(PackageObjectMapping->ContainerID, ChildActor->GetName()), DestLevel, REN_NonTransactional | REN_DoNotDirty | REN_DontCreateRedirectors | REN_AllowPackageLinkerMismatch);
 						}
 					});
 
diff --git a/Engine/Source/Runtime/Experimental/Animation/Constraints/Private/ConstraintsManager.cpp b/Engine/Source/Runtime/Experimental/Animation/Constraints/Private/ConstraintsManager.cpp
--- a/Engine/Source/Runtime/Experimental/Animation/Constraints/Private/ConstraintsManager.cpp
+++ b/Engine/Source/Runtime/Experimental/Animation/Constraints/Private/ConstraintsManager.cpp
@@ -523,7 +523,7 @@ void FConstraintsManagerController::StaticConstraintCreated(UTickableConstraint*
 		{
 			Manager->Modify();
 			Manager->Constraints.Add(InConstraint);
-			InConstraint->Rename(nullptr, Manager, REN_DontCreateRedirectors);
+			InConstraint->Rename(nullptr, Manager, REN_DontCreateRedirectors | REN_AllowPackageLinkerMismatch);
 			Manager->OnConstraintAdded_BP.Broadcast(Manager, InConstraint);
 		}
 	}	
diff --git a/Engine/Source/Runtime/Experimental/Animation/Constraints/Private/Transform/TransformConstraint.cpp b/Engine/Source/Runtime/Experimental/Animation/Constraints/Private/Transform/TransformConstraint.cpp
--- a/Engine/Source/Runtime/Experimental/Animation/Constraints/Private/Transform/TransformConstraint.cpp
+++ b/Engine/Source/Runtime/Experimental/Animation/Constraints/Private/Transform/TransformConstraint.cpp
@@ -287,7 +287,7 @@ void UTickableTransformConstraint::SetupDependencies(const UWorld* InWorld)
 		UTickableTransformConstraint* PTRS = ParentTRSHandle->GetTypedOuter<UTickableTransformConstraint>();
 		if(PTRS != this)
 		{
-			ParentTRSHandle->Rename(nullptr, this, REN_DontCreateRedirectors);
+			ParentTRSHandle->Rename(nullptr, this, REN_DontCreateRedirectors | REN_AllowPackageLinkerMismatch);
 		}
 	}
 	if (ChildTRSHandle)
@@ -295,7 +295,7 @@ void UTickableTransformConstraint::SetupDependencies(const UWorld* InWorld)
 		UTickableTransformConstraint* CTRS = ChildTRSHandle->GetTypedOuter<UTickableTransformConstraint>();
 		if (CTRS != this)
 		{
-			ChildTRSHandle->Rename(nullptr, this, REN_DontCreateRedirectors);
+			ChildTRSHandle->Rename(nullptr, this, REN_DontCreateRedirectors | REN_AllowPackageLinkerMismatch);
 		}
 	}
 
diff --git a/Engine/Source/Runtime/GeometryFramework/Private/Components/DynamicMeshComponent.cpp b/Engine/Source/Runtime/GeometryFramework/Private/Components/DynamicMeshComponent.cpp
--- a/Engine/Source/Runtime/GeometryFramework/Private/Components/DynamicMeshComponent.cpp
+++ b/Engine/Source/Runtime/GeometryFramework/Private/Components/DynamicMeshComponent.cpp
@@ -1306,7 +1306,7 @@ void UDynamicMeshComponent::SetDynamicMesh(UDynamicMesh* NewMesh)
 
 	// set Outer of NewMesh to be this Component, ie transfer ownership. This is done via "renaming", which is
 	// a bit odd, so the flags prevent some standard "renaming" behaviors from happening
-	NewMesh->Rename( nullptr, this, REN_DontCreateRedirectors);
+	NewMesh->Rename( nullptr, this, REN_DontCreateRedirectors | REN_AllowPackageLinkerMismatch);
 	MeshObject = NewMesh;
 	MeshObjectChangedHandle = MeshObject->OnMeshChanged().AddUObject(this, &UDynamicMeshComponent::OnMeshObjectChanged);
 
diff --git a/Engine/Source/Runtime/Landscape/Private/Landscape.cpp b/Engine/Source/Runtime/Landscape/Private/Landscape.cpp
--- a/Engine/Source/Runtime/Landscape/Private/Landscape.cpp
+++ b/Engine/Source/Runtime/Landscape/Private/Landscape.cpp
@@ -1417,7 +1417,7 @@ PRAGMA_ENABLE_DEPRECATION_WARNINGS
 	{
 		if (Object && !Object->HasAllFlags(RF_Public | RF_Standalone) && (Object->GetOuter() != GetOuter()) && (Object->GetOutermost() == GetOutermost()))
 		{
-			Object->Rename(nullptr, GetOuter());
+			Object->Rename(nullptr, GetOuter(), REN_AllowPackageLinkerMismatch);
 			return true;
 		}
 		return false;
@@ -1535,7 +1535,7 @@ PRAGMA_ENABLE_DEPRECATION_WARNINGS
 				Obj->ClearFlags(RF_Public);
 				if (Obj->GetOuter() == Level)
 				{
-					Obj->Rename(nullptr, MyPackage, REN_DoNotDirty | REN_DontCreateRedirectors | REN_NonTransactional);
+					Obj->Rename(nullptr, MyPackage, REN_DoNotDirty | REN_DontCreateRedirectors | REN_NonTransactional | REN_AllowPackageLinkerMismatch);
 				}
 			}
 		}
diff --git a/Engine/Source/Runtime/LevelSequence/Private/LevelSequence.cpp b/Engine/Source/Runtime/LevelSequence/Private/LevelSequence.cpp
--- a/Engine/Source/Runtime/LevelSequence/Private/LevelSequence.cpp
+++ b/Engine/Source/Runtime/LevelSequence/Private/LevelSequence.cpp
@@ -389,9 +389,10 @@ void ULevelSequence::PostLoad()
 		DirectorBlueprint->OnCompiled().RemoveAll(this);
 		DirectorBlueprint->OnCompiled().AddUObject(this, &ULevelSequence::OnDirectorRecompiled);
 
-		if (DirectorBlueprint->Rename(*GetDirectorBlueprintName(), nullptr, (REN_NonTransactional | REN_DoNotDirty | REN_Test | REN_DontCreateRedirectors)))
+		const int32 RenameFlags = REN_AllowPackageLinkerMismatch | REN_NonTransactional | REN_DoNotDirty | REN_DontCreateRedirectors;
+		if (DirectorBlueprint->Rename(*GetDirectorBlueprintName(), nullptr, (RenameFlags | REN_Test)))
 		{
-			DirectorBlueprint->Rename(*GetDirectorBlueprintName(), nullptr, (REN_NonTransactional | REN_DoNotDirty | REN_DontCreateRedirectors));
+			DirectorBlueprint->Rename(*GetDirectorBlueprintName(), nullptr, RenameFlags);
 		}
 	}
 
diff --git a/Engine/Source/Runtime/MovieScene/Private/MovieScene.cpp b/Engine/Source/Runtime/MovieScene/Private/MovieScene.cpp
--- a/Engine/Source/Runtime/MovieScene/Private/MovieScene.cpp
+++ b/Engine/Source/Runtime/MovieScene/Private/MovieScene.cpp
@@ -327,7 +327,7 @@ void UMovieScene::AddSpawnable(const FMovieSceneSpawnable& InNewSpawnable, const
 		if (!Track->IsIn(this))
 		{
 			FName NewName = MakeUniqueObjectName(this, Track->GetClass(), *Track->GetName());
-			Track->Rename(*NewName.ToString(), this, REN_DontCreateRedirectors);
+			Track->Rename(*NewName.ToString(), this, REN_DontCreateRedirectors | REN_AllowPackageLinkerMismatch);
 		}
 	}
 	int32 NewBindingIndex = InsertSorted(ObjectBindings, MoveTemp(NewBinding));
@@ -432,7 +432,7 @@ void UMovieScene::AddPossessable(const FMovieScenePossessable& InNewPossessable,
 		if (!Track->IsIn(this))
 		{
 			FName NewName = MakeUniqueObjectName(this, Track->GetClass(), *Track->GetName());
-			Track->Rename(*NewName.ToString(), this, REN_DontCreateRedirectors);
+			Track->Rename(*NewName.ToString(), this, REN_DontCreateRedirectors | REN_AllowPackageLinkerMismatch);
 		}
 	}
 	int32 NewBindingIndex = InsertSorted(ObjectBindings, MoveTemp(NewBinding));
@@ -1210,7 +1210,7 @@ bool UMovieScene::AddGivenTrack(UMovieSceneTrack* InTrack, const FGuid& ObjectGu
 	Modify();
 	if (FMovieSceneBinding* Binding = FindBinding(ObjectGuid))
 	{
-		InTrack->Rename(nullptr, this, REN_DontCreateRedirectors);
+		InTrack->Rename(nullptr, this, REN_DontCreateRedirectors | REN_AllowPackageLinkerMismatch);
 		Binding->AddTrack(*InTrack, this);
 		return true;
 	}
@@ -1329,7 +1329,7 @@ bool UMovieScene::AddGivenTrack(UMovieSceneTrack* InTrack)
 	{
 		Modify();
 		Tracks.Add(InTrack);
-		InTrack->Rename(nullptr, this, REN_DontCreateRedirectors);
+		InTrack->Rename(nullptr, this, REN_DontCreateRedirectors | REN_AllowPackageLinkerMismatch);
 
 		InTrack->OnAddedToMovieScene(this);
 		EventHandlers.Trigger(&UE::MovieScene::ISequenceDataEventHandler::OnTrackAdded, InTrack);
@@ -1402,7 +1402,7 @@ void UMovieScene::SetCameraCutTrack(UMovieSceneTrack* InTrack)
 	}
 
 	Modify();
-	InTrack->Rename(nullptr, this, REN_DontCreateRedirectors);
+	InTrack->Rename(nullptr, this, REN_DontCreateRedirectors | REN_AllowPackageLinkerMismatch);
 	UMovieSceneTrack* OldCameraCutTrack = CameraCutTrack;
 	CameraCutTrack = InTrack;
 
@@ -1708,7 +1708,7 @@ void UMovieScene::ReplaceBinding(const FGuid& BindingToReplaceGuid, const FMovie
 			if (!Track->IsIn(this))
 			{
 				FName NewName = MakeUniqueObjectName(this, Track->GetClass(), *Track->GetName());
-				Track->Rename(*NewName.ToString(), this, REN_DontCreateRedirectors);
+				Track->Rename(*NewName.ToString(), this, REN_DontCreateRedirectors | REN_AllowPackageLinkerMismatch);
 			}
 		}
 
diff --git a/Engine/Source/Runtime/MovieSceneTracks/Private/Sections/MovieScene3DTransformSection.cpp b/Engine/Source/Runtime/MovieSceneTracks/Private/Sections/MovieScene3DTransformSection.cpp
--- a/Engine/Source/Runtime/MovieSceneTracks/Private/Sections/MovieScene3DTransformSection.cpp
+++ b/Engine/Source/Runtime/MovieSceneTracks/Private/Sections/MovieScene3DTransformSection.cpp
@@ -1211,7 +1211,7 @@ void UMovieScene3DTransformSection::AddConstraintChannel(UTickableConstraint* In
 
 		//make copy that we can spawn if it doesn't exist
 		//the rename changes the outer to this section (from any actor manager)
-		InConstraint->Rename(nullptr, this, REN_DontCreateRedirectors);
+		InConstraint->Rename(nullptr, this, REN_DontCreateRedirectors | REN_AllowPackageLinkerMismatch);
 	
 		CacheChannelProxy();
 
diff --git a/Engine/Source/Runtime/NavigationSystem/Private/NavigationData.cpp b/Engine/Source/Runtime/NavigationSystem/Private/NavigationData.cpp
--- a/Engine/Source/Runtime/NavigationSystem/Private/NavigationData.cpp
+++ b/Engine/Source/Runtime/NavigationSystem/Private/NavigationData.cpp
@@ -437,7 +437,7 @@ void ANavigationData::InstantiateAndRegisterRenderingComponent()
 		if (RenderingComp)
 		{
 			// rename the old rendering component out of the way
-			RenderingComp->Rename(NULL, NULL, REN_DontCreateRedirectors | REN_ForceGlobalUnique | REN_DoNotDirty | REN_NonTransactional);
+			RenderingComp->Rename(NULL, NULL, REN_DontCreateRedirectors | REN_ForceGlobalUnique | REN_DoNotDirty | REN_NonTransactional | REN_AllowPackageLinkerMismatch);
 		}
 
 		RenderingComp = ConstructRenderingComponent();
diff --git a/Engine/Source/Runtime/NavigationSystem/Private/NavigationSystem.cpp b/Engine/Source/Runtime/NavigationSystem/Private/NavigationSystem.cpp
--- a/Engine/Source/Runtime/NavigationSystem/Private/NavigationSystem.cpp
+++ b/Engine/Source/Runtime/NavigationSystem/Private/NavigationSystem.cpp
@@ -4639,14 +4639,14 @@ ANavigationData* UNavigationSystemV1::CreateNavigationDataInstanceInLevel(const
 				}
 
 				// Reset the existing object's name
-				ExistingObject->Rename(nullptr, nullptr, REN_DontCreateRedirectors | REN_ForceGlobalUnique | REN_DoNotDirty | REN_NonTransactional);
+				ExistingObject->Rename(nullptr, nullptr, REN_DontCreateRedirectors | REN_ForceGlobalUnique | REN_DoNotDirty | REN_NonTransactional | REN_AllowPackageLinkerMismatch);
 				// see if there's another one, it does happen when undo/redoing 
 				// nav instance deletion in the editor
 				ExistingObject = StaticFindObject(/*Class=*/ nullptr, Instance->GetOuter(), *StrName, EFindObjectFlags::ExactClass);
 			}
 
 			// Set descriptive name
-			Instance->Rename(*StrName, nullptr, REN_DoNotDirty);
+			Instance->Rename(*StrName, nullptr, REN_DoNotDirty | REN_AllowPackageLinkerMismatch);
 #if WITH_EDITOR
 			if (World->WorldType == EWorldType::Editor)
 			{
diff --git a/Engine/Source/Runtime/UMG/Private/WidgetTree.cpp b/Engine/Source/Runtime/UMG/Private/WidgetTree.cpp
--- a/Engine/Source/Runtime/UMG/Private/WidgetTree.cpp
+++ b/Engine/Source/Runtime/UMG/Private/WidgetTree.cpp
@@ -168,7 +168,7 @@ bool UWidgetTree::TryMoveWidgetToNewTree(UWidget* Widget, UWidgetTree* Destinati
 
 	if (DestinationTree && OriginalTree && OriginalTree != DestinationTree)
 	{
-		bWidgetMoved = Widget->Rename(*Widget->GetName(), DestinationTree, REN_DontCreateRedirectors);
+		bWidgetMoved = Widget->Rename(*Widget->GetName(), DestinationTree, REN_DontCreateRedirectors | REN_AllowPackageLinkerMismatch);
 	}
 
 	return bWidgetMoved;
