Updated IProjectExternalContentInterface to remove assumptions there is a single project.

#rnx
#Synced-CL 49648987

diff --git a/Engine/Source/Editor/UnrealEd/Private/EditorEngine.cpp b/Engine/Source/Editor/UnrealEd/Private/EditorEngine.cpp
--- a/Engine/Source/Editor/UnrealEd/Private/EditorEngine.cpp
+++ b/Engine/Source/Editor/UnrealEd/Private/EditorEngine.cpp
@@ -539,7 +539,7 @@ void FAssetReferenceFilterContext::AddReferencingAssetsFromPropertyHandle(const
 		while (CurrentPropertyHandle.IsValid())
 		{
 			const FProperty* Property = CurrentPropertyHandle->GetProperty();
-			if (Property != nullptr && Property->HasAnyPropertyFlags(CPF_EditorOnly))
+			if (Property != nullptr && Property->IsEditorOnlyProperty())
 			{
 				ReferencerProperties |= EAssetReferenceFilterProperties::EditorOnly;
 				break;
@@ -6796,12 +6796,11 @@ namespace
 	class FProjectExternalContentDefault : public IProjectExternalContentInterface
 	{
 	private:
-		virtual bool IsEnabled() const override { return false; }
-		virtual bool HasExternalContent(const FString& VersePath) const override { return false; }
-		virtual bool IsExternalContentLoaded(const FString& VersePath) const override { return false; }
-		virtual TArray<FString> GetExternalContentVersePaths() const override { return {}; }
-		virtual void AddExternalContent(const FString& VersePath, FAddExternalContentComplete CompleteCallback) override { CompleteCallback.ExecuteIfBound(false, /*Plugins=*/{}); }
-		virtual void RemoveExternalContent(TConstArrayView<FString> VersePaths, FRemoveExternalContentComplete CompleteCallback) override { CompleteCallback.ExecuteIfBound(false); }
+		bool IsEnabled() const override { return false; }
+		void AddExternalContent(const UE::Core::FVersePath& VersePath, TSharedPtr<IPlugin> ReferencingPlugin = {}, FAddExternalContentComplete CompleteCallback = {}) override
+		{
+			CompleteCallback.ExecuteIfBound(false, /*Plugins=*/{});
+		}
 	};
 
 	FProjectExternalContentDefault ProjectExternalContentDefault;
diff --git a/Engine/Source/Editor/UnrealEd/Public/IProjectExternalContentInterface.h b/Engine/Source/Editor/UnrealEd/Public/IProjectExternalContentInterface.h
--- a/Engine/Source/Editor/UnrealEd/Public/IProjectExternalContentInterface.h
+++ b/Engine/Source/Editor/UnrealEd/Public/IProjectExternalContentInterface.h
@@ -5,30 +5,32 @@
 
 class IPlugin;
 
+namespace UE::Core
+{
+	class FVersePath;
+}
+
 /** Interface to manage project references to external content */
 class IProjectExternalContentInterface
 {
 public:
+	IProjectExternalContentInterface() = default;
+	IProjectExternalContentInterface(const IProjectExternalContentInterface&) = delete;
+	virtual ~IProjectExternalContentInterface() = default;
+
+	IProjectExternalContentInterface& operator=(const IProjectExternalContentInterface&) = delete;
+
 	/** Return whether the project can reference external content */
 	virtual bool IsEnabled() const = 0;
 
-	/**
-     * Return whether the specified external content is referenced by the project
-	 * @param VersePath External content Verse path
-	 */
-	virtual bool HasExternalContent(const FString& VersePath) const = 0;
-
-	/** 
-	 * Return whether the specified external content is loaded (and referenced by the project)
-	 * @param VersePath External content Verse path
-	 */
-	virtual bool IsExternalContentLoaded(const FString& VersePath) const = 0;
+	UE_DEPRECATED(5.8, "No longer used.")
+	virtual bool HasExternalContent(const FString& VersePath) const { return false; }
 
-	/** Return the list of external content Verse paths referenced by the project */
-	virtual TArray<FString> GetExternalContentVersePaths() const = 0;
+	UE_DEPRECATED(5.8, "No longer used.")
+	virtual bool IsExternalContentLoaded(const FString& VersePath) const { return false; }
 
-	UE_DEPRECATED(5.6, "GetExternalContentIds is deprecated, use GetExternalContentVersePaths instead")
-	TArray<FString> GetExternalContentIds() const { return GetExternalContentVersePaths(); };
+	UE_DEPRECATED(5.8, "No longer used.")
+	virtual TArray<FString> GetExternalContentVersePaths() const { return {}; }
 
 	/**
 	 * Called upon AddExternalContent completion
@@ -40,25 +42,36 @@ public:
 	/**
 	 * Add a reference to external content to the project and asynchronously downloads/loads the external content
 	 * @param VersePath External content Verse path
+	 * @param ReferencingPlugin Plugin that should reference the external content, or unset if unspecified
 	 * @param CompleteCallback See FAddExternalContentComplete
 	 */
-	virtual void AddExternalContent(const FString& VersePath, FAddExternalContentComplete CompleteCallback = FAddExternalContentComplete()) = 0;
+	virtual void AddExternalContent(const UE::Core::FVersePath& VersePath, TSharedPtr<IPlugin> ReferencingPlugin = {}, FAddExternalContentComplete CompleteCallback = {})
+	{
+		CompleteCallback.ExecuteIfBound(false, /*Plugins=*/{});
+	}
 
-	/**
-	 * Called upon RemoveExternalContent completion
-	 * @param bSuccess Whether the external content was successfully removed from the project (could be canceled by the user)
-	 */
+	UE_DEPRECATED(5.8, "Use the FVersePath overload.")
+	virtual void AddExternalContent(const FString& VersePath, FAddExternalContentComplete CompleteCallback = FAddExternalContentComplete())
+	{
+		CompleteCallback.ExecuteIfBound(false, /*Plugins=*/{});
+	}
+
+	UE_DEPRECATED(5.8, "No longer used.")
 	DECLARE_DELEGATE_OneParam(FRemoveExternalContentComplete, bool /*bSuccess*/);
 
-	/**
-	 * Remove references to external content from the project and unloads the external content
-	 * @param VersePaths External content Verse paths
-	 * @param CompleteCallback See FRemoveExternalContentComplete
-	 */
-	virtual void RemoveExternalContent(TConstArrayView<FString> VersePaths, FRemoveExternalContentComplete CompleteCallback = FRemoveExternalContentComplete()) = 0;
+	PRAGMA_DISABLE_DEPRECATION_WARNINGS
 
+	UE_DEPRECATED(5.8, "No longer used.")
+	virtual void RemoveExternalContent(TConstArrayView<FString> VersePaths, FRemoveExternalContentComplete CompleteCallback = FRemoveExternalContentComplete())
+	{
+		CompleteCallback.ExecuteIfBound(false);
+	}
+
+	UE_DEPRECATED(5.8, "No longer used.")
 	void RemoveExternalContent(const FString& VersePath, FRemoveExternalContentComplete CompleteCallback = FRemoveExternalContentComplete())
 	{
 		RemoveExternalContent(MakeArrayView(&VersePath, 1), MoveTemp(CompleteCallback));
 	}
+
+	PRAGMA_ENABLE_DEPRECATION_WARNINGS
 };
