[Audio Golden Path] Updates to audio auto-reimport
* Folders being watched are stored in new settings for this system, rather than in the global auto-reimport system settings.
* Folders are removed after some time without any updates (30 days by default)
* Watched folders are added without a mount point, to disable automatically adding/moving/creating assets (which was problematic)

#rnx
#Synced-CL 49729169

diff --git a/Engine/Source/Editor/UnrealEd/Classes/Settings/EditorLoadingSavingSettings.h b/Engine/Source/Editor/UnrealEd/Classes/Settings/EditorLoadingSavingSettings.h
--- a/Engine/Source/Editor/UnrealEd/Classes/Settings/EditorLoadingSavingSettings.h
+++ b/Engine/Source/Editor/UnrealEd/Classes/Settings/EditorLoadingSavingSettings.h
@@ -241,8 +241,6 @@ public:
 	DECLARE_EVENT_OneParam(UEditorLoadingSavingSettings, FSettingChangedEvent, FName /*PropertyName*/);
 	FSettingChangedEvent& OnSettingChanged( ) { return SettingChangedEvent; }
 
-protected:
-
 	// UObject overrides
 
 	UNREALED_API virtual void PostEditChangeProperty( struct FPropertyChangedEvent& PropertyChangedEvent ) override;
diff --git a/Engine/Source/Editor/UnrealEd/Private/AutoReimport/AutoReimportManager.cpp b/Engine/Source/Editor/UnrealEd/Private/AutoReimport/AutoReimportManager.cpp
--- a/Engine/Source/Editor/UnrealEd/Private/AutoReimport/AutoReimportManager.cpp
+++ b/Engine/Source/Editor/UnrealEd/Private/AutoReimport/AutoReimportManager.cpp
@@ -126,6 +126,12 @@ public:
 	void IgnoreMovedFile(const FString& SrcFilename, const FString& DstFilename);
 	void IgnoreDeletedFile(const FString& Filename);
 
+	/** Add a directory to watch at runtime. Directories added this way are not saved in the settings. Returns a handle that can be passed to RemoveRuntimeMonitoredDirectory.*/
+	int32 AddRuntimeMonitoredDirectory(const FAutoReimportDirectoryConfig& Directory);
+
+	/** Remove runtime directory with the given handle. */
+	void RemoveRuntimeMonitoredDirectory(const int32 Handle);
+
 	/** Destroy this manager */
 	void Destroy();
 
@@ -160,6 +166,9 @@ private:
 	/** Callback for when an editor user setting has changed */
 	void HandleLoadingSavingSettingChanged(FName PropertyName);
 
+	/** Schedule monitors to be updated after a delay */
+	void SetMonitorsDirty();
+
 	/** Set up monitors to all the monitored content directories */
 	void SetUpDirectoryMonitors();
 
@@ -240,6 +249,10 @@ private:
 	/** The paused state of the state machine */
 	ECurrentState PausedState;
 
+	/** Monitored directories added at runtime */
+	TMap<int32, FAutoReimportDirectoryConfig> RuntimeMonitoredDirectories;
+
+	int32 NextRuntimeMonitoredDirectoryHandle = 0;
 private:
 
 	void OnPauseClicked()
@@ -911,10 +924,15 @@ void FAutoReimportManager::HandleLoadingSavingSettingChanged(FName PropertyName)
 	if (PropertyName == GET_MEMBER_NAME_CHECKED(UEditorLoadingSavingSettings, bMonitorContentDirectories) ||
 		PropertyName == GET_MEMBER_NAME_CHECKED(UEditorLoadingSavingSettings, AutoReimportDirectorySettings))
 	{
-		ResetMonitorsTimeout = DirectoryWatcher::FTimeLimit(5.f);
+		SetMonitorsDirty();
 	}
 }
 
+void FAutoReimportManager::SetMonitorsDirty()
+{
+	ResetMonitorsTimeout = DirectoryWatcher::FTimeLimit(5.f);
+}
+
 void FAutoReimportManager::OnContentPathChanged(const FString& InAssetPath, const FString& FileSystemPath)
 {
 	const UEditorLoadingSavingSettings* Settings = GetDefault<UEditorLoadingSavingSettings>();
@@ -925,6 +943,22 @@ void FAutoReimportManager::OnContentPathChanged(const FString& InAssetPath, cons
 	}
 }
 
+int32 FAutoReimportManager::AddRuntimeMonitoredDirectory(const FAutoReimportDirectoryConfig& Directory)
+{
+	const int32 NewHandle = NextRuntimeMonitoredDirectoryHandle++;
+	RuntimeMonitoredDirectories.Add(NewHandle, Directory);
+	SetMonitorsDirty();
+	return NewHandle;
+}
+
+void FAutoReimportManager::RemoveRuntimeMonitoredDirectory(const int32 Handle)
+{
+	if (RuntimeMonitoredDirectories.Remove(Handle))
+	{
+		SetMonitorsDirty();
+	}
+}
+
 void FAutoReimportManager::SetUpDirectoryMonitors()
 {
 	struct FParsedSettings
@@ -935,26 +969,40 @@ void FAutoReimportManager::SetUpDirectoryMonitors()
 	};
 
 	TArray<FParsedSettings> FinalArray;
-	FString SupportedExtensions = GetAllFactoryExtensions();
+	const FString SupportedExtensions = GetAllFactoryExtensions();
+	auto TryAddMapping = [&FinalArray, &SupportedExtensions](const FAutoReimportDirectoryConfig& Setting, bool bAllowEmptyMountPoint)
+		{
+			FParsedSettings NewMapping;
+			NewMapping.SourceDirectory = Setting.SourceDirectory;
+			NewMapping.MountPoint = Setting.MountPoint;
+
+			if (!FAutoReimportDirectoryConfig::ParseSourceDirectoryAndMountPoint(NewMapping.SourceDirectory, NewMapping.MountPoint))
+			{
+				// Allow the parse failure if it's just because the mount point is empty.
+				if (!(bAllowEmptyMountPoint && FPaths::DirectoryExists(NewMapping.SourceDirectory) && NewMapping.MountPoint.IsEmpty()))
+				{
+					return;
+				}
+			}
+
+			// Only include extensions that match a factory
+			NewMapping.Rules.SetApplicableExtensions(SupportedExtensions);
+			for (const FAutoReimportWildcard& WildcardConfig : Setting.Wildcards)
+			{
+				NewMapping.Rules.AddWildcardRule(WildcardConfig.Wildcard, WildcardConfig.bInclude);
+			}
+
+			FinalArray.Add(NewMapping);
+		};
+	
 	for (const FAutoReimportDirectoryConfig& Setting : GetDefault<UEditorLoadingSavingSettings>()->AutoReimportDirectorySettings)
 	{
-		FParsedSettings NewMapping;
-		NewMapping.SourceDirectory = Setting.SourceDirectory;
-		NewMapping.MountPoint = Setting.MountPoint;
-
-		if (!FAutoReimportDirectoryConfig::ParseSourceDirectoryAndMountPoint(NewMapping.SourceDirectory, NewMapping.MountPoint))
-		{
-			continue;
-		}
+		TryAddMapping(Setting, false);
+	}
 
-		// Only include extensions that match a factory
-		NewMapping.Rules.SetApplicableExtensions(SupportedExtensions);
-		for (const FAutoReimportWildcard& WildcardConfig : Setting.Wildcards)
-		{
-			NewMapping.Rules.AddWildcardRule(WildcardConfig.Wildcard, WildcardConfig.bInclude);
-		}
-		
-		FinalArray.Add(NewMapping);
+	for (const TPair<int32, FAutoReimportDirectoryConfig>& Setting : RuntimeMonitoredDirectories)
+	{
+		TryAddMapping(Setting.Value, true);
 	}
 
 	for (int32 Index = 0; Index < FinalArray.Num(); ++Index)
@@ -1055,6 +1103,16 @@ TArray<FPathAndMountPoint> UAutoReimportManager::GetMonitoredDirectories() const
 	return Implementation->GetMonitoredDirectories();
 }
 
+int32 UAutoReimportManager::AddRuntimeMonitoredDirectory(const FAutoReimportDirectoryConfig& Directory)
+{
+	return Implementation->AddRuntimeMonitoredDirectory(Directory);
+}
+
+void UAutoReimportManager::RemoveRuntimeMonitoredDirectory(const int32 Handle)
+{
+	Implementation->RemoveRuntimeMonitoredDirectory(Handle);
+}
+
 void UAutoReimportManager::BeginDestroy()
 {
 	Super::BeginDestroy();
diff --git a/Engine/Source/Editor/UnrealEd/Private/AutoReimport/ContentDirectoryMonitor.cpp b/Engine/Source/Editor/UnrealEd/Private/AutoReimport/ContentDirectoryMonitor.cpp
--- a/Engine/Source/Editor/UnrealEd/Private/AutoReimport/ContentDirectoryMonitor.cpp
+++ b/Engine/Source/Editor/UnrealEd/Private/AutoReimport/ContentDirectoryMonitor.cpp
@@ -492,7 +492,7 @@ void FContentDirectoryMonitor::ProcessModifications(const DirectoryWatcher::FTim
 			if (Assets.Num() == 1)
 			{
 				UObject* Asset = Assets[0].GetAsset();
-				if (Asset && Utils::ExtractSourceFilePaths(Asset).Num() == 1)
+				if (Asset && Utils::ExtractSourceFilePaths(Asset).Num() == 1 && !MountedContentPath.IsEmpty())
 				{
 					UPackage* ExistingPackage = Asset->GetOutermost();
 
diff --git a/Engine/Source/Editor/UnrealEd/Public/AutoReimport/AutoReimportManager.h b/Engine/Source/Editor/UnrealEd/Public/AutoReimport/AutoReimportManager.h
--- a/Engine/Source/Editor/UnrealEd/Public/AutoReimport/AutoReimportManager.h
+++ b/Engine/Source/Editor/UnrealEd/Public/AutoReimport/AutoReimportManager.h
@@ -8,6 +8,7 @@
 #include "AutoReimportManager.generated.h"
 
 class FAutoReimportManager;
+struct FAutoReimportDirectoryConfig;
 
 /** Struct representing a path on disk, and its virtual mount point */
 struct FPathAndMountPoint
@@ -43,6 +44,12 @@ public:
 	UNREALED_API void IgnoreMovedFile(const FString& SrcFilename, const FString& DstFilename);
 	UNREALED_API void IgnoreDeletedFile(const FString& Filename);
 
+	/** Add a directory to watch at runtime. Directories added this way are not saved in the settings. Returns a handle that can be passed to RemoveRuntimeMonitoredDirectory.*/
+	UNREALED_API int32 AddRuntimeMonitoredDirectory(const FAutoReimportDirectoryConfig& Directory);
+
+	/** Remove runtime directory with the given handle. */
+	UNREALED_API void RemoveRuntimeMonitoredDirectory(const int32 Handle);
+
 private:
 
 	/** UObject Interface */
