IncrementalCook: Bump CookIncrementalVersion due to staging error:
LogIoStore: Error: Missing shader code libary ChunkId. ChunkId = '<ChunkId>'. InFilename = <Path>/ShaderArchive-<GameName>_Chunk<PakFileChunk>-PCD3D_SM5-PCD3D_SM5.ushaderbytecode'. IOStatus = '(Not Found)'.
Looks like this is a bug in incremental cook that occurs whenever a PakFileChunk that previously contained shaders is deleted from the build; the manifest remains on disk from the previous cook and references a now-deleted shader file.
#rnx

#Synced-CL 49660478

diff --git a/Engine/Source/Editor/UnrealEd/Private/Cooker/CookIncrementalVersion.cpp b/Engine/Source/Editor/UnrealEd/Private/Cooker/CookIncrementalVersion.cpp
--- a/Engine/Source/Editor/UnrealEd/Private/Cooker/CookIncrementalVersion.cpp
+++ b/Engine/Source/Editor/UnrealEd/Private/Cooker/CookIncrementalVersion.cpp
@@ -8,6 +8,6 @@ namespace UE::Cook
 {
 
 // Change CookIncrementalVersion to a new guid when all packages in an incremental cook need to be invalidated
-FGuid CookIncrementalVersion(0X48D958B3, 0X565444EE, 0X96DDFC64, 0X4F776273);
+FGuid CookIncrementalVersion(0xAEB44DBC, 0x78124EE2, 0xBEAD667A, 0x1D76D306);
 
 }
