Remove experimental PIE preview for targeted mobile devices

#Synced-CL 49717515

diff --git a/Engine/Source/Editor/UnrealEd/Classes/Settings/EditorExperimentalSettings.h b/Engine/Source/Editor/UnrealEd/Classes/Settings/EditorExperimentalSettings.h
--- a/Engine/Source/Editor/UnrealEd/Classes/Settings/EditorExperimentalSettings.h
+++ b/Engine/Source/Editor/UnrealEd/Classes/Settings/EditorExperimentalSettings.h
@@ -123,10 +123,6 @@ public:
 	UPROPERTY(EditAnywhere, config, Category = Tools)
 	bool bEnableProjectLauncherExtensions;
 
-	/** Enable experimental PIE preview device launch */
-	UPROPERTY(EditAnywhere, config, Category = PIE, meta = (DisplayName = "Enable mobile PIE with preview device launch options."))
-	bool bMobilePIEPreviewDeviceLaunch;
-
 	/** Enables in-editor support for text asset formats */
 	UPROPERTY(EditAnywhere, config, Category = Core)
 	bool bTextAssetFormatSupport;
diff --git a/Engine/Source/Editor/UnrealEd/Classes/Settings/LevelEditorPlaySettings.h b/Engine/Source/Editor/UnrealEd/Classes/Settings/LevelEditorPlaySettings.h
--- a/Engine/Source/Editor/UnrealEd/Classes/Settings/LevelEditorPlaySettings.h
+++ b/Engine/Source/Editor/UnrealEd/Classes/Settings/LevelEditorPlaySettings.h
@@ -67,7 +67,7 @@ enum EPlayModeType : int
 	PlayMode_InMobilePreview,
 
 	/** Runs a mobile preview targeted to a particular device in a new process. */
-	PlayMode_InTargetedMobilePreview,
+	PlayMode_InTargetedMobilePreview UE_DEPRECATED(5.8, "PlayMode_InTargetedMobilePreview is no longer supported, please use PlayMode_InMobilePreview instead"),
 
 	/** Runs a vulkan preview in a new process. */
 	PlayMode_InVulkanPreview UE_DEPRECATED(5.7, "PlayMode_InVulkanPreview is no longer supported, please use PlayMode_InMobilePreview instead"),
diff --git a/Engine/Source/Editor/UnrealEd/Private/Kismet2/DebuggerCommands.cpp b/Engine/Source/Editor/UnrealEd/Private/Kismet2/DebuggerCommands.cpp
--- a/Engine/Source/Editor/UnrealEd/Private/Kismet2/DebuggerCommands.cpp
+++ b/Engine/Source/Editor/UnrealEd/Private/Kismet2/DebuggerCommands.cpp
@@ -163,9 +163,6 @@ public:
 	static bool PlayInVR_CanExecute();
 	static bool PlayInModeIsChecked(EPlayModeType PlayMode);
 
-	static void PlayInNewProcessPreviewDevice_Clicked(FString PIEPreviewDeviceName);
-	static bool PlayInModeAndPreviewDeviceIsChecked(FString PIEPreviewDeviceName);
-
 	static bool PlayInLocation_CanExecute(EPlayModeLocations Location);
 	static void PlayInLocation_Clicked(EPlayModeLocations Location);
 	static bool PlayInLocation_IsChecked(EPlayModeLocations Location);
@@ -373,37 +370,6 @@ void FPlayWorldCommands::RegisterCommands()
 	UI_COMMAND(StepInto, "Step Into", "Step Into the next node to be executed", EUserInterfaceActionType::Button, PLATFORM_MAC ? FInputChord(EModifierKey::Control, EKeys::F11) : FInputChord(EKeys::F11));
 	UI_COMMAND(StepOver, "Step Over", "Step to the next node to be executed in the current graph", EUserInterfaceActionType::Button, FInputChord(EKeys::F10));
 	UI_COMMAND(StepOut, "Step Out", "Step Out to the next node to be executed in the parent graph", EUserInterfaceActionType::Button, FInputChord(EModifierKey::Alt | EModifierKey::Shift, EKeys::F11));
-
-	// PIE mobile preview devices.
-	AddPIEPreviewDeviceCommands();
-}
-
-void FPlayWorldCommands::AddPIEPreviewDeviceCommands()
-{
-	auto PIEPreviewDeviceModule = FModuleManager::LoadModulePtr<FPIEPreviewDeviceModule>(TEXT("PIEPreviewDeviceProfileSelector"));
-	if (PIEPreviewDeviceModule)
-	{
-		TArray<TSharedPtr<FUICommandInfo>>& TargetedMobilePreviewDeviceCommands = PlayInTargetedMobilePreviewDevices;
-		const TArray<FString>& Devices = PIEPreviewDeviceModule->GetPreviewDeviceContainer().GetDeviceSpecificationsLocalizedName();
-		PlayInTargetedMobilePreviewDevices.SetNum(Devices.Num());
-		for (int32 DeviceIndex = 0; DeviceIndex < Devices.Num(); DeviceIndex++)
-		{
-			FFormatNamedArguments Args;
-			Args.Add(TEXT("Device"), FText::FromString(Devices[DeviceIndex]));
-			const FText CommandLabel = FText::Format(LOCTEXT("DevicePreviewLaunchCommandLabel", "{Device}"), Args);
-			const FText CommandDesc = FText::Format(LOCTEXT("DevicePreviewLaunchCommandDesc", "Launch on this computer using {Device}'s settings."), Args);
-
-			FUICommandInfo::MakeCommandInfo(
-				this->AsShared(),
-				TargetedMobilePreviewDeviceCommands[DeviceIndex],
-				FName(*CommandLabel.ToString()),
-				CommandLabel,
-				CommandDesc,
-				FSlateIcon(FAppStyle::GetAppStyleSetName(), "PlayWorld.PlayInMobilePreview"),
-				EUserInterfaceActionType::Check,
-				FInputChord());
-		}
-	}
 }
 
 void FPlayWorldCommands::BindGlobalPlayWorldCommands()
@@ -606,28 +572,6 @@ void FPlayWorldCommands::BindGlobalPlayWorldCommands()
 		FIsActionChecked(),
 		FIsActionButtonVisible::CreateStatic(&FInternalPlayWorldCommandCallbacks::IsStoppedAtBreakpoint)
 	);
-
-	AddPIEPreviewDeviceActions(Commands, ActionList);
-}
-
-void FPlayWorldCommands::AddPIEPreviewDeviceActions(const FPlayWorldCommands &Commands, FUICommandList &ActionList)
-{
-	// PIE preview devices.
-	auto PIEPreviewDeviceModule = FModuleManager::LoadModulePtr<FPIEPreviewDeviceModule>(TEXT("PIEPreviewDeviceProfileSelector"));
-	if (PIEPreviewDeviceModule)
-	{
-		const TArray<TSharedPtr<FUICommandInfo>>& TargetedMobilePreviewDeviceCommands = Commands.PlayInTargetedMobilePreviewDevices;
-		const TArray<FString>& Devices = PIEPreviewDeviceModule->GetPreviewDeviceContainer().GetDeviceSpecifications();
-		for (int32 DeviceIndex = 0; DeviceIndex < Devices.Num(); DeviceIndex++)
-		{
-			ActionList.MapAction(TargetedMobilePreviewDeviceCommands[DeviceIndex],
-				FExecuteAction::CreateStatic(&FInternalPlayWorldCommandCallbacks::PlayInNewProcessPreviewDevice_Clicked, Devices[DeviceIndex]),
-				FCanExecuteAction::CreateStatic(&FInternalPlayWorldCommandCallbacks::PlayInNewProcess_CanExecute),
-				FIsActionChecked::CreateStatic(&FInternalPlayWorldCommandCallbacks::PlayInModeAndPreviewDeviceIsChecked, Devices[DeviceIndex]),
-				FIsActionButtonVisible::CreateStatic(&FInternalPlayWorldCommandCallbacks::CanShowNonPlayWorldOnlyActions)
-			);
-		}
-	}
 }
 
 void FPlayWorldCommands::BuildToolbar(FToolMenuSection& InSection)
@@ -774,46 +718,6 @@ void FPlayWorldCommands::BuildToolbar(FToolMenuSection& InSection)
 	InSection.AddEntry(StepOutEntry);
 }
 
-static void MakePreviewDeviceMenu(FMenuBuilder& MenuBuilder)
-{
-	struct FLocal
-	{
-		static void AddDevicePreviewSubCategories(FMenuBuilder& MenuBuilderIn, TSharedPtr<FPIEPreviewDeviceContainerCategory> PreviewDeviceCategory)
-		{
-			const TArray<TSharedPtr<FUICommandInfo>>& TargetedMobilePreviewDeviceCommands = FPlayWorldCommands::Get().PlayInTargetedMobilePreviewDevices;
-			int32 StartIndex = PreviewDeviceCategory->GetDeviceStartIndex();
-			int32 EndIndex = StartIndex + PreviewDeviceCategory->GetDeviceCount();
-			for (int32 Device = StartIndex; Device < EndIndex; Device++)
-			{
-				MenuBuilderIn.AddMenuEntry(TargetedMobilePreviewDeviceCommands[Device]);
-			}
-
-			// let other classes add to this menu
-			FPIEPreviewDeviceModule* PIEPreviewDeviceModule = FModuleManager::LoadModulePtr<FPIEPreviewDeviceModule>(TEXT("PIEPreviewDeviceProfileSelector"));
-			PIEPreviewDeviceModule->AddToDevicePreviewMenuDelegates.Broadcast(PreviewDeviceCategory->GetCategoryDisplayName(), MenuBuilderIn);
-
-			for (TSharedPtr<FPIEPreviewDeviceContainerCategory> SubCategory : PreviewDeviceCategory->GetSubCategories())
-			{
-				MenuBuilderIn.AddSubMenu(
-					SubCategory->GetCategoryDisplayName(),
-					SubCategory->GetCategoryToolTip(),
-					FNewMenuDelegate::CreateStatic(&FLocal::AddDevicePreviewSubCategories, SubCategory)
-				);
-			}
-		}
-	};
-
-	const TArray<TSharedPtr<FUICommandInfo>>& TargetedMobilePreviewDeviceCommands = FPlayWorldCommands::Get().PlayInTargetedMobilePreviewDevices;
-	auto PIEPreviewDeviceModule = FModuleManager::LoadModulePtr<FPIEPreviewDeviceModule>(TEXT("PIEPreviewDeviceProfileSelector"));
-	if (PIEPreviewDeviceModule)
-	{
-		const FPIEPreviewDeviceContainer& DeviceContainer = PIEPreviewDeviceModule->GetPreviewDeviceContainer();
-		MenuBuilder.BeginSection("LevelEditorPlayModesPreviewDevice", LOCTEXT("PreviewDevicePlayButtonModesSection", "Preview Devices"));
-		FLocal::AddDevicePreviewSubCategories(MenuBuilder, DeviceContainer.GetRootCategory());
-		MenuBuilder.EndSection();
-	}
-}
-
 void SetLastExecutedPlayMode(EPlayModeType PlayMode)
 {
 	ULevelEditorPlaySettings* PlaySettings = GetMutableDefault<ULevelEditorPlaySettings>();
@@ -891,18 +795,6 @@ TSharedRef< SWidget > FPlayWorldCommands::GeneratePlayMenuContent(TSharedRef<FUI
 			FLocal::AddPlayModeMenuEntry(Section, PlayMode_InViewPort);
 			FLocal::AddPlayModeMenuEntry(Section, PlayMode_InEditorFloating);
 			FLocal::AddPlayModeMenuEntry(Section, PlayMode_InVR);
-
-			if (GetDefault<UEditorExperimentalSettings>()->bMobilePIEPreviewDeviceLaunch)
-			{
-				Section.AddSubMenu(
-					"TargetedMobilePreview",
-					LOCTEXT("TargetedMobilePreviewSubMenu", "Standalone Game Target Device Mobile"),
-					LOCTEXT("TargetedMobilePreviewSubMenu_ToolTip", "Play this level using a specified mobile device preview (runs in its own process)"),
-					FNewMenuDelegate::CreateStatic(&MakePreviewDeviceMenu), false,
-					FSlateIcon(FAppStyle::GetAppStyleSetName(), "PlayWorld.PlayInMobilePreview")
-				);
-			}
-
 			FLocal::AddPlayModeMenuEntry(Section, PlayMode_InMobilePreview);
 			FLocal::AddPlayModeMenuEntry(Section, PlayMode_InNewProcess);
 
@@ -1303,21 +1195,6 @@ const TSharedRef < FUICommandInfo > GetLastPlaySessionCommand()
 		Command = Commands.PlayInMobilePreview.ToSharedRef();
 		break;
 
-	case PlayMode_InTargetedMobilePreview:
-	{
-		// Scan through targeted mobile preview commands to find our match.
-		for (auto PreviewerCommand : Commands.PlayInTargetedMobilePreviewDevices)
-		{
-			FName LastExecutedPIEPreviewDevice = FName(*PlaySettings->LastExecutedPIEPreviewDevice);
-			if (PreviewerCommand->GetCommandName() == LastExecutedPIEPreviewDevice)
-			{
-				Command = PreviewerCommand.ToSharedRef();
-				break;
-			}
-		}
-		break;
-	}
-
 	case PlayMode_InNewProcess:
 		Command = Commands.PlayInNewProcess.ToSharedRef();
 		break;
@@ -1373,10 +1250,6 @@ void RecordLastExecutedPlayMode()
 			PlayModeString = TEXT("InMobilePreview");
 			break;
 
-		case PlayMode_InTargetedMobilePreview:
-			PlayModeString = TEXT("InTargetedMobilePreview");
-			break;
-
 		case PlayMode_InNewProcess:
 			PlayModeString = TEXT("InNewProcess");
 			break;
@@ -1675,17 +1548,10 @@ void SetLastExecutedPIEPreviewDevice(FString PIEPreviewDevice)
 	PlaySettings->SaveConfig();
 }
 
-void FInternalPlayWorldCommandCallbacks::PlayInNewProcessPreviewDevice_Clicked(FString PIEPreviewDeviceName)
-{
-	SetLastExecutedPIEPreviewDevice(PIEPreviewDeviceName);
-	PlayInNewProcess_Clicked(PlayMode_InTargetedMobilePreview);
-}
-
 void FInternalPlayWorldCommandCallbacks::PlayInNewProcess_Clicked(EPlayModeType PlayModeType)
 {
 	check(PlayModeType == PlayMode_InNewProcess 
-		|| PlayModeType == PlayMode_InMobilePreview
-		|| PlayModeType == PlayMode_InTargetedMobilePreview);
+		|| PlayModeType == PlayMode_InMobilePreview);
 
 	if (IsStoppedAtBreakpoint_InEngineMode() || DebuggerCommands::IsVREditorActive())
 	{
@@ -1712,13 +1578,8 @@ void FInternalPlayWorldCommandCallbacks::PlayInNewProcess_Clicked(EPlayModeType
 			}
 		}
 
-		if (PlayModeType == PlayMode_InMobilePreview || PlayModeType == PlayMode_InTargetedMobilePreview)
+		if (PlayModeType == PlayMode_InMobilePreview)
 		{
-			if (PlayModeType == PlayMode_InTargetedMobilePreview)
-			{
-				SessionParams.MobilePreviewTargetDevice = GetDefault<ULevelEditorPlaySettings>()->LastExecutedPIEPreviewDevice;
-			}
-
 			SessionParams.SessionPreviewTypeOverride = EPlaySessionPreviewType::MobilePreview;
 		}
 
@@ -1749,11 +1610,6 @@ bool FInternalPlayWorldCommandCallbacks::PlayInNewProcess_CanExecute()
 }
 
 
-bool FInternalPlayWorldCommandCallbacks::PlayInModeAndPreviewDeviceIsChecked(FString PIEPreviewDeviceName)
-{
-	return PlayInModeIsChecked(PlayMode_InTargetedMobilePreview) && GetDefault<ULevelEditorPlaySettings>()->LastExecutedPIEPreviewDevice == PIEPreviewDeviceName;
-}
-
 bool FInternalPlayWorldCommandCallbacks::PlayInModeIsChecked(EPlayModeType PlayMode)
 {
 	return (PlayMode == GetDefault<ULevelEditorPlaySettings>()->LastExecutedPlayModeType);
diff --git a/Engine/Source/Editor/UnrealEd/Private/PlayLevelNewProcess.cpp b/Engine/Source/Editor/UnrealEd/Private/PlayLevelNewProcess.cpp
--- a/Engine/Source/Editor/UnrealEd/Private/PlayLevelNewProcess.cpp
+++ b/Engine/Source/Editor/UnrealEd/Private/PlayLevelNewProcess.cpp
@@ -137,16 +137,7 @@ void UEditorEngine::LaunchNewProcess(const FRequestPlaySessionParams& InParams,
 
 	if (InParams.SessionPreviewTypeOverride.Get(EPlaySessionPreviewType::NoPreview) == EPlaySessionPreviewType::MobilePreview)
 	{
-		// Allow targeting a specific Mobile Device.
-		if (InParams.MobilePreviewTargetDevice.IsSet())
-		{
-			CommandLine += FString::Printf(TEXT(" -MobileTargetDevice=\"%s\""), *InParams.MobilePreviewTargetDevice.GetValue());
-		}
-		else
-		{
-			// Otherwise, we'll fall back to ES31 emulation.
-			CommandLine += TEXT(" -featureleveles31");
-		}
+		CommandLine += TEXT(" -featureleveles31");
 
 		// If we're currently running in OpenGL mode, pass that onto our newly spawned processes.
 		if (IsOpenGLPlatform(GShaderPlatformForFeatureLevel[GMaxRHIFeatureLevel]))
diff --git a/Engine/Source/Editor/UnrealEd/Public/Kismet2/DebuggerCommands.h b/Engine/Source/Editor/UnrealEd/Public/Kismet2/DebuggerCommands.h
--- a/Engine/Source/Editor/UnrealEd/Public/Kismet2/DebuggerCommands.h
+++ b/Engine/Source/Editor/UnrealEd/Public/Kismet2/DebuggerCommands.h
@@ -95,8 +95,6 @@ public:
 	TSharedPtr<FUICommandInfo> PlayInSettings;
 	TSharedPtr<FUICommandInfo> PlayInNetworkSettings;
 
-	TArray< TSharedPtr< FUICommandInfo > > PlayInTargetedMobilePreviewDevices;
-
 	/** SIE & PIE controls */
 	TSharedPtr<FUICommandInfo> ResumePlaySession;
 	TSharedPtr<FUICommandInfo> PausePlaySession;
@@ -126,12 +124,6 @@ protected:
 	 * @return	Menu content widget
 	 */
 	static UNREALED_API TSharedRef< SWidget > GeneratePlayMenuContent( TSharedRef<FUICommandList> InCommandList );
-
-	// Add mobile PIE preview device commands
-	UNREALED_API void AddPIEPreviewDeviceCommands();
-
-	// Add mobile PIE preview device actions
-	static UNREALED_API void AddPIEPreviewDeviceActions(const FPlayWorldCommands &Commands, FUICommandList &ActionList);
 };
 
 
diff --git a/Engine/Source/Editor/UnrealEd/Public/PlayInEditorDataTypes.h b/Engine/Source/Editor/UnrealEd/Public/PlayInEditorDataTypes.h
--- a/Engine/Source/Editor/UnrealEd/Public/PlayInEditorDataTypes.h
+++ b/Engine/Source/Editor/UnrealEd/Public/PlayInEditorDataTypes.h
@@ -185,12 +185,6 @@ public:
 	*/
 	TOptional<FString> AdditionalStandaloneCommandLineParameters;
 
-	/**
-	* If set, this device name will be passed to the -MobileTargetDevice command line option.
-	* Requires EPlaySessionPreviewType::Mobile
-	*/
-	TOptional<FString> MobilePreviewTargetDevice;
-
 	/** If specified, the PIE instance will be created inside this user-provided window instead of creating a default. */
 	TWeakPtr<SWindow> CustomPIEWindow;
 
