Add ScreenToWorld and WorldToScreen function to UnrealEditor subsystem.

Added here because the equivalent functions in GameplayStatics require a PlayerController, which don't exist in the editor (outside of PIE).

#Synced-CL 49654589

diff --git a/Engine/Source/Editor/UnrealEd/Private/Subsystems/UnrealEditorSubsystem.cpp b/Engine/Source/Editor/UnrealEd/Private/Subsystems/UnrealEditorSubsystem.cpp
--- a/Engine/Source/Editor/UnrealEd/Private/Subsystems/UnrealEditorSubsystem.cpp
+++ b/Engine/Source/Editor/UnrealEd/Private/Subsystems/UnrealEditorSubsystem.cpp
@@ -3,6 +3,7 @@
 #include "Subsystems/UnrealEditorSubsystem.h"
 #include "Editor.h"
 #include "SLevelViewport.h"
+#include "SceneView.h"
 #include "LevelEditor.h"
 #include "Utils.h"
 #include "GameFramework/Actor.h"
@@ -35,37 +36,26 @@ namespace InternalUnrealEditorSubsystemLibrary
 
 bool UUnrealEditorSubsystem::GetLevelViewportCameraInfo(FVector& CameraLocation, FRotator& CameraRotation)
 {
-	bool RetVal = false;
 	CameraLocation = FVector::ZeroVector;
 	CameraRotation = FRotator::ZeroRotator;
 
-	for (FLevelEditorViewportClient* LevelVC : GEditor->GetLevelViewportClients())
+	if (FLevelEditorViewportClient* LevelVC = GetViewportClient())
 	{
-		if (LevelVC && LevelVC->IsPerspective())
-		{
-			CameraLocation = LevelVC->GetViewLocation();
-			CameraRotation = LevelVC->GetViewRotation();
-			RetVal = true;
-
-			break;
-		}
+		CameraLocation = LevelVC->GetViewLocation();
+		CameraRotation = LevelVC->GetViewRotation();
+		return true;
 	}
 
-	return RetVal;
+	return false;
 }
 
 void UUnrealEditorSubsystem::SetLevelViewportCameraInfo(FVector CameraLocation, FRotator CameraRotation)
 {
-	for (FLevelEditorViewportClient* LevelVC : GEditor->GetLevelViewportClients())
+	if (FLevelEditorViewportClient* LevelVC = GetViewportClient())
 	{
-		if (LevelVC && LevelVC->IsPerspective())
-		{
-			LevelVC->SetViewLocationForOrbiting(CameraLocation);
-			LevelVC->SetViewLocation(CameraLocation);
-			LevelVC->SetViewRotation(CameraRotation);
-
-			break;
-		}
+		LevelVC->SetViewLocationForOrbiting(CameraLocation);
+		LevelVC->SetViewLocation(CameraLocation);
+		LevelVC->SetViewRotation(CameraRotation);
 	}
 }
 
@@ -87,3 +77,51 @@ UWorld* UUnrealEditorSubsystem::GetGameWorld()
 
 	return InternalUnrealEditorSubsystemLibrary::GetGameWorld();
 }
+
+bool UUnrealEditorSubsystem::ScreenToWorld(const FVector2D& ScreenPosition, FVector& WorldPosition, FVector& WorldDirection) const
+{	
+	if (FLevelEditorViewportClient* LevelVC = GetViewportClient())
+	{		
+		FSceneViewFamily ViewFamily(FSceneViewFamily::ConstructionValues(
+			LevelVC->Viewport,
+			LevelVC->GetScene(),
+			LevelVC->EngineShowFlags).SetRealtimeUpdate(LevelVC->IsRealtime()));
+		TSharedPtr<FSceneView> SceneView = MakeShareable<>(LevelVC->CalcSceneView(&ViewFamily));
+		const FMatrix InvViewProjectionMatrix = SceneView->ViewMatrices.GetInvViewProjectionMatrix();		
+		FSceneView::DeprojectScreenToWorld(ScreenPosition, FIntRect(FIntPoint(0, 0), LevelVC->Viewport->GetSizeXY()), InvViewProjectionMatrix, WorldPosition, WorldDirection);
+		return true;
+	}
+
+	// something went wrong, zero things and return false
+	WorldPosition = FVector::ZeroVector;
+	WorldDirection = FVector::ZeroVector;
+	return false;
+}
+
+bool UUnrealEditorSubsystem::WorldToScreen(const FVector& WorldPosition, FVector2D& ScreenPosition) const
+{
+	if (FLevelEditorViewportClient* LevelVC = GetViewportClient())
+	{
+		FSceneViewFamily ViewFamily(FSceneViewFamily::ConstructionValues(
+			LevelVC->Viewport,
+			LevelVC->GetScene(),
+			LevelVC->EngineShowFlags).SetRealtimeUpdate(LevelVC->IsRealtime()));
+		TSharedPtr<FSceneView> SceneView = MakeShareable<>(LevelVC->CalcSceneView(&ViewFamily));
+		return FSceneView::ProjectWorldToScreen(WorldPosition, FIntRect(FIntPoint(0, 0), LevelVC->Viewport->GetSizeXY()), SceneView->ViewMatrices.GetViewProjectionMatrix(), ScreenPosition);		
+	}
+
+	ScreenPosition = FVector2D::ZeroVector;
+	return false;
+}
+
+FLevelEditorViewportClient* UUnrealEditorSubsystem::GetViewportClient() const
+{
+	for (FLevelEditorViewportClient* LevelVC : GEditor->GetLevelViewportClients())
+	{
+		if (LevelVC && LevelVC->IsPerspective())
+		{
+			return LevelVC;
+		}
+	}
+	return nullptr;
+}
diff --git a/Engine/Source/Editor/UnrealEd/Public/Subsystems/UnrealEditorSubsystem.h b/Engine/Source/Editor/UnrealEd/Public/Subsystems/UnrealEditorSubsystem.h
--- a/Engine/Source/Editor/UnrealEd/Public/Subsystems/UnrealEditorSubsystem.h
+++ b/Engine/Source/Editor/UnrealEd/Public/Subsystems/UnrealEditorSubsystem.h
@@ -7,6 +7,8 @@
 
 #include "UnrealEditorSubsystem.generated.h"
 
+class FLevelEditorViewportClient;
+
 /**
 * UUnrealEditorSubsystem 
 * Subsystem for exposing editor functionality to scripts
@@ -48,4 +50,24 @@ public:
 
 	UFUNCTION(BlueprintCallable, Category = "Editor Scripting | Level Utility")
 	UNREALED_API UWorld* GetGameWorld();
+
+	/**
+	 * Transforms the given 2D screen space coordinate into a 3D world-space point and direction using the primary level editor viewport.	 
+	 * @param ScreenPosition	2D screen space to deproject.
+	 * @param WorldPosition		(out) Corresponding 3D position in world space.
+	 * @param WorldDirection	(out) World space direction vector away from the camera at the given 2d point.
+	 */
+	UFUNCTION(BlueprintPure, Category = "Development|Editor")
+	UNREALED_API bool ScreenToWorld(const FVector2D& ScreenPosition, FVector& WorldPosition, FVector& WorldDirection) const;
+
+	/**
+	 * Transforms the given 3D world-space point into a its 2D screen space coordinate using the primary level editor viewport.	 
+	 * @param WorldPosition		World position to project.
+	 * @param ScreenPosition	(out) Corresponding 2D position in screen space	 
+	 */
+	UFUNCTION(BlueprintPure, Category = "Development|Editor")	
+	UNREALED_API bool WorldToScreen(const FVector& WorldPosition, FVector2D& ScreenPosition) const;
+
+private:
+	FLevelEditorViewportClient* GetViewportClient() const;
 };
