[Navigation] Fix Build->"Build All Levels" that was leaving the navigation system in a broken state.
#Synced-CL 49739737

diff --git a/Engine/Source/Editor/UnrealEd/Private/UnrealEdMisc.cpp b/Engine/Source/Editor/UnrealEd/Private/UnrealEdMisc.cpp
--- a/Engine/Source/Editor/UnrealEd/Private/UnrealEdMisc.cpp
+++ b/Engine/Source/Editor/UnrealEd/Private/UnrealEdMisc.cpp
@@ -105,6 +105,7 @@
 #include "Factories/FbxSkeletalMeshImportData.h"
 #include "Factories/FbxStaticMeshImportData.h"
 #include "OutputLogModule.h"
+#include "AI/NavigationSystemBase.h"
 #include "Editor/EditorPerProjectUserSettings.h"
 #include "Async/TaskGraphInterfaces.h"
 
@@ -1291,6 +1292,13 @@ void FUnrealEdMisc::CB_MapChange( uint32 InFlags )
 			{
 				// Call CleanupWorld/InitWorld to detach components and some other stuff
 				World->ReInitWorld();
+
+				// ReInitWorld (just above) cleanup and initialize the world but InitWorld does not fully initialize the navigation system for the world.
+				// Do it here for map rebuilds ('Build All Levels').
+				if (InFlags & MapChangeEventFlags::MapRebuild)
+				{
+					FNavigationSystem::AddNavigationSystemToWorld(*World, FNavigationSystemRunMode::EditorMode);
+				}
 			}
 		}
 
