TextureAssetActionsBP to exposure TextureAssetActions to Blueprint
with enum return

#Synced-CL 49826549

diff --git a/Engine/Source/Editor/UnrealEd/Private/TextureAssetActions.cpp b/Engine/Source/Editor/UnrealEd/Private/TextureAssetActions.cpp
--- a/Engine/Source/Editor/UnrealEd/Private/TextureAssetActions.cpp
+++ b/Engine/Source/Editor/UnrealEd/Private/TextureAssetActions.cpp
@@ -33,19 +33,14 @@
 #include "TextureImportSettings.h"
 #include "ImageCoreUtils.h"
 
+#include "TextureAssetActionsBP.h"
+
 #define LOCTEXT_NAMESPACE "TextureAssetActions"
 
-namespace
+namespace UE
 {
-
-enum class ETextureAction
+namespace TextureAssetActionsInternal
 {
-	Invalid = 0,
-	Resize,
-	ResizePow2,
-	ConvertTo8bit,
-	JPEG
-};
 
 enum class EAssetActionStatus
 {
@@ -660,9 +655,12 @@ TSharedRef<SWidget> STextureAssetList::CreateAssetLine(int index, const UTexture
 	return Result;
 }
 
-static void DoResizeTextureSourceToPowerOfTwo(UTexture * Texture)
+bool UE::TextureAssetActionsInternal::ResizeTextureSourceToPowerOfTwo(UTexture * Texture)
 {
-	check( ! Texture->Source.AreAllBlocksPowerOfTwo() ); // already filtered for
+	if ( Texture->Source.AreAllBlocksPowerOfTwo() )
+	{
+		return false;
+	}
 
 	const FIntPoint BeforeSourceSize = Texture->Source.GetLogicalSize();
 
@@ -671,7 +669,7 @@ static void DoResizeTextureSourceToPowerOfTwo(UTexture * Texture)
 		UE_LOG(LogTexture, Display, TEXT("Texture (%s) did not resize."), *Texture->GetName());
 
 		// did not resize, but may have done PreEditChange
-		return;
+		return false;
 	}
 	
 	const FIntPoint AfterSourceSize = Texture->Source.GetLogicalSize();
@@ -701,9 +699,13 @@ static void DoResizeTextureSourceToPowerOfTwo(UTexture * Texture)
 
 	// ResizeTextureSourceDataToNearestPowerOfTwo did the PreEditChange
 	Texture->PostEditChange();
+
+	return true;
 }
 
-static void DoResizeTextureSource(UTexture * Texture,int TargetSize)
+// resizes to <= TargetSize
+//	if size was already <= TargetSize , no change is made and we return false
+bool UE::TextureAssetActionsInternal::ResizeTextureSource(UTexture * Texture,int TargetSize)
 {
 	// we do the resizing considering only mip/LOD/build settings for the running Editor platform (eg. Windows)
 	const ITargetPlatform* RunningPlatform = GetTargetPlatformManagerRef().GetRunningTargetPlatform();
@@ -721,7 +723,7 @@ static void DoResizeTextureSource(UTexture * Texture,int TargetSize)
 		UE_LOG(LogTexture, Display, TEXT("Texture (%s) did not resize."), *Texture->GetName());
 
 		// did not resize, but may have done PreEditChange
-		return;
+		return false;
 	}
 	
 	const FIntPoint AfterSourceSize = Texture->Source.GetLogicalSize();
@@ -762,6 +764,8 @@ static void DoResizeTextureSource(UTexture * Texture,int TargetSize)
 
 	// DownsizeTextureSourceData did the PreEditChange
 	Texture->PostEditChange();
+
+	return true;
 }
 
 static ETextureSourceFormat GetReducedTextureSourceFormat(const TextureCompressionSettings TC,const ETextureSourceFormat InTSF,const bool NormalMapsKeep16bits)
@@ -888,19 +892,18 @@ static ETextureSourceFormat GetReducedTextureSourceFormat(const TextureCompressi
 	}
 }
 
-static void DoConvertTo8bitTextureSource(UTexture * Texture,bool NormalMapsKeep16bits)
+bool UE::TextureAssetActionsInternal::ConvertTo8bitTextureSource(UTexture * Texture,bool NormalMapsKeep16bits)
 {
 	if ( Texture->Source.GetNumLayers() > 1 )
 	{
-		check(0);
-		return;
+		return false;
 	}
 
 	ETextureSourceFormat InTSF = Texture->Source.GetFormat();
 	ETextureSourceFormat OutTSF = GetReducedTextureSourceFormat(Texture->CompressionSettings,InTSF,NormalMapsKeep16bits);
 	if ( InTSF == OutTSF )
 	{
-		return;
+		return false;
 	}
 	
 	UE_LOG(LogTexture, Display, TEXT("Texture (%s) changing format from %s to %s, TC = %s"), 
@@ -911,10 +914,12 @@ static void DoConvertTo8bitTextureSource(UTexture * Texture,bool NormalMapsKeep1
 		);
 
 	// calls Pre/PostEditChange :
-	UE::TextureUtilitiesCommon::Experimental::ChangeTextureSourceFormat(Texture,OutTSF);	
+	UE::TextureUtilitiesCommon::Experimental::ChangeTextureSourceFormat(Texture,OutTSF);
+
+	return true;
 }
 
-static void DoCompressTextureSourceWithJPEG(UTexture * Texture)
+bool UE::TextureAssetActionsInternal::CompressTextureSourceWithJPEG(UTexture * Texture)
 {
 	// calls Pre/PostEditChange :
 	bool bDid = UE::TextureUtilitiesCommon::Experimental::CompressTextureSourceWithJPEG(Texture);
@@ -923,6 +928,8 @@ static void DoCompressTextureSourceWithJPEG(UTexture * Texture)
 		*Texture->GetName() ,
 		bDid ? TEXT("was changed to JPEG compression.") : TEXT("was not changed.")
 		);
+
+	return bDid;
 }
 
 void STextureAssetList::DoAction()
@@ -963,16 +970,16 @@ void STextureAssetList::DoAction()
 		switch(Action)
 		{
 		case ETextureAction::Resize:
-			DoResizeTextureSource(Texture,ThresholdValue);
+			UE::TextureAssetActionsInternal::ResizeTextureSource(Texture,ThresholdValue);
 			break;
 		case ETextureAction::ResizePow2:
-			DoResizeTextureSourceToPowerOfTwo(Texture);
+			UE::TextureAssetActionsInternal::ResizeTextureSourceToPowerOfTwo(Texture);
 			break;
 		case ETextureAction::ConvertTo8bit:
-			DoConvertTo8bitTextureSource(Texture,bNormalMapsKeep16bits);
+			UE::TextureAssetActionsInternal::ConvertTo8bitTextureSource(Texture,bNormalMapsKeep16bits);
 			break;
 		case ETextureAction::JPEG:
-			DoCompressTextureSourceWithJPEG(Texture);
+			UE::TextureAssetActionsInternal::CompressTextureSourceWithJPEG(Texture);
 			break;
 		default:
 			check(0);
@@ -981,148 +988,205 @@ void STextureAssetList::DoAction()
 	}
 }
 
-void STextureAssetList::UpdateList()
+static EAssetActionStatus GetTextureFilterStatus(UTexture * Texture, ETextureAction Action, int ThresholdValue,
+		bool bNormalMapsKeep16bits,
+		FString & EnabledActionString,
+		FString & DescriptionString )
 {
-	ExpandedIndexes.SetNum(0);
-
-	AssetListContainer->ClearChildren();
-	int32 AssetListContainerIndex = 0;
-
-	// filter select textures to see if they should be acted on
-	for (int Index=0;Index<TextureList.Num();Index++)
+	EAssetActionStatus Status = EAssetActionStatus::Enabled; 
+				
+	ETextureClass Class = Texture->GetTextureClass();
+		
+	// first filter for exclusions that apply to all actions
+	//	don't query anything in Source until you check IsValid
+	if ( Class != ETextureClass::TwoD &&
+		Class != ETextureClass::Cube &&
+		Class != ETextureClass::Array &&
+		Class != ETextureClass::CubeArray &&
+		Class != ETextureClass::Volume )
 	{
-		TextureListEntry & Entry = TextureList[Index];
-		UTexture * Texture = Entry.Texture;
+		// rendertarget, dynamic, or unknown
+		Status = EAssetActionStatus::WrongType;
+	}
+	else if ( ! Texture->Source.IsValid() )
+	{
+		Status = EAssetActionStatus::NoSource;
+	}
+	else if ( Texture->Source.GetNumLayers() != 1 ) // only 1 layer textures
+	{
+		Status = EAssetActionStatus::HasLayers;
+	}
+	else if ( Texture->Source.GetNumMips() != 1 && Texture->MipGenSettings == TMGS_LeaveExistingMips )
+	{
+		Status = EAssetActionStatus::HasMipsLeaveExisting;
+	}
+	else
+	{
+		// don't prep Description text until after Source IsValid check
 
-		EAssetActionStatus Status = EAssetActionStatus::Enabled; 
-		
-		FString EnabledActionString;
-		FString DescriptionString;
-		
-		ETextureClass Class = Texture->GetTextureClass();
-		
-		// first filter for exclusions that apply to all actions
-		//	don't query anything in Source until you check IsValid
-		if ( Class != ETextureClass::TwoD &&
-			Class != ETextureClass::Cube &&
-			Class != ETextureClass::Array &&
-			Class != ETextureClass::CubeArray &&
-			Class != ETextureClass::Volume )
+		// GetLogicalSize = sum of udim blocks
+		const FIntPoint SourceSize = Texture->Source.GetLogicalSize();
+		const ETextureSourceFormat InTSF = Texture->Source.GetFormat();
+
+		DescriptionString = FString::Printf(TEXT("%dx%d %s"),SourceSize.X,SourceSize.Y,
+			ERawImageFormat::GetName( FImageCoreUtils::ConvertToRawImageFormat(InTSF) ));
+		if ( Texture->VirtualTextureStreaming )
 		{
-			// rendertarget, dynamic, or unknown
-			Status = EAssetActionStatus::WrongType;
+			DescriptionString += FString(TEXT(" VT"));
 		}
-		else if ( ! Texture->Source.IsValid() )
+		if ( Texture->Source.GetSourceCompression() == ETextureSourceCompressionFormat::TSCF_JPEG )
 		{
-			Status = EAssetActionStatus::NoSource;
+			DescriptionString += FString(TEXT(" JPEG"));
 		}
-		else if ( Texture->Source.GetNumLayers() != 1 ) // only 1 layer textures
+		
+		// per-Action filters :
+		if( Action == ETextureAction::Resize || Action == ETextureAction::ResizePow2 )
 		{
-			Status = EAssetActionStatus::HasLayers;
+			int MaxSize = FMath::Max(SourceSize.X,SourceSize.Y);
+
+			if ( Class != ETextureClass::TwoD && Class != ETextureClass::Cube )
+			{
+				// Resize only supports 2d and Cube for now, no arrays or volumes
+				Status = EAssetActionStatus::WrongType;
+			}
+			else if ( Action == ETextureAction::ResizePow2 && Texture->Source.AreAllBlocksPowerOfTwo() )
+			{
+				Status = EAssetActionStatus::AlreadyPow2;
+			}
+			else if ( MaxSize <= ThresholdValue )
+			{
+				Status = EAssetActionStatus::UnderSized;
+			}
+			else
+			{
+				// could fill EnabledActionString
+				// hard to get size we will output exactly right
+			}
 		}
-		else if ( Texture->Source.GetNumMips() != 1 && Texture->MipGenSettings == TMGS_LeaveExistingMips )
+		else if ( Action == ETextureAction::ConvertTo8bit )
 		{
-			Status = EAssetActionStatus::HasMipsLeaveExisting;
+			ETextureSourceFormat OutTSF = GetReducedTextureSourceFormat(Texture->CompressionSettings,InTSF,bNormalMapsKeep16bits);
+			if ( InTSF == OutTSF )
+			{
+				Status = EAssetActionStatus::Already8bit;
+			}
+			else if ( Texture->Source.GetSourceCompression() == ETextureSourceCompressionFormat::TSCF_JPEG )
+			{
+				// JPEG is already 8 bit; we don't want to try to change RGB JPEG to G8, just leave it alone
+				Status = EAssetActionStatus::DontChangeJPEG;
+			}
+			else
+			{
+				// @@ NOTE : ThresholdValue not checked!
+
+				// Enabled!
+				EnabledActionString = FString::Printf(TEXT("TC=%s will convert to %s"),
+					*StaticEnum<TextureCompressionSettings>()->GetDisplayNameTextByValue(Texture->CompressionSettings).ToString(),
+					ERawImageFormat::GetName( FImageCoreUtils::ConvertToRawImageFormat( OutTSF ) ) );
+			}
 		}
-		else
+		else if ( Action == ETextureAction::JPEG )
 		{
-			// don't prep Description text until after Source IsValid check
-
-			// GetLogicalSize = sum of udim blocks
-			const FIntPoint SourceSize = Texture->Source.GetLogicalSize();
-			const ETextureSourceFormat InTSF = Texture->Source.GetFormat();
-
-			DescriptionString = FString::Printf(TEXT("%dx%d %s"),SourceSize.X,SourceSize.Y,
-				ERawImageFormat::GetName( FImageCoreUtils::ConvertToRawImageFormat(InTSF) ));
-			if ( Texture->VirtualTextureStreaming )
+			if ( Class != ETextureClass::TwoD || Texture->Source.GetNumBlocks() > 1 || Texture->Source.GetNumSlices() > 1 )
 			{
-				DescriptionString += FString(TEXT(" VT"));
+				// JPEG only supports 2d
+				// Blocked/UDIM doesn't support JPEG
+				Status = EAssetActionStatus::WrongType;
 			}
-			if ( Texture->Source.GetSourceCompression() == ETextureSourceCompressionFormat::TSCF_JPEG )
+			else if ( Texture->Source.GetSourceCompression() == ETextureSourceCompressionFormat::TSCF_JPEG )
 			{
-				DescriptionString += FString(TEXT(" JPEG"));
+				// Already JPEG
+				Status = EAssetActionStatus::DontChangeJPEG;
 			}
-		
-			// per-Action filters :
-			if( Action == ETextureAction::Resize || Action == ETextureAction::ResizePow2 )
+			else if ( InTSF != TSF_BGRA8 && InTSF != TSF_G8 )
 			{
-				int MaxSize = FMath::Max(SourceSize.X,SourceSize.Y);
-
-				if ( Class != ETextureClass::TwoD && Class != ETextureClass::Cube )
-				{
-					// Resize only supports 2d and Cube for now, no arrays or volumes
-					Status = EAssetActionStatus::WrongType;
-				}
-				else if ( Action == ETextureAction::ResizePow2 && Texture->Source.AreAllBlocksPowerOfTwo() )
-				{
-					Status = EAssetActionStatus::AlreadyPow2;
-				}
-				else if ( MaxSize <= ThresholdValue )
-				{
-					Status = EAssetActionStatus::UnderSized;
-				}
-				else
-				{
-					// could fill EnabledActionString
-					// hard to get size we will output exactly right
-				}
+				Status = EAssetActionStatus::MustBe8BitForJPEG;
 			}
-			else if ( Action == ETextureAction::ConvertTo8bit )
+			else if ( SourceSize.X < 16 || SourceSize.Y < 16 )
 			{
-				ETextureSourceFormat OutTSF = GetReducedTextureSourceFormat(Texture->CompressionSettings,InTSF,bNormalMapsKeep16bits);
-				if ( InTSF == OutTSF )
-				{
-					Status = EAssetActionStatus::Already8bit;
-				}
-				else if ( Texture->Source.GetSourceCompression() == ETextureSourceCompressionFormat::TSCF_JPEG )
-				{
-					// JPEG is already 8 bit; we don't want to try to change RGB JPEG to G8, just leave it alone
-					Status = EAssetActionStatus::DontChangeJPEG;
-				}
-				else
-				{
-					// Enabled!
-					EnabledActionString = FString::Printf(TEXT("TC=%s will convert to %s"),
-						*StaticEnum<TextureCompressionSettings>()->GetDisplayNameTextByValue(Texture->CompressionSettings).ToString(),
-						ERawImageFormat::GetName( FImageCoreUtils::ConvertToRawImageFormat( OutTSF ) ) );
-				}
+				// hard-coded requirement of at least 16 in both dimensions for JPEG
+				Status = EAssetActionStatus::UnderSized;
 			}
-			else if ( Action == ETextureAction::JPEG )
+			else if ( (int64)SourceSize.X*SourceSize.Y < (int64)ThresholdValue*ThresholdValue ) // pixel count
 			{
-				if ( Class != ETextureClass::TwoD || Texture->Source.GetNumBlocks() > 1 || Texture->Source.GetNumSlices() > 1 )
-				{
-					// JPEG only supports 2d
-					// Blocked/UDIM doesn't support JPEG
-					Status = EAssetActionStatus::WrongType;
-				}
-				else if ( Texture->Source.GetSourceCompression() == ETextureSourceCompressionFormat::TSCF_JPEG )
-				{
-					// Already JPEG
-					Status = EAssetActionStatus::DontChangeJPEG;
-				}
-				else if ( InTSF != TSF_BGRA8 && InTSF != TSF_G8 )
-				{
-					Status = EAssetActionStatus::MustBe8BitForJPEG;
-				}
-				else if ( SourceSize.X < 16 || SourceSize.Y < 16 )
-				{
-					// hard-coded requirement of at least 16 in both dimensions
-					Status = EAssetActionStatus::UnderSized;
-				}
-				else if ( (int64)SourceSize.X*SourceSize.Y < (int64)ThresholdValue*ThresholdValue ) // pixel count
-				{
-					Status = EAssetActionStatus::UnderSized;
-				}
-				else
-				{
-					// enabled!
-				}				
+				Status = EAssetActionStatus::UnderSized;
 			}
 			else
 			{
-				check(0);
-			}
+				// enabled!
+			}				
+		}
+		else
+		{
+			check(0);
 		}
+	}
+
+	return Status;
+}
+
+bool UE::TextureAssetActionsInternal::GetTextureFilterEnabled( UTexture * Texture, ETextureAction Action, int ThresholdValue,
+		bool bNormalMapsKeep16bits , bool * bOutExcludedBecauseAlreadySatisfied)
+{
+	FString EnabledActionString;
+	FString DescriptionString;
+
+	EAssetActionStatus Status = GetTextureFilterStatus( Texture, Action, ThresholdValue, bNormalMapsKeep16bits,
+		EnabledActionString, DescriptionString );
+
+	// Special returns for the exclusions that are due to the target already being in the desired state
+	//  return another bool or trinary enum to indicate these
+
+	bool bDummy;
+	if ( bOutExcludedBecauseAlreadySatisfied == nullptr )
+	{
+		bOutExcludedBecauseAlreadySatisfied = &bDummy;
+	}
+
+	if ( Action == ETextureAction::Resize && Status == EAssetActionStatus::UnderSized )
+	{
+		*bOutExcludedBecauseAlreadySatisfied = true;
+		return false;
+	}
+	if ( Action == ETextureAction::ResizePow2 && Status == EAssetActionStatus::AlreadyPow2 )
+	{
+		*bOutExcludedBecauseAlreadySatisfied = true;
+		return false;
+	}
+	if ( Action == ETextureAction::ConvertTo8bit && Status == EAssetActionStatus::Already8bit )
+	{
+		*bOutExcludedBecauseAlreadySatisfied = true;
+		return false;
+	}
+	if ( Action == ETextureAction::JPEG && Status == EAssetActionStatus::DontChangeJPEG )
+	{
+		*bOutExcludedBecauseAlreadySatisfied = true;
+		return false;
+	}
+	
+	*bOutExcludedBecauseAlreadySatisfied = false;
+
+	return ( Status == EAssetActionStatus::Enabled );
+}
+
+void STextureAssetList::UpdateList()
+{
+	ExpandedIndexes.SetNum(0);
+
+	AssetListContainer->ClearChildren();
+	int32 AssetListContainerIndex = 0;
+
+	// filter select textures to see if they should be acted on
+	for (int Index=0;Index<TextureList.Num();Index++)
+	{
+		TextureListEntry & Entry = TextureList[Index];
+		UTexture * Texture = Entry.Texture;
+		
+		FString EnabledActionString;
+		FString DescriptionString;
+
+		EAssetActionStatus Status = GetTextureFilterStatus( Texture, Action, ThresholdValue, bNormalMapsKeep16bits,
+			EnabledActionString, DescriptionString );
 
 		Entry.IsEnabled = ( Status == EAssetActionStatus::Enabled );
 
@@ -1262,29 +1326,30 @@ STextureActionDlg::EResult STextureActionDlg::ShowModal()
 	return UserResponse;
 }
 
+}; // namespace
 }; // namespace
 
 void UE::TextureAssetActions::TextureSource_Resize_WithDialog(const TArray<UTexture*> & InTextures)
 {
-	STextureActionDlg Dlg(InTextures,ETextureAction::Resize);
+	UE::TextureAssetActionsInternal::STextureActionDlg Dlg(InTextures,UE::TextureAssetActionsInternal::ETextureAction::Resize);
 	Dlg.ShowModal();
 }
 
 void UE::TextureAssetActions::TextureSource_ConvertTo8bit_WithDialog(const TArray<UTexture*> & InTextures)
 {
-	STextureActionDlg Dlg(InTextures,ETextureAction::ConvertTo8bit);
+	UE::TextureAssetActionsInternal::STextureActionDlg Dlg(InTextures,UE::TextureAssetActionsInternal::ETextureAction::ConvertTo8bit);
 	Dlg.ShowModal();
 }
 
 void UE::TextureAssetActions::TextureSource_JPEG_WithDialog(const TArray<UTexture*> & InTextures)
 {
-	STextureActionDlg Dlg(InTextures,ETextureAction::JPEG);
+	UE::TextureAssetActionsInternal::STextureActionDlg Dlg(InTextures,UE::TextureAssetActionsInternal::ETextureAction::JPEG);
 	Dlg.ShowModal();
 }
 
 void UE::TextureAssetActions::TextureSource_ResizeToPowerOfTwo_WithDialog(const TArray<UTexture*> & InTextures)
 {
-	STextureActionDlg Dlg(InTextures,ETextureAction::ResizePow2);
+	UE::TextureAssetActionsInternal::STextureActionDlg Dlg(InTextures,UE::TextureAssetActionsInternal::ETextureAction::ResizePow2);
 	Dlg.ShowModal();
 }
 
diff --git a/Engine/Source/Editor/UnrealEd/Private/TextureAssetActionsBP.cpp b/Engine/Source/Editor/UnrealEd/Private/TextureAssetActionsBP.cpp
new file mode 100644
--- /dev/null
+++ b/Engine/Source/Editor/UnrealEd/Private/TextureAssetActionsBP.cpp
@@ -0,0 +1,62 @@
+// Copyright Epic Games, Inc. All Rights Reserved.
+
+#include "TextureAssetActionsBP.h"
+
+#include UE_INLINE_GENERATED_CPP_BY_NAME(TextureAssetActionsBP)
+
+
+
+
+UNREALED_API ETextureAssetActionsBlueprintReturn UTextureAssetActionsBlueprintFns::ResizeTextureSourceToPowerOfTwo(UTexture * Texture, int FilterThresholdValue)
+{
+	bool bExcludedBecauseAlreadySatisfied = false;
+
+	if ( ! UE::TextureAssetActionsInternal::GetTextureFilterEnabled(Texture,UE::TextureAssetActionsInternal::ETextureAction::ResizePow2,FilterThresholdValue,false,&bExcludedBecauseAlreadySatisfied) )
+	{
+		return bExcludedBecauseAlreadySatisfied ? ETextureAssetActionsBlueprintReturn::AlreadyDone : ETextureAssetActionsBlueprintReturn::Error;
+	}
+
+	bool bSuccess = UE::TextureAssetActionsInternal::ResizeTextureSourceToPowerOfTwo(Texture);
+	return bSuccess ? ETextureAssetActionsBlueprintReturn::Success : ETextureAssetActionsBlueprintReturn::Error;
+}
+
+UNREALED_API ETextureAssetActionsBlueprintReturn UTextureAssetActionsBlueprintFns::ConvertTo8bitTextureSource(UTexture * Texture,bool bNormalMapsKeep16bits )
+{
+	bool bExcludedBecauseAlreadySatisfied = false;
+	int FilterThresholdValue = 1;
+
+	if ( ! UE::TextureAssetActionsInternal::GetTextureFilterEnabled(Texture,UE::TextureAssetActionsInternal::ETextureAction::ConvertTo8bit,FilterThresholdValue,bNormalMapsKeep16bits,&bExcludedBecauseAlreadySatisfied) )
+	{
+		return bExcludedBecauseAlreadySatisfied ? ETextureAssetActionsBlueprintReturn::AlreadyDone : ETextureAssetActionsBlueprintReturn::Error;
+	}
+
+	bool bSuccess = UE::TextureAssetActionsInternal::ConvertTo8bitTextureSource(Texture,bNormalMapsKeep16bits);
+	return bSuccess ? ETextureAssetActionsBlueprintReturn::Success : ETextureAssetActionsBlueprintReturn::Error;
+}
+	
+UNREALED_API ETextureAssetActionsBlueprintReturn UTextureAssetActionsBlueprintFns::ResizeTextureSource(UTexture * Texture,int TargetSize )
+{
+	bool bExcludedBecauseAlreadySatisfied = false;
+	int FilterThresholdValue = TargetSize;
+
+	if ( ! UE::TextureAssetActionsInternal::GetTextureFilterEnabled(Texture,UE::TextureAssetActionsInternal::ETextureAction::Resize,FilterThresholdValue,false,&bExcludedBecauseAlreadySatisfied) )
+	{
+		return bExcludedBecauseAlreadySatisfied ? ETextureAssetActionsBlueprintReturn::AlreadyDone : ETextureAssetActionsBlueprintReturn::Error;
+	}
+
+	bool bSuccess = UE::TextureAssetActionsInternal::ResizeTextureSource(Texture,TargetSize);
+	return bSuccess ? ETextureAssetActionsBlueprintReturn::Success : ETextureAssetActionsBlueprintReturn::Error;
+}
+
+UNREALED_API ETextureAssetActionsBlueprintReturn UTextureAssetActionsBlueprintFns::CompressTextureSourceWithJPEG(UTexture * Texture, int FilterThresholdValue )
+{
+	bool bExcludedBecauseAlreadySatisfied = false;
+
+	if ( ! UE::TextureAssetActionsInternal::GetTextureFilterEnabled(Texture,UE::TextureAssetActionsInternal::ETextureAction::JPEG,FilterThresholdValue,false,&bExcludedBecauseAlreadySatisfied) )
+	{
+		return bExcludedBecauseAlreadySatisfied ? ETextureAssetActionsBlueprintReturn::AlreadyDone : ETextureAssetActionsBlueprintReturn::Error;
+	}
+
+	bool bSuccess = UE::TextureAssetActionsInternal::CompressTextureSourceWithJPEG(Texture);
+	return bSuccess ? ETextureAssetActionsBlueprintReturn::Success : ETextureAssetActionsBlueprintReturn::Error;
+}
\ No newline at end of file
diff --git a/Engine/Source/Editor/UnrealEd/Public/TextureAssetActionsBP.h b/Engine/Source/Editor/UnrealEd/Public/TextureAssetActionsBP.h
new file mode 100644
--- /dev/null
+++ b/Engine/Source/Editor/UnrealEd/Public/TextureAssetActionsBP.h
@@ -0,0 +1,143 @@
+// Copyright Epic Games, Inc. All Rights Reserved.
+#pragma once
+
+#include "TextureAssetActions.h"
+#include "UObject/Object.h"
+#include "TextureAssetActionsBP.generated.h"
+
+/***
+
+UTextureAssetActionsBlueprintFns provides BluePrint exposure of the TextureAssetActions functions.
+
+TextureAssetActions provide utilities to modify the UTexture Source data which is stored in the uasset after import.
+
+They do not modify the original source bitmap files on disk.  Those are not used by UE after import unless a reimport is done.
+
+Modifying the UTexture Source data is in contrast to the UTexture Properties, which affect how the runtime GPU is build.
+
+For example the UTexture property "MaxTextureSize" or "StretchToPowerOfTwo" may resize to limit GPU texture size, or stretch the GPU output to power of 2,
+but they do not modify the Source bitmap in the uasset.  To modify the source uasset, use TextureAssetActions::ResizeTextureSource and ResizeTextureSourceToPowerOfTwo.
+
+After these functions modify the Texture Source, the changed asset may need to be saved to retain those changes.
+
+***/
+
+UENUM(BlueprintType)
+enum class ETextureAssetActionsBlueprintReturn : uint8
+{
+	Error = 0,
+	Success = 1,
+	AlreadyDone = 2
+};
+
+UCLASS(MinimalAPI, BlueprintType)
+class UTextureAssetActionsBlueprintFns : public UObject
+{
+	GENERATED_BODY()
+
+public:
+
+	/**
+
+	Resizes the texture Source to power of two, if it wasn't already.
+
+	Only acts on textures larger than FilterThresholdValue.  Use FilterThresholdValue = 0 to disable filtering by size.
+
+	Returns Success if the action was done, returns AlreadyDone if the action did not need to be done because texture already met the goal.
+	Returns Error if the operation could not be applied because the texture is not a suitable type, or an unexpected error was encountered during processing.
+
+	**/
+	UFUNCTION(BlueprintCallable, Category = "TextureAssetActions")
+    static UNREALED_API ETextureAssetActionsBlueprintReturn ResizeTextureSourceToPowerOfTwo(UTexture * Texture, int FilterThresholdValue );
+	
+	/**
+
+	Resizes the texture Source to 8 bits per color channel, if it wasn't already.  HDR textures are converted to 16F (if they were 32F).
+
+	If bNormalMapsKeep16bits is true, then normal maps whose source had >= 16 bits per color channel retain 16 bits per channel (rather than 8).
+	
+	Returns Success if the action was done, returns AlreadyDone if the action did not need to be done because texture already met the goal.
+	Returns Error if the operation could not be applied because the texture is not a suitable type, or an unexpected error was encountered during processing.
+
+	**/
+	UFUNCTION(BlueprintCallable, Category = "TextureAssetActions")
+    static UNREALED_API ETextureAssetActionsBlueprintReturn ConvertTo8bitTextureSource(UTexture * Texture,bool bNormalMapsKeep16bits );
+	
+	/**
+
+	Resizes the texture Source to less or equal to ResizeToLessOrEqualToTargetSize
+
+	Texture LODBias is automatically adjusted to try to keep output size the same when possible.
+
+	Only acts on textures larger than ResizeToLessOrEqualToTargetSize.
+	
+	Returns Success if the action was done, returns AlreadyDone if the action did not need to be done because texture already met the goal.
+	Returns Error if the operation could not be applied because the texture is not a suitable type, or an unexpected error was encountered during processing.
+
+	**/
+	UFUNCTION(BlueprintCallable, Category = "TextureAssetActions")
+    static UNREALED_API ETextureAssetActionsBlueprintReturn ResizeTextureSource(UTexture * Texture,int ResizeToLessOrEqualToTargetSize );
+	
+	/**
+
+	Compresses the texture source with JPEG.  This makes the uasset smaller on disk, but lower quality and slower to access.
+
+	JPEG does not make the runtime game asset smaller on disk, only the source uasset.  JPEG can only be used on simple 2d textures.
+
+	Only acts on texture sources that are 8 bit.  You may wish to use ConvertTo8bitTextureSource before this, optionally.
+
+	Only acts on textures larger than FilterThresholdValue.  Use FilterThresholdValue = 0 to disable filtering by size.
+	
+	Returns Success if the action was done, returns AlreadyDone if the action did not need to be done because texture already met the goal.
+	Returns Error if the operation could not be applied because the texture is not a suitable type, or an unexpected error was encountered during processing.
+
+	**/
+	UFUNCTION(BlueprintCallable, Category = "TextureAssetActions")
+    static UNREALED_API ETextureAssetActionsBlueprintReturn CompressTextureSourceWithJPEG(UTexture * Texture, int FilterThresholdValue );
+};
+
+//================================================
+
+// NOT exported , internal use only, not public APIs :
+namespace UE
+{
+namespace TextureAssetActionsInternal
+{
+
+enum class ETextureAction
+{
+	Invalid = 0,
+	Resize,
+	ResizePow2,
+	ConvertTo8bit,
+	JPEG
+};
+
+// return true if Action should be applied to texture
+//	ThresholdValue is a 1d dimension, textures smaller than Threshold are excluded
+//	(set ThresholdValue == 0 for no size exclusion)
+bool GetTextureFilterEnabled( UTexture * Texture, ETextureAction Action, int ThresholdValue, bool bNormalMapsKeep16bits, 
+		bool * bOutExcludedBecauseAlreadySatisfied );
+
+// Actions return true if something was done, false if not
+//	can return false due to error or because no change was made
+
+// Actions are always supported on simple texture types (no Layers, no UDIM, 2d and Cube)
+//	more complex texture types are supported by some actions and not others YMMV
+
+// you MUST check GetTextureFilterEnabled before calling these :
+bool ResizeTextureSourceToPowerOfTwo(UTexture * Texture);
+bool ConvertTo8bitTextureSource(UTexture * Texture,bool bNormalMapsKeep16bits);
+
+// ResizeTextureSource resizes to <= TargetSize
+//	if size was already <= TargetSize , no change is made and we return false
+bool ResizeTextureSource(UTexture * Texture,int TargetSize);
+
+// Source must be 8 bit for conversion to JPEG
+//	you may optionally call ConvertTo8bitTextureSource before calling CompressTextureSourceWithJPEG
+bool CompressTextureSourceWithJPEG(UTexture * Texture);
+
+}
+};
+
+//================================================
