Added level viewport size accessor to UnrealEditorSubsystem. Also cleaned up some function docs.

#Synced-CL 49766534

diff --git a/Engine/Source/Editor/UnrealEd/Private/Subsystems/UnrealEditorSubsystem.cpp b/Engine/Source/Editor/UnrealEd/Private/Subsystems/UnrealEditorSubsystem.cpp
--- a/Engine/Source/Editor/UnrealEd/Private/Subsystems/UnrealEditorSubsystem.cpp
+++ b/Engine/Source/Editor/UnrealEd/Private/Subsystems/UnrealEditorSubsystem.cpp
@@ -114,6 +114,21 @@ bool UUnrealEditorSubsystem::WorldToScreen(const FVector& WorldPosition, FVector
 	return false;
 }
 
+UNREALED_API bool UUnrealEditorSubsystem::GetLevelViewportSize(FIntPoint& Size) const
+{
+	if (FLevelEditorViewportClient* LevelVC = GetViewportClient())
+	{
+		if (LevelVC->Viewport)
+		{
+			Size = LevelVC->Viewport->GetSizeXY();
+			return true;
+		}
+	}
+
+	Size = FIntPoint::ZeroValue;
+	return false;
+}
+
 FLevelEditorViewportClient* UUnrealEditorSubsystem::GetViewportClient() const
 {
 	for (FLevelEditorViewportClient* LevelVC : GEditor->GetLevelViewportClients())
diff --git a/Engine/Source/Editor/UnrealEd/Public/Subsystems/UnrealEditorSubsystem.h b/Engine/Source/Editor/UnrealEd/Public/Subsystems/UnrealEditorSubsystem.h
--- a/Engine/Source/Editor/UnrealEd/Public/Subsystems/UnrealEditorSubsystem.h
+++ b/Engine/Source/Editor/UnrealEd/Public/Subsystems/UnrealEditorSubsystem.h
@@ -53,9 +53,10 @@ public:
 
 	/**
 	 * Transforms the given 2D screen space coordinate into a 3D world-space point and direction using the primary level editor viewport.	 
-	 * @param ScreenPosition	2D screen space to deproject.
+	 * @param ScreenPosition	2D screen space to deproject. (in pixels)
 	 * @param WorldPosition		(out) Corresponding 3D position in world space.
 	 * @param WorldDirection	(out) World space direction vector away from the camera at the given 2d point.
+	 * @return					True if the position was able to be projected. False otherwise.
 	 */
 	UFUNCTION(BlueprintPure, Category = "Development|Editor")
 	UNREALED_API bool ScreenToWorld(const FVector2D& ScreenPosition, FVector& WorldPosition, FVector& WorldDirection) const;
@@ -63,11 +64,21 @@ public:
 	/**
 	 * Transforms the given 3D world-space point into a its 2D screen space coordinate using the primary level editor viewport.	 
 	 * @param WorldPosition		World position to project.
-	 * @param ScreenPosition	(out) Corresponding 2D position in screen space	 
+	 * @param ScreenPosition	(out) Corresponding 2D position in screen space (in pixels)	 
+	 * @return					True if the position was able to be projected. False otherwise.
 	 */
 	UFUNCTION(BlueprintPure, Category = "Development|Editor")	
 	UNREALED_API bool WorldToScreen(const FVector& WorldPosition, FVector2D& ScreenPosition) const;
 
+	/**
+	 * Returns the size of the primary level editor viewport.
+	 * @param Size	(out)	The dimensions of the viewport (in pixels).
+	 * @return				True if there is a primary level editor viewport.
+	 */
+	UFUNCTION(BlueprintPure, Category = "Development|Editor")
+	UNREALED_API bool GetLevelViewportSize(FIntPoint& Size) const;
+
+
 private:
 	FLevelEditorViewportClient* GetViewportClient() const;
 };
