TextureAssetActions : fix warnings from -Wextra-qualification

#Synced-CL 49826549

diff --git a/Engine/Source/Editor/UnrealEd/Private/TextureAssetActions.cpp b/Engine/Source/Editor/UnrealEd/Private/TextureAssetActions.cpp
--- a/Engine/Source/Editor/UnrealEd/Private/TextureAssetActions.cpp
+++ b/Engine/Source/Editor/UnrealEd/Private/TextureAssetActions.cpp
@@ -655,7 +655,7 @@ TSharedRef<SWidget> STextureAssetList::CreateAssetLine(int index, const UTexture
 	return Result;
 }
 
-bool UE::TextureAssetActionsInternal::ResizeTextureSourceToPowerOfTwo(UTexture * Texture)
+bool ResizeTextureSourceToPowerOfTwo(UTexture * Texture)
 {
 	if ( Texture->Source.AreAllBlocksPowerOfTwo() )
 	{
@@ -705,7 +705,7 @@ bool UE::TextureAssetActionsInternal::ResizeTextureSourceToPowerOfTwo(UTexture *
 
 // resizes to <= TargetSize
 //	if size was already <= TargetSize , no change is made and we return false
-bool UE::TextureAssetActionsInternal::ResizeTextureSource(UTexture * Texture,int TargetSize)
+bool ResizeTextureSource(UTexture * Texture,int TargetSize)
 {
 	// we do the resizing considering only mip/LOD/build settings for the running Editor platform (eg. Windows)
 	const ITargetPlatform* RunningPlatform = GetTargetPlatformManagerRef().GetRunningTargetPlatform();
@@ -892,7 +892,7 @@ static ETextureSourceFormat GetReducedTextureSourceFormat(const TextureCompressi
 	}
 }
 
-bool UE::TextureAssetActionsInternal::ConvertTo8bitTextureSource(UTexture * Texture,bool NormalMapsKeep16bits)
+bool ConvertTo8bitTextureSource(UTexture * Texture,bool NormalMapsKeep16bits)
 {
 	if ( Texture->Source.GetNumLayers() > 1 )
 	{
@@ -919,7 +919,7 @@ bool UE::TextureAssetActionsInternal::ConvertTo8bitTextureSource(UTexture * Text
 	return true;
 }
 
-bool UE::TextureAssetActionsInternal::CompressTextureSourceWithJPEG(UTexture * Texture)
+bool CompressTextureSourceWithJPEG(UTexture * Texture)
 {
 	// calls Pre/PostEditChange :
 	bool bDid = UE::TextureUtilitiesCommon::Experimental::CompressTextureSourceWithJPEG(Texture);
