UE: Bump cookincrementalversion due to shader change

diff --git a/Engine/Source/Editor/UnrealEd/Private/Cooker/CookIncrementalVersion.cpp b/Engine/Source/Editor/UnrealEd/Private/Cooker/CookIncrementalVersion.cpp
--- a/Engine/Source/Editor/UnrealEd/Private/Cooker/CookIncrementalVersion.cpp
+++ b/Engine/Source/Editor/UnrealEd/Private/Cooker/CookIncrementalVersion.cpp
@@ -8,6 +8,6 @@ namespace UE::Cook
 {
 
 // Change CookIncrementalVersion to a new guid when all packages in an incremental cook need to be invalidated
-FGuid CookIncrementalVersion(0xAEB44DBC, 0x78124EE2, 0xBEAD667A, 0x1D76D306);
+FGuid CookIncrementalVersion(0xAEB44DBC, 0x78124EE2, 0xBEAD667A, 0x1D71D306);
 
 }
