TextureAssetActions : fix warnings from -Wextra-qualification part 2

#Synced-CL 49826549

diff --git a/Engine/Source/Editor/UnrealEd/Private/TextureAssetActions.cpp b/Engine/Source/Editor/UnrealEd/Private/TextureAssetActions.cpp
--- a/Engine/Source/Editor/UnrealEd/Private/TextureAssetActions.cpp
+++ b/Engine/Source/Editor/UnrealEd/Private/TextureAssetActions.cpp
@@ -1125,7 +1125,7 @@ static EAssetActionStatus GetTextureFilterStatus(UTexture * Texture, ETextureAct
 	return Status;
 }
 
-bool UE::TextureAssetActionsInternal::GetTextureFilterEnabled( UTexture * Texture, ETextureAction Action, int ThresholdValue,
+bool GetTextureFilterEnabled( UTexture * Texture, ETextureAction Action, int ThresholdValue,
 		bool bNormalMapsKeep16bits , bool * bOutExcludedBecauseAlreadySatisfied)
 {
 	FString EnabledActionString;
