IncrementalCook: Improve feedback about why a class was found to be non-incrementally-skippable when we mark a package non-incrementally-skippable because of one of its classes.
#rnx

#Synced-CL 49850719

diff --git a/Engine/Source/Editor/UnrealEd/Private/TargetDomain/TargetDomainUtils.cpp b/Engine/Source/Editor/UnrealEd/Private/TargetDomain/TargetDomainUtils.cpp
--- a/Engine/Source/Editor/UnrealEd/Private/TargetDomain/TargetDomainUtils.cpp
+++ b/Engine/Source/Editor/UnrealEd/Private/TargetDomain/TargetDomainUtils.cpp
@@ -79,11 +79,6 @@ bool IsIncrementalCookEnabled(FName PackageName, bool bAllowAllClasses)
 
 	if (!bAllowAllClasses)
 	{
-		auto LogInvalidDueTo = [](FName PackageName, FName ClassPath)
-			{
-				UE_LOG(LogEditorDomain, Verbose, TEXT("NonIterative Package %s due to %s"), *PackageName.ToString(), *ClassPath.ToString());
-			};
-
 		UE::EditorDomain::FClassDigestMap& ClassDigests = UE::EditorDomain::GetClassDigests();
 		FReadScopeLock ClassDigestsScopeLock(ClassDigests.Lock);
 		for (FName ClassName : *ImportedClasses)
@@ -107,24 +102,35 @@ bool IsIncrementalCookEnabled(FName PackageName, bool bAllowAllClasses)
 			{
 				// !ExistingData -> !allowed, because caller has already called CalculatePackageDigest, so all
 				// existing classes in the package have been added to ClassDigests.
-				LogInvalidDueTo(PackageName, ClassName);
+				UE_LOG(LogEditorDomain, Verbose, TEXT("NonIterative Package %s due to missing class %s"),
+					*PackageName.ToString(), *ClassPath.ToString());
 				return false;
 			}
-			if (!ExistingData->bNative)
+
+			UE::EditorDomain::FClassDigestData* DataToRead = ExistingData;
+			FTopLevelAssetPath* ClassToRead = &ClassPath;
+			if (!DataToRead->bNative)
 			{
 				// TODO: We need to add a way to mark non-native classes (there can be many of them) as allowed or denied.
 				// Currently we are allowing them all, so long as their closest native is allowed. But this is not completely
 				// safe to do, because non-native classes can add constructionevents that e.g. use the Random function.
-				ExistingData = ClassDigests.Map.Find(ExistingData->ClosestNative);
-				if (!ExistingData)
+				ClassToRead = &DataToRead->ClosestNative;
+				DataToRead = ClassDigests.Map.Find(DataToRead->ClosestNative);
+				if (!DataToRead)
 				{
-					LogInvalidDueTo(PackageName, ClassName);
+					UE_LOG(LogEditorDomain, Verbose,
+						TEXT("NonIterative Package %s due to missing native super class %s of non-native class %s"),
+						*PackageName.ToString(), *ExistingData->ClosestNative.ToString(), *ClassPath.ToString());
 					return false;
 				}
 			}
-			if (!ExistingData->bTargetIterativeEnabled)
+			if (!DataToRead->bTargetIterativeEnabled)
 			{
-				LogInvalidDueTo(PackageName, ClassName);
+				UE_LOG(LogEditorDomain, Verbose,
+					TEXT("NonIterative Package %s due to bTargetIterativeEnabled=false on class %s%s"),
+					*PackageName.ToString(), *ClassToRead->ToString(),
+					(ExistingData == DataToRead ? TEXT("")
+						: *FString::Printf(TEXT(", the closest native superclass of non-native class %s"), *ClassPath.ToString())));
 				return false;
 			}
 		}
