Warn instead of crash when failing to transfer pins while removing a BP interface. The repro is unclear but this can happen if a function node is out of date.

#Synced-CL 49784053

diff --git a/Engine/Source/Editor/UnrealEd/Private/Kismet2/BlueprintEditorUtils.cpp b/Engine/Source/Editor/UnrealEd/Private/Kismet2/BlueprintEditorUtils.cpp
--- a/Engine/Source/Editor/UnrealEd/Private/Kismet2/BlueprintEditorUtils.cpp
+++ b/Engine/Source/Editor/UnrealEd/Private/Kismet2/BlueprintEditorUtils.cpp
@@ -6735,12 +6735,21 @@ void FBlueprintEditorUtils::RemoveInterface(UBlueprint* Blueprint, FTopLevelAsse
 						const UFunction* PreviousSignatureFunction = EventNode->FindEventSignatureFunction();
 						check(PreviousSignatureFunction);
 						UK2Node_CustomEvent* NewEvent = UK2Node_CustomEvent::CreateFromFunction(PreviousNodePos, EventNode->GetGraph(), PreviousNodeName, PreviousSignatureFunction, false);
+						const UEdGraphSchema_K2* Schema = GetDefault<UEdGraphSchema_K2>();
+
 						// Move the pin links from the old pin to the new pin to preserve connections
 						for (UEdGraphPin* CurrentPin : EventNode->Pins)
 						{
-							UEdGraphPin* TargetPin = NewEvent->FindPinChecked(CurrentPin->PinName);
-							const UEdGraphSchema_K2* Schema = GetDefault<UEdGraphSchema_K2>();
-							Schema->MovePinLinks(*CurrentPin, *TargetPin);
+							if (UEdGraphPin* TargetPin = NewEvent->FindPin(CurrentPin->PinName))
+							{
+								Schema->MovePinLinks(*CurrentPin, *TargetPin);
+							}
+							else
+							{
+								UE_LOG(LogBlueprint, Warning,
+									TEXT("Failed to transfer pin %s on function %s while removing interface from blueprint %s"),
+									*CurrentPin->PinName.ToString(), *PreviousSignatureFunction->GetPathName(), *Blueprint->GetPathName());
+							}
 						}
 					}
 
