Profile renderers implement required top class functions.

#tests PC editor dx12
#rnx

#Synced-CL 49883401

diff --git a/Engine/Source/Editor/UnrealEd/Classes/ThumbnailRendering/SpecularProfileRenderer.h b/Engine/Source/Editor/UnrealEd/Classes/ThumbnailRendering/SpecularProfileRenderer.h
--- a/Engine/Source/Editor/UnrealEd/Classes/ThumbnailRendering/SpecularProfileRenderer.h
+++ b/Engine/Source/Editor/UnrealEd/Classes/ThumbnailRendering/SpecularProfileRenderer.h
@@ -23,6 +23,8 @@ class USpecularProfileRenderer : public UTextureThumbnailRenderer
 	// Begin UThumbnailRenderer Object
 	virtual void GetThumbnailSize(UObject* Object, float Zoom, uint32& OutWidth, uint32& OutHeight) const override;
 	virtual void Draw(UObject* Object, int32 X, int32 Y, uint32 Width, uint32 Height, FRenderTarget*, FCanvas* Canvas, bool bAdditionalViewFamily) override;
+	virtual bool CanVisualizeAsset(UObject* Object) override;
+	virtual EThumbnailRenderFrequency GetThumbnailRenderFrequency(UObject* Object) const override;
 	// End UThumbnailRenderer Object
 };
 
diff --git a/Engine/Source/Editor/UnrealEd/Classes/ThumbnailRendering/SubsurfaceProfileRenderer.h b/Engine/Source/Editor/UnrealEd/Classes/ThumbnailRendering/SubsurfaceProfileRenderer.h
--- a/Engine/Source/Editor/UnrealEd/Classes/ThumbnailRendering/SubsurfaceProfileRenderer.h
+++ b/Engine/Source/Editor/UnrealEd/Classes/ThumbnailRendering/SubsurfaceProfileRenderer.h
@@ -24,6 +24,8 @@ class USubsurfaceProfileRenderer : public UTextureThumbnailRenderer
 	// Begin UThumbnailRenderer Object
 	virtual void GetThumbnailSize(UObject* Object, float Zoom, uint32& OutWidth, uint32& OutHeight) const override;
 	virtual void Draw(UObject* Object, int32 X, int32 Y, uint32 Width, uint32 Height, FRenderTarget*, FCanvas* Canvas, bool bAdditionalViewFamily) override;
+	virtual bool CanVisualizeAsset(UObject* Object) override;
+	virtual EThumbnailRenderFrequency GetThumbnailRenderFrequency(UObject* Object) const override;
 	// End UThumbnailRenderer Object
 };
 
diff --git a/Engine/Source/Editor/UnrealEd/Classes/ThumbnailRendering/ToonProfileRenderer.h b/Engine/Source/Editor/UnrealEd/Classes/ThumbnailRendering/ToonProfileRenderer.h
--- a/Engine/Source/Editor/UnrealEd/Classes/ThumbnailRendering/ToonProfileRenderer.h
+++ b/Engine/Source/Editor/UnrealEd/Classes/ThumbnailRendering/ToonProfileRenderer.h
@@ -23,6 +23,8 @@ class UToonProfileRenderer : public UTextureThumbnailRenderer
 	// Begin UThumbnailRenderer Object
 	virtual void GetThumbnailSize(UObject* Object, float Zoom, uint32& OutWidth, uint32& OutHeight) const override;
 	virtual void Draw(UObject* Object, int32 X, int32 Y, uint32 Width, uint32 Height, FRenderTarget*, FCanvas* Canvas, bool bAdditionalViewFamily) override;
+	virtual bool CanVisualizeAsset(UObject* Object) override;
+	virtual EThumbnailRenderFrequency GetThumbnailRenderFrequency(UObject* Object) const override;
 	// End UThumbnailRenderer Object
 };
 
diff --git a/Engine/Source/Editor/UnrealEd/Private/ThumbnailRendering/SpecularProfileRenderer.cpp b/Engine/Source/Editor/UnrealEd/Private/ThumbnailRendering/SpecularProfileRenderer.cpp
--- a/Engine/Source/Editor/UnrealEd/Private/ThumbnailRendering/SpecularProfileRenderer.cpp
+++ b/Engine/Source/Editor/UnrealEd/Private/ThumbnailRendering/SpecularProfileRenderer.cpp
@@ -37,3 +37,13 @@ void USpecularProfileRenderer::Draw(UObject* Object, int32 X, int32 Y, uint32 Wi
 		TextItem.Draw(Canvas);
 	}
 }
+
+bool USpecularProfileRenderer::CanVisualizeAsset(UObject* Object)
+{
+	return true;
+}
+
+EThumbnailRenderFrequency USpecularProfileRenderer::GetThumbnailRenderFrequency(UObject* Object) const
+{
+	return EThumbnailRenderFrequency::OnAssetSave;
+}
diff --git a/Engine/Source/Editor/UnrealEd/Private/ThumbnailRendering/SubsurfaceProfileRenderer.cpp b/Engine/Source/Editor/UnrealEd/Private/ThumbnailRendering/SubsurfaceProfileRenderer.cpp
--- a/Engine/Source/Editor/UnrealEd/Private/ThumbnailRendering/SubsurfaceProfileRenderer.cpp
+++ b/Engine/Source/Editor/UnrealEd/Private/ThumbnailRendering/SubsurfaceProfileRenderer.cpp
@@ -39,3 +39,13 @@ void USubsurfaceProfileRenderer::Draw(UObject* Object, int32 X, int32 Y, uint32
 		TextItem.Draw(Canvas);
 	}
 }
+
+bool USubsurfaceProfileRenderer::CanVisualizeAsset(UObject* Object)
+{
+	return true;
+}
+
+EThumbnailRenderFrequency USubsurfaceProfileRenderer::GetThumbnailRenderFrequency(UObject* Object) const
+{
+	return EThumbnailRenderFrequency::OnAssetSave;
+}
diff --git a/Engine/Source/Editor/UnrealEd/Private/ThumbnailRendering/ToonProfileRenderer.cpp b/Engine/Source/Editor/UnrealEd/Private/ThumbnailRendering/ToonProfileRenderer.cpp
--- a/Engine/Source/Editor/UnrealEd/Private/ThumbnailRendering/ToonProfileRenderer.cpp
+++ b/Engine/Source/Editor/UnrealEd/Private/ThumbnailRendering/ToonProfileRenderer.cpp
@@ -41,3 +41,13 @@ void UToonProfileRenderer::Draw(UObject* Object, int32 X, int32 Y, uint32 Width,
 		TextItem.Draw(Canvas);
 	}
 }
+
+bool UToonProfileRenderer::CanVisualizeAsset(UObject* Object)
+{
+	return true;
+}
+
+EThumbnailRenderFrequency UToonProfileRenderer::GetThumbnailRenderFrequency(UObject* Object) const
+{
+	return EThumbnailRenderFrequency::OnAssetSave;
+}
