Substrate - Toon - added specular curve preview on profile asset.

#tests PC editor
#rnx

#Synced-CL 49883401

diff --git a/Engine/Source/Editor/UnrealEd/Private/ThumbnailRendering/ToonProfileRenderer.cpp b/Engine/Source/Editor/UnrealEd/Private/ThumbnailRendering/ToonProfileRenderer.cpp
--- a/Engine/Source/Editor/UnrealEd/Private/ThumbnailRendering/ToonProfileRenderer.cpp
+++ b/Engine/Source/Editor/UnrealEd/Private/ThumbnailRendering/ToonProfileRenderer.cpp
@@ -32,7 +32,11 @@ void UToonProfileRenderer::Draw(UObject* Object, int32 X, int32 Y, uint32 Width,
 		{
 			Col = LocalToonProfile->Settings.DiffuseRamp.GetLinearColorValue(Time);
 			Col.A = 1;
-			Canvas->DrawTile(float(Width) * Time, 0, float(Width) * (Time + TimeStep), Height, 0, 0, 1, 1, Col);
+			Canvas->DrawTile(float(Width) * Time, 0, float(Width) * (Time + TimeStep), Height/2, 0, 0, 1, 1, Col);
+
+			Col = LocalToonProfile->Settings.SpecularRamp.GetLinearColorValue(Time);
+			Col.A = 1;
+			Canvas->DrawTile(float(Width) * Time, Height / 2, float(Width) * (Time + TimeStep), Height, 0, 0, 1, 1, Col);
 		}
 
 		FCanvasTextItem TextItem(FVector2D(5.0f, 5.0f), FText::FromString(TEXT("TOON")), GEngine->GetLargeFont(), FLinearColor::White);
