Prevent Blueprint actor from having a valid ExternalDataLayerAsset (could happen when using 'Apply instance changes to blueprint')

#Synced-CL 49829988

#Synced-CL 49829988

diff --git a/Engine/Source/Editor/UnrealEd/Private/Kismet2/Kismet2.cpp b/Engine/Source/Editor/UnrealEd/Private/Kismet2/Kismet2.cpp
--- a/Engine/Source/Editor/UnrealEd/Private/Kismet2/Kismet2.cpp
+++ b/Engine/Source/Editor/UnrealEd/Private/Kismet2/Kismet2.cpp
@@ -98,6 +98,12 @@ namespace UE::Private
 		TEXT("BP.bLogExplicitCompilationRequests"), bLogExplicitCompilationRequests,
 		TEXT("If true, CompileBlueprint will emit a log indicating the class getting recompiled."),
 		ECVF_Default);
+
+	static bool ExternalDataLayerAssetCopyPropertyFilter(FProperty& Property, UObject& Object)
+	{
+		static FName ExternalDataLayerAssetPropertyName(AActor::GetExternalDataLayerAssetPropertyName());
+		return !Object.IsA<AActor>() || (Property.GetFName() != ExternalDataLayerAssetPropertyName);
+	};
 }
 
 //////////////////////////////////////////////////////////////////////////
@@ -1544,7 +1550,8 @@ UBlueprint* FKismetEditorUtilities::CreateBlueprintFromActor(const FName Bluepri
 			if (NewBlueprint->GeneratedClass != nullptr)
 			{
 				AActor* CDO = CastChecked<AActor>(NewBlueprint->GeneratedClass->GetDefaultObject());
-				const EditorUtilities::ECopyOptions::Type CopyOptions = (EditorUtilities::ECopyOptions::Type)(EditorUtilities::ECopyOptions::OnlyCopyEditOrInterpProperties | EditorUtilities::ECopyOptions::PropagateChangesToArchetypeInstances | EditorUtilities::ECopyOptions::SkipInstanceOnlyProperties);
+				const EditorUtilities::ECopyOptions::Type CopyOptionFlags = (EditorUtilities::ECopyOptions::Type)(EditorUtilities::ECopyOptions::OnlyCopyEditOrInterpProperties | EditorUtilities::ECopyOptions::PropagateChangesToArchetypeInstances | EditorUtilities::ECopyOptions::SkipInstanceOnlyProperties);
+				const EditorUtilities::FCopyOptions CopyOptions(CopyOptionFlags, UE::Private::ExternalDataLayerAssetCopyPropertyFilter); // Don't allow to copy ExternalDataLayerAsset property to NewBlueprint
 				EditorUtilities::CopyActorProperties(Actor, CDO, CopyOptions);
 
 				if (USceneComponent* DstSceneRoot = CDO->GetRootComponent())
@@ -2200,7 +2207,8 @@ int32 FKismetEditorUtilities::ApplyInstanceChangesToBlueprint(AActor* Actor)
 				AActor* BlueprintCDO = Actor->GetClass()->GetDefaultObject<AActor>();
 				if (BlueprintCDO != NULL)
 				{
-					const EditorUtilities::ECopyOptions::Type CopyOptions = (EditorUtilities::ECopyOptions::Type)(EditorUtilities::ECopyOptions::OnlyCopyEditOrInterpProperties | EditorUtilities::ECopyOptions::PropagateChangesToArchetypeInstances | EditorUtilities::ECopyOptions::SkipInstanceOnlyProperties);
+					const EditorUtilities::ECopyOptions::Type CopyOptionFlags = (EditorUtilities::ECopyOptions::Type)(EditorUtilities::ECopyOptions::OnlyCopyEditOrInterpProperties | EditorUtilities::ECopyOptions::PropagateChangesToArchetypeInstances | EditorUtilities::ECopyOptions::SkipInstanceOnlyProperties);
+					const EditorUtilities::FCopyOptions CopyOptions(CopyOptionFlags, UE::Private::ExternalDataLayerAssetCopyPropertyFilter); // Don't allow to copy ExternalDataLayerAsset property to BlueprintCDO
 					NumChangedProperties = EditorUtilities::CopyActorProperties(Actor, BlueprintCDO, CopyOptions);
 					if (Actor->GetInstanceComponents().Num() > 0)
 					{
diff --git a/Engine/Source/Editor/UnrealEd/Public/Editor.h b/Engine/Source/Editor/UnrealEd/Public/Editor.h
--- a/Engine/Source/Editor/UnrealEd/Public/Editor.h
+++ b/Engine/Source/Editor/UnrealEd/Public/Editor.h
@@ -833,7 +833,10 @@ namespace EditorUtilities
 	struct FCopyOptions
 	{
 		/** Implicit construction for an options enumeration */
-		FCopyOptions(const ECopyOptions::Type InFlags) : Flags(InFlags) {}
+		FCopyOptions(const ECopyOptions::Type InFlags, TFunction<bool(FProperty&, UObject&)> InPropertyFilter = TFunction<bool(FProperty&, UObject&)>()) 
+			: Flags(InFlags)
+			, PropertyFilter(InPropertyFilter)
+		{}
 
 		/** Check whether we can copy the specified property */
 		bool CanCopyProperty(FProperty& Property, UObject& Object) const
diff --git a/Engine/Source/Runtime/Engine/Private/ActorEditor.cpp b/Engine/Source/Runtime/Engine/Private/ActorEditor.cpp
--- a/Engine/Source/Runtime/Engine/Private/ActorEditor.cpp
+++ b/Engine/Source/Runtime/Engine/Private/ActorEditor.cpp
@@ -1844,6 +1844,8 @@ void AActor::FixupDataLayers(bool bRevertChangesOnLockedDataLayer /*= false*/)
 
 	if (IsTemplate())
 	{
+		// Fixup old data: Template actors should never have an ExternalDataLayerAsset
+		ExternalDataLayerAsset = nullptr;
 		return;
 	}
 
