[UnrealEd]

- Ensure the SlateViewport is valid before we attempt to Pin the object. If not valid, early out when toggling between PIE and SIE
- Moves OnPreSwitchBeginPIEAndSIE broadcast to happen only when we will toggle. Previously we would broadcast even in cases when switching is not possible.

#rnx
#Synced-CL 49918891

#Synced-CL 49918891

diff --git a/Engine/Source/Editor/UnrealEd/Private/PlayLevel.cpp b/Engine/Source/Editor/UnrealEd/Private/PlayLevel.cpp
--- a/Engine/Source/Editor/UnrealEd/Private/PlayLevel.cpp
+++ b/Engine/Source/Editor/UnrealEd/Private/PlayLevel.cpp
@@ -2034,8 +2034,6 @@ void UEditorEngine::ToggleBetweenPIEandSIE( bool bNewSession )
 {
 	bIsToggleBetweenPIEandSIEQueued = false;
 
-	FEditorDelegates::OnPreSwitchBeginPIEAndSIE.Broadcast(bIsSimulatingInEditor);
-
 	// The first PIE world context is the one that can toggle between PIE and SIE
 	// Network PIE/SIE toggling is not really meant to be supported.
 	FSlatePlayInEditorInfo * SlateInfoPtr = nullptr;
@@ -2048,11 +2046,15 @@ void UEditorEngine::ToggleBetweenPIEandSIE( bool bNewSession )
 		}
 	}
 
-	if (!SlateInfoPtr)
+	// This is only supported inside SLevelEditor viewports currently
+	if (!SlateInfoPtr || !SlateInfoPtr->DestinationSlateViewport.IsValid())
 	{
+		UE_LOG(LogPlayLevel, Warning, TEXT("Failed to find valid viewport when switching between PIE and SIE."));
 		return;
 	}
 
+	FEditorDelegates::OnPreSwitchBeginPIEAndSIE.Broadcast(bIsSimulatingInEditor);
+
 	if( FEngineAnalytics::IsAvailable() && !bNewSession )
 	{
 		FString ToggleType = bIsSimulatingInEditor ? TEXT("SIEtoPIE") : TEXT("PIEtoSIE");
@@ -2060,94 +2062,89 @@ void UEditorEngine::ToggleBetweenPIEandSIE( bool bNewSession )
 		FEngineAnalytics::GetProvider().RecordEvent(TEXT("Editor.Usage.PIE"), TEXT("ToggleBetweenPIEandSIE"), ToggleType );
 	}
 
-	FSlatePlayInEditorInfo & SlatePlayInEditorSession = *SlateInfoPtr;
-
-	// This is only supported inside SLevelEditor viewports currently
+	FSlatePlayInEditorInfo& SlatePlayInEditorSession = *SlateInfoPtr;
 	TSharedPtr<IAssetViewport> LevelViewport = SlatePlayInEditorSession.DestinationSlateViewport.Pin();
-	if( ensure(LevelViewport.IsValid()) )
+	FEditorViewportClient& EditorViewportClient = LevelViewport->GetAssetViewportClient();
+
+	// Toggle to pie if currently simulating
+	if( bIsSimulatingInEditor )
 	{
-		FEditorViewportClient& EditorViewportClient = LevelViewport->GetAssetViewportClient();
+		// The undo system may have a reference to a SIE object that is about to be destroyed, so clear the transactions
+		ResetTransaction( NSLOCTEXT("UnrealEd", "ToggleBetweenPIEandSIE", "Toggle Between PIE and SIE") );
 
-		// Toggle to pie if currently simulating
-		if( bIsSimulatingInEditor )
-		{
-			// The undo system may have a reference to a SIE object that is about to be destroyed, so clear the transactions
-			ResetTransaction( NSLOCTEXT("UnrealEd", "ToggleBetweenPIEandSIE", "Toggle Between PIE and SIE") );
+		// The Game's viewport needs to know about the change away from simluate before the PC is (potentially) created
+		GameViewport->GetGameViewport()->SetPlayInEditorIsSimulate(false);
 
-			// The Game's viewport needs to know about the change away from simluate before the PC is (potentially) created
-			GameViewport->GetGameViewport()->SetPlayInEditorIsSimulate(false);
+		// The editor viewport client wont be visible so temporarily disable it being realtime
+		const bool bShouldBeRealtime = false;
+		const FText SystemDisplayName = LOCTEXT("RealtimeOverrideMessage_PIE", "Play in Editor");
+		// Remove any previous override since we already applied a override when entering PIE
+		EditorViewportClient.RemoveRealtimeOverride(SystemDisplayName);
+		EditorViewportClient.AddRealtimeOverride(bShouldBeRealtime, SystemDisplayName);
 
-			// The editor viewport client wont be visible so temporarily disable it being realtime
-			const bool bShouldBeRealtime = false;
-			const FText SystemDisplayName = LOCTEXT("RealtimeOverrideMessage_PIE", "Play in Editor");
-			// Remove any previous override since we already applied a override when entering PIE
-			EditorViewportClient.RemoveRealtimeOverride(SystemDisplayName);
-			EditorViewportClient.AddRealtimeOverride(bShouldBeRealtime, SystemDisplayName);
+		if (!SlatePlayInEditorSession.EditorPlayer.IsValid())
+		{
+			OnSwitchWorldsForPIE(true);
 
-			if (!SlatePlayInEditorSession.EditorPlayer.IsValid())
+			UWorld* World = GameViewport->GetWorld();
+			AGameModeBase* AuthGameMode = World->GetAuthGameMode();
+			if (AuthGameMode && GameViewport->GetGameInstance())	// If there is no GameMode, we are probably the client and cannot RestartPlayer.
 			{
-				OnSwitchWorldsForPIE(true);
-
-				UWorld* World = GameViewport->GetWorld();
-				AGameModeBase* AuthGameMode = World->GetAuthGameMode();
-				if (AuthGameMode && GameViewport->GetGameInstance())	// If there is no GameMode, we are probably the client and cannot RestartPlayer.
-				{
-					AuthGameMode->SpawnPlayerFromSimulate(EditorViewportClient.GetViewLocation(), EditorViewportClient.GetViewRotation());
-				}
-
-				OnSwitchWorldsForPIE(false);
+				AuthGameMode->SpawnPlayerFromSimulate(EditorViewportClient.GetViewLocation(), EditorViewportClient.GetViewRotation());
 			}
 
-			// A game viewport already exists, tell the level viewport its in to swap to it
-			LevelViewport->SwapViewportsForPlayInEditor();
+			OnSwitchWorldsForPIE(false);
+		}
+
+		// A game viewport already exists, tell the level viewport its in to swap to it
+		LevelViewport->SwapViewportsForPlayInEditor();
 
-			// No longer simulating
-			GameViewport->SetIsSimulateInEditorViewport(false);
-			EditorViewportClient.SetIsSimulateInEditorViewport(false);
+		// No longer simulating
+		GameViewport->SetIsSimulateInEditorViewport(false);
+		EditorViewportClient.SetIsSimulateInEditorViewport(false);
 
-			FEditorModeRegistry::Get().UnregisterMode(FBuiltinEditorModes::EM_Physics);
+		FEditorModeRegistry::Get().UnregisterMode(FBuiltinEditorModes::EM_Physics);
 			
-			bIsSimulatingInEditor = false;
-		}
-		else
-		{
-			// Swap to simulate from PIE
-			LevelViewport->SwapViewportsForSimulateInEditor();
+		bIsSimulatingInEditor = false;
+	}
+	else
+	{
+		// Swap to simulate from PIE
+		LevelViewport->SwapViewportsForSimulateInEditor();
 	
-			GameViewport->SetIsSimulateInEditorViewport(true);
-			GameViewport->GetGameViewport()->SetPlayInEditorIsSimulate(true);
-			EditorViewportClient.SetIsSimulateInEditorViewport(true);
+		GameViewport->SetIsSimulateInEditorViewport(true);
+		GameViewport->GetGameViewport()->SetPlayInEditorIsSimulate(true);
+		EditorViewportClient.SetIsSimulateInEditorViewport(true);
 
 		
-			TSharedRef<FPhysicsManipulationEdModeFactory> Factory = MakeShareable(new FPhysicsManipulationEdModeFactory);
-			FEditorModeRegistry::Get().RegisterMode(FBuiltinEditorModes::EM_Physics, Factory);
+		TSharedRef<FPhysicsManipulationEdModeFactory> Factory = MakeShareable(new FPhysicsManipulationEdModeFactory);
+		FEditorModeRegistry::Get().RegisterMode(FBuiltinEditorModes::EM_Physics, Factory);
 			
-			bIsSimulatingInEditor = true;
+		bIsSimulatingInEditor = true;
 
-			// Make sure the viewport is in real-time mode
-			const bool bShouldBeRealtime = true;
-			const FText SystemDisplayName = LOCTEXT("RealtimeOverrideMessage_PIE", "Play in Editor");
-			// Remove any previous override since we already applied a override when entering PIE
-			EditorViewportClient.RemoveRealtimeOverride(SystemDisplayName);
-			EditorViewportClient.AddRealtimeOverride(bShouldBeRealtime, SystemDisplayName);
+		// Make sure the viewport is in real-time mode
+		const bool bShouldBeRealtime = true;
+		const FText SystemDisplayName = LOCTEXT("RealtimeOverrideMessage_PIE", "Play in Editor");
+		// Remove any previous override since we already applied a override when entering PIE
+		EditorViewportClient.RemoveRealtimeOverride(SystemDisplayName);
+		EditorViewportClient.AddRealtimeOverride(bShouldBeRealtime, SystemDisplayName);
 
-			// The Simulate window should show stats
-			EditorViewportClient.SetShowStats( true );
+		// The Simulate window should show stats
+		EditorViewportClient.SetShowStats( true );
 
-			if ( SlatePlayInEditorSession.EditorPlayer.IsValid() && SlatePlayInEditorSession.EditorPlayer.Get()->PlayerController )
-			{
-				// Move the editor camera to where the player was.  
-				FVector ViewLocation;
-				FRotator ViewRotation;
-				SlatePlayInEditorSession.EditorPlayer.Get()->PlayerController->GetPlayerViewPoint( ViewLocation, ViewRotation );
-				ViewRotation.Roll = 0.0; // Clear roll since editor camera controls has no way to adjust roll
-				EditorViewportClient.SetViewLocation( ViewLocation );
+		if ( SlatePlayInEditorSession.EditorPlayer.IsValid() && SlatePlayInEditorSession.EditorPlayer.Get()->PlayerController )
+		{
+			// Move the editor camera to where the player was.  
+			FVector ViewLocation;
+			FRotator ViewRotation;
+			SlatePlayInEditorSession.EditorPlayer.Get()->PlayerController->GetPlayerViewPoint( ViewLocation, ViewRotation );
+			ViewRotation.Roll = 0.0; // Clear roll since editor camera controls has no way to adjust roll
+			EditorViewportClient.SetViewLocation( ViewLocation );
 
-				if( EditorViewportClient.IsPerspective() )
-				{
-					// Rotation only matters for perspective viewports not orthographic
-					EditorViewportClient.SetViewRotation( ViewRotation );
-				}
+			if( EditorViewportClient.IsPerspective() )
+			{
+				// Rotation only matters for perspective viewports not orthographic
+				EditorViewportClient.SetViewRotation( ViewRotation );
 			}
 		}
 	}
