Fixed texture thumbnails being rendered smaller than intented in the content browser.
- Texture thumbnail renderers can override the default thumbnail size via GetThumbnailSize but this should only be used when we are rendering back to a FObjectThumbnail and not to a render target, as otherwise we need to use the size of the render target.

This fix restores the behavior pre cl 47908789.

#Synced-CL 49730605

diff --git a/Engine/Source/Editor/UnrealEd/Private/ObjectTools.cpp b/Engine/Source/Editor/UnrealEd/Private/ObjectTools.cpp
--- a/Engine/Source/Editor/UnrealEd/Private/ObjectTools.cpp
+++ b/Engine/Source/Editor/UnrealEd/Private/ObjectTools.cpp
@@ -5642,6 +5642,8 @@ namespace ThumbnailTools
 		// Renderer must be initialized before generating thumbnails
 		check( GIsRHIInitialized );
 
+		const bool bIsRenderingToSpecificRenderTarget = (InOutTextureRenderTargetResource != nullptr);
+
 		// Grab the actual render target resource from the texture.  Note that we're absolutely NOT ALLOWED to
 		// dereference this pointer.  We're just passing it along to other functions that will use it on the render
 		// thread.  The only thing we're allowed to do is check to see if it's NULL or not.
@@ -5711,6 +5713,9 @@ namespace ThumbnailTools
 			uint32 DrawWidth = InOutImageWidth;
 			uint32 DrawHeight = InOutImageHeight;
 
+			// If we're not rendering to a specific render target, then allow the thumbnail renderer to override the thumbnail size.
+			// This happens when we render the thumbnail that will be saved in the asset package.
+			if (!bIsRenderingToSpecificRenderTarget)
 			{
 				// Find how big the thumbnail WANTS to be
 				uint32 DesiredWidth = 0;
