Mute editor when out of focus and "Allow Background Audio" is unchecked

#tests Verified that an ambient loop placed in the level muted correctly in the background
#rnx

#Synced-CL 49867512

diff --git a/Engine/Source/Editor/UnrealEd/Private/EditorEngine.cpp b/Engine/Source/Editor/UnrealEd/Private/EditorEngine.cpp
--- a/Engine/Source/Editor/UnrealEd/Private/EditorEngine.cpp
+++ b/Engine/Source/Editor/UnrealEd/Private/EditorEngine.cpp
@@ -1906,7 +1906,7 @@ void UEditorEngine::Tick( float DeltaSeconds, bool bIdleMode )
 		}
 	}
 
-	// Find out if the editor has focus. Audio should only play if the editor has focus.
+	// Find out if the editor has focus. Audio should only play if the editor has focus unless "Allow background audio" is set.
 	const bool bHasFocus = FApp::HasFocus();
 
 	if (bHasFocus || GetDefault<ULevelEditorMiscSettings>()->bAllowBackgroundAudio)
@@ -1922,6 +1922,10 @@ void UEditorEngine::Tick( float DeltaSeconds, bool bIdleMode )
 			FApp::SetVolumeMultiplier( 1.0f );
 		}
 	}
+	else
+	{
+		FApp::SetVolumeMultiplier(0.0f);
+	}
 
 	if (!bHasFocus)
 	{
@@ -5639,8 +5643,8 @@ void UEditorEngine::UpdateAutoLoadProject()
 		{
 			if(FSlateApplication::IsInitialized())
 			{
-				FString SupressSettingName(FString(TEXT("UpdateMacOSX_")) + VERSION_STRINGIFY(ENGINE_MAJOR_VERSION) + TEXT("_") + VERSION_STRINGIFY(ENGINE_MINOR_VERSION) + TEXT("_") + VERSION_STRINGIFY(ENGINE_PATCH_VERSION));
-				FSuppressableWarningDialog::FSetupInfo Info( LOCTEXT("UpdateMacOSX_Body","Please update to the latest version of macOS for best performance and stability."), LOCTEXT("UpdateMacOSX_Title","Update macOS"), *SupressSettingName, GEditorSettingsIni );
+				FString SuppressSettingName(FString(TEXT("UpdateMacOSX_")) + VERSION_STRINGIFY(ENGINE_MAJOR_VERSION) + TEXT("_") + VERSION_STRINGIFY(ENGINE_MINOR_VERSION) + TEXT("_") + VERSION_STRINGIFY(ENGINE_PATCH_VERSION));
+				FSuppressableWarningDialog::FSetupInfo Info( LOCTEXT("UpdateMacOSX_Body","Please update to the latest version of macOS for best performance and stability."), LOCTEXT("UpdateMacOSX_Title","Update macOS"), *SuppressSettingName, GEditorSettingsIni );
 				Info.ConfirmText = LOCTEXT( "OK", "OK");
 				Info.bDefaultToSuppressInTheFuture = false;
 				FSuppressableWarningDialog OSUpdateWarning( Info );
