[StructUtils] When compiling a UserDefinedStruct default values, if the property ID was not found, try using the property name instead. The UserDefinedStruct editor automatically generates an ID with the name. When converting a PropertyBag to a UserDefinedStruct, we do not want the ID to be generated with the name. Systems like PropertyBinding find the properties by name. We would need to modify the engine to use CustomFindProperty instead.

#Synced-CL 49855772
#Synced-CL 49855772

diff --git a/Engine/Source/Editor/UnrealEd/Private/UserDefinedStructEditorData.cpp b/Engine/Source/Editor/UnrealEd/Private/UserDefinedStructEditorData.cpp
--- a/Engine/Source/Editor/UnrealEd/Private/UserDefinedStructEditorData.cpp
+++ b/Engine/Source/Editor/UnrealEd/Private/UserDefinedStructEditorData.cpp
@@ -255,9 +255,13 @@ void UUserDefinedStructEditorData::ReinitializeDefaultInstance(FString* OutLog)
 			FProperty* Property = *It;
 			if (Property)
 			{
-				FGuid VarGuid = FStructureEditorUtils::GetGuidFromPropertyName(Property->GetFName());
-
+				const FGuid VarGuid = FStructureEditorUtils::GetGuidFromPropertyName(Property->GetFName());
 				FStructVariableDescription* VarDesc = VariablesDescriptions.FindByPredicate(FStructureEditorUtils::FFindByGuidHelper<FStructVariableDescription>(VarGuid));
+				if (VarDesc == nullptr)
+				{
+					VarDesc = VariablesDescriptions.FindByPredicate(FStructureEditorUtils::FFindByNameHelper<FStructVariableDescription>(Property->GetFName()));
+				}
+
 				if (VarDesc && !VarDesc->CurrentDefaultValue.IsEmpty())
 				{
 					if (!FBlueprintEditorUtils::PropertyValueFromString(Property, VarDesc->CurrentDefaultValue, StructData, ScriptStruct, PPF_DuplicateVerbatim))
@@ -303,7 +307,13 @@ void UUserDefinedStructEditorData::RefreshValuesFromDefaultInstance()
 			if (const FProperty* Property = *It)
 			{
 				const FGuid VarGuid = FStructureEditorUtils::GetGuidFromPropertyName(Property->GetFName());
-				if (FStructVariableDescription* VarDesc = VariablesDescriptions.FindByPredicate(FStructureEditorUtils::FFindByGuidHelper<FStructVariableDescription>(VarGuid)))
+				FStructVariableDescription* VarDesc = VariablesDescriptions.FindByPredicate(FStructureEditorUtils::FFindByGuidHelper<FStructVariableDescription>(VarGuid));
+				if (VarDesc == nullptr)
+				{
+					VarDesc = VariablesDescriptions.FindByPredicate(FStructureEditorUtils::FFindByNameHelper<FStructVariableDescription>(Property->GetFName()));
+				}
+
+				if (VarDesc)
 				{
 					// If the default value has been changed elsewhere, don't refresh it here. This is typically the result of an edit action prior to compilation.
 					const bool bHasDefaultValueChanged = (VarDesc->DefaultValue != VarDesc->CurrentDefaultValue);
