FAssetEditorToolkit : speculative fix for a crash where EditingObjects is empty from the attribute bound to the asset editor's tab label

The code that is being replaced used `RemoveAllSwap` and could potentially leave the EditingObjects array empty such that the next call to FAssetEditorToolkit::GetEditingObjects() would assert on check( EditingObjects.Num() > 0 ).  With this change, removal is done through a call to RemoveEditingAsset which takes care to close the editor tab if it's the last element in the array.

#Synced-CL 49923003

diff --git a/Engine/Source/Editor/UnrealEd/Private/Toolkits/AssetEditorToolkit.cpp b/Engine/Source/Editor/UnrealEd/Private/Toolkits/AssetEditorToolkit.cpp
--- a/Engine/Source/Editor/UnrealEd/Private/Toolkits/AssetEditorToolkit.cpp
+++ b/Engine/Source/Editor/UnrealEd/Private/Toolkits/AssetEditorToolkit.cpp
@@ -1708,7 +1708,12 @@ void FAssetEditorToolkit::FGCEditingObjects::AddReferencedObjects(FReferenceColl
 	Collector.AddReferencedObjects(OwnerToolkit.EditingObjects);
 	
 	// Remove null objects as a safe guard against assets being forcibly GC'd
-	OwnerToolkit.EditingObjects.RemoveAllSwap([](UObject* Obj) { return Obj == nullptr; } );
+	TArray<TObjectPtr<UObject>> ObjectsToRemove = OwnerToolkit.EditingObjects.FilterByPredicate([](UObject* Obj) { return Obj == nullptr; });
+	for (TObjectPtr<UObject> ObjectToRemove : ObjectsToRemove)
+	{
+		// Call RemoveEditingAsset so that the editor tab gets closed if we remove the last EditingObject
+		OwnerToolkit.RemoveEditingAsset(ObjectToRemove);
+	}
 }
 
 FString FAssetEditorToolkit::FGCEditingObjects::GetReferencerName() const
