[HLODBuilder] Batch HLOD actor deletion operations into a single source control operation, significantly reducing source control overhead.
In a test deleting 64 HLOD actors, time spent on source control operations dropped from ~1m 9s to ~1s.

#Synced-CL 49955272

diff --git a/Engine/Source/Editor/UnrealEd/Private/WorldPartition/WorldPartitionBuilderHelpers.cpp b/Engine/Source/Editor/UnrealEd/Private/WorldPartition/WorldPartitionBuilderHelpers.cpp
--- a/Engine/Source/Editor/UnrealEd/Private/WorldPartition/WorldPartitionBuilderHelpers.cpp
+++ b/Engine/Source/Editor/UnrealEd/Private/WorldPartition/WorldPartitionBuilderHelpers.cpp
@@ -132,6 +132,16 @@ bool FSourceControlHelper::Delete(const FString& PackageName) const
 	return bDeleted;
 }
 
+bool FSourceControlHelper::Delete(const TArray<FString>& PackageNames) const
+{
+	bool bDeleted = PackageHelper.Delete(PackageNames);
+	if (bDeleted)
+	{
+		ModifiedFiles.Append(FBuilderModifiedFiles::EFileOperation::FileDeleted, PackageNames);
+	}
+	return bDeleted;
+}
+
 bool FSourceControlHelper::Delete(UPackage* Package) const 
 {
 	FString PackageName = GetFilename(Package);
diff --git a/Engine/Source/Editor/UnrealEd/Private/WorldPartition/WorldPartitionHLODsBuilder.cpp b/Engine/Source/Editor/UnrealEd/Private/WorldPartition/WorldPartitionHLODsBuilder.cpp
--- a/Engine/Source/Editor/UnrealEd/Private/WorldPartition/WorldPartitionHLODsBuilder.cpp
+++ b/Engine/Source/Editor/UnrealEd/Private/WorldPartition/WorldPartitionHLODsBuilder.cpp
@@ -1011,21 +1011,19 @@ bool UWorldPartitionHLODsBuilder::DeleteHLODActors()
 	WorldPartition->Uninitialize();
 	FWorldPartitionHelpers::DoCollectGarbage();
 
-	for (int32 PackageIndex = 0; PackageIndex < PackagesToDelete.Num(); ++PackageIndex)
+	bool bDeleted = SourceControlHelper->Delete(PackagesToDelete);
+	if (bDeleted)
 	{
-		const FString& PackageName = PackagesToDelete[PackageIndex];
-
-		bool bDeleted = SourceControlHelper->Delete(PackageName);
-		if (bDeleted)
+		for (int32 PackageIndex = 0; PackageIndex < PackagesToDelete.Num(); PackageIndex++)
 		{
-			UE_LOG(LogWorldPartitionHLODsBuilder, Display, TEXT("[%d / %d] Deleting %s..."), PackageIndex + 1, PackagesToDelete.Num(), *PackageName);
-		}
-		else
-		{
-			UE_LOG(LogWorldPartitionHLODsBuilder, Error, TEXT("Failed to delete %s, exiting..."), *PackageName);
-			return false;
+			UE_LOG(LogWorldPartitionHLODsBuilder, Display, TEXT("[%d / %d] Deleting %s..."), PackageIndex + 1, PackagesToDelete.Num(), *PackagesToDelete[PackageIndex]);
 		}
 	}
+	else
+	{
+		UE_LOG(LogWorldPartitionHLODsBuilder, Error, TEXT("Failed to delete HLOD actors, exiting..."));
+		return false;
+	}
 
 	UE_LOG(LogWorldPartitionHLODsBuilder, Display, TEXT("#### Deleted %d HLOD actors ####"), PackagesToDelete.Num());
 
diff --git a/Engine/Source/Editor/UnrealEd/Public/WorldPartition/WorldPartitionBuilderHelpers.h b/Engine/Source/Editor/UnrealEd/Public/WorldPartition/WorldPartitionBuilderHelpers.h
--- a/Engine/Source/Editor/UnrealEd/Public/WorldPartition/WorldPartitionBuilderHelpers.h
+++ b/Engine/Source/Editor/UnrealEd/Public/WorldPartition/WorldPartitionBuilderHelpers.h
@@ -40,6 +40,7 @@ public:
 	virtual bool Checkout(UPackage* Package) const override;
 	virtual bool Add(UPackage* Package) const override;
 	virtual bool Delete(const FString& PackageName) const override;
+	virtual bool Delete(const TArray<FString>& PackageNames) const override;
 	virtual bool Delete(UPackage* Package) const override;
 	virtual bool Save(UPackage* Package) const override;
 	virtual bool Copy(const FString& SrcFilePath, const FString& DstFilePath) const override;
diff --git a/Engine/Source/Runtime/Engine/Public/WorldPartition/WorldPartition.h b/Engine/Source/Runtime/Engine/Public/WorldPartition/WorldPartition.h
--- a/Engine/Source/Runtime/Engine/Public/WorldPartition/WorldPartition.h
+++ b/Engine/Source/Runtime/Engine/Public/WorldPartition/WorldPartition.h
@@ -111,6 +111,7 @@ public:
 	virtual bool Checkout(UPackage* Package) const =0;
 	virtual bool Add(UPackage* Package) const =0;
 	virtual bool Delete(const FString& PackageName) const =0;
+	virtual bool Delete(const TArray<FString>& PackageNames) const =0;
 	virtual bool Delete(UPackage* Package) const =0;
 	virtual bool Save(UPackage* Package) const =0;
 	virtual bool Copy(const FString& SrcFilePath, const FString& DstFilePath) const =0;
