[HLOD] Cancel the HLOD build when closing the "Unsaved Actors" prompt using the close icon

#Synced-CL 49955272

diff --git a/Engine/Source/Editor/UnrealEd/Private/WorldPartition/WorldPartitionHLODsBuilder.cpp b/Engine/Source/Editor/UnrealEd/Private/WorldPartition/WorldPartitionHLODsBuilder.cpp
--- a/Engine/Source/Editor/UnrealEd/Private/WorldPartition/WorldPartitionHLODsBuilder.cpp
+++ b/Engine/Source/Editor/UnrealEd/Private/WorldPartition/WorldPartitionHLODsBuilder.cpp
@@ -230,7 +230,7 @@ bool UWorldPartitionHLODsBuilder::BuildWithFilters(UWorld* InWorld, const TArray
 		EAppReturnType::Type Response = FMessageDialog::Open(EAppMsgType::YesNo,
 			LOCTEXT("HLODBuildUnsavedActorsMessage", "There are unsaved actors that won't be included in this HLOD build. Do you want to continue?"));
 
-		if (Response == EAppReturnType::No)
+		if (Response == EAppReturnType::No || Response == EAppReturnType::Cancel)
 		{
 			return false;
 		}
