Cooker: Add more feedback to "Content is missing from cook" warning: display the instigator chain of the source package.
Also change the display of instigators from a package to be more intuitive:
Old:
{ GeneratedPackage: /ShooterMaps/Maps/L_Expanse }
New:
{ /ShooterMaps/Maps/L_Expanse -> GeneratedPackage }

#rnx

#Synced-CL 49917028

diff --git a/Engine/Source/Editor/UnrealEd/Private/Commandlets/CookCommandlet.cpp b/Engine/Source/Editor/UnrealEd/Private/Commandlets/CookCommandlet.cpp
--- a/Engine/Source/Editor/UnrealEd/Private/Commandlets/CookCommandlet.cpp
+++ b/Engine/Source/Editor/UnrealEd/Private/Commandlets/CookCommandlet.cpp
@@ -12,6 +12,7 @@
 #include "CookOnTheSide/CookOnTheFlyServer.h"
 #include "Cooker/CookImportsChecker.h"
 #include "Cooker/CookProfilingPrivate.h"
+#include "Cooker/CookTypes.h"
 #include "CookerSettings.h"
 #include "Editor.h"
 #include "EngineGlobals.h"
@@ -615,13 +616,19 @@ void UCookCommandlet::RunCookByTheBookCook(UCookOnTheFlyServer* CookOnTheFlyServ
 	if (bShouldVerifyEDLCookInfo)
 	{
 		bool bFullReferencesExpected = !(CookOptions & ECookByTheBookOptions::SkipHardReferences);
-		FEDLCookChecker::Verify([](UE::FLogRecord&& Record)
+		auto WriteLogRecord = [](UE::FLogRecord&& Record)
 			{
 #if !NO_LOGGING
 				Record.SetCategory(LogCook.GetCategoryName());
 				UE::DispatchDynamicLogRecord(Record);
 #endif
-			}, bFullReferencesExpected);
+			};
+		auto GetInstigatorChain = [CookOnTheFlyServer](FName PackageName)
+			{
+				return UE::Cook::InstigatorChainToString(
+					CookOnTheFlyServer->GetInstigatorChain(PackageName), PackageName.ToString());
+			};
+		FEDLCookChecker::Verify(WriteLogRecord, GetInstigatorChain, bFullReferencesExpected);
 	}
 
 #if ENABLE_LOW_LEVEL_MEM_TRACKER
diff --git a/Engine/Source/Editor/UnrealEd/Private/Cooker/CookImportsChecker.cpp b/Engine/Source/Editor/UnrealEd/Private/Cooker/CookImportsChecker.cpp
--- a/Engine/Source/Editor/UnrealEd/Private/Cooker/CookImportsChecker.cpp
+++ b/Engine/Source/Editor/UnrealEd/Private/Cooker/CookImportsChecker.cpp
@@ -672,6 +672,7 @@ FEDLCookChecker FEDLCookChecker::AccumulateAndClear()
 }
 
 void FEDLCookChecker::Verify(const TFunction<void(UE::FLogRecord&& Record)>& MessageCallback,
+	const TFunction<FString(FName PackageName)>& GetInstigatorChain,
 	bool bFullReferencesExpected)
 {
 	LLM_SCOPE_BYTAG(EDLCookChecker);
@@ -740,7 +741,9 @@ void FEDLCookChecker::Verify(const TFunction<void(UE::FLogRecord&& Record)>& Mes
 					Record.SetUtf8Format(UTF8TEXT("Content is missing from cook. Source package referenced an object in target package but {Reason}.")
 						UTF8TEXT("\n\tSource package: {Source}")
 						UTF8TEXT("\n\tTarget package: {Target}")
-						UTF8TEXT("\n\tReferenced object: {ReferencedObject}"));
+						UTF8TEXT("\n\tReferenced object: {ReferencedObject}")
+						UTF8TEXT("\n\tSourceInstigator: {SourceInstigator}")
+					);
 					{
 						FCbWriter Writer;
 						Writer.BeginObject();
@@ -752,6 +755,7 @@ void FEDLCookChecker::Verify(const TFunction<void(UE::FLogRecord&& Record)>& Mes
 							NodeData.AppendPathName(Accumulator, ReferencedObjectStr);
 							Writer << "ReferencedObject" << ReferencedObjectStr;
 						}
+						Writer << "SourceInstigator" << GetInstigatorChain(PackageName);
 						Writer.EndObject();
 						Record.SetFields(Writer.Save().AsObject());
 					}
diff --git a/Engine/Source/Editor/UnrealEd/Private/Cooker/CookImportsChecker.h b/Engine/Source/Editor/UnrealEd/Private/Cooker/CookImportsChecker.h
--- a/Engine/Source/Editor/UnrealEd/Private/Cooker/CookImportsChecker.h
+++ b/Engine/Source/Editor/UnrealEd/Private/Cooker/CookImportsChecker.h
@@ -97,6 +97,7 @@ struct FEDLCookChecker
 
 	static void StartSavingEDLCookInfoForVerification();
 	static void Verify(const TFunction<void(UE::FLogRecord&& Record)>& MessageCallback,
+		const TFunction<FString(FName PackageName)>& GetInstigatorChain,
 		bool bFullReferencesExpected);
 	static void MoveToCompactBinaryAndClear(FCbWriter& Writer, bool& bOutHasData);
 	static bool AppendFromCompactBinary(FCbFieldView Field);
diff --git a/Engine/Source/Editor/UnrealEd/Private/Cooker/CookPackageData.cpp b/Engine/Source/Editor/UnrealEd/Private/Cooker/CookPackageData.cpp
--- a/Engine/Source/Editor/UnrealEd/Private/Cooker/CookPackageData.cpp
+++ b/Engine/Source/Editor/UnrealEd/Private/Cooker/CookPackageData.cpp
@@ -3106,20 +3106,8 @@ void FPackageDatas::DebugInstigator(FPackageData& PackageData)
 	if (ShowInstigatorPackageData == &PackageData)
 	{
 		TArray<FInstigator> Chain = CookOnTheFlyServer.GetInstigatorChain(PackageData.GetPackageName());
-		TStringBuilder<256> ChainText;
-		if (Chain.Num() == 0)
-		{
-			ChainText << TEXT("<NoInstigator>");
-		}
-		bool bFirst = true;
-		for (FInstigator& Instigator : Chain)
-		{
-			if (!bFirst) ChainText << TEXT(" <- ");
-			ChainText << TEXT("{ ") << Instigator.ToString() << TEXT(" }");
-			bFirst = false;
-		}
 		UE_LOG(LogCook, Display, TEXT("Instigator chain of %s: %s"),
-			*PackageData.GetPackageName().ToString(), ChainText.ToString());
+			*PackageData.GetPackageName().ToString(), *UE::Cook::InstigatorChainToString(Chain));
 	}
 	UpdateThreadsafePackageData(PackageData);
 }
diff --git a/Engine/Source/Editor/UnrealEd/Private/Cooker/CookTypes.cpp b/Engine/Source/Editor/UnrealEd/Private/Cooker/CookTypes.cpp
--- a/Engine/Source/Editor/UnrealEd/Private/Cooker/CookTypes.cpp
+++ b/Engine/Source/Editor/UnrealEd/Private/Cooker/CookTypes.cpp
@@ -18,6 +18,8 @@
 #include "Math/UnrealMathUtility.h"
 #include "Misc/CommandLine.h"
 #include "Misc/Parse.h"
+#include "Misc/ReverseIterate.h"
+#include "Misc/StringBuilder.h"
 #include "Serialization/PackageWriterToSharedBuffer.h"
 
 LLM_DEFINE_TAG(Cooker_CachedPlatformData);
@@ -737,6 +739,30 @@ bool LoadFromCompactBinary(FCbFieldView Field, FDiscoveredPlatformSet& OutValue,
 	return bOk;
 }
 
+FString InstigatorChainToString(TConstArrayView<FInstigator> Chain, FStringView FinalTargetToAppend)
+{
+	TStringBuilder<256> ChainText;
+	if (Chain.Num() == 0)
+	{
+		ChainText << TEXT("<NoInstigator>");
+	}
+	else
+	{
+		bool bFirst = true;
+		for (const FInstigator& Instigator : ReverseIterate(Chain))
+		{
+			if (!bFirst) ChainText << TEXT(" -> ");
+			ChainText << TEXT("{ ") << Instigator.ToString() << TEXT(" }");
+			bFirst = false;
+		}
+	}
+	if (!FinalTargetToAppend.IsEmpty())
+	{
+		ChainText << TEXT(" -> ") << FinalTargetToAppend;
+	}
+	return ChainText.ToString();
+}
+
 }
 
 FCbWriter& operator<<(FCbWriter& Writer, const UE::Cook::FInitializeConfigSettings& Value)
diff --git a/Engine/Source/Editor/UnrealEd/Private/Cooker/CookTypes.h b/Engine/Source/Editor/UnrealEd/Private/Cooker/CookTypes.h
--- a/Engine/Source/Editor/UnrealEd/Private/Cooker/CookTypes.h
+++ b/Engine/Source/Editor/UnrealEd/Private/Cooker/CookTypes.h
@@ -3,8 +3,10 @@
 #pragma once
 
 #include "AssetRegistry/PackageReader.h"
+#include "Containers/ArrayView.h"
 #include "Containers/Map.h"
 #include "Containers/Set.h"
+#include "Containers/StringView.h"
 #include "Containers/UnrealString.h"
 #include "Cooker/CookProfilingPrivate.h"
 #include "Cooker/MPCollector.h"
@@ -428,7 +430,7 @@ namespace UE::Cook
 	void SetIsSchedulerThread(bool bValue);
 
 
-	/** Placeholder to handle executing BuildDefintions for requested but not-yet-loaded packages. */
+	/** Placeholder to handle executing BuildDefinitions for requested but not-yet-loaded packages. */
 	class FBuildDefinitions
 	{
 	public:
@@ -564,7 +566,7 @@ public:
 	/** Should we generate streaming install manifests (only valid option in cook by the book) */
 	bool							bGenerateStreamingInstallManifests = false;
 
-	/** Should we generate a seperate manifest for map dependencies */
+	/** Should we generate a separate manifest for map dependencies */
 	bool							bGenerateDependenciesForMaps = false;
 
 	/** error when detecting engine content being used in this cook */
@@ -672,6 +674,8 @@ struct FCookArtifactsCurrentSettings
 	bool bRequestInvalidate = false;
 };
 
+FString InstigatorChainToString(TConstArrayView<FInstigator> Chain, FStringView FinalTargetToAppend = FStringView());
+
 } // namespace UE::Cook
 
 
diff --git a/Engine/Source/Runtime/CoreUObject/Private/UObject/ICookInfo.cpp b/Engine/Source/Runtime/CoreUObject/Private/UObject/ICookInfo.cpp
--- a/Engine/Source/Runtime/CoreUObject/Private/UObject/ICookInfo.cpp
+++ b/Engine/Source/Runtime/CoreUObject/Private/UObject/ICookInfo.cpp
@@ -58,10 +58,9 @@ const TCHAR* LexToString(EInstigator Value)
 FString FInstigator::ToString() const
 {
 	TStringBuilder<256> Result;
-	Result << LexToString(Category);
 	if (!Referencer.IsNone())
 	{
-		Result << TEXT(": ") << Referencer;
+		Result << Referencer << TEXT(" -> ") << LexToString(Category);
 	}
 	else
 	{
@@ -73,9 +72,13 @@ FString FInstigator::ToString() const
 #undef EINSTIGATOR_VALUE_CALLBACK
 		default: break;
 		}
-		if (!bCategoryAllowsUnparameterized)
+		if (bCategoryAllowsUnparameterized)
 		{
-			Result << TEXT(": <NoReferencer>");
+			Result << LexToString(Category);
+		}
+		else
+		{
+			Result << TEXT("<NoReferencer>") << TEXT(" -> ") << LexToString(Category);
 		}
 	}
 	return FString(Result);
