EditorUsability : Commands
- Removing TexStreamAccMaterialTextureScaleSingle bundle of 256 texture commands
- Removing RequiredTextureResolutionSingle bundle of 256 texture commands
- Refactoring view mode options menu to use MenuBuilder and populate entries for texture actions to replace those 2 commands

#Synced-CL 50006372

diff --git a/Engine/Source/Editor/UnrealEd/Private/EditorViewportCommands.cpp b/Engine/Source/Editor/UnrealEd/Private/EditorViewportCommands.cpp
--- a/Engine/Source/Editor/UnrealEd/Private/EditorViewportCommands.cpp
+++ b/Engine/Source/Editor/UnrealEd/Private/EditorViewportCommands.cpp
@@ -9,6 +9,7 @@
 #include "Editor.h"
 #include "Modules/ModuleManager.h"
 #include "ContentBrowserModule.h"
+#include "EditorViewportClient.h"
 #include "IContentBrowserSingleton.h"
 #include "Materials/MaterialInterface.h"
 #include "MaterialShared.h"
@@ -96,20 +97,6 @@ void FEditorViewportCommands::RegisterCommands()
 	UI_COMMAND( TexStreamAccMaterialTextureScaleAll, "All Textures", "Visualize the scales accuracy of all textures", EUserInterfaceActionType::RadioButton, FInputChord() );
 	UI_COMMAND( RequiredTextureResolutionMode, "Required Texture Resolution View Mode", "Visualize the ratio between the currently streamed texture resolution and the resolution wanted by the GPU", EUserInterfaceActionType::RadioButton, FInputChord() );
 
-	for (int32 TextureIndex = 0; TextureIndex < TEXSTREAM_MAX_NUM_TEXTURES_PER_MATERIAL; ++TextureIndex)
-	{
-		const FName TextureScaleCommandName = *FString::Printf(TEXT("ShowTextureTextureScale%d"), TextureIndex);
-		const FName TextureResolutionCommandName = *FString::Printf(TEXT("ShowTextureTextureResolution%d"), TextureIndex);
-
-		const FText LocalizedName = FText::Format(LOCTEXT("ShowTextureCommands", "Texture {0}"), TextureIndex);
-		
-		const FText LocalizedTextureScaleTooltip = FText::Format(LOCTEXT("ShowTextureCommands_TextureScale_ToolTip", "Visualize the scale accuracy of texture {0}"), TextureIndex);
-		const FText LocalizedTextureResolutionTooltip = FText::Format(LOCTEXT("ShowTextureCommands_TextureResolution_ToolTip", "Visualize the ratio between the currently streamed resolution of texture {0} texture resolution and the resolution wanted by the GPU."), TextureIndex);
-
-		TexStreamAccMaterialTextureScaleSingle[TextureIndex] = FUICommandInfoDecl(this->AsShared(), TextureScaleCommandName, LocalizedName, LocalizedTextureScaleTooltip, ShowTextureScaleBundle).UserInterfaceType(EUserInterfaceActionType::RadioButton);
-		RequiredTextureResolutionSingle[TextureIndex] = FUICommandInfoDecl(this->AsShared(), TextureResolutionCommandName, LocalizedName, LocalizedTextureResolutionTooltip, ShowTextureResolutionBundle).UserInterfaceType(EUserInterfaceActionType::RadioButton);
-	}
-
 	UI_COMMAND( StationaryLightOverlapMode, "Stationary Light Overlap View Mode", "Visualizes overlap of stationary lights", EUserInterfaceActionType::RadioButton, FInputChord() );
 	UI_COMMAND( LightmapDensityMode, "Lightmap Density View Mode", "Renders the scene with lightmap density visualization", EUserInterfaceActionType::RadioButton, FInputChord( EModifierKey::Alt, EKeys::Zero ) );
 
@@ -445,14 +432,26 @@ static void AppendMaterialInfoToMenu(const UMaterialInterface* MaterialInterface
 
 TSharedRef<SWidget> BuildViewModeOptionsMenu(TSharedPtr<FUICommandList> CommandList, EViewModeIndex ViewModeIndex, ERHIFeatureLevel::Type FeatureLevel, TMap<int32, FName>& ParamNameMap)
 {
-	const FEditorViewportCommands& Commands = FEditorViewportCommands::Get();
 	FMenuBuilder MenuBuilder( true, CommandList );
+	const UE::UnrealEd::FViewModeOptionsArgs ViewModeArgs
+	{
+		.ViewModeIndex = ViewModeIndex,
+		.FeatureLevel = FeatureLevel,
+		.ParamNameMap = ParamNameMap,
+		.ViewportClient = nullptr // Use UE::UnrealEd::BuildViewModeOptionsMenu instead
+	};
+	BuildViewModeOptionsMenu(MenuBuilder, ViewModeArgs);
+	return MenuBuilder.MakeWidget();
+}
 
+void UE::UnrealEd::BuildViewModeOptionsMenu(FMenuBuilder& MenuBuilder, const FViewModeOptionsArgs& InArgs)
+{
+	const FEditorViewportCommands& Commands = FEditorViewportCommands::Get();
 	const FString MenuName = LOCTEXT("TexStreamAccMaterialTextureScaleSingleDisplayName", "Texture").ToString();
 
-	ParamNameMap.Reset();
+	InArgs.ParamNameMap.Reset();
 
-	if (ViewModeIndex == VMI_MeshUVDensityAccuracy)
+	if (InArgs.ViewModeIndex == VMI_MeshUVDensityAccuracy)
 	{
 		MenuBuilder.AddMenuEntry(Commands.TexStreamAccMeshUVDensityAll, NAME_None, LOCTEXT("TexStreamAccMeshUVDensityAllDisplayName", "All UV Channels"));
 		for (int32 TexCoordIndex = 0; TexCoordIndex < TEXSTREAM_MAX_NUM_UVCHANNELS; ++TexCoordIndex)
@@ -460,7 +459,7 @@ TSharedRef<SWidget> BuildViewModeOptionsMenu(TSharedPtr<FUICommandList> CommandL
 			MenuBuilder.AddMenuEntry(Commands.TexStreamAccMeshUVDensitySingle[TexCoordIndex], NAME_None, FText::FromString(FString::Printf(TEXT("%s %d"), *MenuName, TexCoordIndex)));
 		}
 	}
-	else if (ViewModeIndex == VMI_MaterialTextureScaleAccuracy || ViewModeIndex == VMI_RequiredTextureResolution)
+	else if (InArgs.ViewModeIndex == VMI_MaterialTextureScaleAccuracy || InArgs.ViewModeIndex == VMI_RequiredTextureResolution)
 	{
 		TArray<const UMaterialInterface*> SelectedMaterials;
 		GetSelectedMaterials(SelectedMaterials);
@@ -468,7 +467,7 @@ TSharedRef<SWidget> BuildViewModeOptionsMenu(TSharedPtr<FUICommandList> CommandL
 		TMap<int32, TArray<FString> > DataPerTextureIndex;
 		TMap<FName, FMenuDataPerTexture> DataPerTextureName;
 
-		if (ViewModeIndex == VMI_MaterialTextureScaleAccuracy)
+		if (InArgs.ViewModeIndex == VMI_MaterialTextureScaleAccuracy)
 		{
 			MenuBuilder.AddMenuEntry(Commands.TexStreamAccMaterialTextureScaleAll, NAME_None, LOCTEXT("TexStreamAccMaterialTextureScaleAllDisplayName", "All Textures"));
 
@@ -481,68 +480,79 @@ TSharedRef<SWidget> BuildViewModeOptionsMenu(TSharedPtr<FUICommandList> CommandL
 		{
 			for (const UMaterialInterface* MaterialInterface : SelectedMaterials)
 			{
-				AppendMaterialInfoToMenu(MaterialInterface, FeatureLevel, MenuName, DataPerTextureIndex, DataPerTextureName);
+				AppendMaterialInfoToMenu(MaterialInterface, InArgs.FeatureLevel, MenuName, DataPerTextureIndex, DataPerTextureName);
 			}
 		}
 
-		const TSharedPtr<FUICommandInfo>* PerTextureCommands = ViewModeIndex == VMI_MaterialTextureScaleAccuracy ? Commands.TexStreamAccMaterialTextureScaleSingle : Commands.RequiredTextureResolutionSingle;
-
-		// If there is not too many textures, show the data per name.
-		if (DataPerTextureName.Num() > 0 && DataPerTextureName.Num() < TEXSTREAM_MAX_NUM_TEXTURES_PER_MATERIAL)
+		if (InArgs.ViewportClient.IsValid())
 		{
-			TMap<FString, FName> SortedFNames;
-			for (auto Ite = DataPerTextureName.CreateConstIterator(); Ite; ++Ite)
+			TSharedRef<FEditorViewportClient> ViewportClient = InArgs.ViewportClient.ToSharedRef();
+			auto GetUIAction = [&ViewportClient](int32 InViewModeIndex)
 			{
-				SortedFNames.Add(Ite.Key().ToString(), Ite.Key());
-			}
-			SortedFNames.KeySort(TLess<FString>());
-
-			int32 CommandIndex = 0;
-			for (auto Ite = SortedFNames.CreateConstIterator(); Ite; ++Ite)
-			{ 
-				const FMenuDataPerTexture& MaterialInfos = DataPerTextureName.FindOrAdd(Ite.Value());
-				FString ToolTipOverride = MaterialInfos.ToolTips[0];
-				for (int32 MaterialIndex = 1; MaterialIndex < MaterialInfos.ToolTips.Num(); ++MaterialIndex)
+				return FUIAction(
+					FExecuteAction::CreateSP( ViewportClient, &FEditorViewportClient::SetViewModeParam, InViewModeIndex ),
+					FCanExecuteAction(),
+					FIsActionChecked::CreateSP( ViewportClient, &FEditorViewportClient::IsViewModeParam, InViewModeIndex )
+				);
+			};
+
+			// If there is not too many textures, show the data per name.
+			if (DataPerTextureName.Num() > 0 && DataPerTextureName.Num() < TEXSTREAM_MAX_NUM_TEXTURES_PER_MATERIAL)
+			{
+				TMap<FString, FName> SortedFNames;
+				for (auto Ite = DataPerTextureName.CreateConstIterator(); Ite; ++Ite)
 				{
-					ToolTipOverride = FString::Printf(TEXT("%s\n%s"), *ToolTipOverride, *MaterialInfos.ToolTips[MaterialIndex]);
+					SortedFNames.Add(Ite.Key().ToString(), Ite.Key());
+				}
+				SortedFNames.KeySort(TLess<FString>());
+
+				int32 CommandIndex = 0;
+				for (auto Ite = SortedFNames.CreateConstIterator(); Ite; ++Ite)
+				{ 
+					const FMenuDataPerTexture& MaterialInfos = DataPerTextureName.FindOrAdd(Ite.Value());
+					FString ToolTipOverride = MaterialInfos.ToolTips[0];
+					for (int32 MaterialIndex = 1; MaterialIndex < MaterialInfos.ToolTips.Num(); ++MaterialIndex)
+					{
+						ToolTipOverride = FString::Printf(TEXT("%s\n%s"), *ToolTipOverride, *MaterialInfos.ToolTips[MaterialIndex]);
+					}
+					
+					MenuBuilder.AddMenuEntry(FText::FromString(MaterialInfos.DisplayName), FText::FromString(ToolTipOverride), FSlateIcon(), GetUIAction(CommandIndex), NAME_None, EUserInterfaceActionType::RadioButton);
+					
+					InArgs.ParamNameMap.Add(CommandIndex, Ite.Value());
+					++CommandIndex;
 				}
-				
-				MenuBuilder.AddMenuEntry(PerTextureCommands[CommandIndex], NAME_None, FText::FromString(MaterialInfos.DisplayName), FText::FromString(ToolTipOverride));
-				
-				ParamNameMap.Add(CommandIndex, Ite.Value());
-				++CommandIndex;
 			}
-		}
-		else if (DataPerTextureIndex.Num() > 0) // Otherwise show the data per index, with extra info (also hidding entries with no info)
-		{
-			for (int32 TextureIndex = 0; TextureIndex < TEXSTREAM_MAX_NUM_TEXTURES_PER_MATERIAL; ++TextureIndex)
+			else if (DataPerTextureIndex.Num() > 0) // Otherwise show the data per index, with extra info (also hidding entries with no info)
 			{
-				const TArray<FString>* TextureDataAtIndex = DataPerTextureIndex.Find(TextureIndex);
-				if (TextureDataAtIndex && TextureDataAtIndex->Num() > 0)
+				for (int32 TextureIndex = 0; TextureIndex < TEXSTREAM_MAX_NUM_TEXTURES_PER_MATERIAL; ++TextureIndex)
 				{
-					FString NameOverride = FString::Printf(TEXT("%s %d (%s)"), *MenuName, TextureIndex, *(*TextureDataAtIndex)[0]);
-					if (TextureDataAtIndex->Num() > 1)
-					{
-						NameOverride.Append(TEXT(" ..."));
-					}
-					FString ToolTipOverride = (*TextureDataAtIndex)[0];
-					for (int32 Index = 1; Index < (*TextureDataAtIndex).Num(); ++Index)
+					const TArray<FString>* TextureDataAtIndex = DataPerTextureIndex.Find(TextureIndex);
+					if (TextureDataAtIndex && TextureDataAtIndex->Num() > 0)
 					{
-						ToolTipOverride = FString::Printf(TEXT("%s\n%s"), *ToolTipOverride, *(*TextureDataAtIndex)[Index]);
+						FString NameOverride = FString::Printf(TEXT("%s %d (%s)"), *MenuName, TextureIndex, *(*TextureDataAtIndex)[0]);
+						if (TextureDataAtIndex->Num() > 1)
+						{
+							NameOverride.Append(TEXT(" ..."));
+						}
+						FString ToolTipOverride = (*TextureDataAtIndex)[0];
+						for (int32 Index = 1; Index < (*TextureDataAtIndex).Num(); ++Index)
+						{
+							ToolTipOverride = FString::Printf(TEXT("%s\n%s"), *ToolTipOverride, *(*TextureDataAtIndex)[Index]);
+						}
+						MenuBuilder.AddMenuEntry(FText::FromString(NameOverride), FText::FromString(ToolTipOverride), FSlateIcon(), GetUIAction(TextureIndex), NAME_None, EUserInterfaceActionType::RadioButton);
 					}
-					MenuBuilder.AddMenuEntry(PerTextureCommands[TextureIndex], NAME_None, FText::FromString(NameOverride), FText::FromString(ToolTipOverride));
 				}
 			}
-		}
-		else // If nothing selected, just display a list of name
-		{
-			for (int32 TextureIndex = 0; TextureIndex < TEXSTREAM_MAX_NUM_TEXTURES_PER_MATERIAL; ++TextureIndex)
+			else // If nothing selected, just display a list of name
 			{
-				MenuBuilder.AddMenuEntry(PerTextureCommands[TextureIndex], NAME_None, FText::FromString(FString::Printf(TEXT("%s %d"), *MenuName, TextureIndex)));
+				for (int32 TextureIndex = 0; TextureIndex < TEXSTREAM_MAX_NUM_TEXTURES_PER_MATERIAL; ++TextureIndex)
+				{
+					const FText Label = FText::FromString(FString::Printf(TEXT("%s %d"), *MenuName, TextureIndex));
+					MenuBuilder.AddMenuEntry(Label, Label, FSlateIcon(), GetUIAction(TextureIndex), NAME_None, EUserInterfaceActionType::RadioButton);
+				}
 			}
 		}
 	}
-	return MenuBuilder.MakeWidget();
 }
 
 #undef LOCTEXT_NAMESPACE
diff --git a/Engine/Source/Editor/UnrealEd/Private/SCommonEditorViewportToolbarBase.cpp b/Engine/Source/Editor/UnrealEd/Private/SCommonEditorViewportToolbarBase.cpp
--- a/Engine/Source/Editor/UnrealEd/Private/SCommonEditorViewportToolbarBase.cpp
+++ b/Engine/Source/Editor/UnrealEd/Private/SCommonEditorViewportToolbarBase.cpp
@@ -610,9 +610,19 @@ TSharedRef<SWidget> SCommonEditorViewportToolbarBase::GenerateViewModeOptionsMen
 {
 	GetInfoProvider().OnFloatingButtonClicked();
 	TSharedRef<SEditorViewport> ViewportRef = GetInfoProvider().GetViewportWidget();
-	FEditorViewportClient& ViewClient = GetViewportClient();
-	const UWorld* World = ViewClient.GetWorld();
-	return BuildViewModeOptionsMenu(ViewportRef->GetCommandList(), ViewClient.GetViewMode(), World ? World->GetFeatureLevel() : GMaxRHIFeatureLevel, ViewClient.GetViewModeParamNameMap());
+	TSharedRef<FEditorViewportClient> ViewClient = GetInfoProvider().GetViewportWidget()->GetViewportClient().ToSharedRef();
+	const UWorld* World = ViewClient->GetWorld();
+	
+	FMenuBuilder MenuBuilder( true, ViewportRef->GetCommandList() );
+	const UE::UnrealEd::FViewModeOptionsArgs ViewModeArgs
+	{
+		.ViewModeIndex = ViewClient->GetViewMode(),
+		.FeatureLevel = World ? World->GetFeatureLevel() : GMaxRHIFeatureLevel,
+		.ParamNameMap = ViewClient->GetViewModeParamNameMap(),
+		.ViewportClient = ViewClient
+	};
+	BuildViewModeOptionsMenu(MenuBuilder, ViewModeArgs);
+	return MenuBuilder.MakeWidget();
 }
 
 
diff --git a/Engine/Source/Editor/UnrealEd/Private/SEditorViewport.cpp b/Engine/Source/Editor/UnrealEd/Private/SEditorViewport.cpp
--- a/Engine/Source/Editor/UnrealEd/Private/SEditorViewport.cpp
+++ b/Engine/Source/Editor/UnrealEd/Private/SEditorViewport.cpp
@@ -710,11 +710,6 @@ void SEditorViewport::BindCommands()
 	}
 
 	MAP_VIEWMODEPARAM_ACTION( Commands.TexStreamAccMaterialTextureScaleAll, -1 );
-	for (int32 TextureIndex = 0; TextureIndex < TEXSTREAM_MAX_NUM_TEXTURES_PER_MATERIAL; ++TextureIndex)
-	{
-		MAP_VIEWMODEPARAM_ACTION( Commands.TexStreamAccMaterialTextureScaleSingle[TextureIndex], TextureIndex );
-		MAP_VIEWMODEPARAM_ACTION( Commands.RequiredTextureResolutionSingle[TextureIndex], TextureIndex );
-	}
 
 	MAP_VIEWMODE_ACTION( Commands.GeometryInspectionClayMode, VMI_Clay );
 	MAP_VIEWMODE_ACTION( Commands.GeometryInspectionFrontBackFaceMode, VMI_FrontBackFace );
diff --git a/Engine/Source/Editor/UnrealEd/Private/ViewportToolbar/UnrealEdViewportToolbar.cpp b/Engine/Source/Editor/UnrealEd/Private/ViewportToolbar/UnrealEdViewportToolbar.cpp
--- a/Engine/Source/Editor/UnrealEd/Private/ViewportToolbar/UnrealEdViewportToolbar.cpp
+++ b/Engine/Source/Editor/UnrealEd/Private/ViewportToolbar/UnrealEdViewportToolbar.cpp
@@ -2352,26 +2352,21 @@ void PopulateViewModesMenu(UToolMenu* InMenu)
 			LOCTEXT("ViewModeOptionsToolTip", "Options for the current view mode."),
 			FNewMenuDelegate::CreateLambda([WeakViewport = EditorViewport.ToWeakPtr()](FMenuBuilder& MenuBuilder)
 			{
-				TSharedRef<SWidget> OptionsMenuContent = SNullWidget::NullWidget;
-			
 				if (TSharedPtr<SEditorViewport> Viewport = WeakViewport.Pin())
 				{
 					if (TSharedPtr<FEditorViewportClient> Client = Viewport->GetViewportClient())
 					{
 						const UWorld* World = Client->GetWorld();
-						
-						// TODO: This is a very awkward way of exposing this menu.
-						// The function should be modified to allow defining the menu builder.
-						// It was done this way to maintain ABI compatability and should be fixed at the earliest possible convenience.
-						OptionsMenuContent = ::BuildViewModeOptionsMenu(
-							Viewport->GetCommandList(),
-							Client->GetViewMode(),
-							World ? World->GetFeatureLevel() : GMaxRHIFeatureLevel,
-							Client->GetViewModeParamNameMap()
-						);
-					}	
+						const FViewModeOptionsArgs ViewModeArgs
+						{
+							.ViewModeIndex = Client->GetViewMode(),
+							.FeatureLevel = World ? World->GetFeatureLevel() : GMaxRHIFeatureLevel,
+							.ParamNameMap = Client->GetViewModeParamNameMap(),
+							.ViewportClient = Client
+						};
+						BuildViewModeOptionsMenu(MenuBuilder, ViewModeArgs);
+					}
 				}
-				MenuBuilder.AddWidget(OptionsMenuContent, FText::GetEmpty(), true, false);
 			})
 		);
 					
diff --git a/Engine/Source/Editor/UnrealEd/Public/EditorViewportCommands.h b/Engine/Source/Editor/UnrealEd/Public/EditorViewportCommands.h
--- a/Engine/Source/Editor/UnrealEd/Public/EditorViewportCommands.h
+++ b/Engine/Source/Editor/UnrealEd/Public/EditorViewportCommands.h
@@ -10,6 +10,8 @@
 #include "Styling/AppStyle.h"
 #include "Engine/TextureStreamingTypes.h"
 
+class FEditorViewportClient;
+class FMenuBuilder;
 class FUICommandList;
 
 UE_DECLARE_TCOMMANDS(class FEditorViewportCommands, UNREALED_API)
@@ -108,10 +110,7 @@ public:
 	/** Changes the viewport to material texture scale accuracy mode */
 	TSharedPtr< FUICommandInfo > TexStreamAccMaterialTextureScaleMode;
 	TSharedPtr< FUICommandInfo > TexStreamAccMaterialTextureScaleAll; // Similar as TexStreamAccMaterialTextureScaleMode, but used in the view mode options menu.
-	TSharedPtr< FUICommandInfo > TexStreamAccMaterialTextureScaleSingle[TEXSTREAM_MAX_NUM_TEXTURES_PER_MATERIAL];
-
 	TSharedPtr< FUICommandInfo > RequiredTextureResolutionMode;
-	TSharedPtr< FUICommandInfo > RequiredTextureResolutionSingle[TEXSTREAM_MAX_NUM_TEXTURES_PER_MATERIAL];
 
 	/** Changes the viewport to stationary light overlap mode */
 	TSharedPtr< FUICommandInfo > StationaryLightOverlapMode;
@@ -313,4 +312,24 @@ UNREALED_API FText GetViewModeOptionsMenuLabel(EViewModeIndex ViewModeIndex);
 
  * @return	The widget containing the viewmode options.
  */
+UE_DEPRECATED(5.8, "Use UE::UnrealEd::BuildViewModeOptionsMenu instead.")
 UNREALED_API TSharedRef<SWidget> BuildViewModeOptionsMenu(TSharedPtr<FUICommandList> CommandList, EViewModeIndex ViewModeIndex, ERHIFeatureLevel::Type FeatureLevel, TMap<int32, FName>& ParamNameMap);
+
+namespace UE::UnrealEd
+{
+	struct FViewModeOptionsArgs
+	{
+		EViewModeIndex ViewModeIndex;
+		ERHIFeatureLevel::Type FeatureLevel;
+		TMap<int32, FName>& ParamNameMap;
+		TSharedPtr<FEditorViewportClient> ViewportClient;
+	};
+
+	/**
+	 * Build the viewmode otions menu.
+	 * 
+	 * @param FMenuBuilder The menu builder to create the menu
+	 * @param FViewModeArgs The parameters to customize the view mode menu
+	 */
+	UNREALED_API void BuildViewModeOptionsMenu(FMenuBuilder& InMenuBuilder, const FViewModeOptionsArgs& InArgs);
+}
\ No newline at end of file
