Adding export macros so UWorldPartitionNavigationDataBuilder can be used as a base class in project code.

#rnx
#Synced-CL 0

diff --git a/Engine/Source/Editor/UnrealEd/Private/WorldPartition/WorldPartitionNavigationDataBuilder.cpp b/Engine/Source/Editor/UnrealEd/Private/WorldPartition/WorldPartitionNavigationDataBuilder.cpp
--- a/Engine/Source/Editor/UnrealEd/Private/WorldPartition/WorldPartitionNavigationDataBuilder.cpp
+++ b/Engine/Source/Editor/UnrealEd/Private/WorldPartition/WorldPartitionNavigationDataBuilder.cpp
@@ -23,6 +23,11 @@
 
 DEFINE_LOG_CATEGORY_STATIC(LogWorldPartitionNavigationDataBuilder, Log, All);
 
+UWorldPartitionNavigationDataBuilder::UWorldPartitionNavigationDataBuilder(FVTableHelper& Helper)
+	: Super(Helper)
+{
+}
+
 UWorldPartitionNavigationDataBuilder::UWorldPartitionNavigationDataBuilder(const FObjectInitializer& ObjectInitializer)
 	: Super(ObjectInitializer)
 {
diff --git a/Engine/Source/Editor/UnrealEd/Public/WorldPartition/WorldPartitionNavigationDataBuilder.h b/Engine/Source/Editor/UnrealEd/Public/WorldPartition/WorldPartitionNavigationDataBuilder.h
--- a/Engine/Source/Editor/UnrealEd/Public/WorldPartition/WorldPartitionNavigationDataBuilder.h
+++ b/Engine/Source/Editor/UnrealEd/Public/WorldPartition/WorldPartitionNavigationDataBuilder.h
@@ -10,23 +10,29 @@ class UWorldPartition;
 UCLASS()
 class UWorldPartitionNavigationDataBuilder : public UWorldPartitionBuilder
 {
-	GENERATED_UCLASS_BODY()
+	GENERATED_BODY()
 
 public:
+
+	UNREALED_API UWorldPartitionNavigationDataBuilder(FVTableHelper& Helper);
+	UNREALED_API UWorldPartitionNavigationDataBuilder(const FObjectInitializer& ObjectInitializer);
+	UWorldPartitionNavigationDataBuilder() = default;
+	virtual ~UWorldPartitionNavigationDataBuilder() = default;
+
 	// UWorldPartitionBuilder interface begin
 	virtual bool RequiresCommandletRendering() const override { return false; }
 	virtual ELoadingMode GetLoadingMode() const override { return ELoadingMode::IterativeCells2D; }
 
 protected:
-	virtual bool PreRun(UWorld* World, FPackageSourceControlHelper& PackageHelper) override;
-	virtual bool RunInternal(UWorld* World, const FCellInfo& InCellInfo, FPackageSourceControlHelper& PackageHelper) override;
-	virtual bool PostRun(UWorld* World, FPackageSourceControlHelper& PackageHelper, const bool bInRunSuccess) override;
+	UNREALED_API virtual bool PreRun(UWorld* World, FPackageSourceControlHelper& PackageHelper) override;
+	UNREALED_API virtual bool RunInternal(UWorld* World, const FCellInfo& InCellInfo, FPackageSourceControlHelper& PackageHelper) override;
+	UNREALED_API virtual bool PostRun(UWorld* World, FPackageSourceControlHelper& PackageHelper, const bool bInRunSuccess) override;
 	// UWorldPartitionBuilder interface end
 
-	bool GenerateNavigationData(UWorldPartition* WorldPartition, const FBox& LoadedBounds, const FBox& GeneratingBounds) const;
+	UNREALED_API bool GenerateNavigationData(UWorldPartition* WorldPartition, const FBox& LoadedBounds, const FBox& GeneratingBounds) const;
 
-	bool SavePackages(const TArray<UPackage*>& PackagesToSave) const;
-	bool DeletePackages(const FPackageSourceControlHelper& PackageHelper, const TArray<UPackage*>& PackagesToDelete) const;
+	UNREALED_API bool SavePackages(const TArray<UPackage*>& PackagesToSave) const;
+	UNREALED_API bool DeletePackages(const FPackageSourceControlHelper& PackageHelper, const TArray<UPackage*>& PackagesToDelete) const;
 
 	bool bCleanBuilderPackages = false;
 
