DDB: Updates to DerivedDataBuildLocalExecutor

- Added multi-process support to avoid contention on the scratch paths.
- Share the workers between builds to avoid the cost of thousands of copies.
- Improved the level of detail in logs to assist with debugging.

#Synced-CL 49858055

#Synced-CL 49858055

diff --git a/Engine/Source/Editor/UnrealEd/Private/DerivedDataBuildLocalExecutor.cpp b/Engine/Source/Editor/UnrealEd/Private/DerivedDataBuildLocalExecutor.cpp
--- a/Engine/Source/Editor/UnrealEd/Private/DerivedDataBuildLocalExecutor.cpp
+++ b/Engine/Source/Editor/UnrealEd/Private/DerivedDataBuildLocalExecutor.cpp
@@ -2,11 +2,13 @@
 
 #include "DerivedDataBuildLocalExecutor.h"
 
+#include "Async/ManualResetEvent.h"
+#include "Async/Mutex.h"
+#include "Async/UniqueLock.h"
 #include "Compression/CompressedBuffer.h"
 #include "Containers/Array.h"
-#include "Containers/ArrayView.h"
-#include "Containers/ContainersFwd.h"
-#include "Containers/StringFwd.h"
+#include "Containers/Map.h"
+#include "Containers/SharedString.h"
 #include "Containers/StringView.h"
 #include "Containers/UnrealString.h"
 #include "DerivedDataBuild.h"
@@ -20,16 +22,22 @@
 #include "DerivedDataValue.h"
 #include "Features/IModularFeatures.h"
 #include "HAL/FileManager.h"
+#include "HAL/PlatformFile.h"
+#include "HAL/PlatformFileManager.h"
 #include "HAL/PlatformProcess.h"
+#include "Hash/xxhash.h"
 #include "IO/IoHash.h"
 #include "Logging/LogCategory.h"
 #include "Logging/LogMacros.h"
+#include "Logging/StructuredLog.h"
 #include "Memory/CompositeBuffer.h"
 #include "Memory/SharedBuffer.h"
 #include "Misc/EnumClassFlags.h"
+#include "Misc/Guid.h"
 #include "Misc/Optional.h"
-#include "Misc/PathViews.h"
 #include "Misc/Paths.h"
+#include "Misc/PathViews.h"
+#include "Misc/ScopeExit.h"
 #include "Misc/StringBuilder.h"
 #include "Serialization/Archive.h"
 #include "Serialization/CompactBinary.h"
@@ -50,30 +58,30 @@ namespace UE::DerivedData
 DEFINE_LOG_CATEGORY_STATIC(LogDerivedDataBuildLocalExecutor, Log, All);
 
 /**
- * This implements a simple local DDC2 executor which spawns all jobs
+ * This implements a simple local build worker executor which spawns all jobs
  * as local processes. This is intentionally as simple as possible and
  * everything is synchronous. This is not meant to be used in production,
  * and is perhaps primarily useful as a debugging aid.
  */
 class FLocalBuildWorkerExecutor final : public IBuildWorkerExecutor
 {
-	FString SandboxRootDir = FPaths::EngineSavedDir() / TEXT("LocalExec");
-
 public:
 	FLocalBuildWorkerExecutor()
 	{
-		// Clean out any leftovers from a previous run
-
-		UE_LOG(LogDerivedDataBuildLocalExecutor, Warning, TEXT("Deleting existing local execution state from '%s'"), *SandboxRootDir);
-
-		constexpr bool RequireExists = false;
-		constexpr bool Tree = true;
-		IFileManager::Get().DeleteDirectory(*SandboxRootDir, RequireExists, Tree);
+		// Delete previous local execution state from the primary process.
+		if (GetMultiprocessId() == 0)
+		{
+			UE_LOGFMT(LogDerivedDataBuildLocalExecutor, Display, "Deleting existing local execution state from '{SandboxRoot}'.", *SandboxRoot);
+			IFileManager::Get().DeleteDirectory(*SandboxRoot, /*bRequireExists*/ false, /*bTree*/ true);
+		}
 
 		IModularFeatures::Get().RegisterModularFeature(IBuildWorkerExecutor::FeatureName, this);
 	}
 
-	~FLocalBuildWorkerExecutor() final = default;
+	~FLocalBuildWorkerExecutor() final
+	{
+		IModularFeatures::Get().UnregisterModularFeature(IBuildWorkerExecutor::FeatureName, this);
+	}
 
 	void Build(
 		const FBuildAction& Action,
@@ -84,12 +92,12 @@ public:
 		IRequestOwner& Owner,
 		FOnBuildWorkerActionComplete&& OnComplete) final
 	{
-		// Review build action inputs to determine if they need to be materialized/propagated 
+		// Review build action inputs to determine if they need to be materialized/propagated
 		// (right now, they always will be)
 
 		TArray<FUtf8StringView> MissingInputs;
 
-		Action.IterateInputs([&](FUtf8StringView Key, const FIoHash& RawHash, uint64 RawSize)
+		Action.IterateInputs([&Inputs, &MissingInputs](FUtf8StringView Key, const FIoHash& RawHash, uint64 RawSize)
 		{
 			if (Inputs.IsNull() || Inputs.Get().FindInput(Key).IsNull())
 			{
@@ -111,101 +119,79 @@ public:
 		// this could be problematic for large runs so we should probably add support for
 		// configurable cleanup policies
 
-		static std::atomic<int32> SerialNo = 0;
-
-		TStringBuilder<256> SandboxRoot;
-		FPathViews::Append(SandboxRoot, SandboxRootDir, TEXT("Scratch"));
-		SandboxRoot.Appendf(TEXT("%06d"), ++SerialNo);
-
-		// Manifest worker in scratch area
-
+		TStringBuilder<256> WorkerPath;
+		if (!ManifestWorker(Worker, WorkerPath))
 		{
-			TArray<FIoHash> MissingWorkerData;
-
-			TArray<FIoHash> WorkerFileHashes;
-			TArray<TTuple<FStringView, bool>> WorkerFileMeta;
-			Worker.IterateExecutables([&WorkerFileHashes, &WorkerFileMeta](FStringView Path, const FIoHash& RawHash, uint64 RawSize)
-			{
-				WorkerFileHashes.Emplace(RawHash);
-				WorkerFileMeta.Emplace(Path, true);
-			});
-
-			Worker.IterateFiles([&WorkerFileHashes, &WorkerFileMeta](FStringView Path, const FIoHash& RawHash, uint64 RawSize)
-			{
-				WorkerFileHashes.Emplace(RawHash);
-				WorkerFileMeta.Emplace(Path, false);
-			});
-
-			FRequestOwner BlockingOwner(EPriority::Blocking);
-			Worker.FindFileData(WorkerFileHashes, BlockingOwner, [&SandboxRoot, &WorkerFileMeta](FBuildWorkerFileDataCompleteParams&& Params)
-			{
-				uint32 MetaIndex = 0;
-				for (const FCompressedBuffer& Buffer : Params.Files)
-				{
-					const TTuple<FStringView, bool>& Meta = WorkerFileMeta[MetaIndex];
-
-					TStringBuilder<256> Path;
-					FPathViews::Append(Path, SandboxRoot, Meta.Key);
-					if (TUniquePtr<FArchive> Ar{IFileManager::Get().CreateFileWriter(*Path, FILEWRITE_Silent)})
-					{
-						FCompositeBuffer DecompressedComposite = Buffer.DecompressToComposite();
-
-						for (auto& Segment : DecompressedComposite.GetSegments())
-						{
-							Ar->Serialize((void*)Segment.GetData(), Segment.GetSize());
-						}
-					}
-
-					++MetaIndex;
-				}
-			});
-			BlockingOwner.Wait();
-
-			// This directory must exist in order for the builder to run correctly
-			TStringBuilder<256> BinPath;
-			FPathViews::Append(BinPath, SandboxRoot, TEXT("Engine/Binaries/Win64"));
-			IFileManager::Get().MakeDirectory(*BinPath);
+			return OnComplete({Action.GetKey(), {}, {}, EStatus::Error});
 		}
+		FPathViews::Append(WorkerPath, Worker.GetPath());
 
-		// Manifest inputs in scratch area
+		// Create a scratch directory for the build at {SandboxRoot}/Builds/######/####.
+		TStringBuilder<256> BuildRoot;
+		{
+			static std::atomic<int32> SerialNo = 0;
+			const int32 BuildSerialNo = ++SerialNo;
+			TStringBuilder<8> OuterDir, InnerDir;
+			OuterDir.Appendf(TEXT("%02d%04d"), GetMultiprocessId(), BuildSerialNo / 10000);
+			InnerDir.Appendf(TEXT("%04d"), BuildSerialNo % 10000);
+			FPathViews::Append(BuildRoot, SandboxRoot, TEXTVIEW("Builds"), OuterDir, InnerDir);
+		}
 
+		// Manifest inputs under the build root.
 		if (!Inputs.IsNull())
 		{
-			Inputs.Get().IterateInputs([&SandboxRoot](FUtf8StringView Key, const FCompressedBuffer& Buffer)
+			std::atomic<bool> bOk = true;
+			Inputs.Get().IterateInputs([&BuildRoot, &Action, &Worker, &bOk](FUtf8StringView Key, const FCompressedBuffer& Buffer)
 			{
 				TStringBuilder<256> Path;
-				FPathViews::Append(Path, SandboxRoot, TEXT("Inputs"), FIoHash(Buffer.GetRawHash()));
+				FPathViews::Append(Path, BuildRoot, TEXT("Inputs"), FIoHash(Buffer.GetRawHash()));
 				if (TUniquePtr<FArchive> Ar{IFileManager::Get().CreateFileWriter(*Path, FILEWRITE_Silent)})
 				{
-					// Workaround for FArchive::operator<< not accepting const objects
-					FCompressedBuffer Copy = Buffer;
-					*Ar << Copy;
+					Buffer.Save(*Ar);
+				}
+				else
+				{
+					UE_LOGFMT(LogDerivedDataBuildLocalExecutor, Warning,
+						"Failed to open '{Path}' to write input for build worker {WorkerName} for '{Name}'.",
+						Path, Worker.GetName(), Action.GetName());
+					bOk = false;
 				}
 			});
+			if (!bOk)
+			{
+				return OnComplete({Action.GetKey(), {}, {}, EStatus::Error});
+			}
 		}
 
-		// Serialize action specification
-
+		// Manifest the action under the build root.
 		{
 			TStringBuilder<256> Path;
-			FPathViews::Append(Path, SandboxRoot, TEXT("Build.action"));
+			FPathViews::Append(Path, BuildRoot, TEXT("Build.action"));
 			if (TUniquePtr<FArchive> Ar{IFileManager::Get().CreateFileWriter(*Path, FILEWRITE_Silent)})
 			{
-				FCbWriter BuildActionWriter;
-				Action.Save(BuildActionWriter);
-				BuildActionWriter.Save(*Ar);
+				FCbWriter Writer;
+				Action.Save(Writer);
+				Writer.Save(*Ar);
+			}
+			else
+			{
+				UE_LOGFMT(LogDerivedDataBuildLocalExecutor, Warning,
+					"Failed to open '{Path}' to write action for build worker {WorkerName} for '{Name}'.",
+					Path, Worker.GetName(), Action.GetName());
+				return OnComplete({Action.GetKey(), {}, {}, EStatus::Error});
 			}
 		}
 
+		// Launch the worker process and wait for it to complete.
+
+		// TODO: Handle Worker.IterateEnvironment(...)
+
 		const bool bLaunchDetached = false;
 		const bool bLaunchHidden = false;
 		const bool bLaunchReallyHidden = false;
 		uint32 ProcessID = 0;
 		const int PriorityModifier = 0;
-		const TCHAR* WorkingDirectory = *SandboxRoot;
-		TStringBuilder<256> WorkerPath;
-		FPathViews::Append(WorkerPath, SandboxRoot, Worker.GetPath());
-
+		const TCHAR* WorkingDirectory = *BuildRoot;
 		FProcHandle ProcHandle = FPlatformProcess::CreateProc(
 			*WorkerPath,
 			TEXT("-Build=Build.action"),
@@ -223,37 +209,39 @@ public:
 		int32 ExitCode = -1;
 		if (!FPlatformProcess::GetProcReturnCode(ProcHandle, &ExitCode) || ExitCode != 0)
 		{
-			UE_LOG(LogDerivedDataBuildLocalExecutor, Warning, TEXT("Worker process exit code = %d!"), ExitCode);
+			UE_LOGFMT(LogDerivedDataBuildLocalExecutor, Warning,
+				"Worker {WorkerName} exited with code {ExitCode} for '{Name}' at '{BuildRoot}' with worker '{WorkerPath}'.",
+				Worker.GetName(), ExitCode, Action.GetName(), BuildRoot, WorkerPath);
 			return OnComplete({Action.GetKey(), {}, {}, EStatus::Error});
 		}
 
-		// Gather results
-
+		// Load the output object from the build root.
 		FOptionalBuildOutput RemoteBuildOutput;
-
 		{
-			TStringBuilder<256> OutputPath;
-			FPathViews::Append(OutputPath, SandboxRoot, TEXT("Build.output"));
-			if (TUniquePtr<FArchive> Ar{IFileManager::Get().CreateFileReader(*OutputPath, FILEREAD_Silent)})
+			TStringBuilder<256> Path;
+			FPathViews::Append(Path, BuildRoot, TEXT("Build.output"));
+			if (TUniquePtr<FArchive> Ar{IFileManager::Get().CreateFileReader(*Path, FILEREAD_Silent)})
 			{
 				FCbObject BuildOutput = LoadCompactBinary(*Ar).AsObject();
-
 				if (Ar->IsError())
 				{
-					UE_LOG(LogDerivedDataBuildLocalExecutor, Warning, TEXT("Worker error: build output structure not valid!"));
+					UE_LOGFMT(LogDerivedDataBuildLocalExecutor, Warning,
+						"Worker {WorkerName} build output '{Path}' is not valid compact binary for '{Name}' at '{BuildRoot}' with worker '{WorkerPath}'.",
+						Worker.GetName(), Path, Action.GetName(), BuildRoot, WorkerPath);
 					return OnComplete({Action.GetKey(), {}, {}, EStatus::Error});
 				}
-
 				RemoteBuildOutput = FBuildOutput::Load(Action.GetName(), Action.GetFunction(), BuildOutput);
 			}
+			if (RemoteBuildOutput.IsNull())
+			{
+				UE_LOGFMT(LogDerivedDataBuildLocalExecutor, Warning,
+					"Worker {WorkerName} is missing build output '{Path}' for '{Name}' at '{BuildRoot}' with worker '{WorkerPath}'.",
+					Worker.GetName(), Path, Action.GetName(), BuildRoot, WorkerPath);
+				return OnComplete({Action.GetKey(), {}, {}, EStatus::Error});
+			}
 		}
 
-		if (RemoteBuildOutput.IsNull())
-		{
-			UE_LOG(LogDerivedDataBuildLocalExecutor, Warning, TEXT("Remote execution system error: build output blob missing!"));
-			return OnComplete({Action.GetKey(), {}, {}, EStatus::Error});
-		}
-
+		// Load the output values from the build root.
 		FBuildOutputBuilder OutputBuilder = BuildSystem.CreateOutput(Action.GetName(), Action.GetFunction());
 		for (const FValueWithId& Value : RemoteBuildOutput.Get().GetValues())
 		{
@@ -263,22 +251,36 @@ public:
 			}
 			else
 			{
-				FCompressedBuffer BufferForValue;
+				FCompressedBuffer OutputData;
 
 				TStringBuilder<128> Path;
-				FPathViews::Append(Path, SandboxRoot, TEXT("Outputs"), FIoHash(Value.GetRawHash()));
+				FPathViews::Append(Path, BuildRoot, TEXT("Outputs"), FIoHash(Value.GetRawHash()));
 				if (TUniquePtr<FArchive> Ar{IFileManager::Get().CreateFileReader(*Path, FILEREAD_Silent)})
 				{
-					BufferForValue = FCompressedBuffer::Load(*Ar);
+					OutputData = FCompressedBuffer::Load(*Ar);
 				}
 
-				if (BufferForValue.IsNull())
+				if (OutputData.IsNull())
 				{
-					UE_LOG(LogDerivedDataBuildLocalExecutor, Warning, TEXT("Remote execution system error: payload blob missing!"));
+					UE_LOGFMT(LogDerivedDataBuildLocalExecutor, Warning,
+						"Worker {WorkerName} is missing build output data '{Path}' for '{Name}' at '{BuildRoot}' with worker '{WorkerPath}'.",
+						Worker.GetName(), Path, Action.GetName(), BuildRoot, WorkerPath);
 					return OnComplete({Action.GetKey(), {}, {}, EStatus::Error});
 				}
 
-				OutputBuilder.AddValue(Value.GetId(), FValue(MoveTemp(BufferForValue)));
+				if (OutputData.GetRawHash() != Value.GetRawHash() || OutputData.GetRawSize() != Value.GetRawSize())
+				{
+					UE_LOGFMT(LogDerivedDataBuildLocalExecutor, Warning,
+						"Worker {WorkerName} has mismatched build output data at '{Path}' for '{Name}' at '{BuildRoot}' with worker '{WorkerPath}'. "
+						"Expected hash {ExpectedHash} and size {ExpectedSize}. "
+						"Received hash {ReceivedHash} and size {ReceivedSize}.",
+						Worker.GetName(), Path, Action.GetName(), BuildRoot, WorkerPath,
+						Value.GetRawHash(), Value.GetRawSize(),
+						OutputData.GetRawHash(), OutputData.GetRawSize());
+					return OnComplete({Action.GetKey(), {}, {}, EStatus::Error});
+				}
+
+				OutputBuilder.AddValue(Value.GetId(), FValue(MoveTemp(OutputData)));
 			}
 		}
 
@@ -286,6 +288,136 @@ public:
 		return OnComplete({Action.GetKey(), BuildOutput, {}, EStatus::Ok});
 	}
 
+	[[nodiscard]] bool ManifestWorker(const FBuildWorker& Worker, FStringBuilderBase& OutPath)
+	{
+		// Calculate the worker hash that uniquely identifies the files of the worker.
+		FXxHash128Builder WorkerHashBuilder;
+		TArray<FIoHash> WorkerFileHashes;
+		TArray<TTuple<FStringView, bool>> WorkerFileMeta;
+
+		Worker.IterateExecutables([&WorkerHashBuilder, &WorkerFileHashes, &WorkerFileMeta](FStringView Path, const FIoHash& RawHash, uint64 RawSize)
+		{
+			WorkerHashBuilder.Update(MakeMemoryView(&RawHash, 1));
+			WorkerFileHashes.Emplace(RawHash);
+			WorkerFileMeta.Emplace(Path, true);
+		});
+
+		Worker.IterateFiles([&WorkerHashBuilder, &WorkerFileHashes, &WorkerFileMeta](FStringView Path, const FIoHash& RawHash, uint64 RawSize)
+		{
+			WorkerHashBuilder.Update(MakeMemoryView(&RawHash, 1));
+			WorkerFileHashes.Emplace(RawHash);
+			WorkerFileMeta.Emplace(Path, false);
+		});
+
+		const FXxHash128 WorkerHash = WorkerHashBuilder.Finalize();
+
+		OutPath.Reset();
+		FPathViews::Append(OutPath, SandboxRoot, TEXTVIEW("Workers"), WorkerHash);
+
+		// Return if a worker exists with this worker hash.
+		IFileManager& FileManager = IFileManager::Get();
+		if (FileManager.DirectoryExists(*OutPath))
+		{
+			return true;
+		}
+
+		// Allow the first thread in this process to manifest the worker and have the others wait.
+		// Multiple processes may race to rename their temporary directory.
+		bool bWaitForManifest = false;
+		FManifestWorkerState* State;
+		{
+			TUniqueLock Lock(ManifestingWorkersMutex);
+			TUniquePtr<FManifestWorkerState>& StatePtr = ManifestingWorkers.FindOrAdd(WorkerHash);
+			if (StatePtr)
+			{
+				bWaitForManifest = true;
+			}
+			else
+			{
+				StatePtr = MakeUnique<FManifestWorkerState>();
+			}
+			State = StatePtr.Get();
+			++State->RefCount;
+		}
+		ON_SCOPE_EXIT
+		{
+			TUniqueLock Lock(ManifestingWorkersMutex);
+			if (--State->RefCount == 0)
+			{
+				ManifestingWorkers.Remove(WorkerHash);
+			}
+		};
+		if (bWaitForManifest)
+		{
+			State->ReadyEvent.Wait();
+			return State->bOk;
+		}
+
+		UE_LOGFMT(LogDerivedDataBuildLocalExecutor, Display,
+			"Manifesting build worker {WorkerName} in '{WorkerPath}'...", Worker.GetName(), OutPath);
+
+		// Manifest the worker in a temporary directory.
+		TStringBuilder<256> ScratchPath(InPlace, OutPath, '.', FGuid::NewGuid());
+
+		std::atomic<bool> bOk = true;
+		FRequestOwner BlockingOwner(EPriority::Blocking);
+		Worker.FindFileData(WorkerFileHashes, BlockingOwner, [&Worker, &FileManager, &ScratchPath, &WorkerFileMeta, &bOk](FBuildWorkerFileDataCompleteParams&& Params)
+		{
+			uint32 MetaIndex = 0;
+			for (const FCompressedBuffer& Buffer : Params.Files)
+			{
+				TStringBuilder<256> FilePath;
+				FPathViews::Append(FilePath, ScratchPath, WorkerFileMeta[MetaIndex].Key);
+				if (TUniquePtr<FArchive> Ar{FileManager.CreateFileWriter(*FilePath, FILEWRITE_Silent)})
+				{
+					for (auto& Segment : Buffer.DecompressToComposite().GetSegments())
+					{
+						Ar->Serialize((void*)Segment.GetData(), Segment.GetSize());
+					}
+				}
+				else
+				{
+					UE_LOGFMT(LogDerivedDataBuildLocalExecutor, Display,
+						"Failed to open '{Path}' for writing for build worker {WorkerName}.",
+						FilePath, Worker.GetName());
+					bOk = false;
+				}
+
+				++MetaIndex;
+			}
+		});
+		BlockingOwner.Wait();
+
+		// Ensure this directory exists because the build worker cannot run correctly without it.
+		{
+			TStringBuilder<256> BinariesPath;
+			FPathViews::Append(BinariesPath, ScratchPath, TEXTVIEW("Engine/Binaries/Win64"));
+			if (!FileManager.MakeDirectory(*BinariesPath, /*bTree*/ true))
+			{
+				UE_LOGFMT(LogDerivedDataBuildLocalExecutor, Display,
+					"Failed to create directory '{Path}' for build worker {WorkerName}.",
+					BinariesPath, Worker.GetName());
+				bOk = false;
+			}
+		}
+
+		// Move the temporary directory to its final location or delete it if there was a race.
+		IPlatformFile& PlatformFile = FPlatformFileManager::Get().GetPlatformFile();
+		if (bOk && !PlatformFile.DirectoryExists(*OutPath))
+		{
+			PlatformFile.MoveFile(*OutPath, *ScratchPath);
+		}
+		PlatformFile.DeleteDirectoryRecursively(*ScratchPath);
+		const bool bExists = PlatformFile.DirectoryExists(*OutPath);
+		UE_CLOGFMT(!bExists, LogDerivedDataBuildLocalExecutor, Display,
+			"Failed to rename '{TempPath}' to '{Path}' for build worker {WorkerName}.",
+			ScratchPath, OutPath, Worker.GetName());
+
+		State->bOk = bOk;
+		State->ReadyEvent.Notify();
+		return bExists;
+	}
+
 	TConstArrayView<FStringView> GetHostPlatforms() const final
 	{
 		static constexpr FStringView HostPlatforms[]{TEXTVIEW("Win64")};
@@ -295,9 +427,21 @@ public:
 	void DumpStats()
 	{
 	}
+
+private:
+	struct FManifestWorkerState
+	{
+		int32 RefCount = 0;
+		bool bOk = false;
+		FManualResetEvent ReadyEvent;
+	};
+
+	FString SandboxRoot = FPaths::EngineSavedDir() / TEXT("LocalExec");
+	TMap<FXxHash128, TUniquePtr<FManifestWorkerState>> ManifestingWorkers;
+	FMutex ManifestingWorkersMutex;
 };
 
-}  // namespace UE::DerivedData
+} // UE::DerivedData
 
 TOptional<UE::DerivedData::FLocalBuildWorkerExecutor> GLocalBuildWorkerExecutor;
 
