Cooker: CookImportsChecker: Change container of NodePrereqs to a TMap<TArray> to avoid overflowing the int32 limit of the allocation size of a TMultiMap's buckets.
#rnx

#Synced-CL 50035849

diff --git a/Engine/Source/Editor/UnrealEd/Private/Cooker/CookImportsChecker.cpp b/Engine/Source/Editor/UnrealEd/Private/Cooker/CookImportsChecker.cpp
--- a/Engine/Source/Editor/UnrealEd/Private/Cooker/CookImportsChecker.cpp
+++ b/Engine/Source/Editor/UnrealEd/Private/Cooker/CookImportsChecker.cpp
@@ -498,7 +498,9 @@ void FEDLCookChecker::AddPackageWithUnknownExports(FName LongPackageName)
 
 void FEDLCookChecker::AddDependency(FEDLNodeID SourceID, FEDLNodeID TargetID)
 {
-	NodePrereqs.Add(SourceID, TargetID);
+	// Add the edge; to save time don't check for uniqueness. Uniqueness is checked on demand, in
+	// UniquifyNodePrereqsEdgeList.
+	NodePrereqs.FindOrAdd(SourceID).Add(TargetID);
 }
 
 void FEDLCookChecker::StartSavingEDLCookInfoForVerification()
@@ -512,6 +514,28 @@ void FEDLCookChecker::StartSavingEDLCookInfoForVerification()
 	}
 }
 
+void FEDLCookChecker::UniquifyNodePrereqsEdgeList(TArray<FEDLNodeID>& EdgeList)
+{
+	bool bNeedsShrink = false;
+
+	// Avoid the cost of sort+unique if unnecessary, to reduce the cost if called multiple times.
+	// This loop is similar to !Algo::IsSorted, but also returns true if any adjacent elements are equal.
+	int32 NumToCompare = EdgeList.Num() - 1;
+	for (int32 Index = 0; Index < NumToCompare; ++Index)
+	{
+		if (EdgeList[Index + 1] <= EdgeList[Index])
+		{
+			bNeedsShrink = true;
+			break;
+		}
+	}
+	if (bNeedsShrink)
+	{
+		Algo::Sort(EdgeList);
+		EdgeList.SetNum(Algo::Unique(EdgeList));
+	}
+}
+
 bool FEDLCookChecker::CheckForCyclesInner(TSet<FEDLNodeID>& Visited, TSet<FEDLNodeID>& Stack, const FEDLNodeID& Visit, FEDLNodeID& FailNode)
 {
 	bool bResult = false;
@@ -527,9 +551,13 @@ bool FEDLCookChecker::CheckForCyclesInner(TSet<FEDLNodeID>& Visited, TSet<FEDLNo
 		if (!bWasAlreadyTested)
 		{
 			Stack.Add(Visit);
-			for (auto It = NodePrereqs.CreateConstKeyIterator(Visit); !bResult && It; ++It)
+			TArray<FEDLNodeID>* EdgeList = NodePrereqs.Find(Visit);
+			if (EdgeList)
 			{
-				bResult = CheckForCyclesInner(Visited, Stack, It.Value(), FailNode);
+				for (int32 EdgeListIndex = 0; !bResult && EdgeListIndex < EdgeList->Num(); ++EdgeListIndex)
+				{
+					bResult = CheckForCyclesInner(Visited, Stack, (*EdgeList)[EdgeListIndex], FailNode);
+				}
 			}
 			Stack.Remove(Visit);
 		}
@@ -627,11 +655,17 @@ void FEDLCookChecker::Merge(FEDLCookChecker&& Other)
 			RemapIDs.Add(NodeID);
 		}
 
-		for (const TPair<FEDLNodeID, FEDLNodeID>& Prereq : Other.NodePrereqs)
+		for (const TPair<FEDLNodeID, TArray<FEDLNodeID>>& Prereq : Other.NodePrereqs)
 		{
+			// Merge the edgelists; to save time don't check for uniqueness. Uniqueness is checked on demand, in
+			// UniquifyNodePrereqsEdgeList.
 			FEDLNodeID SourceID = RemapIDs[Prereq.Key];
-			FEDLNodeID TargetID = RemapIDs[Prereq.Value];
-			AddDependency(SourceID, TargetID);
+			TArray<FEDLNodeID>& EdgeList = NodePrereqs.FindOrAdd(SourceID);
+			for (FEDLNodeID OtherTargetID : Prereq.Value)
+			{
+				FEDLNodeID TargetID = RemapIDs[OtherTargetID];
+				EdgeList.Add(TargetID);
+			}
 		}
 
 		Other.NodePrereqs.Empty();
@@ -770,6 +804,10 @@ void FEDLCookChecker::Verify(const TFunction<void(UE::FLogRecord&& Record)>& Mes
 		TSet<FEDLNodeID> Visited;
 		TSet<FEDLNodeID> Stack;
 		bool bHadCycle = false;
+		for (TPair<FEDLNodeID, TArray<FEDLNodeID>>& NodePreqsPair : Accumulator.NodePrereqs)
+		{
+			UniquifyNodePrereqsEdgeList(NodePreqsPair.Value);
+		}
 		for (const FEDLNodeData& NodeData : Accumulator.Nodes)
 		{
 			if (!NodeData.bIsExport)
@@ -838,10 +876,22 @@ void FEDLCookChecker::WriteToCompactBinary(FCbWriter& Writer)
 		}
 		Writer.EndArray();
 		Writer.BeginArray("NodePrereqs");
-		for (const TPair<FEDLNodeID, FEDLNodeID>& Pair : NodePrereqs)
+		// Transmit NumSources to use for reservation size
+		uint32 NumSources = static_cast<uint32>(NodePrereqs.Num());
+		Writer << NumSources;
+		// Transmit the edgelist as (Source, NumEdgeList, EdgeListElements). Use a flat list of concatenated
+		// objects to reduce transmission size.
+		for (TPair<FEDLNodeID, TArray<FEDLNodeID>>& Pair : NodePrereqs)
 		{
 			Writer << static_cast<uint32>(Pair.Key);
-			Writer << static_cast<uint32>(Pair.Value);
+			TArray<FEDLNodeID>& EdgeList = Pair.Value;
+			UniquifyNodePrereqsEdgeList(EdgeList);
+			uint32 NumDests = static_cast<uint32>(EdgeList.Num());
+			Writer << NumDests;
+			for (FEDLNodeID Dest : EdgeList)
+			{
+				Writer << static_cast<uint32>(Dest);
+			}
 		}
 		Writer.EndArray();
 		Writer.BeginArray("PackagesWithUnknownExports");
@@ -900,26 +950,24 @@ bool FEDLCookChecker::ReadFromCompactBinary(FCbFieldView Field)
 	}
 
 	FCbFieldView PrereqsField = Field["NodePrereqs"];
-	const int64 NumNodePrereqs = PrereqsField.AsArrayView().Num() / 2;
-	if (NumNodePrereqs > MAX_int32)
-	{
-		return false;
-	}
-	NodePrereqs.Reserve(static_cast<int32>(NumNodePrereqs));
-	if (PrereqsField.HasError())
-	{
-		return false;
-	}
 	FCbFieldViewIterator PrereqsIter = PrereqsField.CreateViewIterator();
+	uint32 NumSources = 0;
+	if (!LoadFromCompactBinary(PrereqsIter++, NumSources)) { return false; }
+	NodePrereqs.Reserve(static_cast<int32>(NumSources));
 	while (PrereqsIter)
 	{
-		uint32 Key;
-		uint32 Value;
-		if (!LoadFromCompactBinary(PrereqsIter, Key)) { return false; }
-		++PrereqsIter;
-		if (!LoadFromCompactBinary(PrereqsIter, Value)) { return false; }
-		++PrereqsIter;
-		NodePrereqs.Add(static_cast<FEDLNodeID>(Key), static_cast<FEDLNodeID>(Value));
+		uint32 SourceAsInt;
+		if (!LoadFromCompactBinary(PrereqsIter++, SourceAsInt)) { return false; }
+		uint32 NumDests;
+		if (!LoadFromCompactBinary(PrereqsIter++, NumDests)) { return false; }
+		TArray<FEDLNodeID>& EdgeList = NodePrereqs.FindOrAdd(static_cast<FEDLNodeID>(SourceAsInt));
+		EdgeList.Reserve(NumDests);
+		for (uint32 DestIndex = 0; DestIndex < NumDests; ++DestIndex)
+		{
+			uint32 DestAsInt;
+			if (!LoadFromCompactBinary(PrereqsIter++, DestAsInt)) { return false; }
+			EdgeList.Add(DestAsInt);
+		}
 	}
 
 	FCbFieldView PackagesWithUnknownExportsField = Field["PackagesWithUnknownExports"];
diff --git a/Engine/Source/Editor/UnrealEd/Private/Cooker/CookImportsChecker.h b/Engine/Source/Editor/UnrealEd/Private/Cooker/CookImportsChecker.h
--- a/Engine/Source/Editor/UnrealEd/Private/Cooker/CookImportsChecker.h
+++ b/Engine/Source/Editor/UnrealEd/Private/Cooker/CookImportsChecker.h
@@ -233,6 +233,7 @@ private:
 	void Merge(FEDLCookChecker&& Other);
 	bool CheckForCyclesInner(TSet<FEDLNodeID>& Visited, TSet<FEDLNodeID>& Stack, const FEDLNodeID& Visit, FEDLNodeID& FailNode);
 	void AddDependency(FEDLNodeID SourceID, FEDLNodeID TargetID);
+	static void UniquifyNodePrereqsEdgeList(TArray<FEDLNodeID>& EdgeList);
 
 	/**
 	 * All the FEDLNodeDatas that have been created for this checker. These are allocated as elements of an array rather than pointers to reduce cputime and
@@ -241,8 +242,15 @@ private:
 	TArray<FEDLNodeData> Nodes;
 	/** A map to lookup the node for a UObject or for the corresponding node in another thread's FEDLCookChecker. */
 	TMap<FEDLNodeHash, FEDLNodeID> NodeHashToNodeID;
-	/** The graph of dependencies between nodes. */
-	TMultiMap<FEDLNodeID, FEDLNodeID> NodePrereqs;
+	/**
+	 * The graph of dependencies between nodes. Format is a map from Source node to a custom-format edgelist. 
+	 * The custom format edge-list is an optionally (sorted and unique) array of Dest nodes. We make it sorted and unique
+	 * when we read it. This format is performant for our use case, which is a period of unique Source->Target writes
+	 * with no reads followed by a period of reads with no writes. Size of the graph can be very large on large projects
+	 * (20 million vertices, 100 million edges); we use this format to prevent int32 overflow of the allocation size
+	 * of the map buckets and to otherwise be as memory-efficient as possible.
+	 */
+	TMap<FEDLNodeID, TArray<FEDLNodeID>> NodePrereqs;
 	/**
 	 * Packages that were cooked with LegacyIterative and therefore have an unknown set of exports.
 	 * We suppress warnings for exports missing from these packages.
