[ITF Viewport Interactions] Add a delegate so that EdModes can opt out of certain interactions

This CL adds the `OnGetUnsupportedInteractions` delegate to `UModeManagerInteractiveToolsContext`: add interaction name(s) to this list to prevent the interaction from being added

#Synced-CL 50039254

diff --git a/Engine/Source/Editor/Experimental/EditorInteractiveToolsFramework/Private/ViewportInteractions/ViewportInteractionsBehaviorSource.cpp b/Engine/Source/Editor/Experimental/EditorInteractiveToolsFramework/Private/ViewportInteractions/ViewportInteractionsBehaviorSource.cpp
--- a/Engine/Source/Editor/Experimental/EditorInteractiveToolsFramework/Private/ViewportInteractions/ViewportInteractionsBehaviorSource.cpp
+++ b/Engine/Source/Editor/Experimental/EditorInteractiveToolsFramework/Private/ViewportInteractions/ViewportInteractionsBehaviorSource.cpp
@@ -366,6 +366,11 @@ UViewportInteraction* UViewportInteractionsBehaviorSource::AddInteraction(const
 				return nullptr;
 			}
 		}
+
+		if (UnsupportedInteractions.Contains(DefaultInteraction->GetInteractionName()))
+		{
+			return nullptr;
+		}
 	}
 
 	UViewportInteraction* Interaction = NewObject<UViewportInteraction>(GetTransientPackage(), InInteractionClass);
@@ -645,6 +650,11 @@ FEditorViewportClient* UViewportInteractionsBehaviorSource::GetHoveredEditorView
 	return EditorViewportClient;
 }
 
+void UViewportInteractionsBehaviorSource::SetUnsupportedViewportInteractions(const TArray<FName>& InUnsupportedInteractions)
+{
+	UnsupportedInteractions = InUnsupportedInteractions;
+}
+
 void UViewportInteractionsBehaviorSource::SetAltKeyState(bool bInDown)
 {
 	bIsAltDown = bInDown;
diff --git a/Engine/Source/Editor/Experimental/EditorInteractiveToolsFramework/Public/ViewportInteractions/ViewportInteractionsBehaviorSource.h b/Engine/Source/Editor/Experimental/EditorInteractiveToolsFramework/Public/ViewportInteractions/ViewportInteractionsBehaviorSource.h
--- a/Engine/Source/Editor/Experimental/EditorInteractiveToolsFramework/Public/ViewportInteractions/ViewportInteractionsBehaviorSource.h
+++ b/Engine/Source/Editor/Experimental/EditorInteractiveToolsFramework/Public/ViewportInteractions/ViewportInteractionsBehaviorSource.h
@@ -301,6 +301,13 @@ public:
 		return OnShiftKeyStateChangedDelegate;
 	}
 
+	EDITORINTERACTIVETOOLSFRAMEWORK_API void SetUnsupportedViewportInteractions(const TArray<FName>& InUnsupportedInteractions);
+
+	const TArray<FName>& GetUnsupportedViewportInteractions() const
+	{
+		return UnsupportedInteractions;
+	}
+
 private:
 	/**
 	 * All Input Behaviors from registered Viewport Interactions should be added to this Behavior Set (done with ::AddInteraction() )
@@ -366,4 +373,6 @@ private:
 
 	bool bGizmoDragging = false;
 	bool bCameraHasMoved = false;
+
+	TArray<FName> UnsupportedInteractions;
 };
diff --git a/Engine/Source/Editor/UnrealEd/Private/Tools/EdModeInteractiveToolsContext.cpp b/Engine/Source/Editor/UnrealEd/Private/Tools/EdModeInteractiveToolsContext.cpp
--- a/Engine/Source/Editor/UnrealEd/Private/Tools/EdModeInteractiveToolsContext.cpp
+++ b/Engine/Source/Editor/UnrealEd/Private/Tools/EdModeInteractiveToolsContext.cpp
@@ -1538,9 +1538,12 @@ void UModeManagerInteractiveToolsContext::RegisterViewportInteractions()
 	PRAGMA_DISABLE_DEPRECATION_WARNINGS
 	OnRegisterViewportInteractions().Broadcast();
 	PRAGMA_ENABLE_DEPRECATION_WARNINGS
-	
+
+	TArray<FName> UnsupportedInteractions;
+	OnGetUnsupportedInteractions().Broadcast(UnsupportedInteractions);
+	ViewportInteractionsBehaviorSource->SetUnsupportedViewportInteractions(UnsupportedInteractions);
 	OnBuildViewportInteractions().Broadcast(ViewportInteractionsBehaviorSource);
-	
+
 	ViewportInteractionsBehaviorSource->RegisterBehaviorSources();
 }
 
diff --git a/Engine/Source/Editor/UnrealEd/Public/Tools/EdModeInteractiveToolsContext.h b/Engine/Source/Editor/UnrealEd/Public/Tools/EdModeInteractiveToolsContext.h
--- a/Engine/Source/Editor/UnrealEd/Public/Tools/EdModeInteractiveToolsContext.h
+++ b/Engine/Source/Editor/UnrealEd/Public/Tools/EdModeInteractiveToolsContext.h
@@ -264,6 +264,7 @@ private:
 	bool bIsActive = false;
 };
 
+DECLARE_MULTICAST_DELEGATE_OneParam(FGetUnsupportedInteractions, TArray<FName>& /* Interactions */)
 
 /**
  * UModeManagerInteractiveToolsContext extends UEditorInteractiveToolsContext with various functions for handling 
@@ -284,6 +285,15 @@ public:
 
 	UE_API virtual bool ProcessEditDelete() override;
 
+	/**
+	 * Called when a new interaction is about to be added. Argument is a TArray<FName>& list of interactions which should not be currently
+	 * supported. For example, an Ed Mode could prevent a specific interaction from being added by registering to this delegate, and adding the interaction name to the array.
+	 */
+	FGetUnsupportedInteractions& OnGetUnsupportedInteractions()
+	{
+		return GetUnsupportedInteractionsDelegate;
+	}
+
 	//
 	// Input handling, these functions forward ViewportClient events to the UInputRouter
 	//
@@ -341,6 +351,7 @@ protected:
 private:
 	bool bIsTrackingMouse;
 
+	FGetUnsupportedInteractions GetUnsupportedInteractionsDelegate;
 
 
 
