[UnrealEd]
* Fixed include to make sure FreeImage is not creating typedef that xcode sdk does not like

#Synced-CL 50144691

diff --git a/Engine/Source/Editor/UnrealEd/Private/Factories/TIFFLoader.cpp b/Engine/Source/Editor/UnrealEd/Private/Factories/TIFFLoader.cpp
--- a/Engine/Source/Editor/UnrealEd/Private/Factories/TIFFLoader.cpp
+++ b/Engine/Source/Editor/UnrealEd/Private/Factories/TIFFLoader.cpp
@@ -17,6 +17,10 @@ PRAGMA_DISABLE_DEPRECATION_WARNINGS
 #include "Windows/AllowWindowsPlatformTypes.h"
 #endif // PLATFORM_WINDOWS
 
+#if PLATFORM_MAC
+#include "Mac/MacSystemIncludes.h" // This is needed to make sure FreeImage is not typedeffing BOOL
+#endif
+
 THIRD_PARTY_INCLUDES_START
 #include "FreeImage.h"
 THIRD_PARTY_INCLUDES_END
