[AutoViz] Fix for: "Entering new Geometry Inspection viewmodes destroys editor performance forever"

#Synced-CL 50096329

#Synced-CL 50096329

diff --git a/Engine/Source/Editor/UnrealEd/Private/EditorViewportClient.cpp b/Engine/Source/Editor/UnrealEd/Private/EditorViewportClient.cpp
--- a/Engine/Source/Editor/UnrealEd/Private/EditorViewportClient.cpp
+++ b/Engine/Source/Editor/UnrealEd/Private/EditorViewportClient.cpp
@@ -6462,6 +6462,28 @@ void FEditorViewportClient::SetViewMode(EViewModeIndex InViewModeIndex)
 	Invalidate();
 }
 
+void FEditorViewportClient::SetGeometryInspectionViewMode(EViewModeIndex ViewModeIndex)
+{
+	static const IConsoleVariable* ICNaniteVar = IConsoleManager::Get().FindConsoleVariable(TEXT("r.Nanite"));
+	
+	const bool bNanitesEnabled = (ICNaniteVar == nullptr || ICNaniteVar->GetBool());
+
+	if (bNanitesEnabled)
+	{
+		const FText WarningText(LOCTEXT("GeometryInspectionViewModesDontWorkWithNanite", "Geometry Inspection view modes only display non-Nanite geometry."));
+		UE_LOG(LogEditorModes, Warning, TEXT("%s"), *WarningText.ToString());
+
+		FNotificationInfo Info(WarningText);
+		Info.Image = FAppStyle::GetBrush("Icons.WarningWithColor");
+		Info.ExpireDuration = 5.0f;
+		Info.SubText = LOCTEXT("GeometryInspectionViewModesDontWorkWithNaniteSubText", "View fallback meshes by disabling Nanite in the Details panel (Rendering -> Disallow Nanite).");
+
+		FSlateNotificationManager::Get().AddNotification(Info);
+	}
+
+	SetViewMode(ViewModeIndex);
+}
+
 void FEditorViewportClient::SetZebraViewMode(bool bHorizontalZebraMode)
 {
 	if (UMaterialInstanceDynamic* MaterialInstance = Cast<UMaterialInstanceDynamic>(GEngine->ZebraMaterialInstance))
@@ -6471,7 +6493,7 @@ void FEditorViewportClient::SetZebraViewMode(bool bHorizontalZebraMode)
 		MaterialInstance->SetVectorParameterValue(GeometryInspectionZebraMaterialParameterNames::Orientation, Color);
 	}
 
-	SetViewMode(VMI_Zebra);
+	SetGeometryInspectionViewMode(VMI_Zebra);
 }
 
 bool FEditorViewportClient::IsZebraViewModeEnabled(bool bHorizontal) const
diff --git a/Engine/Source/Editor/UnrealEd/Private/SEditorViewport.cpp b/Engine/Source/Editor/UnrealEd/Private/SEditorViewport.cpp
--- a/Engine/Source/Editor/UnrealEd/Private/SEditorViewport.cpp
+++ b/Engine/Source/Editor/UnrealEd/Private/SEditorViewport.cpp
@@ -711,9 +711,18 @@ void SEditorViewport::BindCommands()
 
 	MAP_VIEWMODEPARAM_ACTION( Commands.TexStreamAccMaterialTextureScaleAll, -1 );
 
-	MAP_VIEWMODE_ACTION( Commands.GeometryInspectionClayMode, VMI_Clay );
-	MAP_VIEWMODE_ACTION( Commands.GeometryInspectionFrontBackFaceMode, VMI_FrontBackFace );
-	MAP_VIEWMODE_ACTION( Commands.GeometryInspectionRandomColorMode, VMI_RandomColor );
+	CommandListRef.MapAction( Commands.GeometryInspectionClayMode, 
+		FExecuteAction::CreateSP( ClientRef, &FEditorViewportClient::SetGeometryInspectionViewMode, VMI_Clay ),
+		FCanExecuteAction(),
+		FIsActionChecked::CreateSP( ClientRef, &FEditorViewportClient::IsViewModeEnabled, VMI_Clay ));
+	CommandListRef.MapAction( Commands.GeometryInspectionFrontBackFaceMode, 
+		FExecuteAction::CreateSP( ClientRef, &FEditorViewportClient::SetGeometryInspectionViewMode, VMI_FrontBackFace ),
+		FCanExecuteAction(),
+		FIsActionChecked::CreateSP( ClientRef, &FEditorViewportClient::IsViewModeEnabled, VMI_FrontBackFace ));
+	CommandListRef.MapAction( Commands.GeometryInspectionRandomColorMode, 
+		FExecuteAction::CreateSP( ClientRef, &FEditorViewportClient::SetGeometryInspectionViewMode, VMI_RandomColor ),
+		FCanExecuteAction(),
+		FIsActionChecked::CreateSP( ClientRef, &FEditorViewportClient::IsViewModeEnabled, VMI_RandomColor ));
 	CommandListRef.MapAction( Commands.GeometryInspectionZebraHorizontalMode,
 		FExecuteAction::CreateSP( ClientRef, &FEditorViewportClient::SetZebraViewMode, true ),
 		FCanExecuteAction(),
diff --git a/Engine/Source/Editor/UnrealEd/Public/EditorViewportClient.h b/Engine/Source/Editor/UnrealEd/Public/EditorViewportClient.h
--- a/Engine/Source/Editor/UnrealEd/Public/EditorViewportClient.h
+++ b/Engine/Source/Editor/UnrealEd/Public/EditorViewportClient.h
@@ -1004,6 +1004,13 @@ public:
 	 */
 	UNREALED_API virtual void SetViewMode(EViewModeIndex InViewModeIndex);
 	
+	/**
+	 * Sets how the viewport is displayed (Geometry Inspection Modes) and displays warning
+	 * 
+	 * @param	InViewModeIndex				View mode to set for the current viewport type
+	 */
+	UNREALED_API virtual void SetGeometryInspectionViewMode(EViewModeIndex InViewModeIndex);
+
 	/**
 	 * Sets the viewport type to Geometry Inspection Zebra mode, horizontal or vertical type
 	 * 
@@ -2240,7 +2247,6 @@ protected:
 	/** A list of all the stat names which are enabled for this viewport */
 	TArray<FString> EnabledStats;
 
-
 	// Used by the custom temporal upscaler plugin
 public:
 	class ICustomTemporalUpscalerData
diff --git a/Engine/Source/Runtime/Engine/Private/ShowFlags.cpp b/Engine/Source/Runtime/Engine/Private/ShowFlags.cpp
--- a/Engine/Source/Runtime/Engine/Private/ShowFlags.cpp
+++ b/Engine/Source/Runtime/Engine/Private/ShowFlags.cpp
@@ -576,10 +576,6 @@ void EngineShowFlagOverride(EShowFlagInitMode ShowFlagInitMode, EViewModeIndex V
 			EngineShowFlags.SetFog(false);
 		}
 
-
-		static TOptional<bool> NanitePreviousValue;
-		static const auto ICNaniteVar = IConsoleManager::Get().FindConsoleVariable(TEXT("r.Nanite"));
-
 		if (ViewModeIndex == VMI_Clay ||
 			ViewModeIndex == VMI_Zebra ||
 			ViewModeIndex == VMI_FrontBackFace ||
@@ -592,21 +588,7 @@ void EngineShowFlagOverride(EShowFlagInitMode ShowFlagInitMode, EViewModeIndex V
 
 			EngineShowFlags.SetAtmosphere(false);
 			EngineShowFlags.SetFog(false);
-
-			if (ICNaniteVar)
-			{
-				if (!NanitePreviousValue.IsSet())
-				{
-					NanitePreviousValue = ICNaniteVar->GetBool();
-				}
-				ICNaniteVar->Set(false);
-			}
 		}
-		else if (ICNaniteVar && NanitePreviousValue.IsSet())
-		{
-			ICNaniteVar->Set(NanitePreviousValue.GetValue());
-		}
-
 
 		if (ViewModeIndex == VMI_ShadowCasters)
 		{
