[CrashFix] UEditorEngine::CreatePIEWorldByDuplication() PIE >= 10 crashes when Net Mode set to Play as Client

[Fix] In Net Mode set to Play as Client, bRequiresExtraListenServer in PlayLevel.cpp creates an additional client making the PIEInstanceID in LazyObjectPtr 10 which is bigger than the array. To fix this, I added 1 to MAX_PIE_INSTANCES inside LazyObjectPtr.
[Refactored] Changed UI sliders values to reflect the clamped values from 0 to MAX_PIE_INSTANCES

#Synced-CL 50045910

diff --git a/Engine/Source/Editor/UnrealEd/Classes/Settings/LevelEditorPlaySettings.h b/Engine/Source/Editor/UnrealEd/Classes/Settings/LevelEditorPlaySettings.h
--- a/Engine/Source/Editor/UnrealEd/Classes/Settings/LevelEditorPlaySettings.h
+++ b/Engine/Source/Editor/UnrealEd/Classes/Settings/LevelEditorPlaySettings.h
@@ -386,11 +386,11 @@ private:
 	bool RunUnderOneProcess;
 
 	/** The number of client windows to open. The first one to open will respect the Play In Editor "Modes" option (PIE, PINW), additional clients respect the RunUnderOneProcess setting. */
-	UPROPERTY(config, EditAnywhere, Category="Multiplayer Options|Client", meta=(ClampMin = "1", UIMin = "1", UIMax = "64"))
+	UPROPERTY(config, EditAnywhere, Category="Multiplayer Options|Client", meta=(ClampMin = "1", UIMin = "1", UIMax = "10"))
 	int32 PlayNumberOfClients;
 
 	/** In multiplayer PIE which client will be the 'primary'. (default = 0, the first client)*/
-	UPROPERTY(config, EditAnywhere, Category = "Multiplayer Options|Client", meta = (ClampMin = "-1", UIMin = "-1", UIMax = "64", ToolTip = "In multiplayer PIE which client will be the 'primary'. Considered most important and given a larger client window, access to unique hardware like a VirtualReality HMD, etc. Intended to help test issues that affect the second, etc client.  0 is the first client. If the setting is >= than the number of clients the last will be primary. -1 will result in no primary.  Note that this is an index only of PIE instance windows, in netmode 'Play as Client' pie instance zero is a windowless dedicated server, so setting 0 here would make the fist pie window the primary which would be PIEInstance 1, rather than 0 as in other netmodes.", DisplayName = "Primary PIE Client Index"))
+	UPROPERTY(config, EditAnywhere, Category = "Multiplayer Options|Client", meta = (ClampMin = "-1", UIMin = "-1", UIMax = "9", ToolTip = "In multiplayer PIE which client will be the 'primary'. Considered most important and given a larger client window, access to unique hardware like a VirtualReality HMD, etc. Intended to help test issues that affect the second, etc client.  0 is the first client. If the setting is >= than the number of clients the last will be primary. -1 will result in no primary.  Note that this is an index only of PIE instance windows, in netmode 'Play as Client' pie instance zero is a windowless dedicated server, so setting 0 here would make the fist pie window the primary which would be PIEInstance 1, rather than 0 as in other netmodes.", DisplayName = "Primary PIE Client Index"))
 	int PrimaryPIEClientIndex = 0;
 
 	/** What port used by the server for simple networking */
diff --git a/Engine/Source/Runtime/CoreUObject/Public/UObject/LazyObjectPtr.h b/Engine/Source/Runtime/CoreUObject/Public/UObject/LazyObjectPtr.h
--- a/Engine/Source/Runtime/CoreUObject/Public/UObject/LazyObjectPtr.h
+++ b/Engine/Source/Runtime/CoreUObject/Public/UObject/LazyObjectPtr.h
@@ -28,7 +28,8 @@ class UObject;
 template <typename T> struct TIsPODType;
 template <typename T> struct TIsWeakPointerType;
 
-#define MAX_PIE_INSTANCES 10
+// Max of 10 PIE instance + 1 Play As Client PIE instance
+#define MAX_PIE_INSTANCES 11
 
 /**
  * Wrapper structure for a GUID that uniquely identifies registered UObjects.
