Reset OldCDOSourceObjects to stop it from adding the same objects thousands of times during a large reinstance

#Synced-CL 49929924

diff --git a/Engine/Source/Editor/UnrealEd/Private/Kismet2/KismetReinstanceUtilities.cpp b/Engine/Source/Editor/UnrealEd/Private/Kismet2/KismetReinstanceUtilities.cpp
--- a/Engine/Source/Editor/UnrealEd/Private/Kismet2/KismetReinstanceUtilities.cpp
+++ b/Engine/Source/Editor/UnrealEd/Private/Kismet2/KismetReinstanceUtilities.cpp
@@ -2942,6 +2942,7 @@ void FBlueprintCompileReinstancer::ReplaceInstancesOfClass_Inner(const TMap<UCla
 		// Always map old to new instances of CDOs along with any owned subobject(s). This allows delegates to be
 		// notified that these instances have been replaced. However, we don't proactively find and replace those
 		// references ourselves unless input parameters have explicitly configured this path to do so (see below).
+		OldCDOSourceObjects.Reset();
 		UE::Private::FReplaceReferenceHelper::IncludeCDO(OldClass, NewClass, OldToNewInstanceMap, OldCDOSourceObjects, InOriginalCDO, Params.OldToNewTemplates);
 		if (bReplaceReferencesToOldCDOs)
 		{
