Add option to show components in prefab editor outliner

#Synced-CL 50204572

diff --git a/Engine/Plugins/Experimental/EditorDataStorageFeatures/Source/TedsTableViewer/Private/Widgets/Composite/STedsCompositeHierarchyViewer.cpp b/Engine/Plugins/Experimental/EditorDataStorageFeatures/Source/TedsTableViewer/Private/Widgets/Composite/STedsCompositeHierarchyViewer.cpp
--- a/Engine/Plugins/Experimental/EditorDataStorageFeatures/Source/TedsTableViewer/Private/Widgets/Composite/STedsCompositeHierarchyViewer.cpp
+++ b/Engine/Plugins/Experimental/EditorDataStorageFeatures/Source/TedsTableViewer/Private/Widgets/Composite/STedsCompositeHierarchyViewer.cpp
@@ -6,11 +6,53 @@
 #include "TedsQueryStackInterfaces.h"
 #include "Widgets/STedsFilterBar.h"
 #include "Widgets/STedsSearchBox.h"
+#include "Widgets/Input/SComboButton.h"
+
+#define LOCTEXT_NAMESPACE "TedsCompositeHierarchyViewer"
 
 namespace UE::Editor::DataStorage
 {
+
+	FTedsCompositeHierarchyViewerShowOption::FTedsCompositeHierarchyViewerShowOption(const FText& InShowSettingTitle, const FText& InShowSettingTitleToolTip, bool bInActive)
+		: ShowOptionTitle(InShowSettingTitle)
+		, ShowOptionToolTip(InShowSettingTitleToolTip)
+		, bActive(bInActive)
+	{
+	}
+
+	bool FTedsCompositeHierarchyViewerShowOption::IsActive() const
+	{
+		return bActive;
+	}
+
+	void FTedsCompositeHierarchyViewerShowOption::ToggleIsActive()
+	{
+		bActive = !bActive;
+		OnToggleEvent.Broadcast(bActive);
+	}
+
+	void FTedsCompositeHierarchyViewerShowOption::AddMenu(FMenuBuilder& InMenuBuilder)
+	{
+		TSharedRef<FTedsCompositeHierarchyViewerShowOption> SharedThis = AsShared();
+
+		InMenuBuilder.AddMenuEntry(
+			ShowOptionTitle,
+			ShowOptionToolTip,
+			FSlateIcon(),
+			FUIAction(
+				FExecuteAction::CreateSP(SharedThis, &FTedsCompositeHierarchyViewerShowOption::ToggleIsActive),
+				FCanExecuteAction(),
+				FIsActionChecked::CreateSP(SharedThis, &FTedsCompositeHierarchyViewerShowOption::IsActive)
+			),
+			NAME_None,
+			EUserInterfaceActionType::ToggleButton
+		);
+	}
+
 	void STedsCompositeHierarchyViewer::Construct(const FArguments& InArgs, TSharedPtr<IHierarchyViewerDataInterface> InHierarchyInterface)
 	{
+		ShowOptions = InArgs._ShowOptions;
+
 		TSharedPtr<SVerticalBox> ContentWidget = SNew(SVerticalBox);
 
 		const bool bFilteringEnabled = InArgs._EnableFiltering;
@@ -101,6 +143,23 @@ namespace UE::Editor::DataStorage
 					FilterBar.ToSharedRef()
 				];
 			}
+
+			// View mode combo button
+			FilterSearchBar->AddSlot()
+			.VAlign(VAlign_Center)
+			.AutoWidth()
+			[
+				SAssignNew(SettingsComboButton, SComboButton)
+					.ComboButtonStyle(FAppStyle::Get(), "SimpleComboButtonWithIcon") // Use the tool bar item style for this button
+					.OnGetMenuContent(this, &STedsCompositeHierarchyViewer::GetSettingsButtonContent)
+					.HasDownArrow(false)
+					.ButtonContent()
+					[
+						SNew(SImage)
+							.ColorAndOpacity(FSlateColor::UseForeground())
+							.Image(FAppStyle::Get().GetBrush("Icons.Settings"))
+					]
+			];
 		}
 
 		HierarchyViewerArgs.AllNodeProvider(QueryStackNode);
@@ -256,4 +315,24 @@ namespace UE::Editor::DataStorage
 	{
 		HierarchyViewer->SetDropHandler(MoveTemp(InDropHandler));
 	}
+
+	TSharedRef<SWidget> STedsCompositeHierarchyViewer::GetSettingsButtonContent()
+	{
+		FMenuBuilder MenuBuilder(true, nullptr);
+
+		MenuBuilder.BeginSection("Show", LOCTEXT("ShowHeading", "Show"));
+
+		for (TSharedPtr<FTedsCompositeHierarchyViewerShowOption>& ShowOption : ShowOptions)
+		{
+			ShowOption->AddMenu(MenuBuilder);
+		}
+
+		MenuBuilder.EndSection();
+
+		return MenuBuilder.MakeWidget();
+	}
+
 } // namespace UE::Editor::DataStorage
+
+
+#undef LOCTEXT_NAMESPACE
diff --git a/Engine/Plugins/Experimental/EditorDataStorageFeatures/Source/TedsTableViewer/Public/Widgets/Composite/STedsCompositeHierarchyViewer.h b/Engine/Plugins/Experimental/EditorDataStorageFeatures/Source/TedsTableViewer/Public/Widgets/Composite/STedsCompositeHierarchyViewer.h
--- a/Engine/Plugins/Experimental/EditorDataStorageFeatures/Source/TedsTableViewer/Public/Widgets/Composite/STedsCompositeHierarchyViewer.h
+++ b/Engine/Plugins/Experimental/EditorDataStorageFeatures/Source/TedsTableViewer/Public/Widgets/Composite/STedsCompositeHierarchyViewer.h
@@ -12,11 +12,41 @@ namespace UE::Editor::DataStorage
 	class STedsFilterBar;
 	class STedsSearchBox;
 }
+class SComboButton;
 
 namespace UE::Editor::DataStorage
 {
 	class FTedsTableViewerColumn;
 
+	/**
+	* A class corresponding to one "show" boolean option (e.g. "Show Components"). An array of these objects can be passed into the 
+	* STedsCompositeHierarchyViewer and it will auto-populate a menu containing toggle buttons for each object.
+	*/
+	class FTedsCompositeHierarchyViewerShowOption : public TSharedFromThis<FTedsCompositeHierarchyViewerShowOption>
+	{
+	public:
+
+		TEDSTABLEVIEWER_API FTedsCompositeHierarchyViewerShowOption(const FText& InShowSettingTitle, const FText& InShowSettingTitleToolTip, bool bInActive);
+
+		TEDSTABLEVIEWER_API bool IsActive() const;
+
+		TEDSTABLEVIEWER_API void ToggleIsActive();
+
+		TEDSTABLEVIEWER_API void AddMenu(FMenuBuilder& InMenuBuilder);
+
+		DECLARE_EVENT_OneParam(FTedsCompositeHierarchyViewerShowOption, FOnToggle, bool);
+		TEDSTABLEVIEWER_API FOnToggle& OnToggle() { return OnToggleEvent; }
+
+	private:
+
+		FText ShowOptionTitle;
+		FText ShowOptionToolTip;
+		bool bActive;
+
+		FOnToggle OnToggleEvent;
+	};
+
+
 	/*
 	 * A hierarchy viewer widget can be used to show a visual representation of data in TEDS. This composite widget adds features in a 'default' layout
 	 * such as searching and filtering. The rows to display can be specified using a RowQueryStack, and the columns to display are directly input 
@@ -60,6 +90,9 @@ namespace UE::Editor::DataStorage
 		// Array of class filters to add to the filter bar *OR Filters
 		SLATE_ARGUMENT(TArray<UClass*>, ClassFilters)
 		
+		// Array of boolean option objects used to populate a visibility menu
+		SLATE_ARGUMENT(TArray<TSharedPtr<FTedsCompositeHierarchyViewerShowOption>>, ShowOptions)
+
 		// Whether to use submenus or sections for categories in the filter menu
 		SLATE_ARGUMENT(bool, UseSectionsForCategories)
 
@@ -127,6 +160,14 @@ namespace UE::Editor::DataStorage
 		TEDSTABLEVIEWER_API void SetDropHandler(TUniquePtr<FWidgetDropHandler> InDropHandler);
 		
 	private:
+
+		// Build a menu of Toggle Buttons from the ShowOptions objects
+		TSharedRef<SWidget> GetSettingsButtonContent();
+
+
+		// Option objects used to construct the "Show" menu
+		TArray<TSharedPtr<FTedsCompositeHierarchyViewerShowOption>> ShowOptions;
+
 		// Row Node passed to the SearchBox that receives the searched result
 		TSharedPtr<QueryStack::IRowNode> SearchNode;
 
@@ -138,7 +179,10 @@ namespace UE::Editor::DataStorage
 
 		// The filter bar widget (uses MakeAddFilterButton to create the menu dropdown)
 		TSharedPtr<STedsFilterBar> FilterBar;
-		
+
+		// Button used to bring up the settings menu
+		TSharedPtr<SComboButton> SettingsComboButton;
+
 		// The actual HierarchyViewer widget that displays the rows
 		TSharedPtr<SHierarchyViewer> HierarchyViewer;
 	};
diff --git a/Engine/Source/Editor/UnrealEd/Private/Elements/Component/ComponentElementEditorSelectionInterface.cpp b/Engine/Source/Editor/UnrealEd/Private/Elements/Component/ComponentElementEditorSelectionInterface.cpp
--- a/Engine/Source/Editor/UnrealEd/Private/Elements/Component/ComponentElementEditorSelectionInterface.cpp
+++ b/Engine/Source/Editor/UnrealEd/Private/Elements/Component/ComponentElementEditorSelectionInterface.cpp
@@ -2,13 +2,17 @@
 
 #include "Elements/Component/ComponentElementEditorSelectionInterface.h"
 
+#include "DataStorage/Features.h"
 #include "Components/ActorComponent.h"
 #include "Components/ChildActorComponent.h"
 #include "Components/SceneComponent.h"
+#include "Elements/Columns/TypedElementCompatibilityColumns.h"
+#include "Elements/Columns/TypedElementSelectionColumns.h"
 #include "Elements/Component/ComponentElementData.h"
 #include "Elements/Framework/EngineElementsLibrary.h"
 #include "Elements/Framework/TypedElementHandle.h"
 #include "Elements/Framework/TypedElementList.h"
+#include "Elements/Interfaces/TypedElementDataStorageCompatibilityInterface.h"
 #include "Elements/Interfaces/TypedElementSelectionInterface.h"
 #include "Elements/Object/ObjectElementEditorSelectionInterface.h"
 #include "GameFramework/Actor.h"
@@ -56,6 +60,51 @@ bool UComponentElementEditorSelectionInterface::IsElementSelected(const FTypedEl
 	return SelectionSetPtr && Component && IsComponentSelected(Component, SelectionSetPtr.ToSharedRef(), InSelectionOptions);
 }
 
+bool UComponentElementEditorSelectionInterface::SelectElement(const FTypedElementHandle& InElementHandle, const FTypedElementListPtr& InSelectionSet, const FTypedElementSelectionOptions& InSelectionOptions)
+{
+	using namespace UE::Editor::DataStorage;
+
+	const UActorComponent* const Component = ComponentElementDataUtil::GetComponentFromHandle(InElementHandle);
+
+	// Add a selection column in TEDS
+	if (const ICompatibilityProvider* const Compatibility = GetDataStorageFeature<ICompatibilityProvider>(CompatibilityFeatureName))
+	{
+		const RowHandle Row = Compatibility->FindRowWithCompatibleObject(Component);
+		if (Row != InvalidRowHandle)
+		{
+			if (ICoreProvider* const DataStorage = GetMutableDataStorageFeature<ICoreProvider>(StorageFeatureName))
+			{
+				DataStorage->AddColumn<FTypedElementSelectionColumn>(Row, FTypedElementSelectionColumn{ .SelectionSet = InSelectionOptions.GetNameForTEDSIntegration() });
+			}
+		}
+	}
+
+	return Super::SelectElement(InElementHandle, InSelectionSet, InSelectionOptions);
+}
+
+bool UComponentElementEditorSelectionInterface::DeselectElement(const FTypedElementHandle& InElementHandle, const FTypedElementListPtr& InSelectionSet, const FTypedElementSelectionOptions& InSelectionOptions)
+{
+	using namespace UE::Editor::DataStorage;
+
+	const UActorComponent* const Component = ComponentElementDataUtil::GetComponentFromHandle(InElementHandle);
+
+	const ICompatibilityProvider* const Compatibility = GetDataStorageFeature<ICompatibilityProvider>(CompatibilityFeatureName);
+	ICoreProvider* const DataStorage = GetMutableDataStorageFeature<ICoreProvider>(StorageFeatureName);
+
+	if (!Compatibility || !DataStorage)
+	{
+		return Super::DeselectElement(InElementHandle, InSelectionSet, InSelectionOptions);
+	}
+
+	const RowHandle Row = Compatibility->FindRowWithCompatibleObject(Component);
+	if (DataStorage->IsRowAssigned(Row))
+	{
+		DataStorage->RemoveColumn(Row, FTypedElementSelectionColumn::StaticStruct());
+	}
+
+	return Super::DeselectElement(InElementHandle, InSelectionSet, InSelectionOptions);
+}
+
 bool UComponentElementEditorSelectionInterface::ShouldPreventTransactions(const FTypedElementHandle& InElementHandle)
 {
 	const UActorComponent* Component = ComponentElementDataUtil::GetComponentFromHandle(InElementHandle);
diff --git a/Engine/Source/Editor/UnrealEd/Public/Elements/Component/ComponentElementEditorSelectionInterface.h b/Engine/Source/Editor/UnrealEd/Public/Elements/Component/ComponentElementEditorSelectionInterface.h
--- a/Engine/Source/Editor/UnrealEd/Public/Elements/Component/ComponentElementEditorSelectionInterface.h
+++ b/Engine/Source/Editor/UnrealEd/Public/Elements/Component/ComponentElementEditorSelectionInterface.h
@@ -23,6 +23,8 @@ class UComponentElementEditorSelectionInterface : public UComponentElementSelect
 
 public:
 	UNREALED_API virtual bool IsElementSelected(const FTypedElementHandle& InElementHandle, const FTypedElementListConstPtr& InSelectionSet, const FTypedElementIsSelectedOptions& InSelectionOptions) override;
+	UNREALED_API virtual bool SelectElement(const FTypedElementHandle& InElementHandle, const FTypedElementListPtr& InSelectionSet, const FTypedElementSelectionOptions& InSelectionOptions) override;
+	UNREALED_API virtual bool DeselectElement(const FTypedElementHandle& InElementHandle, const FTypedElementListPtr& InSelectionSet, const FTypedElementSelectionOptions& InSelectionOptions) override;
 	UNREALED_API virtual bool ShouldPreventTransactions(const FTypedElementHandle& InElementHandle) override;
 	UNREALED_API virtual TUniquePtr<ITypedElementTransactedElement> CreateTransactedElementImpl() override;
 
