Fixed saving a world package containing UChildActorComponents with an uninitialized physics scene resulting in the UChildActorComponents being dirty after save.

#rnx
#Synced-CL 50036996

diff --git a/Engine/Source/Editor/UnrealEd/Classes/Editor/EditorEngine.h b/Engine/Source/Editor/UnrealEd/Classes/Editor/EditorEngine.h
--- a/Engine/Source/Editor/UnrealEd/Classes/Editor/EditorEngine.h
+++ b/Engine/Source/Editor/UnrealEd/Classes/Editor/EditorEngine.h
@@ -2512,6 +2512,9 @@ public:
 	FProjectExternalContentInterfaceGetter ProjectExternalContentInterfaceGetter;
 
 private:
+	/**  Helper struct to cache the world package dirty state and reset the state for any packages dirtied during the scope. */
+	struct FWorldScopedPackageDirtyGuard;
+
 	UPROPERTY()
 	mutable TObjectPtr<ULevelEditorDragDropHandler> DragDropHandler;
 
diff --git a/Engine/Source/Editor/UnrealEd/Private/EditorEngine.cpp b/Engine/Source/Editor/UnrealEd/Private/EditorEngine.cpp
--- a/Engine/Source/Editor/UnrealEd/Private/EditorEngine.cpp
+++ b/Engine/Source/Editor/UnrealEd/Private/EditorEngine.cpp
@@ -4319,6 +4319,48 @@ void UEditorEngine::OnDismissSourceControlFailureNotification()
 	}
 }
 
+struct UEditorEngine::FWorldScopedPackageDirtyGuard final
+{
+	explicit FWorldScopedPackageDirtyGuard(const UWorld* InWorld)
+		: World(InWorld)
+	{
+		const UPackage* const WorldPackage = World->GetPackage();
+
+		bWorldPackageDirty = WorldPackage->IsDirty();
+
+		for (const UPackage* ExternalPackage : WorldPackage->GetExternalPackages())
+		{
+			if (ExternalPackage->IsDirty())
+			{
+				DirtyExternalPackages.Add(ExternalPackage);
+			}
+		}
+	}
+
+	~FWorldScopedPackageDirtyGuard()
+	{
+		UPackage* const WorldPackage = World->GetPackage();
+
+		for (UPackage* ExternalPackage : WorldPackage->GetExternalPackages())
+		{
+			if (ExternalPackage->IsDirty() && !DirtyExternalPackages.Contains(ExternalPackage))
+			{
+				ExternalPackage->SetDirtyFlag(false);
+			}
+		}
+
+		if (WorldPackage->IsDirty() && !bWorldPackageDirty)
+		{
+			WorldPackage->SetDirtyFlag(false);
+		}
+	}
+
+private:
+	const UWorld* const World;
+	bool bWorldPackageDirty;
+	TSet<const UPackage*> DirtyExternalPackages;
+};
+
 bool UEditorEngine::InitializePhysicsSceneForSaveIfNecessary(UWorld* World, bool &bOutForceInitialized)
 {
 	TRACE_CPUPROFILER_EVENT_SCOPE(UEditorEngine::InitializePhysicsSceneForSaveIfNecessary);
@@ -4342,6 +4384,9 @@ bool UEditorEngine::InitializePhysicsSceneForSaveIfNecessary(UWorld* World, bool
 	// If we didn't find any physics scene we will synthesize one and remove it after save
 	if (!bHasPhysicsScene)
 	{
+		// Capture Dirty packages so we can reset unwanted dirtyness 
+		FWorldScopedPackageDirtyGuard WorldScopedPackageDirtyGuard(World);
+
 		// Clear world components first so that UpdateWorldComponents below properly adds them all to the physics scene
 		World->ClearWorldComponents();
 
@@ -4381,16 +4426,7 @@ bool UEditorEngine::InitializePhysicsSceneForSaveIfNecessary(UWorld* World, bool
 void UEditorEngine::CleanupPhysicsSceneThatWasInitializedForSave(UWorld* World, bool bForceInitialized)
 {
 	// Capture Dirty packages so we can reset unwanted dirtyness 
-	TSet<UPackage*> DirtyPackages;
-	const bool bWorldPackageDirty = World->GetPackage()->IsDirty();
-
-	for (UPackage* ExternalPackage : World->GetPackage()->GetExternalPackages())
-	{
-		if (ExternalPackage->IsDirty())
-		{
-			DirtyPackages.Add(ExternalPackage);
-		}
-	}
+	FWorldScopedPackageDirtyGuard WorldScopedPackageDirtyGuard(World);
 
 	// Make sure we clean up the physics scene here. If we leave too many scenes in memory, undefined behavior occurs when locking a scene for read/write.
 	World->ClearWorldComponents();
@@ -4407,20 +4443,6 @@ void UEditorEngine::CleanupPhysicsSceneThatWasInitializedForSave(UWorld* World,
 		// Update components again in case it was a world without a physics scene but did have rendered components.
 		World->UpdateWorldComponents(true, true);
 	}
-
-	for (UPackage* ExternalPackage : World->GetPackage()->GetExternalPackages())
-
-	{
-		if (!DirtyPackages.Contains(ExternalPackage) && ExternalPackage->IsDirty())
-		{
-			ExternalPackage->SetDirtyFlag(false);
-		}
-	}
-
-	if (!bWorldPackageDirty && World->GetPackage()->IsDirty())
-	{
-		World->GetPackage()->SetDirtyFlag(false);
-	}
 }
 
 UE_TRACE_EVENT_BEGIN(Cpu, EditorEngine_Save, NoSync)
