Fix naked dereference of optional PIE SessionInfo

Cannot be SIE there is no PIE session.

#Synced-CL 50182883

diff --git a/Engine/Source/Editor/UnrealEd/Private/PlayLevel.cpp b/Engine/Source/Editor/UnrealEd/Private/PlayLevel.cpp
--- a/Engine/Source/Editor/UnrealEd/Private/PlayLevel.cpp
+++ b/Engine/Source/Editor/UnrealEd/Private/PlayLevel.cpp
@@ -806,7 +806,7 @@ void UEditorEngine::TeardownPlaySession(FWorldContext& PieWorldContext)
 	GIsPlayInEditorWorld = true;
 	
 	// Remember Simulating flag so that we know if OnSimulateSessionFinished is required after everything has been cleaned up. 
-	bool bWasSimulatingInEditor = PlayInEditorSessionInfo->OriginalRequestParams.WorldType == EPlaySessionWorldType::SimulateInEditor;
+	bool bWasSimulatingInEditor = PlayInEditorSessionInfo.IsSet() ? PlayInEditorSessionInfo->OriginalRequestParams.WorldType == EPlaySessionWorldType::SimulateInEditor : false;
 	
 	// Stop all audio and remove references to temp level.
 	if (FAudioDevice* AudioDevice = PlayWorld->GetAudioDeviceRaw())
