Blueprints: Protect ptr from potential null deref

#rnx

#Synced-CL 50220239

diff --git a/Engine/Source/Editor/UnrealEd/Private/Kismet2/KismetReinstanceUtilities.cpp b/Engine/Source/Editor/UnrealEd/Private/Kismet2/KismetReinstanceUtilities.cpp
--- a/Engine/Source/Editor/UnrealEd/Private/Kismet2/KismetReinstanceUtilities.cpp
+++ b/Engine/Source/Editor/UnrealEd/Private/Kismet2/KismetReinstanceUtilities.cpp
@@ -2155,7 +2155,10 @@ UClass* FBlueprintCompileReinstancer::MoveCDOToNewClass(UClass* OwnerClass, cons
 	CopyOfOwnerClass->StaticLink(true);
 
 	// Mark this class as being an internal class to ease filtering when enumerating all BP classes. 
-	BPGDuplicatedClass->bIsInternalClass = true;
+	if (BPGDuplicatedClass)
+	{
+		BPGDuplicatedClass->bIsInternalClass = true;
+	}
 	
 	// make sure we've bound to our native sparse data - this should have been done in
 	// link but don't want to destabilize early adopters:
