Fixed crash in UWorldFolders::RebuildListInternal (invalid world).

#Synced-CL 50205858

diff --git a/Engine/Source/Editor/UnrealEd/Private/WorldFolders.cpp b/Engine/Source/Editor/UnrealEd/Private/WorldFolders.cpp
--- a/Engine/Source/Editor/UnrealEd/Private/WorldFolders.cpp
+++ b/Engine/Source/Editor/UnrealEd/Private/WorldFolders.cpp
@@ -92,7 +92,11 @@ void UWorldFolders::RebuildListInternal()
 {
 	TRACE_CPUPROFILER_EVENT_SCOPE(UWorldFolders::RebuildList);
 
-	if (GetWorld()->IsGameWorld())
+	// No need to rebuild if world is not valid or if it's a game world
+	// Because life cycle of UWorldFolders is decoupled from its World, it's 
+	// possible that the world is not valid anymore (see FActorFolders::Housekeeping 
+	// and FActorFolders::AddReferencedObjects).
+	if (!GetWorld() || GetWorld()->IsGameWorld())
 	{
 		return;
 	}
