Guarding against nullptr subobject in GetOwnedSubobjectsRecursiv. It seems wrong that there can even be a nullptr subobject but we shouldn't crash.

#Synced-CL 47552827

diff --git a/Engine/Source/Editor/UnrealEd/Private/Kismet2/KismetReinstanceUtilities.cpp b/Engine/Source/Editor/UnrealEd/Private/Kismet2/KismetReinstanceUtilities.cpp
--- a/Engine/Source/Editor/UnrealEd/Private/Kismet2/KismetReinstanceUtilities.cpp
+++ b/Engine/Source/Editor/UnrealEd/Private/Kismet2/KismetReinstanceUtilities.cpp
@@ -414,7 +414,7 @@ struct FReplaceReferenceHelper
 				const uint8* ValuePtr = Prop->ContainerPtrToValuePtr<uint8>(Container, ArrayIdx);
 				ForEachSubObject(Prop, Container, ValuePtr, [&OutObjects, Root](const UObject* Ref)
 					{
-						if (!OutObjects.Contains(Ref) && Ref->IsIn(Root))
+						if (Ref && !OutObjects.Contains(Ref) && Ref->IsIn(Root))
 						{
 							OutObjects.Add(const_cast<UObject*>(Ref)); // consumer is not const correct
 							GetOwnedSubobjectsRecursive(Ref, OutObjects, Root);
