[KismetResintanceUtilities/ALT]
- UnregisterGCObject in the most derived class before the virtual table is destructed to avoid race conditions with UGCObjectReferencer::AddReferencedObjects ran from another thread

#Synced-CL 50254652

diff --git a/Engine/Source/Editor/UnrealEd/Private/Kismet2/KismetReinstanceUtilities.cpp b/Engine/Source/Editor/UnrealEd/Private/Kismet2/KismetReinstanceUtilities.cpp
--- a/Engine/Source/Editor/UnrealEd/Private/Kismet2/KismetReinstanceUtilities.cpp
+++ b/Engine/Source/Editor/UnrealEd/Private/Kismet2/KismetReinstanceUtilities.cpp
@@ -736,6 +736,8 @@ FBlueprintCompileReinstancer::~FBlueprintCompileReinstancer()
 			CompiledBlueprintsToSave.Empty();		
 		}
 	}
+
+	UnregisterGCObject();
 }
 
 class FReinstanceFinalizer : public TSharedFromThis<FReinstanceFinalizer>
