Unshelved from pending changelist '49797722':

PR #14275: Fix bug where temporary ActorComponent assets downloaded for diff were being enumerated by the component type registry, causing them to be selectable in real blueprints

#Synced-CL 50036996

diff --git a/Engine/Source/Editor/UnrealEd/Private/ComponentTypeRegistry.cpp b/Engine/Source/Editor/UnrealEd/Private/ComponentTypeRegistry.cpp
--- a/Engine/Source/Editor/UnrealEd/Private/ComponentTypeRegistry.cpp
+++ b/Engine/Source/Editor/UnrealEd/Private/ComponentTypeRegistry.cpp
@@ -536,8 +536,9 @@ void FComponentTypeRegistryData::ForceRefreshComponentList()
 
 			const bool bOutOfDateClass = Class->HasAnyClassFlags(CLASS_NewerVersionExists);
 			const bool bBlueprintSkeletonClass = FKismetEditorUtilities::IsClassABlueprintSkeleton(Class);
+			const bool bPackageForDiffing = Class->GetPackage()->HasAnyPackageFlags(PKG_ForDiffing);
 
-			if (!bOutOfDateClass &&	!bBlueprintSkeletonClass)
+			if (!bOutOfDateClass &&	!bBlueprintSkeletonClass && !bPackageForDiffing)
 			{
 				if (FKismetEditorUtilities::IsClassABlueprintSpawnableComponent(Class))
 				{
@@ -631,8 +632,9 @@ void FComponentTypeRegistryData::ForceRefreshComponentList()
 			}
 			if (!ClassPath.IsNull())
 			{
-				if (!InMemoryClassPaths.Contains(ClassPath) && 
-					ActorComponentDerivedClassNames.Contains(ClassPath))
+				if (!InMemoryClassPaths.Contains(ClassPath)
+					&& ActorComponentDerivedClassNames.Contains(ClassPath)
+					&& !BPAsset.HasAnyPackageFlags(PKG_ForDiffing))
 				{
 					AddUnloadedComponentClass(ClassPath, BPAsset);
 				}
@@ -649,7 +651,8 @@ void FComponentTypeRegistryData::ForceRefreshComponentList()
 		for (FAssetData& BPGCAsset : Assets)
 		{
 			FTopLevelAssetPath BPGCAssetClassPathName(BPGCAsset.PackagePath, BPGCAsset.AssetName);
-			if (ActorComponentDerivedClassNames.Contains(BPGCAssetClassPathName) && !InMemoryClassPaths.Contains(BPGCAssetClassPathName))
+			if (ActorComponentDerivedClassNames.Contains(BPGCAssetClassPathName) && !InMemoryClassPaths.Contains(BPGCAssetClassPathName)
+				&& !BPGCAsset.HasAnyPackageFlags(PKG_ForDiffing))
 			{
 				AddUnloadedComponentClass(BPGCAssetClassPathName, BPGCAsset);
 			}
