[KismetResintanceUtilities/ALT]
- RegisterGCObject in the most derived class after the virtual table is fully setup to avoid race conditions with UGCObjectReferencer::AddReferencedObjects ran from another thread

#Synced-CL 50254652

diff --git a/Engine/Source/Editor/UnrealEd/Private/Kismet2/KismetReinstanceUtilities.cpp b/Engine/Source/Editor/UnrealEd/Private/Kismet2/KismetReinstanceUtilities.cpp
--- a/Engine/Source/Editor/UnrealEd/Private/Kismet2/KismetReinstanceUtilities.cpp
+++ b/Engine/Source/Editor/UnrealEd/Private/Kismet2/KismetReinstanceUtilities.cpp
@@ -495,7 +495,8 @@ UClass* FBlueprintCompileReinstancer::HotReloadedOldClass = nullptr;
 UClass* FBlueprintCompileReinstancer::HotReloadedNewClass = nullptr;
 
 FBlueprintCompileReinstancer::FBlueprintCompileReinstancer(UClass* InClassToReinstance, EBlueprintCompileReinstancerFlags Flags)
-	: ClassToReinstance(InClassToReinstance)
+	: FGCObject(EFlags::RegisterLater) 
+	, ClassToReinstance(InClassToReinstance)
 	, DuplicatedClass(nullptr)
 	, OriginalCDO(nullptr)
 	, OriginalSCD(nullptr)
@@ -506,6 +507,16 @@ FBlueprintCompileReinstancer::FBlueprintCompileReinstancer(UClass* InClassToRein
 	, bIsRootReinstancer(false)
 	, bAllowResaveAtTheEndIfRequested(false)
 {
+	// This is not needed during async loading since GC cannot run and we don't want
+	// to expose any objects we're currently working on to other threads if they try to use
+	// the GC interface to collect objects.
+	if (!IsInActualLoadingThread() && !IsInParallelLoadingThread())
+	{
+		// Call this in the most derived class to avoid pure calls when registering from the base class
+		// and the vtbl isn't yet ready.
+		RegisterGCObject();
+	}
+
 	if( InClassToReinstance != nullptr && InClassToReinstance->GetDefaultObject(false) )
 	{
 		bool bAutoInferSaveOnCompile = !!(Flags & EBlueprintCompileReinstancerFlags::AutoInferSaveOnCompile);
