World Partition - HLOD: Ensure distributed HLOD builds can work with UE projects outside of the engine root.

#Synced-CL 50353017

diff --git a/Engine/Build/Graph/Examples/BuildWorldPartitionHLODs.xml b/Engine/Build/Graph/Examples/BuildWorldPartitionHLODs.xml
--- a/Engine/Build/Graph/Examples/BuildWorldPartitionHLODs.xml
+++ b/Engine/Build/Graph/Examples/BuildWorldPartitionHLODs.xml
@@ -3,6 +3,7 @@
 
 	<!-- Options -->
 	<Option Name="ProjectName" DefaultValue="" Description="Name of the project."/>
+	<Option Name="ProjectPath" DefaultValue="$(RootDir)/$(ProjectName)" Description="Path to the project directory."/>
 	<Option Name="EditorTarget" DefaultValue="" Description="Name of the editor target to use."/>
 	<Option Name="MapName" DefaultValue="" Description="The map (or collection) for which HLODs will be generated."/>
 	<Option Name="BuilderCount" DefaultValue="10" Description="Number of machines to use"/>
@@ -55,7 +56,7 @@
 	<Agent Name="HLOD Setup" Type="Win64">
 		<Node Name="HLOD Setup" Requires="$(HLODSetupDependencies)">
 			<Command Name="WorldPartitionBuilder" Arguments="-Project=$(ProjectName) -Builder=$(HLODBuilder) -CommandletArgs=&quot;$(HLODCommonBuilderArgs) $(HLODBuilderSCC_None) -SetupHLODs -BuilderCount=$(BuilderCount)&quot;"/>
-			<GatherBuildProductsFromFile BuildProductsFile="$(RootDir)/HLODTemp/BuildProducts.txt"/>
+			<GatherBuildProductsFromFile BuildProductsFile="$(ProjectPath)/Intermediate/HLODTemp/BuildProducts.txt" BaseDirectory="$(ProjectPath)/Intermediate/HLODTemp"/>
 		</Node>
 	</Agent>
 
@@ -76,7 +77,7 @@
 					<Command Name="WorldPartitionBuilder" Arguments="-Project=$(ProjectName) -Builder=$(HLODBuilder) -CommandletArgs=&quot;$(HLODCommonBuilderArgs) $(HLODBuilderSCC_None) -AllowCommandletRendering -BuildHLODs -BuilderIdx=$(HLODBuildJobIndex)&quot;"/>
 				</Node>
 				<Node Name="$(HLODBuildCompleteNode)" After="$(HLODBuildNode)">
-					<GatherBuildProductsFromFile BuildProductsFile="$(RootDir)/HLODTemp/BuildProducts.txt"/>
+					<GatherBuildProductsFromFile BuildProductsFile="$(ProjectPath)/Intermediate/HLODTemp/BuildProducts.txt" BaseDirectory="$(ProjectPath)/Intermediate/HLODTemp"/>
 				</Node>
 			</Agent>
 		</Do>
diff --git a/Engine/Source/Editor/UnrealEd/Private/WorldPartition/WorldPartitionHLODsBuilder.cpp b/Engine/Source/Editor/UnrealEd/Private/WorldPartition/WorldPartitionHLODsBuilder.cpp
--- a/Engine/Source/Editor/UnrealEd/Private/WorldPartition/WorldPartitionHLODsBuilder.cpp
+++ b/Engine/Source/Editor/UnrealEd/Private/WorldPartition/WorldPartitionHLODsBuilder.cpp
@@ -67,6 +67,7 @@ static const FString BuildProductsFileName = TEXT("BuildProducts.txt");
 
 FString GetHLODBuilderFolderName(uint32 BuilderIndex) { return FString::Printf(TEXT("HLODBuilder%d"), BuilderIndex); }
 FString GetToSubmitFolderName() { return TEXT("ToSubmit"); }
+FString GetDistributedBuildBaseDir() { return FPaths::ProjectIntermediateDir() / DistributedBuildWorkingDirName / TEXT(""); }
 
 void CacheUnsavedHLODActors(UWorld* InWorld, TMap<FGuid, TWeakObjectPtr<AWorldPartitionHLOD> >& OutUnsavedHLODActors)
 {
@@ -316,7 +317,7 @@ bool UWorldPartitionHLODsBuilder::ValidateParams() const
 FString GetDistributedBuildWorkingDir(UWorld* InWorld)
 {
 	uint32 WorldPackageHash = GetTypeHash(InWorld->GetPackage()->GetFullName());
-	return FString::Printf(TEXT("%s/%s/%08x"), *FPaths::RootDir(), *DistributedBuildWorkingDirName, WorldPackageHash);
+	return FString::Printf(TEXT("%s%08x"), *GetDistributedBuildBaseDir(), WorldPackageHash);
 }
 
 bool UWorldPartitionHLODsBuilder::ShouldProcessWorld(UWorld* InWorld) const
@@ -669,7 +670,10 @@ bool UWorldPartitionHLODsBuilder::SetupHLODActors()
 			}
 
 			// The build manifest must also be included as a build product to be available in the next steps
-			BuildProducts.Add(BuildManifest);
+			// Store as relative path for portability
+			FString RelativeBuildManifest = BuildManifest;
+			FPaths::MakePathRelativeTo(RelativeBuildManifest, *GetDistributedBuildBaseDir());
+			BuildProducts.Add(RelativeBuildManifest);
 
 			// Write build products to a file
 			if (!AddBuildProducts(BuildProducts))
@@ -1528,17 +1532,21 @@ const FName FileAction_Delete(TEXT("Delete"));
 bool UWorldPartitionHLODsBuilder::CopyFilesToWorkingDir(const FString& TargetDir, const FBuilderModifiedFiles& Files, const FString& WorkingDir, TArray<FString>& BuildProducts)
 {
 	const FString AbsoluteTargetDir = WorkingDir / TargetDir / TEXT("");
+	const FString BaseDir = GetDistributedBuildBaseDir();
 
 	bool bSuccess = true;
 
 	auto CopyFileToWorkingDir = [&](const FString& SourceFilename, const FName FileAction)
 	{
-		FString SourceFilenameRelativeToRoot = SourceFilename;
-		FPaths::MakePathRelativeTo(SourceFilenameRelativeToRoot, *FPaths::RootDir());
+		FString SourceFilenameRelativeToProject = SourceFilename;
+		FPaths::MakePathRelativeTo(SourceFilenameRelativeToProject, *FPaths::ProjectDir());
 
-		FString TargetFilename = AbsoluteTargetDir / FileAction.ToString() / SourceFilenameRelativeToRoot;
+		FString TargetFilename = AbsoluteTargetDir / FileAction.ToString() / SourceFilenameRelativeToProject;
 
-		BuildProducts.Add(TargetFilename);
+		// Store path relative to HLODTemp base directory for portability
+		FString RelativeTargetFilename = TargetFilename;
+		FPaths::MakePathRelativeTo(RelativeTargetFilename, *BaseDir);
+		BuildProducts.Add(RelativeTargetFilename);
 
 		if (FileAction != FileAction_Delete)
 		{
@@ -1639,24 +1647,24 @@ bool UWorldPartitionHLODsBuilder::CopyFilesFromWorkingDir(const FString& SourceD
 		}
 
 		FPaths::MakePathRelativeTo(PathRelativeToWorkingDir, *(FileActionString / TEXT("")));
-		FString FullPathInRootDirectory =  FPaths::RootDir() / PathRelativeToWorkingDir;
+		FString FullPathInProjectDirectory = FPaths::ProjectDir() / PathRelativeToWorkingDir;
 
 		FName FileAction(FileActionString);
 		if (FileAction == FileAction_Add)
 		{
-			FilesToAdd.Add(*FullPathInRootDirectory, File);
+			FilesToAdd.Add(FullPathInProjectDirectory, File);
 		}
 		else if (FileAction == FileAction_Edit)
 		{
-			FilesToEdit.Add(FullPathInRootDirectory, File);
+			FilesToEdit.Add(FullPathInProjectDirectory, File);
 		}
 		else if (FileAction == FileAction_Delete)
 		{
-			FilesToDelete.Add(*FullPathInRootDirectory);
+			FilesToDelete.Add(FullPathInProjectDirectory);
 		}
 		else
 		{
-			UE_LOG(LogWorldPartitionHLODsBuilder, Error, TEXT("Unsupported file action %s for file %s"), *FileActionString, *FullPathInRootDirectory);
+			UE_LOG(LogWorldPartitionHLODsBuilder, Error, TEXT("Unsupported file action %s for file %s"), *FileActionString, *FullPathInProjectDirectory);
 		}
 	}
 
@@ -1752,8 +1760,8 @@ bool UWorldPartitionHLODsBuilder::CopyFilesFromWorkingDir(const FString& SourceD
 
 bool UWorldPartitionHLODsBuilder::AddBuildProducts(const TArray<FString>& BuildProducts) const
 {
-	// Write build products to a file
-	FString BuildProductsFile = FString::Printf(TEXT("%s/%s/%s"), *FPaths::RootDir(), *DistributedBuildWorkingDirName, *BuildProductsFileName);
+	// Write build products to a file in the project's Intermediate directory
+	FString BuildProductsFile = FString::Printf(TEXT("%s%s"), *GetDistributedBuildBaseDir(), *BuildProductsFileName);
 	bool bRet = FFileHelper::SaveStringArrayToFile(BuildProducts, *BuildProductsFile, FFileHelper::EEncodingOptions::AutoDetect, &IFileManager::Get(), EFileWrite::FILEWRITE_Append);
 	if (!bRet)
 	{
diff --git a/Engine/Source/Programs/AutomationTool/BuildGraph/Tasks/GatherBuildProductsFromFileTask.cs b/Engine/Source/Programs/AutomationTool/BuildGraph/Tasks/GatherBuildProductsFromFileTask.cs
--- a/Engine/Source/Programs/AutomationTool/BuildGraph/Tasks/GatherBuildProductsFromFileTask.cs
+++ b/Engine/Source/Programs/AutomationTool/BuildGraph/Tasks/GatherBuildProductsFromFileTask.cs
@@ -16,10 +16,16 @@ namespace AutomationTool.Tasks
 	public class GatherBuildProductsFromFileTaskParameters
 	{
 		/// <summary>
-		/// 
+		/// Path to the file containing the list of build products
 		/// </summary>
 		[TaskParameter]
 		public string BuildProductsFile { get; set; }
+
+		/// <summary>
+		/// Optional base directory to apply to entries in the BuildProducts file
+		/// </summary>
+		[TaskParameter(Optional = true)]
+		public string BaseDirectory { get; set; }
 	}
 
 	[TaskElement("GatherBuildProductsFromFile", typeof(GatherBuildProductsFromFileTaskParameters))]
@@ -39,8 +45,13 @@ namespace AutomationTool.Tasks
 				string[] fileBuildProducts = await File.ReadAllLinesAsync(_parameters.BuildProductsFile);
 				foreach (string buildProduct in fileBuildProducts)
 				{
-					Logger.LogInformation("Adding file to build products: {BuildProduct}", buildProduct);
-					buildProducts.Add(new FileReference(buildProduct));
+					string fullPath = buildProduct;
+					if (!String.IsNullOrEmpty(_parameters.BaseDirectory) && !Path.IsPathRooted(buildProduct))
+					{
+						fullPath = Path.Combine(_parameters.BaseDirectory, buildProduct);
+					}
+					Logger.LogInformation("Adding file to build products: {BuildProduct}", fullPath);
+					buildProducts.Add(new FileReference(fullPath));
 				}
 			}
 			catch (Exception ex)
