Fixed incomplete application of 'EncryptionKey' tag to encrypted assets, when they are added to keychain via pull-in.

This broke Cosmo mapping of PrimaryAssetId to pak EncryptionKey in cooked builds, leading to incorrect Como URL requests (URL hash for unencrypted asset, when it should be encrypted asset only).

#Synced-CL -1

diff --git a/Engine/Source/Editor/UnrealEd/Private/Commandlets/AssetRegistryGenerator.cpp b/Engine/Source/Editor/UnrealEd/Private/Commandlets/AssetRegistryGenerator.cpp
--- a/Engine/Source/Editor/UnrealEd/Private/Commandlets/AssetRegistryGenerator.cpp
+++ b/Engine/Source/Editor/UnrealEd/Private/Commandlets/AssetRegistryGenerator.cpp
@@ -1081,6 +1081,40 @@ void FAssetRegistryGenerator::InjectEncryptionData(FAssetRegistryState& TargetSt
 	FContentEncryptionConfig EncryptionConfig;
 	AssetManager.GetContentEncryptionConfig(EncryptionConfig);
 
+	auto ApplyEncryptionTag = 
+		[&GuidCache, &AssetManager, &TargetState](FAssetData* AssetData)
+		{
+			FString GuidString;
+			const FAssetData::FChunkArrayView ChunkIDs = AssetData->GetChunkIDs();
+			if (ChunkIDs.Num() > 1)
+			{
+				UE_LOG(LogAssetRegistryGenerator, Warning, TEXT("Encrypted root asset '%s' exists in two chunks. Only secondary assets should be shared between chunks."), *AssetData->GetObjectPathString());
+			}
+			else if (ChunkIDs.Num() == 1)
+			{
+				int32 ChunkID = ChunkIDs[0];
+				FGuid Guid;
+
+				if (GuidCache.Contains(ChunkID))
+				{
+					Guid = GuidCache[ChunkID];
+				}
+				else
+				{
+					Guid = GuidCache.Add(ChunkID, AssetManager.GetChunkEncryptionKeyGuid(ChunkID));
+				}
+
+				if (Guid.IsValid())
+				{
+					FAssetDataTagMap TagsAndValues = AssetData->TagsAndValues.CopyMap();
+					TagsAndValues.Add(UAssetManager::GetEncryptionKeyAssetTagName(), Guid.ToString());
+					FAssetData NewAssetData = FAssetData(AssetData->PackageName, AssetData->PackagePath, AssetData->AssetName, AssetData->AssetClassPath, TagsAndValues, ChunkIDs, AssetData->PackageFlags);
+					NewAssetData.TaggedAssetBundles = AssetData->TaggedAssetBundles;
+					TargetState.UpdateAssetData(AssetData, MoveTemp(NewAssetData));
+				}
+			}
+		};
+
 	for (FContentEncryptionConfig::TGroupMap::ElementType EncryptedAssetSetElement : EncryptionConfig.GetPackageGroupMap())
 	{
 		FName SetName = EncryptedAssetSetElement.Key;
@@ -1089,41 +1123,31 @@ void FAssetRegistryGenerator::InjectEncryptionData(FAssetRegistryState& TargetSt
 		for (FName EncryptedRootPackageName : EncryptedRootAssets)
 		{
 			TargetState.EnumerateMutableAssetsByPackageName(EncryptedRootPackageName,
-				[&GuidCache, &AssetManager, &TargetState](FAssetData* AssetData)
+				[&ApplyEncryptionTag](FAssetData* AssetData)
 				{
-					FString GuidString;
-					const FAssetData::FChunkArrayView ChunkIDs = AssetData->GetChunkIDs();
-					if (ChunkIDs.Num() > 1)
-					{
-						UE_LOG(LogAssetRegistryGenerator, Warning, TEXT("Encrypted root asset '%s' exists in two chunks. Only secondary assets should be shared between chunks."), *AssetData->GetObjectPathString());
-					}
-					else if (ChunkIDs.Num() == 1)
-					{
-						int32 ChunkID = ChunkIDs[0];
-						FGuid Guid;
+					ApplyEncryptionTag(AssetData);
 
-						if (GuidCache.Contains(ChunkID))
-						{
-							Guid = GuidCache[ChunkID];
-						}
-						else
-						{
-							Guid = GuidCache.Add(ChunkID, AssetManager.GetChunkEncryptionKeyGuid(ChunkID));
-						}
-
-						if (Guid.IsValid())
-						{
-							FAssetDataTagMap TagsAndValues = AssetData->TagsAndValues.CopyMap();
-							TagsAndValues.Add(UAssetManager::GetEncryptionKeyAssetTagName(), Guid.ToString());
-							FAssetData NewAssetData = FAssetData(AssetData->PackageName, AssetData->PackagePath, AssetData->AssetName, AssetData->AssetClassPath, TagsAndValues, ChunkIDs, AssetData->PackageFlags);
-							NewAssetData.TaggedAssetBundles = AssetData->TaggedAssetBundles;
-							TargetState.UpdateAssetData(AssetData, MoveTemp(NewAssetData));
-						}
-					}
 					return true; // Keep iterating assets in the package
 				});
 		}
 	}
+
+	// Since the encryption collections won't list all assets going into an encrypted chunk, repeat the search based on chunk id,
+	// and look for any missed encrypted assets. Will not work for fast cook.
+	TargetState.EnumerateAllMutableAssets(
+		[&ApplyEncryptionTag](FAssetData& AssetData)
+		{
+			const FAssetData::FChunkArrayView ChunkIDs = AssetData.GetChunkIDs();
+			const bool bEncrypted = ChunkIDs.ContainsByPredicate([](const int32 ChunkID) { return ChunkID >= 1000 && ChunkID < 2000; });
+			const bool bHasTag = AssetData.TagsAndValues.Contains(UAssetManager::GetEncryptionKeyAssetTagName());
+
+			if (bEncrypted && !bHasTag)
+			{
+				ApplyEncryptionTag(&AssetData);
+			}
+
+			return true;
+		});
 }
 
 bool FAssetRegistryGenerator::SaveManifests(UE::Cook::FCookSandbox& InSandboxFile, int64 InOverrideChunkSize,
