Added a function to check if there are any dirty packages.

#rnx

#Synced-CL 50221278

diff --git a/Engine/Plugins/Developer/PlasticSourceControl/Source/PlasticSourceControl/Private/PackageUtils.cpp b/Engine/Plugins/Developer/PlasticSourceControl/Source/PlasticSourceControl/Private/PackageUtils.cpp
--- a/Engine/Plugins/Developer/PlasticSourceControl/Source/PlasticSourceControl/Private/PackageUtils.cpp
+++ b/Engine/Plugins/Developer/PlasticSourceControl/Source/PlasticSourceControl/Private/PackageUtils.cpp
@@ -59,10 +59,7 @@ bool SaveDirtyPackages()
 	// bSaved can be true if the user selects to not save an asset by un-checking it and clicking "save"
 	if (bSaved)
 	{
-		TArray<UPackage*> DirtyPackages;
-		FEditorFileUtils::GetDirtyWorldPackages(DirtyPackages);
-		FEditorFileUtils::GetDirtyContentPackages(DirtyPackages);
-		bSaved = DirtyPackages.Num() == 0;
+		bSaved = !FEditorFileUtils::HasDirtyPackages();
 	}
 
 	return bSaved;
diff --git a/Engine/Source/Editor/MainFrame/Private/Frame/MainFrameActions.cpp b/Engine/Source/Editor/MainFrame/Private/Frame/MainFrameActions.cpp
--- a/Engine/Source/Editor/MainFrame/Private/Frame/MainFrameActions.cpp
+++ b/Engine/Source/Editor/MainFrame/Private/Frame/MainFrameActions.cpp
@@ -348,8 +348,7 @@ bool FMainFrameActionCallbacks::IsVisibleBranchExplorer()
 void FMainFrameActionCallbacks::CheckOutModifiedFiles()
 {
 	TArray<UPackage*> PackagesToSave;
-	FEditorFileUtils::GetDirtyWorldPackages(PackagesToSave);
-	FEditorFileUtils::GetDirtyContentPackages(PackagesToSave);
+	FEditorFileUtils::GetDirtyPackages(PackagesToSave);
 
 	FEditorFileUtils::FPromptForCheckoutAndSaveParams SaveParams;
 	SaveParams.bCheckDirty = true;
@@ -365,11 +364,7 @@ bool FMainFrameActionCallbacks::CanCheckOutModifiedFiles()
 	if (ISourceControlModule::Get().IsEnabled() &&
 		ISourceControlModule::Get().GetProvider().IsAvailable())
 	{
-		TArray<UPackage*> PackagesToSave;
-		FEditorFileUtils::GetDirtyWorldPackages(PackagesToSave);
-		FEditorFileUtils::GetDirtyContentPackages(PackagesToSave);
-
-		return PackagesToSave.Num() > 0;
+		return FEditorFileUtils::HasDirtyPackages();
 	}
 
 	return false;
diff --git a/Engine/Source/Editor/SourceControlWindows/Private/SourceControlWindows.cpp b/Engine/Source/Editor/SourceControlWindows/Private/SourceControlWindows.cpp
--- a/Engine/Source/Editor/SourceControlWindows/Private/SourceControlWindows.cpp
+++ b/Engine/Source/Editor/SourceControlWindows/Private/SourceControlWindows.cpp
@@ -77,11 +77,7 @@ static bool SaveDirtyPackages(bool bUseDialog)
 	// bSaved can be true if the user selects to not save an asset by unchecking it and clicking "save"
 	if (bSaved)
 	{
-		TArray<UPackage*> DirtyPackages;
-		FEditorFileUtils::GetDirtyWorldPackages(DirtyPackages);
-		FEditorFileUtils::GetDirtyContentPackages(DirtyPackages);
-
-		bSaved = DirtyPackages.Num() == 0;
+		bSaved = !FEditorFileUtils::HasDirtyPackages();
 	}
 
 	return bSaved;
diff --git a/Engine/Source/Editor/StatusBar/Private/SourceControlMenuHelpers.cpp b/Engine/Source/Editor/StatusBar/Private/SourceControlMenuHelpers.cpp
--- a/Engine/Source/Editor/StatusBar/Private/SourceControlMenuHelpers.cpp
+++ b/Engine/Source/Editor/StatusBar/Private/SourceControlMenuHelpers.cpp
@@ -218,11 +218,7 @@ bool FSourceControlCommands::CheckOutModifiedFiles_CanExecute()
 	if (ISourceControlModule::Get().IsEnabled() &&
 		ISourceControlModule::Get().GetProvider().IsAvailable())
 	{
-		TArray<UPackage*> PackagesToSave;
-		FEditorFileUtils::GetDirtyWorldPackages(PackagesToSave);
-		FEditorFileUtils::GetDirtyContentPackages(PackagesToSave);
-
-		return PackagesToSave.Num() > 0;
+		return FEditorFileUtils::HasDirtyPackages();
 	}
 
 	return false;
@@ -236,8 +232,7 @@ bool FSourceControlCommands::CheckOutModifiedFiles_IsVisible()
 void FSourceControlCommands::CheckOutModifiedFiles_Clicked()
 {
 	TArray<UPackage*> PackagesToSave;
-	FEditorFileUtils::GetDirtyWorldPackages(PackagesToSave);
-	FEditorFileUtils::GetDirtyContentPackages(PackagesToSave);
+	FEditorFileUtils::GetDirtyPackages(PackagesToSave);
 
 	FEditorFileUtils::FPromptForCheckoutAndSaveParams SaveParams;
 	SaveParams.bCheckDirty = true;
diff --git a/Engine/Source/Editor/UnrealEd/Private/FileHelpers.cpp b/Engine/Source/Editor/UnrealEd/Private/FileHelpers.cpp
--- a/Engine/Source/Editor/UnrealEd/Private/FileHelpers.cpp
+++ b/Engine/Source/Editor/UnrealEd/Private/FileHelpers.cpp
@@ -5544,7 +5544,40 @@ FString FEditorFileUtils::ExtractPackageName(const FString& ObjectPath)
 	return ObjectPath;
 }
 
+bool FEditorFileUtils::HasDirtyWorldPackages(const FShouldIgnorePackageFunctionRef& ShouldIgnorePackageFunction)
+{
+	bool bHasDirtyWorldPackage = false;
+
+	TFunction<bool(UPackage*)> Callback = [&bHasDirtyWorldPackage, &ShouldIgnorePackageFunction](UPackage* DirtyPackage)
+	{
+		if (!ShouldIgnorePackageFunction(DirtyPackage))
+		{
+			bHasDirtyWorldPackage = true;
+			return false;
+		}
+		return true;
+	};
+
+	EnumerateDirtyWorldPackages(Callback);
+
+	return bHasDirtyWorldPackage;
+}
+
 void FEditorFileUtils::GetDirtyWorldPackages(TArray<UPackage*>& OutDirtyPackages, const FShouldIgnorePackageFunctionRef& ShouldIgnorePackageFunction)
+{
+	TFunction<bool(UPackage*)> Callback = [&OutDirtyPackages, &ShouldIgnorePackageFunction](UPackage* DirtyPackage)
+	{
+		if (!ShouldIgnorePackageFunction(DirtyPackage))
+		{
+			OutDirtyPackages.Add(DirtyPackage);
+		}
+		return true;
+	};
+
+	EnumerateDirtyWorldPackages(Callback);
+}
+
+void FEditorFileUtils::EnumerateDirtyWorldPackages(const FEnumeratePackageFunctionRef& Callback)
 {
 	FAssetToolsModule& AssetToolsModule = FModuleManager::LoadModuleChecked<FAssetToolsModule>("AssetTools");
 	const TSharedRef<FPathPermissionList>& WritableFolderFilter = AssetToolsModule.Get().GetWritableFolderPermissionList();
@@ -5566,10 +5599,13 @@ void FEditorFileUtils::GetDirtyWorldPackages(TArray<UPackage*>& OutDirtyPackages
 			)
 		{
 			bool bDirtyNewWorldPackage = false;
-			if (WorldPackage->IsDirty() && !ShouldIgnorePackageFunction(WorldPackage))
+			if (WorldPackage->IsDirty())
 			{
-				// IF the package is dirty and its not a pie package, add the world package to the list of packages to save
-				OutDirtyPackages.Add(WorldPackage);
+				// If the package is dirty and its not a pie package, add the world package to the list of packages to save
+				if (!Callback(WorldPackage))
+				{
+					return;
+				}
 			}
 
 			// Add the Map built data as well if world is
@@ -5583,7 +5619,7 @@ void FEditorFileUtils::GetDirtyWorldPackages(TArray<UPackage*>& OutDirtyPackages
 					{
 						// Mark built data package dirty if has not been given name yet
 						// Otherwise SaveDirtyPackages will fail to create built data file on disk due to re-entrance guard in PromptForCheckoutAndSave preventing a second pop-up window
-						if (!FPackageName::IsValidLongPackageName(BuiltDataPackage->GetName(), /*bIncludeReadOnlyRoots= */ false))
+						if (!FPackageName::IsValidLongPackageName(FNameBuilder(BuiltDataPackage->GetFName()), /*bIncludeReadOnlyRoots= */ false))
 						{
 							BuiltDataPackage->MarkPackageDirty();
 						}
@@ -5593,9 +5629,9 @@ void FEditorFileUtils::GetDirtyWorldPackages(TArray<UPackage*>& OutDirtyPackages
 					{
 						bDirtyNewWorldPackage = true;
 
-						if (!ShouldIgnorePackageFunction(BuiltDataPackage))
+						if (!Callback(BuiltDataPackage))
 						{
-							OutDirtyPackages.Add(BuiltDataPackage);
+							return;
 						}
 					}
 				}
@@ -5629,28 +5665,30 @@ void FEditorFileUtils::GetDirtyWorldPackages(TArray<UPackage*>& OutDirtyPackages
 					{
 						bDirtyNewWorldPackage = true;
 
-						if (!ShouldIgnorePackageFunction(ExternalPackage))
+						// Filter out Actors that might be unsaved (/Temp folder)
+						if (!FPackageName::IsValidLongPackageName(FNameBuilder(ExternalPackage->GetFName())))
 						{
-							bool bActorPackageNeedsToSave = true;
+							continue;
+						}
 
-							// Skip unsaved empty actor packages, unless we've been asked to include them
-							if (ExternalPackage->HasAnyPackageFlags(PKG_NewlyCreated) && UPackage::IsEmptyPackage(ExternalPackage))
+						// Skip unsaved empty actor packages, unless we've been asked to include them
+						if (ExternalPackage->HasAnyPackageFlags(PKG_NewlyCreated) && UPackage::IsEmptyPackage(ExternalPackage))
+						{
+							if (!AdditionalSaveCandidates)
 							{
-								if (!AdditionalSaveCandidates)
-								{
-									AdditionalSaveCandidates.Emplace();
-									UE::FileHelpers::Internal::GetAdditionalInteractiveSavePackageCandidates.Broadcast(AdditionalSaveCandidates.GetValue());
-								}
-								bActorPackageNeedsToSave &= AdditionalSaveCandidates->Contains(ExternalPackage);
+								AdditionalSaveCandidates.Emplace();
+								UE::FileHelpers::Internal::GetAdditionalInteractiveSavePackageCandidates.Broadcast(AdditionalSaveCandidates.GetValue());
 							}
-
-							// Filter out Actors that might be unsaved (/Temp folder)
-							bActorPackageNeedsToSave &= FPackageName::IsValidLongPackageName(ExternalPackage->GetName());
-							if (bActorPackageNeedsToSave)
+							if (!AdditionalSaveCandidates->Contains(ExternalPackage))
 							{
-								OutDirtyPackages.Add(ExternalPackage);
+								continue;
 							}
 						}
+
+						if (!Callback(ExternalPackage))
+						{
+							return;
+						}
 					}
 				}
 			}
@@ -5659,16 +5697,49 @@ void FEditorFileUtils::GetDirtyWorldPackages(TArray<UPackage*>& OutDirtyPackages
 			{
 				// If world package does not have a name yet add the world package so a user is prompted to have a name chosen
 				WorldPackage->MarkPackageDirty();
-				if (!ShouldIgnorePackageFunction(WorldPackage))
+				if (!Callback(WorldPackage))
 				{
-					OutDirtyPackages.Add(WorldPackage);
+					return;
 				}
 			}
 		}
 	}
 }
 
+bool FEditorFileUtils::HasDirtyContentPackages(const FShouldIgnorePackageFunctionRef& ShouldIgnorePackageFunction)
+{
+	bool bHasDirtyContentPackage = false;
+
+	TFunction<bool(UPackage*)> Callback = [&bHasDirtyContentPackage, &ShouldIgnorePackageFunction](UPackage* DirtyPackage)
+	{
+		if (!ShouldIgnorePackageFunction(DirtyPackage))
+		{
+			bHasDirtyContentPackage = true;
+			return false;
+		}
+		return true;
+	};
+
+	EnumerateDirtyContentPackages(Callback);
+
+	return bHasDirtyContentPackage;
+}
+
 void FEditorFileUtils::GetDirtyContentPackages(TArray<UPackage*>& OutDirtyPackages, const FShouldIgnorePackageFunctionRef& ShouldIgnorePackageFunction)
+{
+	TFunction<bool(UPackage*)> Callback = [&OutDirtyPackages, &ShouldIgnorePackageFunction](UPackage* DirtyPackage)
+	{
+		if (!ShouldIgnorePackageFunction(DirtyPackage))
+		{
+			OutDirtyPackages.Add(DirtyPackage);
+		}
+		return true;
+	};
+
+	EnumerateDirtyContentPackages(Callback);
+}
+
+void FEditorFileUtils::EnumerateDirtyContentPackages(const FEnumeratePackageFunctionRef& Callback)
 {
 	FAssetToolsModule& AssetToolsModule = FModuleManager::LoadModuleChecked<FAssetToolsModule>("AssetTools");
 	const TSharedRef<FPathPermissionList>& WritableFolderFilter = AssetToolsModule.Get().GetWritableFolderPermissionList();
@@ -5677,59 +5748,77 @@ void FEditorFileUtils::GetDirtyContentPackages(TArray<UPackage*>& OutDirtyPackag
 	// Make a list of all content packages that we should save
 	for (TObjectIterator<UPackage> It; It; ++It)
 	{
-		UPackage*	Package = *It;
-		bool		bShouldIgnorePackage = false;
+		UPackage* const Package = *It;
 
 		// Only look at root packages.
-		bShouldIgnorePackage |= Package->GetOuter() != NULL;
+		if (Package->GetOuter() != nullptr)
+		{
+			continue;
+		}
 		// Don't try to save "Transient" package.
-		bShouldIgnorePackage |= Package == GetTransientPackage();
+		if (Package == GetTransientPackage())
+		{
+			continue;
+		}
 		// Don't try to save packages with the RF_Transient flag
-		bShouldIgnorePackage |= Package->HasAnyFlags(RF_Transient);
+		if (Package->HasAnyFlags(RF_Transient))
+		{
+			continue;
+		}
 		// Ignore UHT/Verse generated packages, PIE packages, or packages containing map data
-		bShouldIgnorePackage |= Package->HasAnyPackageFlags(PKG_CompiledIn | PKG_PlayInEditor | PKG_ContainsMapData);
+		if (Package->HasAnyPackageFlags(PKG_CompiledIn | PKG_PlayInEditor | PKG_ContainsMapData))
+		{
+			continue;
+		}
 		// Ignore packages that haven't been modified.
-		bShouldIgnorePackage |= !Package->IsDirty();
-
-		if (!bShouldIgnorePackage)
+		if (!Package->IsDirty())
 		{
-			UObject* Asset = Package->FindAssetInPackage();
-			const bool bIsMapPackage = Cast<UWorld>(Asset) != nullptr;
-			const bool bIsExternalMapObject = Asset && Asset->GetTypedOuter<UWorld>() != nullptr;
+			continue;
+		}
 
+		const UObject* const Asset = Package->FindAssetInPackage();
+		if (Asset != nullptr)
+		{
 			// Ignore map packages, they are caught above.
-			bShouldIgnorePackage |= bIsMapPackage;
+			if (Cast<const UWorld>(Asset) != nullptr)
+			{
+				continue;
+			}
 
 			// Ignore external actors, they are caught alongside maps
-			bShouldIgnorePackage |= bIsExternalMapObject;
-
-			if (!bShouldIgnorePackage)
+			if (Asset->GetTypedOuter<UWorld>() != nullptr)
 			{
-				FString PackageName = Package->GetName();
+				continue;
+			}
+		}
 
-				// Ignore packages with long, invalid names. This culls out packages with paths in read-only roots such as /Temp.
-				bShouldIgnorePackage |= (!FPackageName::IsShortPackageName(Package->GetFName()) && !FPackageName::IsValidLongPackageName(PackageName, /*bIncludeReadOnlyRoots=*/false));
+		const FNameBuilder PackageNameBuilder(Package->GetFName());
+		const FStringView PackageName = PackageNameBuilder.ToView();
 
-				// Ignore packages that cannot be saved due to a custom filter
-				if (!bShouldIgnorePackage && bHasWritableFolderFilter)
-				{
-					bShouldIgnorePackage |= (!WritableFolderFilter->PassesStartsWithFilter(PackageName));
-				}
-			}
+		// Ignore packages with long, invalid names. This culls out packages with paths in read-only roots such as /Temp.
+		if (!FPackageName::IsShortPackageName(PackageName) && !FPackageName::IsValidLongPackageName(PackageName, /*bIncludeReadOnlyRoots=*/false))
+		{
+			continue;
 		}
 
-		if (!bShouldIgnorePackage)
+		// Ignore packages that cannot be saved due to a custom filter
+		if (bHasWritableFolderFilter && !WritableFolderFilter->PassesStartsWithFilter(PackageName))
 		{
-			bShouldIgnorePackage |= ShouldIgnorePackageFunction(Package);
+			continue;
 		}
 
-		if (!bShouldIgnorePackage)
+		if (!Callback(Package))
 		{
-			OutDirtyPackages.Add(Package);
+			return;
 		}
 	}
 }
 
+bool FEditorFileUtils::HasDirtyPackages(const FEditorFileUtils::FShouldIgnorePackageFunctionRef& ShouldIgnorePackageFunction)
+{
+	return HasDirtyWorldPackages(ShouldIgnorePackageFunction) || HasDirtyContentPackages(ShouldIgnorePackageFunction);
+}
+
 void FEditorFileUtils::GetDirtyPackages(TArray<UPackage*>& OutDirtyPackages, const FEditorFileUtils::FShouldIgnorePackageFunctionRef& ShouldIgnorePackageFunction)
 {
 	GetDirtyWorldPackages(OutDirtyPackages, ShouldIgnorePackageFunction);
diff --git a/Engine/Source/Editor/UnrealEd/Public/FileHelpers.h b/Engine/Source/Editor/UnrealEd/Public/FileHelpers.h
--- a/Engine/Source/Editor/UnrealEd/Public/FileHelpers.h
+++ b/Engine/Source/Editor/UnrealEd/Public/FileHelpers.h
@@ -412,25 +412,48 @@ public:
 	*/
 	UNREALED_API static bool SaveDirtyContentPackages(TArray<UClass*>& SaveContentClasses, const bool bPromptUserToSave, const bool bFastSave = false, const bool bNotifyNoPackagesSaved = false, const bool bCanBeDeclined = true);
 
+	/**
+	 * Returns true if there currently are any dirty world packages.
+	 *
+	 * @param ShouldIgnorePackageFunction Allows filtering out some dirty packages.
+	 */
+	UNREALED_API static bool HasDirtyWorldPackages(const FShouldIgnorePackageFunctionRef& ShouldIgnorePackageFunction = FShouldIgnorePackage::Default);
+
 	/**
 	 * Appends array with all currently dirty world packages.
 	 *
 	 * @param OutDirtyPackages Array to append dirty packages to.
+	 * @param ShouldIgnorePackageFunction Allows filtering out some dirty packages.
 	 */
 	UNREALED_API static void GetDirtyWorldPackages(TArray<UPackage*>& OutDirtyPackages, const FShouldIgnorePackageFunctionRef& ShouldIgnorePackageFunction = FShouldIgnorePackage::Default);
 
+	/**
+	 * Returns true if there currently are any dirty content packages.
+	 *
+	 * @param ShouldIgnorePackageFunction Allows filtering out some dirty packages.
+	 */
+	UNREALED_API static bool HasDirtyContentPackages(const FShouldIgnorePackageFunctionRef& ShouldIgnorePackageFunction = FShouldIgnorePackage::Default);
+
 	/**
 	 * Appends array with all currently dirty content packages.
 	 *
 	 * @param OutDirtyPackages Array to append dirty packages to.
+	 * @param ShouldIgnorePackageFunction Allows filtering out some dirty packages.
 	 */
 	UNREALED_API static void GetDirtyContentPackages(TArray<UPackage*>& OutDirtyPackages, const FShouldIgnorePackageFunctionRef& ShouldIgnorePackageFunction = FShouldIgnorePackage::Default);
 
+	/**
+	 * Returns true if there currently are any dirty packages.
+	 *
+	 * @param ShouldIgnorePackageFunction Allows filtering out some dirty packages.
+	 */
+	UNREALED_API static bool HasDirtyPackages(const FShouldIgnorePackageFunctionRef& ShouldIgnorePackageFunction = FShouldIgnorePackage::Default);
+
 	/**
 	 * Appends array with all currently dirty packages
 	 *
 	 * @param OutDirtyPackages Array to append dirty packages to.
-	 * @param FilterFunction Allows filtering out some dirty packages.
+	 * @param ShouldIgnorePackageFunction Allows filtering out some dirty packages.
 	 */
 	UNREALED_API static void GetDirtyPackages(TArray<UPackage*>& OutDirtyPackages, const FShouldIgnorePackageFunctionRef& ShouldIgnorePackageFunction = FShouldIgnorePackage::Default);
 
@@ -688,6 +711,11 @@ public:
 	UNREALED_API static FString GetFilename(UObject* LevelObject);
 
 private:
+	using FEnumeratePackageFunctionRef = TFunctionRef<bool(UPackage*)>;
+
+	static void EnumerateDirtyWorldPackages(const FEnumeratePackageFunctionRef& Callback);
+
+	static void EnumerateDirtyContentPackages(const FEnumeratePackageFunctionRef& Callback);
 
 	/** Private method used to build a list of items requiring checkout */
 	static bool AddCheckoutPackageItems(bool bCheckDirty, const TArray<UPackage*>& PackagesToCheckOut, const TSet<FName>& ReadOnlyPackages, TArray<UPackage*>* OutPackagesNotNeedingCheckout, bool* bOutHavePackageToCheckOut);
