DDB: Improvements to the build worker and local executor

- Optimized build worker startup by disabling the config system.
- Reduced Display verbosity to Log verbosity to avoid log spam in the console.
- Fixed directory deletion in the local executor to rename the directory before deleting it to avoid another process observing a partially-deleted worker directory and assuming it is valid to use.

#Synced-CL 50253651

diff --git a/Engine/Source/Editor/UnrealEd/Private/DerivedDataBuildLocalExecutor.cpp b/Engine/Source/Editor/UnrealEd/Private/DerivedDataBuildLocalExecutor.cpp
--- a/Engine/Source/Editor/UnrealEd/Private/DerivedDataBuildLocalExecutor.cpp
+++ b/Engine/Source/Editor/UnrealEd/Private/DerivedDataBuildLocalExecutor.cpp
@@ -71,8 +71,13 @@ public:
 		// Delete previous local execution state from the primary process.
 		if (GetMultiprocessId() == 0)
 		{
-			UE_LOGFMT(LogDerivedDataBuildLocalExecutor, Display, "Deleting existing local execution state from '{SandboxRoot}'.", *SandboxRoot);
-			IFileManager::Get().DeleteDirectory(*SandboxRoot, /*bRequireExists*/ false, /*bTree*/ true);
+			TStringBuilder<256> DeleteRoot(InPlace, SandboxRoot, '.', FGuid::NewGuid());
+			if (IFileManager::Get().DirectoryExists(*SandboxRoot))
+			{
+				UE_LOGFMT(LogDerivedDataBuildLocalExecutor, Display, "Deleting existing local execution state from '{SandboxRoot}'.", *SandboxRoot);
+				IFileManager::Get().Move(*DeleteRoot, *SandboxRoot);
+				IFileManager::Get().DeleteDirectory(*DeleteRoot, /*bRequireExists*/ false, /*bTree*/ true);
+			}
 		}
 
 		IModularFeatures::Get().RegisterModularFeature(IBuildWorkerExecutor::FeatureName, this);
diff --git a/Engine/Source/Programs/DerivedDataBuildWorker/DerivedDataBuildWorker.Build.cs b/Engine/Source/Programs/DerivedDataBuildWorker/DerivedDataBuildWorker.Build.cs
--- a/Engine/Source/Programs/DerivedDataBuildWorker/DerivedDataBuildWorker.Build.cs
+++ b/Engine/Source/Programs/DerivedDataBuildWorker/DerivedDataBuildWorker.Build.cs
@@ -21,6 +21,7 @@ public abstract class DerivedDataBuildWorkerTarget : TargetRules
 		bCompileICU 						= false;
 		bCompileWithAccessibilitySupport	= false;
 		bCompileWithPluginSupport			= false;
+		bEnableConfigSystem					= false;
 		bUsesSlate							= false;
 		bUseXGEController					= false;
 		bWithLiveCoding						= false;
diff --git a/Engine/Source/Programs/DerivedDataBuildWorker/Private/DerivedDataBuildWorker.cpp b/Engine/Source/Programs/DerivedDataBuildWorker/Private/DerivedDataBuildWorker.cpp
--- a/Engine/Source/Programs/DerivedDataBuildWorker/Private/DerivedDataBuildWorker.cpp
+++ b/Engine/Source/Programs/DerivedDataBuildWorker/Private/DerivedDataBuildWorker.cpp
@@ -472,7 +472,7 @@ DERIVEDDATABUILDWORKER_API int32 BuildWorkerMain(int32 ArgC, TCHAR* ArgV[], cons
 		return 1;
 	}
 
-	UE_LOG(LogDerivedDataBuildWorker, Display, TEXT("Worker completed in %fms"), FPlatformTime::ToMilliseconds64(FPlatformTime::Cycles64()-WorkerStartTime));
+	UE_LOG(LogDerivedDataBuildWorker, Log, TEXT("Worker completed in %.3fms"), FPlatformTime::ToMilliseconds64(FPlatformTime::Cycles64()-WorkerStartTime));
 
 	return 0;
 }
