[ITF] Fix issues where the unbounded mouse position cache in `UModeManagerInteractiveToolsContext` didn't use the viewport-sourced variables.

#rnx

#Synced-CL 50484185

diff --git a/Engine/Source/Editor/UnrealEd/Private/Tools/EdModeInteractiveToolsContext.cpp b/Engine/Source/Editor/UnrealEd/Private/Tools/EdModeInteractiveToolsContext.cpp
--- a/Engine/Source/Editor/UnrealEd/Private/Tools/EdModeInteractiveToolsContext.cpp
+++ b/Engine/Source/Editor/UnrealEd/Private/Tools/EdModeInteractiveToolsContext.cpp
@@ -1137,11 +1137,7 @@ bool UModeManagerInteractiveToolsContext::MouseEnter(FEditorViewportClient* View
 
 	if (!bIsTrackingMouse)
 	{
-		CurrentMouseState.Mouse.Position2D = FVector2D(x, y);
-		CurrentMouseState.Mouse.WorldRay = GetRayFromMousePos(ViewportClient, Viewport, x, y);
-
-		CurrentMouseState.Mouse.UnboundedPosition2D = CurrentMouseState.Mouse.Position2D;
-		CurrentMouseState.Mouse.UnboundedWorldRay = CurrentMouseState.Mouse.WorldRay;
+		UpdateMousePosition(ViewportClient, Viewport, x, y);
 	}
 
 	return false;
@@ -1154,17 +1150,11 @@ bool UModeManagerInteractiveToolsContext::MouseMove(FEditorViewportClient* Viewp
 	UE_LOG(LogInteractiveToolsContext, Warning, TEXT("HOVER %p"), ViewportClient);
 #endif
 
-	CurrentMouseState.Mouse.Position2D = FVector2D(x, y);
-	CurrentMouseState.Mouse.WorldRay = GetRayFromMousePos(ViewportClient, Viewport, x, y);
+	UpdateMousePosition(ViewportClient, Viewport, x, y);
+
 	FInputDeviceState InputState = CurrentMouseState;
 	InputState.InputDevice = EInputDevices::Mouse;
 
-	if (!bIsTrackingMouse)
-	{
-		CurrentMouseState.Mouse.UnboundedPosition2D = CurrentMouseState.Mouse.Position2D;
-		CurrentMouseState.Mouse.UnboundedWorldRay = CurrentMouseState.Mouse.WorldRay;
-	}
-
 	FModifierKeysState ModifierKeys = FSlateApplication::Get().GetModifierKeys();
 	InputState.SetModifierKeyStates(
 		ModifierKeys.IsShiftDown(), ModifierKeys.IsAltDown(),
@@ -1199,12 +1189,6 @@ bool UModeManagerInteractiveToolsContext::MouseLeave(FEditorViewportClient* View
 	UE_LOG(LogInteractiveToolsContext, Warning, TEXT("MOUSE LEAVE"));
 #endif
 
-	if (!bIsTrackingMouse)
-	{
-		CurrentMouseState.Mouse.UnboundedPosition2D = CurrentMouseState.Mouse.Position2D;
-		CurrentMouseState.Mouse.UnboundedWorldRay = CurrentMouseState.Mouse.WorldRay;
-	}
-
 	return false;
 }
 
@@ -1276,6 +1260,21 @@ bool UModeManagerInteractiveToolsContext::EndTracking(FEditorViewportClient* InV
 	return false;
 }
 
+void UModeManagerInteractiveToolsContext::UpdateMousePosition(FEditorViewportClient* InViewportClient, FViewport* InViewport, int32 X, int32 Y)
+{
+	FVector2D NewPosition(X, Y);
+	
+	CurrentMouseState.Mouse.Position2D = NewPosition;
+	CurrentMouseState.Mouse.WorldRay = GetRayFromMousePos(InViewportClient, InViewport, X, Y);
+	
+	FIntPoint UnboundedMousePos;
+	InViewport->GetUnboundedMousePos(UnboundedMousePos);
+	const FVector2D UnboundedPosition = UnboundedMousePos;
+
+	CurrentMouseState.Mouse.UnboundedPosition2D = UnboundedPosition;
+	CurrentMouseState.Mouse.UnboundedWorldRay = GetRayFromMousePos(InViewportClient, InViewport, UnboundedPosition.X, UnboundedPosition.Y, false);
+}
+
 bool UModeManagerInteractiveToolsContext::GetCursor(EMouseCursor::Type& OutCursor) const
 {	
 	UToolsContextCursorAPI* ToolsContextCursorAPI = ToolManager->GetContextObjectStore()->FindContext<UToolsContextCursorAPI>();
diff --git a/Engine/Source/Editor/UnrealEd/Public/Tools/EdModeInteractiveToolsContext.h b/Engine/Source/Editor/UnrealEd/Public/Tools/EdModeInteractiveToolsContext.h
--- a/Engine/Source/Editor/UnrealEd/Public/Tools/EdModeInteractiveToolsContext.h
+++ b/Engine/Source/Editor/UnrealEd/Public/Tools/EdModeInteractiveToolsContext.h
@@ -336,6 +336,8 @@ protected:
 	UE_API virtual void Shutdown() override;
 		
 	FDeviceButtonState OnGetButtonStateFallback(const FKey& Key) const;
+	
+	void UpdateMousePosition(FEditorViewportClient* InViewportClient, FViewport* InViewport, int32 X, int32 Y);
 
 	/** Input event instance used to keep track of various button states, etc, that we cannot directly query on-demand */
 	FInputDeviceState CurrentMouseState;
