Editor: allow typed elements to provide their own component visualizers, to allow "bridge" component visualizers for non-actors.

#Synced-CL 50538707

diff --git a/Engine/Source/Editor/EditorFramework/Public/Elements/Interfaces/TypedElementEditorVisualizerInterface.h b/Engine/Source/Editor/EditorFramework/Public/Elements/Interfaces/TypedElementEditorVisualizerInterface.h
new file mode 100644
--- /dev/null
+++ b/Engine/Source/Editor/EditorFramework/Public/Elements/Interfaces/TypedElementEditorVisualizerInterface.h
@@ -0,0 +1,53 @@
+// Copyright Epic Games, Inc. All Rights Reserved.
+
+#pragma once
+
+#include "ComponentVisualizer.h" // FCachedComponentVisualizer
+#include "Elements/Framework/TypedElementHandle.h"
+#include "UObject/Interface.h"
+
+#include "TypedElementEditorVisualizerInterface.generated.h"
+
+class UTypedElementSelectionSet;
+
+UINTERFACE(MinimalAPI, meta = (CannotImplementInterfaceInBlueprint))
+class UTypedElementEditorVisualizerInterface : public UInterface
+{
+	GENERATED_BODY()
+};
+
+UE_EXPERIMENTAL(5.8, "ITypedElementEditorVisualizerInterface is currently used to bridge non-actor "
+	"elements to actor component visualizers, but it may be removed as another visualizer system comes online.")
+class ITypedElementEditorVisualizerInterface
+{
+	GENERATED_BODY()
+
+public:
+
+	//~ Mirrors UUnrealEdEngine::FComponentVisualizerForSelection to avoid exposing that struct, as this interface might
+	//~  be temporary.
+	struct FComponentVisualizerForSelection
+	{
+		FCachedComponentVisualizer ComponentVisualizer;
+		TOptional<TFunction<bool(void)>> IsEnabledDelegate;
+	};
+
+	/**
+	 * For given element, appends any additional visualizers to use
+	 */
+	virtual bool GetComponentVisualizersToUse(const FTypedElementHandle& InElementHandle, const UTypedElementSelectionSet* SelectionSet, 
+		TArray<FComponentVisualizerForSelection>& VisualizersOut) const
+	{ 
+		return false; 
+	}
+};
+
+template <>
+struct TTypedElement<ITypedElementEditorVisualizerInterface> : public TTypedElementBase<ITypedElementEditorVisualizerInterface>
+{
+	bool GetComponentVisualizersToUse(const UTypedElementSelectionSet* SelectionSet,
+		TArray<ITypedElementEditorVisualizerInterface::FComponentVisualizerForSelection>& VisualizersOut) const
+	{ 
+		return InterfacePtr->GetComponentVisualizersToUse(*this, SelectionSet, VisualizersOut);
+	}
+};
diff --git a/Engine/Source/Editor/UnrealEd/Private/EditorSelectUtils.cpp b/Engine/Source/Editor/UnrealEd/Private/EditorSelectUtils.cpp
--- a/Engine/Source/Editor/UnrealEd/Private/EditorSelectUtils.cpp
+++ b/Engine/Source/Editor/UnrealEd/Private/EditorSelectUtils.cpp
@@ -33,6 +33,7 @@
 #include "Elements/Framework/TypedElementList.h"
 #include "Elements/Framework/TypedElementSelectionSet.h"
 #include "Elements/Framework/EngineElementsLibrary.h"
+#include "Elements/Interfaces/TypedElementEditorVisualizerInterface.h"
 #include "Elements/Interfaces/TypedElementWorldInterface.h"
 #include "Elements/Interfaces/TypedElementObjectInterface.h"
 
@@ -378,6 +379,31 @@ void UUnrealEdEngine::OnEditorElementSelectionChanged(const UTypedElementSelecti
 			});
 	}
 
+	// Try getting visualizers to use from typed elements. This could be used as the sole path instead
+	//  of doing the above if we choose to implement the visualization interface for actors and components. 
+	//  For now, however, it is used for non-actors that implement their visualizations via bridge/hack actors.
+	LevelEditorSelectionSet->ForEachSelectedElement<ITypedElementEditorVisualizerInterface>(
+		[this, SelectionSet](const TTypedElement<ITypedElementEditorVisualizerInterface>& Element)
+	{
+		TArray<ITypedElementEditorVisualizerInterface::FComponentVisualizerForSelection> Visualizers;
+		if (Element.GetComponentVisualizersToUse(SelectionSet, Visualizers))
+		{
+			for (ITypedElementEditorVisualizerInterface::FComponentVisualizerForSelection& Visualizer : Visualizers)
+			{
+				// We'll avoid adding visualizers that target the same component since we don't currently remove
+				//  any conflicting ones from regular component/actor selection
+				bool bComponentAlreadyTargeted = VisualizersForSelection.ContainsByPredicate([&Visualizer](const FComponentVisualizerForSelection& Existing)
+					{ return Existing.ComponentVisualizer.ComponentPropertyPath == Visualizer.ComponentVisualizer.ComponentPropertyPath; });
+
+				if (!bComponentAlreadyTargeted)
+				{
+					VisualizersForSelection.Add({ Visualizer.ComponentVisualizer, Visualizer.IsEnabledDelegate });
+				}
+			}//end for each returned visualizer
+		}
+		return true; // continue iterating
+	});
+	
 	// Restore the active component visualizer, since an undo/redo may have changed the selection
 	bool bDidSetVizThisTick = false;
 	if (VisualizersForSelection.Num() > 0)
