LevelEditorViewport: Make snap-to-surface support non-actors. Remove PreDragActorTransforms.

#Synced-CL 50541686

diff --git a/Engine/Source/Editor/ComponentVisualizers/Private/SplineComponentVisualizer.cpp b/Engine/Source/Editor/ComponentVisualizers/Private/SplineComponentVisualizer.cpp
--- a/Engine/Source/Editor/ComponentVisualizers/Private/SplineComponentVisualizer.cpp
+++ b/Engine/Source/Editor/ComponentVisualizers/Private/SplineComponentVisualizer.cpp
@@ -1532,7 +1532,7 @@ bool FSplineComponentVisualizer::HandleInputDelta(FEditorViewportClient* Viewpor
 
 	// It's tough to port the actor "surface snapping" toggle behavior here because the original code only works on actors
 	// and involves things like keeping track of last transform, snapping to surface even if moving in a different plane, etc
-	// (see FLevelEditorViewportClient::ProjectActorsIntoWorld). However we can at least support the limited but useful case
+	// (see FLevelEditorViewportClient::ProjectSelectedElementsIntoWorld). However we can at least support the limited but useful case
 	// of trying to drag points around by the ball of the widget.
 	const ULevelEditorViewportSettings* ViewportSettings = GetDefault<ULevelEditorViewportSettings>();
 	const float SnapOffsetExtent = (ViewportSettings->SnapToSurface.bEnabled) ? (ViewportSettings->SnapToSurface.SnapOffsetExtent) : (0.0f);
diff --git a/Engine/Source/Editor/LevelEditor/Private/Elements/Actor/ActorElementLevelEditorViewportInteractionCustomization.cpp b/Engine/Source/Editor/LevelEditor/Private/Elements/Actor/ActorElementLevelEditorViewportInteractionCustomization.cpp
--- a/Engine/Source/Editor/LevelEditor/Private/Elements/Actor/ActorElementLevelEditorViewportInteractionCustomization.cpp
+++ b/Engine/Source/Editor/LevelEditor/Private/Elements/Actor/ActorElementLevelEditorViewportInteractionCustomization.cpp
@@ -101,7 +101,7 @@ void FActorElementLevelEditorViewportInteractionCustomization::ModifyScale(AActo
 
 		const FBox LocalBox = InActor->GetComponentsBoundingBox(true);
 		const FVector ScaledExtents = LocalBox.GetExtent() * CurrentScale;
-		const FTransform PreDragTransform = GetMutableLevelEditorViewportClient()->CachePreDragActorTransform(InActor);
+		const FTransform PreDragTransform = GetMutableLevelEditorViewportClient()->GetPreDragActorTransform(InActor);
 
 		FComponentElementLevelEditorViewportInteractionCustomization::ValidateScale(PreDragTransform.GetScale3D(), InDragAxis, CurrentScale, ScaledExtents, ScaleDelta, bCheckSmallExtent);
 
diff --git a/Engine/Source/Editor/LevelEditor/Private/Elements/Component/ComponentElementLevelEditorViewportInteractionCustomization.cpp b/Engine/Source/Editor/LevelEditor/Private/Elements/Component/ComponentElementLevelEditorViewportInteractionCustomization.cpp
--- a/Engine/Source/Editor/LevelEditor/Private/Elements/Component/ComponentElementLevelEditorViewportInteractionCustomization.cpp
+++ b/Engine/Source/Editor/LevelEditor/Private/Elements/Component/ComponentElementLevelEditorViewportInteractionCustomization.cpp
@@ -94,7 +94,7 @@ void FComponentElementLevelEditorViewportInteractionCustomization::GizmoManipula
 void FComponentElementLevelEditorViewportInteractionCustomization::ModifyScale(USceneComponent* InComponent, const EAxisList::Type InDragAxis, FVector& ScaleDelta)
 {
 	AActor* Actor = InComponent->GetOwner();
-	const FTransform PreDragTransform = GetMutableLevelEditorViewportClient()->CachePreDragActorTransform(Actor);
+	const FTransform PreDragTransform = GetMutableLevelEditorViewportClient()->GetPreDragActorTransform(Actor);
 	const FBox LocalBox = Actor->GetComponentsBoundingBox(true);
 	const FVector ScaledExtents = LocalBox.GetExtent() * InComponent->GetRelativeScale3D();
 
diff --git a/Engine/Source/Editor/UnrealEd/Private/LevelEditorViewport.cpp b/Engine/Source/Editor/UnrealEd/Private/LevelEditorViewport.cpp
--- a/Engine/Source/Editor/UnrealEd/Private/LevelEditorViewport.cpp
+++ b/Engine/Source/Editor/UnrealEd/Private/LevelEditorViewport.cpp
@@ -43,6 +43,7 @@
 #include "EditorSupportDelegates.h"
 #include "EditorGizmos/EditorTransformGizmoUtil.h"
 #include "DataStorage/Features.h"
+#include "Elements/Actor/ActorElementData.h"
 #include "Elements/Framework/TypedElementRegistry.h"
 #include "Elements/Framework/EngineElementsLibrary.h"
 #include "Elements/Framework/TypedElementCommonActions.h"
@@ -2035,7 +2036,7 @@ bool FLevelEditorViewportClient::DropObjectsAtCoordinates(int32 MouseX, int32 Mo
 					[this](const TTypedElement<ITypedElementWorldInterface>& InElement)
 				{
 					FTransform Transform;
-					if (InElement.GetWorldTransform(Transform))
+					if (ensure(InElement.GetWorldTransform(Transform)))
 					{
 						PreDragElementTransforms.Add(InElement, Transform);
 					}
@@ -2053,8 +2054,6 @@ bool FLevelEditorViewportClient::DropObjectsAtCoordinates(int32 MouseX, int32 Mo
 						return true; // true means continue
 					}
 
-					PreDragActorTransforms.Add(NewActor, NewActor->GetTransform());
-
 					NewActor->SetActorEnableCollision(false);
 
 					// This boolean already gets set via UActorFactory::PlaceAsset in code paths that go through
@@ -3128,9 +3127,8 @@ static void GetAttachedActorsRecursive(const AActor* InActor, TArray<AActor*>& O
 	OutActors.Append(AttachedActors);
 };
 
-void FLevelEditorViewportClient::ProjectActorsIntoWorld(const TArray<AActor*>& Actors, FViewport* InViewport, const FVector& Drag, const FRotator& Rot)
+void FLevelEditorViewportClient::ProjectSelectedElementsIntoWorld(FViewport* InViewport, const FVector& Drag, const FRotator& Rot)
 {
-	// Compile an array of selected actors
 	FSceneViewFamilyContext ViewFamily(FSceneViewFamily::ConstructionValues(
 		InViewport, 
 		GetScene(),
@@ -3148,25 +3146,35 @@ void FLevelEditorViewportClient::ProjectActorsIntoWorld(const TArray<AActor*>& A
 	
 	const FVector DeltaTranslation = (ModeTools->PivotLocation - ModeTools->CachedLocation) + Drag;
 
+	FTypedElementListConstRef ElementsToManipulate = GetElementsToManipulate();
+	FCollisionQueryParams CollisionQueryParams;
+	ElementsToManipulate->ForEachElement<ITypedElementWorldInterface>(
+		[&CollisionQueryParams](const TTypedElement<ITypedElementWorldInterface>& InElement)
+	{
+		InElement.AddIgnoredElementToCollisionQueryParams(InElement, CollisionQueryParams, /*bAlsoIgnoreSubElements*/ true);
+		return true; // continue iterating
+	});
 
-	// Loop over all the actors and attempt to snap them along the drag axis normal
-	for (AActor* Actor : Actors)
+	// Loop over all the elements and attempt to snap them along the drag axis normal
+	ElementsToManipulate->ForEachElement<ITypedElementWorldInterface>([&](const TTypedElement<ITypedElementWorldInterface>& InElement) -> bool
 	{
 		// Use the Delta of the Mode tool with the actor pre drag location to avoid accumulating snapping offsets
-		FVector NewActorPosition;
-		if (const FTransform* PreDragTransform = PreDragActorTransforms.Find(Actor))
+		FVector NewPosition;
+		if (const FTransform* PreDragTransform = PreDragElementTransforms.Find(InElement))
 		{
-			NewActorPosition = PreDragTransform->GetLocation() + DeltaTranslation;
+			NewPosition = PreDragTransform->GetLocation() + DeltaTranslation;
 		}
 		else
 		{
-			const FTransform& ActorTransform = PreDragActorTransforms.Add(Actor, Actor->GetTransform());
-			NewActorPosition = ActorTransform.GetLocation() + DeltaTranslation;
+			FTransform ElementTransform;
+			ensure(InElement.GetWorldTransform(ElementTransform));
+			PreDragElementTransforms.Add(InElement, ElementTransform);
+			NewPosition = ElementTransform.GetLocation() + DeltaTranslation;
 		}
 
 		FViewportCursorLocation Cursor(SceneView, this, 0, 0);
 
-		FActorPositionTraceResult TraceResult;
+		UE::Positioning::FObjectPositioningTraceResult TraceResult;
 		bool bSnapped = false;
 
 		bool bIsOnScreen = false;
@@ -3174,7 +3182,7 @@ void FLevelEditorViewportClient::ProjectActorsIntoWorld(const TArray<AActor*>& A
 			// We only snap things that are on screen
 			FVector2D ScreenPos;
 			FIntPoint ViewportSize = InViewport->GetSizeXY();
-			if (SceneView->WorldToPixel(NewActorPosition, ScreenPos) && FMath::IsWithin(ScreenPos.X, 0, ViewportSize.X) && FMath::IsWithin(ScreenPos.Y, 0, ViewportSize.Y))
+			if (SceneView->WorldToPixel(NewPosition, ScreenPos) && FMath::IsWithin(ScreenPos.X, 0, ViewportSize.X) && FMath::IsWithin(ScreenPos.Y, 0, ViewportSize.Y))
 			{
 				bIsOnScreen = true;
 				Cursor = FViewportCursorLocation(SceneView, this, static_cast<int32>(ScreenPos.X), static_cast<int32>(ScreenPos.Y));
@@ -3183,10 +3191,6 @@ void FLevelEditorViewportClient::ProjectActorsIntoWorld(const TArray<AActor*>& A
 
 		if (bIsOnScreen)
 		{
-			TArray<AActor*> IgnoreActors; 
-			IgnoreActors.Append(Actors);  // Add the whole list of actors so you can't hit the moving set with the ray
-			GetAttachedActorsRecursive(Actor, IgnoreActors);
-
 			// Determine how we're going to attempt to project the object onto the world
 			if (CurrentAxis == EAxisList::XY || CurrentAxis == EAxisList::XZ || CurrentAxis == EAxisList::YZ)
 			{
@@ -3200,51 +3204,67 @@ void FLevelEditorViewportClient::ProjectActorsIntoWorld(const TArray<AActor*>& A
 					TraceDirection = -TraceDirection;
 				}
 
-				FVector RayStart	= NewActorPosition - (TraceDirection * HALF_WORLD_MAX/2);
-				FVector RayEnd		= NewActorPosition + (TraceDirection * HALF_WORLD_MAX/2);
+				FVector RayStart	= NewPosition - (TraceDirection * HALF_WORLD_MAX/2);
+				FVector RayEnd		= NewPosition + (TraceDirection * HALF_WORLD_MAX/2);
 
 				TrimLineToFrustum(Frustum, RayStart, RayEnd);
 
-				TraceResult = FActorPositioning::TraceWorldForPosition(*GetWorld(), *SceneView, RayStart, RayEnd, &IgnoreActors);
+				TraceResult = UE::Positioning::TraceWorldForPosition(*GetWorld(), *SceneView, RayStart, RayEnd, &CollisionQueryParams);
 			}
 			else
 			{
-				TraceResult = FActorPositioning::TraceWorldForPosition(Cursor, *SceneView, &IgnoreActors);
+				TraceResult = UE::Positioning::TraceWorldForPosition(Cursor, *SceneView, &CollisionQueryParams);
 			}
 		}
 				
-		if (TraceResult.State == FActorPositionTraceResult::HitSuccess)
+		if (TraceResult.State == UE::Positioning::FObjectPositioningTraceResult::HitSuccess)
 		{
-			// Move the actor to the position of the trace hit using the spawn offset rules
-			// We only do this if we found a valid hit (we don't want to move the actor in front of the camera by default)
+			// Move the element to the position of the trace hit using the spawn offset rules
+			// We only do this if we found a valid hit (we don't want to move the element in front of the camera by default)
 
-			const UActorFactory* Factory = GEditor->FindActorFactoryForActorClass(Actor->GetClass());
+			// Actors have a few snapping options set for them in their actor factory.
+			const UActorFactory* ActorFactory = nullptr;
+			AActor* Actor = ActorElementDataUtil::GetActorFromHandle(InElement);
+			if (Actor)
+			{
+				ActorFactory = GEditor->FindActorFactoryForActorClass(Actor->GetClass());
+			}
 			
-			const FTransform* PreDragActorTransform = PreDragActorTransforms.Find(Actor);
-			check(PreDragActorTransform);
+			const FTransform* PreDragTransform = PreDragElementTransforms.Find(InElement);
+			if (!ensure(PreDragTransform))
+			{
+				return true; // continue to next element
+			}
 
 			// Compute the surface aligned transform. Note we do not use the snapped version here as our DragDelta is already snapped
 
 			const FPositioningData PositioningData = FPositioningData(TraceResult.Location, TraceResult.SurfaceNormal)
-				.UseStartTransform(*PreDragActorTransform)
-				.UsePlacementExtent(Actor->GetPlacementExtent())
-				.UseFactory(Factory);
+				.UseStartTransform(*PreDragTransform)
+				.UsePlacementExtent(Actor ? Actor->GetPlacementExtent() : FVector::ZeroVector)
+				// It is ok for this to be null if we're not dealing with non-actors
+				.UseFactory(ActorFactory);
+			FTransform NewTransform = FActorPositioning::GetSurfaceAlignedTransform(PositioningData);
 
-			FTransform ActorTransform = FActorPositioning::GetSurfaceAlignedTransform(PositioningData);
-			
-			ActorTransform.SetScale3D(Actor->GetActorScale3D());
-			if (USceneComponent* RootComponent = Actor->GetRootComponent())
+			FTransform CurrentTransform;
+			InElement.GetWorldTransform(CurrentTransform);
+
+			NewTransform.SetScale3D(CurrentTransform.GetScale3D());
+
+			InElement.SetWorldTransform(NewTransform);
+			if (Actor)
 			{
-				RootComponent->SetWorldTransform(ActorTransform);
-				Actor->PostEditMove(false);
+				Actor->PostEditMove(/*bFinished*/ false);
 			}
 		}
 		else
 		{
 			// Didn't find a valid surface snapping candidate, just apply the deltas directly
-			ApplyDeltaToActor(Actor, NewActorPosition - Actor->GetActorLocation(), Rot, FVector(0.f, 0.f, 0.f));
+			FTransform CurrentTransform;
+			InElement.GetWorldTransform(CurrentTransform);
+			ApplyDeltaToElement(InElement, FTransform(Rot, NewPosition - CurrentTransform.GetLocation(), FVector(0.f, 0.f, 0.f)));
 		}
-	}
+		return true; // continue iterating
+	});
 }
 
 bool FLevelEditorViewportClient::InputWidgetDelta(FViewport* InViewport, EAxisList::Type InCurrentAxis, FVector& Drag, FRotator& Rot, FVector& Scale)
@@ -3336,16 +3356,7 @@ bool FLevelEditorViewportClient::InputWidgetDelta(FViewport* InViewport, EAxisLi
 					const bool bSingleAxisDrag = CurrentAxis == EAxisList::X || CurrentAxis == EAxisList::Y || CurrentAxis == EAxisList::Z;
 					if (!bSnapped && !bSingleAxisDrag && GetDefault<ULevelEditorViewportSettings>()->SnapToSurface.bEnabled && bOnlyTranslation)
 					{
-						TArray<AActor*> SelectedActors;
-						for (FSelectionIterator It(GEditor->GetSelectedActorIterator()); It; ++It)
-						{
-							if (AActor* Actor = Cast<AActor>(*It))
-							{
-								SelectedActors.Add(Actor);
-							}
-						}
-
-						ProjectActorsIntoWorld(SelectedActors, InViewport, Drag, Rot);
+						ProjectSelectedElementsIntoWorld(InViewport, Drag, Rot);
 					}
 					else
 					{
@@ -3641,7 +3652,6 @@ void FLevelEditorViewportClient::TrackingStarted( const FInputEventState& InInpu
 	bNeedToRestoreComponentBeingMovedFlag = false;
 	bHasBegunGizmoManipulation = false;
 
-	PreDragActorTransforms.Empty();
 	PreDragElementTransforms.Empty();
 
 	// Track BSP changes
@@ -3834,7 +3844,6 @@ void FLevelEditorViewportClient::TrackingStopped()
 		RedrawAllViewportsIntoThisScene();
 	}
 
-	PreDragActorTransforms.Empty();
 	PreDragElementTransforms.Empty();
 }
 
@@ -4594,13 +4603,14 @@ bool FLevelEditorViewportClient::GetFocusBounds(FTypedElementListConstRef InElem
 	return ViewportInteraction->GetFocusBounds(InElements, OutBounds);
 }
 
-FTransform FLevelEditorViewportClient::CachePreDragActorTransform(const AActor* InActor)
+FTransform FLevelEditorViewportClient::GetPreDragActorTransform(const AActor* InActor) const
 {
-	if (const FTransform* PreDragTransform = PreDragActorTransforms.Find(InActor))
+	FTypedElementHandle Handle = UEngineElementsLibrary::AcquireEditorActorElementHandle(InActor, /*bAllowCreate*/ false);
+	if (const FTransform* PreDragTransform = Handle ? PreDragElementTransforms.Find(Handle) : nullptr)
 	{
 		return *PreDragTransform;
 	}
-	return PreDragActorTransforms.Add(InActor, InActor->GetTransform());
+	return InActor->GetTransform();
 }
 
 EMouseCursor::Type FLevelEditorViewportClient::GetCursor(FViewport* InViewport,int32 X,int32 Y)
@@ -5491,7 +5501,7 @@ bool FLevelEditorViewportClient::BeginTransform(const FGizmoState& InState)
 	bNeedToRestoreComponentBeingMovedFlag = false;
 	MouseDeltaTracker->SetExternalMovement(true);
 
-	PreDragActorTransforms.Empty();
+	PreDragElementTransforms.Empty();
 
 	Widget->SetSnapEnabled(true);
 
@@ -5606,7 +5616,7 @@ bool FLevelEditorViewportClient::EndTransform(const FGizmoState& InState)
 		RedrawAllViewportsIntoThisScene();
 	}
 
-	PreDragActorTransforms.Empty();
+	PreDragElementTransforms.Empty();
 	MouseDeltaTracker->SetExternalMovement(false);
 
 	return true;
@@ -5937,7 +5947,7 @@ void FLevelEditorViewportClient::DuplicateSelection(bool bInSelectNewElements)
 			RedrawAllViewportsIntoThisScene();
 
 			// Prevent issues with dragged elements sometimes being set back to their starting position after stamping the duplicate
-			PreDragActorTransforms.Empty();
+			PreDragElementTransforms.Empty();
 		}
 	}
 }
diff --git a/Engine/Source/Editor/UnrealEd/Public/LevelEditorViewport.h b/Engine/Source/Editor/UnrealEd/Public/LevelEditorViewport.h
--- a/Engine/Source/Editor/UnrealEd/Public/LevelEditorViewport.h
+++ b/Engine/Source/Editor/UnrealEd/Public/LevelEditorViewport.h
@@ -753,7 +753,9 @@ protected:
 	UNREALED_API virtual void RedrawAllViewportsIntoThisScene() override;
 
 private:
-	UNREALED_API FTransform CachePreDragActorTransform(const AActor* InActor);
+	// Currently used by a couple friended classes 
+	// TODO: Should probably either rewrite those to not use it, or make some element-based version public
+	UNREALED_API FTransform GetPreDragActorTransform(const AActor* InActor) const;
 
 	/**
 	 * Checks to see the viewports locked actor need updating
@@ -870,8 +872,8 @@ private:
 	 */
 	UNREALED_API bool DropObjectsOnWidget(FSceneView* View, struct FViewportCursorLocation& Cursor, const TArray<UObject*>& DroppedObjects, bool bCreateDropPreview = false);
 
-	/** Project the specified actors into the world according to the current drag parameters */
-	UNREALED_API void ProjectActorsIntoWorld(const TArray<AActor*>& Actors, FViewport* Viewport, const FVector& Drag, const FRotator& Rot);
+	/** Project the specified elements into the world according to the current drag parameters */
+	void ProjectSelectedElementsIntoWorld(FViewport* Viewport, const FVector& Drag, const FRotator& Rot);
 
 	/** Draw additional details for brushes in the world */
 	UNREALED_API void DrawBrushDetails(const FSceneView* View, FPrimitiveDrawInterface* PDI);
@@ -960,12 +962,7 @@ private:
 	/** If currently creating a preview actor. */
 	static UNREALED_API bool bIsDroppingPreviewActor;
 
-	// TODO: Remove this to always use PreDragElementTransforms. That requires modifications to ProjectActorsIntoWorld
-	/** A map of actor locations before a drag operation */
-	mutable TMap<TWeakObjectPtr<const AActor>, FTransform> PreDragActorTransforms;
-
-	/** Map of element locations keyed by their handle. Currently used for the preview drop elements. Should
-	 eventually replace PreDragActorTransforms entirely. */
+	/** Map of element locations before a drag operation, keyed by their handle. */
 	TMap<FTypedElementHandle, FTransform> PreDragElementTransforms;
 
 	/** The elements (from the current selection set) that this viewport can manipulate (eg, via the transform gizmo) */
