- added -CookPackageDir= to the cooker to allow cooking a directory, and when combined with -cooksinglepackagenorefs, then only that directory will be cooked

#Synced-CL 50462590

diff --git a/Engine/Source/Editor/UnrealEd/Private/CookOnTheFlyServer.cpp b/Engine/Source/Editor/UnrealEd/Private/CookOnTheFlyServer.cpp
--- a/Engine/Source/Editor/UnrealEd/Private/CookOnTheFlyServer.cpp
+++ b/Engine/Source/Editor/UnrealEd/Private/CookOnTheFlyServer.cpp
@@ -11588,6 +11588,27 @@ void UCookOnTheFlyServer::GenerateInitialRequests(FBeginCookContext& BeginContex
 			}
 		}
 	}
+	if (FParse::Value(FCommandLine::Get(), TEXT("-CookPackageDir="), Text))
+	{
+		TArray<FString> Array;
+		Text.ParseIntoArray(Array, CommandLineDelimiters.GetData(), CommandLineDelimiters.Num(), true);
+		FARFilter Filter;
+		Filter.bRecursivePaths = true;
+		// search in all paths specified
+		for (FString& PackageDir : Array)
+		{
+			Filter.PackagePaths.Add(*PackageDir);
+		}
+
+		// add each discovered package to the list of packages to cook
+		TArray<FAssetData> FoundAssets;
+		AssetRegistry->GetAssets(Filter, FoundAssets);
+		for (FAssetData& AssetData : FoundAssets)
+		{
+			// add path to package, not any asset within
+			CookMaps.AddUnique(AssetData.GetSoftObjectPath().GetLongPackageName());
+		}
+	}
 
 
 	TArray<FName> FilesInPath;
