DDB: Fixed a non-unity compile error in the editor build worker registration

#rnx
#Synced-CL 50253651

diff --git a/Engine/Source/Editor/UnrealEd/Private/DerivedDataBuildWorkers.cpp b/Engine/Source/Editor/UnrealEd/Private/DerivedDataBuildWorkers.cpp
--- a/Engine/Source/Editor/UnrealEd/Private/DerivedDataBuildWorkers.cpp
+++ b/Engine/Source/Editor/UnrealEd/Private/DerivedDataBuildWorkers.cpp
@@ -2,15 +2,14 @@
 
 #include "DerivedDataBuildWorkers.h"
 
-#include "HAL/Platform.h"
-#include "HAL/PlatformCrt.h"
+#include "Containers/StringView.h"
 #include "Misc/Optional.h"
 #include "MultiPlatformTargetReceiptBuildWorkers.h"
 
 struct FDerivedDataBuildWorkers
 {
-	FMultiPlatformTargetReceiptBuildWorkers BaseTextureBuildWorkerFactory{TEXT("$(EngineDir)/Binaries/$(Platform)/BaseTextureBuildWorker.target")};
-	FMultiPlatformTargetReceiptBuildWorkers ShaderBuildWorkerFactory{TEXT("$(EngineDir)/Binaries/$(Platform)/ShaderBuildWorker.target")};
+	FMultiPlatformTargetReceiptBuildWorkers BaseTextureBuildWorkerFactory{TEXTVIEW("$(EngineDir)/Binaries/$(Platform)/BaseTextureBuildWorker.target")};
+	FMultiPlatformTargetReceiptBuildWorkers ShaderBuildWorkerFactory{TEXTVIEW("$(EngineDir)/Binaries/$(Platform)/ShaderBuildWorker.target")};
 };
 
 TOptional<FDerivedDataBuildWorkers> GBuildWorkers;
