fix: Increment PackageDataFromBaseGameNum on DLC plugin only cooks

#Synced-CL 50586249
#Synced-CL 50586249

diff --git a/Engine/Source/Editor/UnrealEd/Private/CookOnTheFlyServer.cpp b/Engine/Source/Editor/UnrealEd/Private/CookOnTheFlyServer.cpp
--- a/Engine/Source/Editor/UnrealEd/Private/CookOnTheFlyServer.cpp
+++ b/Engine/Source/Editor/UnrealEd/Private/CookOnTheFlyServer.cpp
@@ -11786,7 +11786,7 @@ void UCookOnTheFlyServer::RecordDLCPackagesFromBaseGame(FBeginCookContext& Begin
 		FString DLCPath = GetMountedAssetPathForDLC();
 		UE_LOG(LogCook, Log, TEXT("Performing plugin only cook for: %s"), *DLCPath);
 
-		PackageDatas->LockAndEnumeratePackageDatas([&BeginContext,&DLCPath](FPackageData* PackageData)
+		PackageDatas->LockAndEnumeratePackageDatas([&BeginContext,&DLCPath, this](FPackageData* PackageData)
 		{
 			FString PackageName = PackageData->GetPackageName().ToString();
 
@@ -11798,6 +11798,10 @@ void UCookOnTheFlyServer::RecordDLCPackagesFromBaseGame(FBeginCookContext& Begin
 					FPackagePlatformData& PlatformData = PackageData->FindOrAddPlatformData(TargetPlatform);
 					PlatformData.SetWhereCooked(EWhereCooked::BaseGame);
 				}
+
+				// Incremented only for the first platform, similar to 
+				// NumPackagesIncrementallySkipped and other relevant stats. 
+				++PackageDataFromBaseGameNum;
 			}
 		});
 
