[ITF Viewport Interactions] Unregister any registered viewport interactions when they are marked disabled.

#rnx

#Synced-CL 50605375

diff --git a/Engine/Source/Editor/UnrealEd/Private/Tools/EdModeInteractiveToolsContext.cpp b/Engine/Source/Editor/UnrealEd/Private/Tools/EdModeInteractiveToolsContext.cpp
--- a/Engine/Source/Editor/UnrealEd/Private/Tools/EdModeInteractiveToolsContext.cpp
+++ b/Engine/Source/Editor/UnrealEd/Private/Tools/EdModeInteractiveToolsContext.cpp
@@ -1501,6 +1501,18 @@ void UModeManagerInteractiveToolsContext::SetBypassAltModifier(bool bInBypassAlt
 	bBypassAltModifier = bInBypassAltModifier;
 }
 
+void UModeManagerInteractiveToolsContext::SetViewportInteractionsEnabled(bool bInEnabled)
+{
+	if (bUsesViewportInteractions != bInEnabled)
+	{
+		bUsesViewportInteractions = bInEnabled;
+		if (!bInEnabled)
+		{
+			UnregisterViewportInteractions();
+		}
+	}
+}
+
 bool UModeManagerInteractiveToolsContext::UsesViewportInteractions() const
 {
 	return bUsesViewportInteractions && UE::Editor::ViewportInteractions::UseEditorViewportInteractions();
diff --git a/Engine/Source/Editor/UnrealEd/Public/Tools/EdModeInteractiveToolsContext.h b/Engine/Source/Editor/UnrealEd/Public/Tools/EdModeInteractiveToolsContext.h
--- a/Engine/Source/Editor/UnrealEd/Public/Tools/EdModeInteractiveToolsContext.h
+++ b/Engine/Source/Editor/UnrealEd/Public/Tools/EdModeInteractiveToolsContext.h
@@ -393,10 +393,7 @@ public:
 	/**
 	 * Enable ITF Input tools support (e.g. Viewport Interactions as camera movement and selection)
 	 */
-	void SetViewportInteractionsEnabled(bool bInEnabled)
-	{
-		bUsesViewportInteractions = bInEnabled;
-	}
+	UE_API void SetViewportInteractionsEnabled(bool bInEnabled);
 
 	/**
 	 * Are Drag Tools supported? e.g. Marquee Select
