UE: Init ExternalThumbnailCache before doing a ParallelFor to make sure it initializes in the calling thread if it hadn't already

#Synced-CL 50582365

diff --git a/Engine/Source/Editor/UnrealEd/Private/ThumbnailExternalCache.cpp b/Engine/Source/Editor/UnrealEd/Private/ThumbnailExternalCache.cpp
--- a/Engine/Source/Editor/UnrealEd/Private/ThumbnailExternalCache.cpp
+++ b/Engine/Source/Editor/UnrealEd/Private/ThumbnailExternalCache.cpp
@@ -115,6 +115,9 @@ void FThumbnailExternalCache::LoadCompressAndAppend(const TArrayView<FAssetData>
 		}
 	}
 
+	// Init() is not threadsafe, so call it here before doing ParallelFor below in case it hasn't already been previously called.
+	Init();
+
 	// Load and compress
 	ParallelFor(AssetsToProcess.Num(), [&AssetsToProcess, &CombinedCache](int32 Index)
 	{
