[ITF Viewport Interactions] Move the application focus lost handling into `UModeManagerInteractiveToolsContext`. Reset cached input state when application focus is lost.

#rnx

#Synced-CL 50653461

diff --git a/Engine/Source/Editor/UnrealEd/Private/Tools/EdModeInteractiveToolsContext.cpp b/Engine/Source/Editor/UnrealEd/Private/Tools/EdModeInteractiveToolsContext.cpp
--- a/Engine/Source/Editor/UnrealEd/Private/Tools/EdModeInteractiveToolsContext.cpp
+++ b/Engine/Source/Editor/UnrealEd/Private/Tools/EdModeInteractiveToolsContext.cpp
@@ -1276,6 +1276,15 @@ void UModeManagerInteractiveToolsContext::UpdateMousePosition(FEditorViewportCli
 	}
 }
 
+void UModeManagerInteractiveToolsContext::OnApplicationFocusChanged(bool bInIsFocused)
+{
+	if (!bInIsFocused)
+	{
+		// Any pressed mouse or modifier keys need to be reset as they are no longer reliable.	
+		CurrentMouseState = FInputDeviceState();
+	}
+}
+
 bool UModeManagerInteractiveToolsContext::GetCursor(EMouseCursor::Type& OutCursor) const
 {	
 	UToolsContextCursorAPI* ToolsContextCursorAPI = ToolManager->GetContextObjectStore()->FindContext<UToolsContextCursorAPI>();
@@ -1409,6 +1418,11 @@ void UModeManagerInteractiveToolsContext::Initialize(IToolsContextQueriesAPI* Qu
 	UE::Editor::ViewportInteractions::OnEditorViewportInteractionsDeactivated().AddUObject(this, &UModeManagerInteractiveToolsContext::UnregisterViewportInteractions);
 	
 	CurrentMouseState.OnGetButtonStateFallback.BindUObject(this, &UModeManagerInteractiveToolsContext::OnGetButtonStateFallback);
+	
+	if (FSlateApplication::IsInitialized())
+	{
+		FSlateApplication::Get().OnApplicationActivationStateChanged().AddUObject(this, &UModeManagerInteractiveToolsContext::OnApplicationFocusChanged);
+	}
 }
 
 void UModeManagerInteractiveToolsContext::Shutdown()
@@ -1435,6 +1449,11 @@ void UModeManagerInteractiveToolsContext::Shutdown()
 
 	FWorldDelegates::PreLevelRemovedFromWorld.RemoveAll(this);
 
+	if (FSlateApplication::IsInitialized())
+	{
+		FSlateApplication::Get().OnApplicationActivationStateChanged().RemoveAll(this);
+	}
+
 	Super::Shutdown();
 }
 
diff --git a/Engine/Source/Editor/UnrealEd/Public/Tools/EdModeInteractiveToolsContext.h b/Engine/Source/Editor/UnrealEd/Public/Tools/EdModeInteractiveToolsContext.h
--- a/Engine/Source/Editor/UnrealEd/Public/Tools/EdModeInteractiveToolsContext.h
+++ b/Engine/Source/Editor/UnrealEd/Public/Tools/EdModeInteractiveToolsContext.h
@@ -337,8 +337,6 @@ protected:
 		
 	FDeviceButtonState OnGetButtonStateFallback(const FKey& Key) const;
 	
-	void UpdateMousePosition(FEditorViewportClient* InViewportClient, FViewport* InViewport, int32 X, int32 Y);
-
 	/** Input event instance used to keep track of various button states, etc, that we cannot directly query on-demand */
 	FInputDeviceState CurrentMouseState;
 	// called when PIE is about to start, shuts down active tools
@@ -351,6 +349,9 @@ protected:
 	FDelegateHandle ViewportClientListChangedHandle;
 
 private:
+	void UpdateMousePosition(FEditorViewportClient* InViewportClient, FViewport* InViewport, int32 X, int32 Y);
+	void OnApplicationFocusChanged(bool bInIsFocused);
+	
 	bool bIsTrackingMouse;
 
 	FGetUnsupportedInteractions GetUnsupportedInteractionsDelegate;
