UE-364664 - No warning when deleting a local, in-use Blueprint variable

FBlueprintEditorUtils::IsVariableUsed was only checking for in-use class variables. It was missing logic to check for local variables if LocalGraphScope is non-null.

This fix also revealed several missing consts from the FindLocal* API, as well as an extraneous FindLocalVariable overload.

#Synced-CL 50676684

diff --git a/Engine/Source/Editor/UnrealEd/Private/Kismet2/BlueprintEditorUtils.cpp b/Engine/Source/Editor/UnrealEd/Private/Kismet2/BlueprintEditorUtils.cpp
--- a/Engine/Source/Editor/UnrealEd/Private/Kismet2/BlueprintEditorUtils.cpp
+++ b/Engine/Source/Editor/UnrealEd/Private/Kismet2/BlueprintEditorUtils.cpp
@@ -5694,12 +5694,6 @@ void FBlueprintEditorUtils::RenameLocalVariable(UBlueprint* InBlueprint, const U
 	}
 }
 
-FBPVariableDescription* FBlueprintEditorUtils::FindLocalVariable(UBlueprint* InBlueprint, const UStruct* InScope, const FName InVariableName)
-{
-	UK2Node_FunctionEntry* DummyFunctionEntry = nullptr;
-	return FindLocalVariable(InBlueprint, InScope, InVariableName, &DummyFunctionEntry);
-}
-
 FBPVariableDescription* FBlueprintEditorUtils::FindLocalVariable(const UBlueprint* InBlueprint, const UEdGraph* InScopeGraph, const FName InVariableName, class UK2Node_FunctionEntry** OutFunctionEntry)
 {
 	FBPVariableDescription* ReturnVariable = nullptr;
@@ -5741,7 +5735,7 @@ FBPVariableDescription* FBlueprintEditorUtils::FindLocalVariable(const UBlueprin
 	return FindLocalVariable(InBlueprint, ScopeGraph, InVariableName, OutFunctionEntry);
 }
 
-FName FBlueprintEditorUtils::FindLocalVariableNameByGuid(UBlueprint* InBlueprint, const FGuid& InVariableGuid)
+FName FBlueprintEditorUtils::FindLocalVariableNameByGuid(const UBlueprint* InBlueprint, const FGuid& InVariableGuid)
 {
 	// Search through all function entry nodes for a local variable with the passed Guid
 	TArray<UK2Node_FunctionEntry*> FunctionEntryNodes;
@@ -5768,7 +5762,7 @@ FName FBlueprintEditorUtils::FindLocalVariableNameByGuid(UBlueprint* InBlueprint
 	return ReturnVariableName;
 }
 
-FGuid FBlueprintEditorUtils::FindLocalVariableGuidByName(UBlueprint* InBlueprint, const UStruct* InScope, const FName InVariableName)
+FGuid FBlueprintEditorUtils::FindLocalVariableGuidByName(const UBlueprint* InBlueprint, const UStruct* InScope, const FName InVariableName)
 {
 	FGuid ReturnGuid;
 	if(FBPVariableDescription* LocalVariable = FindLocalVariable(InBlueprint, InScope, InVariableName))
@@ -5779,7 +5773,7 @@ FGuid FBlueprintEditorUtils::FindLocalVariableGuidByName(UBlueprint* InBlueprint
 	return ReturnGuid;
 }
 
-FGuid FBlueprintEditorUtils::FindLocalVariableGuidByName(UBlueprint* InBlueprint, const UEdGraph* InScopeGraph, const FName InVariableName)
+FGuid FBlueprintEditorUtils::FindLocalVariableGuidByName(const UBlueprint* InBlueprint, const UEdGraph* InScopeGraph, const FName InVariableName)
 {
 	FGuid ReturnGuid;
 	if(FBPVariableDescription* LocalVariable = FindLocalVariable(InBlueprint, InScopeGraph, InVariableName))
@@ -6017,7 +6011,15 @@ bool FBlueprintEditorUtils::IsVariableUsed(const UBlueprint* Blueprint, const FN
 
 	// Retrieve the corresponding variable guid from the blueprint
 	FGuid VariableGuid;
-	UBlueprint::GetGuidFromClassByFieldName<FProperty>(Blueprint->SkeletonGeneratedClass, VariableName, VariableGuid);
+
+	if (LocalGraphScope)
+	{
+		VariableGuid = FindLocalVariableGuidByName(Blueprint, LocalGraphScope, VariableName);
+	}
+	else
+	{
+		UBlueprint::GetGuidFromClassByFieldName<FProperty>(Blueprint->SkeletonGeneratedClass, VariableName, VariableGuid);
+	}
 	
 	if (!VariableGuid.IsValid())
 	{
diff --git a/Engine/Source/Editor/UnrealEd/Public/Kismet2/BlueprintEditorUtils.h b/Engine/Source/Editor/UnrealEd/Public/Kismet2/BlueprintEditorUtils.h
--- a/Engine/Source/Editor/UnrealEd/Public/Kismet2/BlueprintEditorUtils.h
+++ b/Engine/Source/Editor/UnrealEd/Public/Kismet2/BlueprintEditorUtils.h
@@ -973,7 +973,7 @@ public:
 	static UNREALED_API bool AddLocalVariable(UBlueprint* Blueprint, UEdGraph* InTargetGraph, const FName InNewVarName, const FEdGraphPinType& InNewVarType, const FString& DefaultValue = FString());
 
 	/**
-	 * Removes a member variable if it was declared in this blueprint and not in a base class.
+	 * Removes a local variable if it was declared in this blueprint and not in a base class.
 	 *
 	 * @param	InBlueprint			The Blueprint the local variable can be found in
 	 * @param	InScope				Local variable's scope
@@ -981,15 +981,6 @@ public:
 	 */
 	static UNREALED_API void RemoveLocalVariable(UBlueprint* InBlueprint, const UStruct* InScope, const FName InVarName);
 
-	/**
-	 * Returns a local variable with the function entry it was found in
-	 *
-	 * @param InBlueprint		Blueprint to search for the local variable
-	 * @param InVariableName	Name of the variable to search for
-	 * @return					The local variable description
-	 */
-	static UNREALED_API FBPVariableDescription* FindLocalVariable(UBlueprint* InBlueprint, const UStruct* InScope, const FName InVariableName);
-
 	/**
 	 * Returns a local variable
 	 *
@@ -999,7 +990,7 @@ public:
 	 * @param OutFunctionEntry	Optional output parameter. If not null, the found function entry is returned.
 	 * @return					The local variable description
 	 */
-	static UNREALED_API FBPVariableDescription* FindLocalVariable(const UBlueprint* InBlueprint, const UEdGraph* InScopeGraph, const FName InVariableName, class UK2Node_FunctionEntry** OutFunctionEntry = NULL);
+	static UNREALED_API FBPVariableDescription* FindLocalVariable(const UBlueprint* InBlueprint, const UEdGraph* InScopeGraph, const FName InVariableName, class UK2Node_FunctionEntry** OutFunctionEntry = nullptr);
 
 	/**
 	 * Returns a local variable
@@ -1010,7 +1001,7 @@ public:
 	 * @param OutFunctionEntry	Optional output parameter. If not null, the found function entry is returned.
 	 * @return					The local variable description
 	 */
-	static UNREALED_API FBPVariableDescription* FindLocalVariable(const UBlueprint* InBlueprint, const UStruct* InScope, const FName InVariableName, class UK2Node_FunctionEntry** OutFunctionEntry = NULL);
+	static UNREALED_API FBPVariableDescription* FindLocalVariable(const UBlueprint* InBlueprint, const UStruct* InScope, const FName InVariableName, class UK2Node_FunctionEntry** OutFunctionEntry = nullptr);
 
 	/**
 	 * Finds a local variable name using the variable's Guid
@@ -1019,7 +1010,7 @@ public:
 	 * @param InVariableGuid	Guid to identify the local variable with
 	 * @return					Local variable's name
 	 */
-	static UNREALED_API FName FindLocalVariableNameByGuid(UBlueprint* InBlueprint, const FGuid& InVariableGuid);
+	static UNREALED_API FName FindLocalVariableNameByGuid(const UBlueprint* InBlueprint, const FGuid& InVariableGuid);
 
 	/**
 	 * Finds a local variable Guid using the variable's name
@@ -1029,7 +1020,7 @@ public:
 	 * @param InVariableGuid	Local variable's name to search for
 	 * @return					The Guid associated with the local variable
 	 */
-	static UNREALED_API FGuid FindLocalVariableGuidByName(UBlueprint* InBlueprint, const UStruct* InScope, const FName InVariableName);
+	static UNREALED_API FGuid FindLocalVariableGuidByName(const UBlueprint* InBlueprint, const UStruct* InScope, const FName InVariableName);
 
 	/**
 	 * Finds a local variable Guid using the variable's name
@@ -1039,7 +1030,7 @@ public:
 	 * @param InVariableGuid	Local variable's name to search for
 	 * @return					The Guid associated with the local variable
 	 */
-	static UNREALED_API FGuid FindLocalVariableGuidByName(UBlueprint* InBlueprint, const UEdGraph* InScopeGraph, const FName InVariableName);
+	static UNREALED_API FGuid FindLocalVariableGuidByName(const UBlueprint* InBlueprint, const UEdGraph* InScopeGraph, const FName InVariableName);
 
 	/**
 	 * Rename a local variable
