[ITF] Fix an issue where  `CurrentMouseState.OnGetButtonStateFallback` was reset when application focus was lost.

#rnx
#Synced-CL 50708145

diff --git a/Engine/Source/Editor/UnrealEd/Private/Tools/EdModeInteractiveToolsContext.cpp b/Engine/Source/Editor/UnrealEd/Private/Tools/EdModeInteractiveToolsContext.cpp
--- a/Engine/Source/Editor/UnrealEd/Private/Tools/EdModeInteractiveToolsContext.cpp
+++ b/Engine/Source/Editor/UnrealEd/Private/Tools/EdModeInteractiveToolsContext.cpp
@@ -1281,7 +1281,9 @@ void UModeManagerInteractiveToolsContext::OnApplicationFocusChanged(bool bInIsFo
 	if (!bInIsFocused)
 	{
 		// Any pressed mouse or modifier keys need to be reset as they are no longer reliable.	
-		CurrentMouseState = FInputDeviceState();
+		FInputDeviceState ResetState = FInputDeviceState();
+		ResetState.OnGetButtonStateFallback = CurrentMouseState.OnGetButtonStateFallback;
+		CurrentMouseState = ResetState;
 	}
 }
 
