Anim In Engine. Add Sequencer cvar for defaulting Set Retarget Source Asset to true, defaults false, to turn on use AnimationExport.SetRetargetSourceAsset true.  Note this will only be for the default, first time use.

#Synced-CL 50674404

diff --git a/Engine/Source/Editor/UnrealEd/Classes/Exporters/AnimSeqExportOption.h b/Engine/Source/Editor/UnrealEd/Classes/Exporters/AnimSeqExportOption.h
--- a/Engine/Source/Editor/UnrealEd/Classes/Exporters/AnimSeqExportOption.h
+++ b/Engine/Source/Editor/UnrealEd/Classes/Exporters/AnimSeqExportOption.h
@@ -11,7 +11,6 @@
 #include "Curves/RealCurve.h"
 #include "AnimSeqExportOption.generated.h"
 
-
 UCLASS(MinimalAPI, BlueprintType)
 class UAnimSeqExportOption : public UObject
 {
@@ -109,30 +108,8 @@ public:
 	/** Custom frame rate that the anim sequence will be recorded at*/
 	UPROPERTY(EditAnywhere, BlueprintReadWrite, meta = (EditCondition = bUseCustomFrameRate), AdvancedDisplay, Category = Export)
 	FFrameRate CustomFrameRate = FFrameRate(30, 1);
-
-
-	void ResetToDefault()
-	{
-		bExportTransforms = true;
-		bExportMorphTargets = true;
-		bExportAttributeCurves = true;
-		bExportMaterialCurves = true;
-		bRecordInWorldSpace = false;
-		bSetRetargetSourceAsset = false;
-		Interpolation = EAnimInterpolationType::Linear;
-		CurveInterpolation = ERichCurveInterpMode::RCIM_Linear;
-		bEvaluateAllSkeletalMeshComponents = true;
-		WarmUpFrames = 0;
-		DelayBeforeStart = 0;
-		bTransactRecording = true;
-		bBakeTimecode = false;
-		bTimecodeRateOverride = false;
-		OverrideTimecodeRate = FFrameRate(30, 1);
-		bUseCustomTimeRange = false;
-		CustomStartFrame = 0;
-		CustomEndFrame = 120;
-		CustomDisplayRate = FFrameRate(30, 1);
-		bUseCustomFrameRate = false;
-		CustomFrameRate = FFrameRate(30, 1);
-	}
+	
+	UNREALED_API void ResetToDefault();
 };
+
+#undef UE_API
diff --git a/Engine/Source/Editor/UnrealEd/Private/Animation/AnimSeqExportOption.cpp b/Engine/Source/Editor/UnrealEd/Private/Animation/AnimSeqExportOption.cpp
--- a/Engine/Source/Editor/UnrealEd/Private/Animation/AnimSeqExportOption.cpp
+++ b/Engine/Source/Editor/UnrealEd/Private/Animation/AnimSeqExportOption.cpp
@@ -2,12 +2,40 @@
 #include "Exporters/AnimSeqExportOption.h"
 
 #include "UObject/Object.h"
+#include "HAL/IConsoleManager.h"
 
 #include UE_INLINE_GENERATED_CPP_BY_NAME(AnimSeqExportOption)
 
+TAutoConsoleVariable<bool> CVarSetRetargetSourceAsset(TEXT("AnimationExport.SetRetargetSourceAsset"), false, TEXT("When true we set the Retarget Source Asset on the anim sequence when recording, if false we do not, by default false."));
+
+
 UAnimSeqExportOption::UAnimSeqExportOption(const FObjectInitializer& ObjectInitializer)
 	: Super(ObjectInitializer)
 {
 	ResetToDefault();
 }
 
+void UAnimSeqExportOption::ResetToDefault()
+{
+	bExportTransforms = true;
+	bExportMorphTargets = true;
+	bExportAttributeCurves = true;
+	bExportMaterialCurves = true;
+	bRecordInWorldSpace = false;
+	bSetRetargetSourceAsset = CVarSetRetargetSourceAsset.GetValueOnGameThread();
+	Interpolation = EAnimInterpolationType::Linear;
+	CurveInterpolation = ERichCurveInterpMode::RCIM_Linear;
+	bEvaluateAllSkeletalMeshComponents = true;
+	WarmUpFrames = 0;
+	DelayBeforeStart = 0;
+	bTransactRecording = true;
+	bBakeTimecode = false;
+	bTimecodeRateOverride = false;
+	OverrideTimecodeRate = FFrameRate(30, 1);
+	bUseCustomTimeRange = false;
+	CustomStartFrame = 0;
+	CustomEndFrame = 120;
+	CustomDisplayRate = FFrameRate(30, 1);
+	bUseCustomFrameRate = false;
+	CustomFrameRate = FFrameRate(30, 1);
+}
\ No newline at end of file
