Always remap PIE paths if running with editor. Allow remapping for writes even if there's no valid UNetConnection.

#Synced-CL 50723010

diff --git a/Engine/Source/Editor/UnrealEd/Private/EditorEngine.cpp b/Engine/Source/Editor/UnrealEd/Private/EditorEngine.cpp
--- a/Engine/Source/Editor/UnrealEd/Private/EditorEngine.cpp
+++ b/Engine/Source/Editor/UnrealEd/Private/EditorEngine.cpp
@@ -5795,11 +5795,11 @@ FString NetworkRemapPath_TestLevelScriptActor(const ALevelScriptActor* LevelScri
 	return ResultStr;
 }
 
-FORCEINLINE bool NetworkRemapPath_local(FWorldContext& Context, FString& Str, bool bReading, bool bIsReplay)
+FORCEINLINE bool NetworkRemapPath_local(const FWorldContext* Context, FString& Str, bool bReading, bool bIsReplay)
 {
 	if (bReading)
 	{
-		UWorld* const World = Context.World();
+		const UWorld* World = Context->World();
 		if (World == nullptr)
 		{
 			return false;
@@ -5817,7 +5817,7 @@ FORCEINLINE bool NetworkRemapPath_local(FWorldContext& Context, FString& Str, bo
 
 		// First strip any source prefix, then add the appropriate prefix for this context
 		FSoftObjectPath Path = UWorld::RemovePIEPrefix(Str);
-		
+		const int32 PIEInstance = Context->PIEInstance;
 		if (bIsReplay)
 		{
 			const FString AssetName = Path.GetAssetName();
@@ -5826,8 +5826,8 @@ FORCEINLINE bool NetworkRemapPath_local(FWorldContext& Context, FString& Str, bo
 			FPackageName::TryConvertFilenameToLongPackageName(PackageNameOnly, PackageNameOnly);
 			const FString ShortName = FPackageName::GetShortName(PackageNameOnly);
 
-			const FString PrefixedFullName = UWorld::ConvertToPIEPackageName(Str, Context.PIEInstance);
-			const FString PrefixedPackageName = UWorld::ConvertToPIEPackageName(PackageNameOnly, Context.PIEInstance);
+			const FString PrefixedFullName = UWorld::ConvertToPIEPackageName(Str, PIEInstance);
+			const FString PrefixedPackageName = UWorld::ConvertToPIEPackageName(PackageNameOnly, PIEInstance);
 			const FString WorldPackageName = World->GetOutermost()->GetName();
 
 			if (WorldPackageName == PrefixedPackageName)
@@ -5854,12 +5854,12 @@ FORCEINLINE bool NetworkRemapPath_local(FWorldContext& Context, FString& Str, bo
 			const bool bActorClass = FPackageName::IsValidObjectPath(Path.ToString()) && !AssetName.IsEmpty() && !ShortName.IsEmpty() && (AssetName == (ShortName + TEXT("_C")));
 			if (!bActorClass)
 			{
-				Path.FixupForPIE(Context.PIEInstance);
+				Path.FixupForPIE(PIEInstance);
 			}
 		}
 		else
 		{
-			Path.FixupForPIE(Context.PIEInstance);
+			Path.FixupForPIE(PIEInstance);
 		}
 
 		FString Remapped = Path.ToString();
@@ -5884,20 +5884,30 @@ FORCEINLINE bool NetworkRemapPath_local(FWorldContext& Context, FString& Str, bo
 
 bool UEditorEngine::NetworkRemapPath(UNetConnection* Connection, FString& Str, bool bReading)
 {
-	if (Connection == nullptr || Connection->GetWorld() == nullptr)
+	if (bReading)
 	{
-		return false;
-	}
+		// Reading requires a valid world which we retrieve from the supplied connection
+		if (Connection == nullptr || Connection->GetWorld() == nullptr)
+		{
+			return false;
+		}
 
-	// Pretty sure there's no case where you can't have a world by this point.
-	FWorldContext& Context = GetWorldContextFromWorldChecked(Connection->GetWorld());
-	return NetworkRemapPath_local(Context, Str, bReading, Connection->IsReplay());
+		const bool bIsReplay = Connection->IsReplay();
+		const FWorldContext& Context = GetWorldContextFromWorldChecked(Connection->GetWorld());
+		return NetworkRemapPath_local(&Context, Str, bReading, bIsReplay);
+	}
+	else
+	{
+		constexpr bool bIsReplay = false;
+		const FWorldContext* Context = nullptr;
+		return NetworkRemapPath_local(Context, Str, bReading, bIsReplay);
+	}
 }
 
 bool UEditorEngine::NetworkRemapPath( UPendingNetGame *PendingNetGame, FString& Str, bool bReading)
 {
-	FWorldContext& Context = GetWorldContextFromPendingNetGameChecked(PendingNetGame);
-	return NetworkRemapPath_local(Context, Str, bReading, PendingNetGame->GetDemoNetDriver() != nullptr);
+	const FWorldContext& Context = GetWorldContextFromPendingNetGameChecked(PendingNetGame);
+	return NetworkRemapPath_local(&Context, Str, bReading, PendingNetGame->GetDemoNetDriver() != nullptr);
 }
 
 void UEditorEngine::CheckAndHandleStaleWorldObjectReferences(FWorldContext* InWorldContext)
diff --git a/Engine/Source/Runtime/Net/Iris/Private/Iris/ReplicationSystem/ObjectReferenceCache.cpp b/Engine/Source/Runtime/Net/Iris/Private/Iris/ReplicationSystem/ObjectReferenceCache.cpp
--- a/Engine/Source/Runtime/Net/Iris/Private/Iris/ReplicationSystem/ObjectReferenceCache.cpp
+++ b/Engine/Source/Runtime/Net/Iris/Private/Iris/ReplicationSystem/ObjectReferenceCache.cpp
@@ -236,8 +236,15 @@ bool FObjectReferenceCache::ShouldIgnoreWhenMissing(FNetRefHandle RefHandle) con
 
 bool FObjectReferenceCache::RenamePathForPie(uint32 ConnectionId, FString& Str, bool bReading, UReplicationSystem* ReplicationSystem)
 {
-	UObjectReplicationBridge* ReplicationBridge = ReplicationSystem->GetReplicationBridgeAs<UObjectReplicationBridge>();
-	return UE::GetPlayInEditorID() != -1 ? ReplicationBridge->RemapPathForPIE(ConnectionId, Str, bReading) : false;
+	if constexpr (WITH_EDITOR)
+	{
+		UObjectReplicationBridge* ReplicationBridge = ReplicationSystem->GetReplicationBridgeAs<UObjectReplicationBridge>();
+		return ReplicationBridge->RemapPathForPIE(ConnectionId, Str, bReading);
+	}
+	else
+	{
+		return false;
+	}
 }
 
 FNetRefHandle FObjectReferenceCache::CreateObjectReferenceHandle(const UObject* Object, UReplicationSystem* ReplicationSystem)
