Fix SA issue.

#rnx

#Synced-CL 50723010

#Synced-CL 50723010

diff --git a/Engine/Source/Editor/UnrealEd/Private/EditorEngine.cpp b/Engine/Source/Editor/UnrealEd/Private/EditorEngine.cpp
--- a/Engine/Source/Editor/UnrealEd/Private/EditorEngine.cpp
+++ b/Engine/Source/Editor/UnrealEd/Private/EditorEngine.cpp
@@ -5799,6 +5799,7 @@ FORCEINLINE bool NetworkRemapPath_local(const FWorldContext* Context, FString& S
 {
 	if (bReading)
 	{
+		CA_ASSUME(Context != nullptr);
 		const UWorld* World = Context->World();
 		if (World == nullptr)
 		{
