[ShowFlags] Support sorting flags in specific groups by declaration index.

#Synced-CL 50789286

diff --git a/Engine/Source/Editor/UnrealEd/Private/EditorShowFlags.cpp b/Engine/Source/Editor/UnrealEd/Private/EditorShowFlags.cpp
--- a/Engine/Source/Editor/UnrealEd/Private/EditorShowFlags.cpp
+++ b/Engine/Source/Editor/UnrealEd/Private/EditorShowFlags.cpp
@@ -4,6 +4,17 @@
 #include "EditorShowFlags.h"
 #include "LevelEditorViewport.h"
 
+bool ShouldFlagsInGroupKeepDeclarationOrder(EShowFlagGroup InGroup)
+{
+	switch (InGroup)
+	{
+	case SFG_Lumen:
+	case SFG_MegaLights:
+		return true;
+	default:
+		return false;
+	}
+}
 
 FShowFlagData::FShowFlagData(const FString& InName, const FText& InDisplayName, const uint32 InEngineShowFlagIndex, EShowFlagGroup InGroup )
 	:	DisplayName( InDisplayName )
@@ -111,15 +122,29 @@ TArray<FShowFlagData>& GetShowFlagMenuItems()
 			FEngineShowFlags::IterateAllFlags(Sink);
 		}	
 
-		// Sort the show flags alphabetically by string.
-		struct FCompareFShowFlagDataByName
+		// Sort the show flags
+		struct FShowFlagSorter
 		{
 			FORCEINLINE bool operator()( const FShowFlagData& A, const FShowFlagData& B ) const
 			{
-				return A.DisplayName.ToString() < B.DisplayName.ToString();
+				if (A.Group != B.Group)
+				{
+					// first sort by group
+					return A.Group < B.Group;
+				}
+				else if (ShouldFlagsInGroupKeepDeclarationOrder(A.Group))
+				{
+					// if group wants to keep declaration order do that
+					return A.EngineShowFlagIndex < B.EngineShowFlagIndex;
+				}
+				else
+				{
+					// otherwise sort alphabetically by display name
+					return A.DisplayName.ToString() < B.DisplayName.ToString();
+				}
 			}
 		};
-		OutShowFlags.Sort( FCompareFShowFlagDataByName() );
+		OutShowFlags.Sort(FShowFlagSorter());
 	}
 
 	if (bInvalidated)
