[ITF Viewport Interactions] Alt + RMB Dolly panning motion issues

Making sure Dolly movement is not actually going past LookAt position.
This prevents jumpiness around the final part of the motion.

Note: dolly is always performed with respect to latest LookAt position, which might not always be what the user expects.
Some tests have been done in order to update the LookAt based on current camera state. Results are reasonable, but not always predictable, so that will be addressed separately.

#Synced-CL 50767940

diff --git a/Engine/Source/Editor/UnrealEd/Private/EditorViewportClient.cpp b/Engine/Source/Editor/UnrealEd/Private/EditorViewportClient.cpp
--- a/Engine/Source/Editor/UnrealEd/Private/EditorViewportClient.cpp
+++ b/Engine/Source/Editor/UnrealEd/Private/EditorViewportClient.cpp
@@ -2572,18 +2572,21 @@ void FEditorViewportClient::DollyPerspectiveCamera(double InDelta)
 	
 	FViewportCameraTransform& ViewTransform = GetViewTransform();
 	
+	const FVector LookAtLocation = GetLookAtLocation();
+	const FVector ViewLocation = GetViewLocation();
 	const float CameraSpeed = GetCameraSpeed();
-	const float CameraSpeedDistanceScale = ShouldScaleCameraSpeedByDistance() ? FMath::Min(FVector::Dist( GetViewLocation(), GetLookAtLocation() ) / 1000.f, 1000.f) : 1.0f;
+	const float CameraSpeedDistanceScale = ShouldScaleCameraSpeedByDistance() ? FMath::Min(FVector::Dist( ViewLocation, LookAtLocation ) / 1000.f, 1000.f) : 1.0f;
 	
 	FMatrix OrbitMatrix = ViewTransform.ComputeOrbitMatrix().InverseFast();
 
 	FVector DeltaLocation = FVector(0, InDelta * CameraSpeed * CameraSpeedDistanceScale, 0);
 
-	FVector LookAt = ViewTransform.GetLookAt();
+	// Make sure we don't move past the LookAt
+	DeltaLocation = FVector::Min(DeltaLocation, ViewLocation - LookAtLocation);
 
 	// Orient the delta down the view direction towards the look at
 	FMatrix RotMat =
-		FTranslationMatrix( -LookAt ) *
+		FTranslationMatrix( -LookAtLocation ) *
 		FRotationMatrix( FRotator(0,GetViewRotation().Yaw,0) ) *
 		FRotationMatrix( FRotator(0, 0, GetViewRotation().Pitch));
 
@@ -3151,6 +3154,9 @@ void FEditorViewportClient::InputAxisForOrbit(FViewport* InViewport, const FVect
 
 			FVector LookAt = ViewTransform.GetLookAt();
 
+			// Make sure we don't move past the LookAt
+			DeltaLocation = FVector::Min(DeltaLocation, GetViewLocation() - LookAt);
+
 			// Orient the delta down the view direction towards the look at
 			FMatrix RotMat =
 				FTranslationMatrix( -LookAt ) *
