[PackageTracker/ALT]
- Add missing lock when adding package to a TSet

#Synced-CL 50642529

diff --git a/Engine/Source/Editor/UnrealEd/Private/Cooker/PackageTracker.cpp b/Engine/Source/Editor/UnrealEd/Private/Cooker/PackageTracker.cpp
--- a/Engine/Source/Editor/UnrealEd/Private/Cooker/PackageTracker.cpp
+++ b/Engine/Source/Editor/UnrealEd/Private/Cooker/PackageTracker.cpp
@@ -134,6 +134,8 @@ void FPackageTracker::NotifyUObjectCreated(const class UObjectBase* Object, int3
 		if (Package->GetOuter() == nullptr) // Nested packages are no longer created, but can still exist in old data.
 		{
 			OnCreatePackage(Package->GetFName());
+
+			FWriteScopeLock ScopeLock(Lock);
 			LoadedPackages.Add(Package);
 		}
 	}
