[ITF Viewport Interactions] Alt + RMB dolly interaction should be smoother

Previous solution was backed out, since it was not always working as expected.
This CL provides an alternative fix.

#Synced-CL 50816236

diff --git a/Engine/Source/Editor/UnrealEd/Private/EditorViewportClient.cpp b/Engine/Source/Editor/UnrealEd/Private/EditorViewportClient.cpp
--- a/Engine/Source/Editor/UnrealEd/Private/EditorViewportClient.cpp
+++ b/Engine/Source/Editor/UnrealEd/Private/EditorViewportClient.cpp
@@ -2571,25 +2571,29 @@ void FEditorViewportClient::DollyPerspectiveCamera(double InDelta)
 	ToggleOrbitCamera(true);
 	
 	FViewportCameraTransform& ViewTransform = GetViewTransform();
-	
-	const float CameraSpeed = GetCameraSpeed();
-	const float CameraSpeedDistanceScale = ShouldScaleCameraSpeedByDistance() ? FMath::Min(FVector::Dist( GetViewLocation(), GetLookAtLocation() ) / 1000.f, 1000.f) : 1.0f;
-	
-	FMatrix OrbitMatrix = ViewTransform.ComputeOrbitMatrix().InverseFast();
 
-	FVector DeltaLocation = FVector(0, InDelta * CameraSpeed * CameraSpeedDistanceScale, 0);
+	const FVector ViewLocation = GetViewLocation();
+	const FVector LookAtLocation = GetLookAtLocation();
+	const FVector DeltaToTarget = LookAtLocation - ViewLocation;
 
-	FVector LookAt = ViewTransform.GetLookAt();
+	const float DistanceToTarget = DeltaToTarget.Length();
+	const float CameraSpeed = GetCameraSpeed();
+	const float CameraSpeedDistanceScale = ShouldScaleCameraSpeedByDistance() ? FMath::Min(DistanceToTarget / 1000.f, 1000.f) : 1.0f;
+	const float DeltaLength = InDelta * CameraSpeed * CameraSpeedDistanceScale;
+
+	// Make sure we don't move past the LookAt
+	FVector DeltaLocation = FVector::UnitY() * FMath::Min(DeltaLength, DistanceToTarget);
 
 	// Orient the delta down the view direction towards the look at
 	FMatrix RotMat =
-		FTranslationMatrix( -LookAt ) *
+		FTranslationMatrix( -LookAtLocation ) *
 		FRotationMatrix( FRotator(0,GetViewRotation().Yaw,0) ) *
 		FRotationMatrix( FRotator(0, 0, GetViewRotation().Pitch));
 
 	FVector TransformedDelta = RotMat.InverseFast().TransformVector(DeltaLocation);
 
-	SetViewLocation( OrbitMatrix.GetOrigin() + TransformedDelta );
+	FMatrix OrbitMatrix = ViewTransform.ComputeOrbitMatrix().InverseFast();
+	SetViewLocation(OrbitMatrix.GetOrigin() + TransformedDelta);
 }
 
 /**
@@ -3140,25 +3144,30 @@ void FEditorViewportClient::InputAxisForOrbit(FViewport* InViewport, const FVect
 		}
 		else if ( IsOrbitZoomMode( InViewport ) )
 		{
-		const bool bInvertY = GetDefault<ULevelEditorViewportSettings>()->bInvertRightMouseDollyYAxis;
-
-			FMatrix OrbitMatrix = ViewTransform.ComputeOrbitMatrix().InverseFast();
+			const bool bInvertY = GetDefault<ULevelEditorViewportSettings>()->bInvertRightMouseDollyYAxis;
 
 			const float CameraSpeed = GetCameraSpeed();
 			Drag *= CameraSpeed * CameraSpeedDistanceScale;
 
-			FVector DeltaLocation = bInvertY ? FVector(0, Drag.X + Drag.Y, 0) : FVector(0, Drag.X+ -Drag.Y, 0);
+			const FVector ViewLocation = GetViewLocation();
+			const FVector LookAtLocation = GetLookAtLocation();
+			const FVector DeltaToTarget = LookAtLocation - ViewLocation;
 
-			FVector LookAt = ViewTransform.GetLookAt();
+			const float DistanceToTarget = DeltaToTarget.Length();
+			const float DeltaLength = bInvertY ? Drag.X + Drag.Y : Drag.X+ -Drag.Y;
+
+			// Make sure we don't move past the LookAt
+			FVector DeltaLocation = FVector::UnitY() * FMath::Min(DeltaLength, DistanceToTarget);
 
 			// Orient the delta down the view direction towards the look at
 			FMatrix RotMat =
-				FTranslationMatrix( -LookAt ) *
+				FTranslationMatrix( -LookAtLocation ) *
 				FRotationMatrix( FRotator(0,GetViewRotation().Yaw,0) ) *
 				FRotationMatrix( FRotator(0, 0, GetViewRotation().Pitch));
 
 			FVector TransformedDelta = RotMat.InverseFast().TransformVector(DeltaLocation);
 
+			FMatrix OrbitMatrix = ViewTransform.ComputeOrbitMatrix().InverseFast();
 			SetViewLocation( OrbitMatrix.GetOrigin() + TransformedDelta );
 
 			FEditorViewportStats::Using(IsPerspective() ? FEditorViewportStats::CAT_PERSPECTIVE_MOUSE_ORBIT_ZOOM : FEditorViewportStats::CAT_ORTHOGRAPHIC_MOUSE_ORBIT_ZOOM);
