SnappingEnabled - Added IsRotationSnappingEnabled, IsTranslationSnappingEnabled, and IsScaleSnappingEnabled to ISnappingPolicy...

#Synced-CL 50475248

diff --git a/Engine/Plugins/2D/Paper2D/Source/SmartSnapping/Private/PlanarConstraintSnapPolicy.cpp b/Engine/Plugins/2D/Paper2D/Source/SmartSnapping/Private/PlanarConstraintSnapPolicy.cpp
--- a/Engine/Plugins/2D/Paper2D/Source/SmartSnapping/Private/PlanarConstraintSnapPolicy.cpp
+++ b/Engine/Plugins/2D/Paper2D/Source/SmartSnapping/Private/PlanarConstraintSnapPolicy.cpp
@@ -41,6 +41,21 @@ void FPlanarConstraintSnapPolicy::DrawSnappingHelpers(const FSceneView* View, FP
 {
 }
 
+bool FPlanarConstraintSnapPolicy::IsTranslationSnappingEnabled() const
+{
+	return IsEnabled();
+}
+
+bool FPlanarConstraintSnapPolicy::IsRotationSnappingEnabled() const
+{
+	return false; // rotation snapping not supported
+}
+
+bool FPlanarConstraintSnapPolicy::IsScaleSnappingEnabled() const
+{
+	return false; // scale snapping not supported
+}
+
 bool FPlanarConstraintSnapPolicy::IsEnabled() const
 {
 	return bIsEnabled;
diff --git a/Engine/Plugins/2D/Paper2D/Source/SmartSnapping/Private/PlanarConstraintSnapPolicy.h b/Engine/Plugins/2D/Paper2D/Source/SmartSnapping/Private/PlanarConstraintSnapPolicy.h
--- a/Engine/Plugins/2D/Paper2D/Source/SmartSnapping/Private/PlanarConstraintSnapPolicy.h
+++ b/Engine/Plugins/2D/Paper2D/Source/SmartSnapping/Private/PlanarConstraintSnapPolicy.h
@@ -30,5 +30,8 @@ public:
 	virtual void SnapRotatorToGrid(FRotator& Rotation) override;
 	virtual void ClearSnappingHelpers(bool bClearImmediately) override;
 	virtual void DrawSnappingHelpers(const FSceneView* View, FPrimitiveDrawInterface* PDI) override;
+	virtual bool IsTranslationSnappingEnabled() const override;
+	virtual bool IsRotationSnappingEnabled() const override;
+	virtual bool IsScaleSnappingEnabled() const override;
 	// End of ISnappingPolicy interface
 };
diff --git a/Engine/Source/Editor/UnrealEd/Private/SnappingUtils.cpp b/Engine/Source/Editor/UnrealEd/Private/SnappingUtils.cpp
--- a/Engine/Source/Editor/UnrealEd/Private/SnappingUtils.cpp
+++ b/Engine/Source/Editor/UnrealEd/Private/SnappingUtils.cpp
@@ -44,6 +44,9 @@ public:
 	virtual void SnapRotatorToGrid(FRotator& Rotation) override;
 	virtual void ClearSnappingHelpers(bool bClearImmediately = false) override;
 	virtual void DrawSnappingHelpers(const FSceneView* View, FPrimitiveDrawInterface* PDI) override;
+	virtual bool IsTranslationSnappingEnabled() const override;
+	virtual bool IsRotationSnappingEnabled() const override;
+	virtual bool IsScaleSnappingEnabled() const override;
 	// End of FEditorViewportSnapping interface
 
 	bool IsSnapToGridEnabled();
@@ -118,6 +121,20 @@ public:
 
 	bool SnapToBSPVertex( FVector& Location, FVector GridBase, FRotator& Rotation );
 
+private:
+	bool IsSnapRotationEnabledInternal() const;
+	bool IsSnapScaleEnabledInternal() const;
+
+	/**
+	 * @return true if snapping to vertices is enabled
+	 */
+	bool IsSnapToVertexEnabledInternal(bool bIsPivot = false) const;
+
+	/**
+	 * @return true if snapping actors to other actors is enabled
+	 */
+	bool IsSnapToActorEnabledInternal() const;
+	
 private:
 
 	/** Vertex snapping implementation */
@@ -129,21 +146,37 @@ private:
 
 bool FEditorViewportSnapping::IsSnapToGridEnabled()
 {
-	return GetDefault<ULevelEditorViewportSettings>()->GridEnabled && !IsSnapToVertexEnabled();
+	return GetDefault<ULevelEditorViewportSettings>()->GridEnabled && !IsSnapToVertexEnabledInternal();
 }
 
 bool FEditorViewportSnapping::IsSnapRotationEnabled()
 {	
-	// Ask Current Editor Mode if Rotation Snap is enabled
-	return GLevelEditorModeTools().IsSnapRotationEnabled();
+	return IsSnapRotationEnabledInternal();
 }
 
 bool FEditorViewportSnapping::IsSnapScaleEnabled()
+{
+	return IsSnapScaleEnabledInternal();
+}
+
+bool FEditorViewportSnapping::IsSnapRotationEnabledInternal() const
+{
+	// Ask Current Editor Mode if Rotation Snap is enabled
+	const ULevelEditorViewportSettings* Settings = GetDefault<ULevelEditorViewportSettings>();
+	return GLevelEditorModeTools().IsSnapRotationEnabled() || (Settings->SnapToSurface.bEnabled && Settings->SnapToSurface.bSnapRotation);
+}
+
+bool FEditorViewportSnapping::IsSnapScaleEnabledInternal() const
 {
 	return GetDefault<ULevelEditorViewportSettings>()->SnapScaleEnabled;
 }
 
 bool FEditorViewportSnapping::IsSnapToVertexEnabled(bool bIsPivot)
+{
+	return IsSnapToVertexEnabledInternal(bIsPivot);
+}
+
+bool FEditorViewportSnapping::IsSnapToVertexEnabledInternal(bool bIsPivot) const
 {
 	if( GetDefault<ULevelEditorViewportSettings>()->bSnapVertices )
 	{
@@ -163,7 +196,12 @@ bool FEditorViewportSnapping::IsSnapToVertexEnabled(bool bIsPivot)
 
 bool FEditorViewportSnapping::IsSnapToActorEnabled()
 {
-	return GetDefault<ULevelEditorViewportSettings>()->bEnableActorSnap && !IsSnapToVertexEnabled();
+	return IsSnapToActorEnabledInternal();
+}
+
+bool FEditorViewportSnapping::IsSnapToActorEnabledInternal() const
+{
+	return GetDefault<ULevelEditorViewportSettings>()->bEnableActorSnap && !IsSnapToVertexEnabledInternal();
 }
 
 void FEditorViewportSnapping::EnableActorSnap(bool bEnable)
@@ -392,6 +430,22 @@ void FEditorViewportSnapping::DrawSnappingHelpers(const FSceneView* View, FPrimi
 	VertexSnappingImpl.DrawSnappingHelpers( View, PDI );
 }
 
+bool FEditorViewportSnapping::IsTranslationSnappingEnabled() const
+{
+	const ULevelEditorViewportSettings* Settings = GetDefault<ULevelEditorViewportSettings>();
+	return Settings->GridEnabled || Settings->SnapToSurface.bEnabled || IsSnapToVertexEnabledInternal() || IsSnapToActorEnabledInternal();
+}
+
+bool FEditorViewportSnapping::IsRotationSnappingEnabled() const
+{
+	return IsSnapRotationEnabledInternal();
+}
+
+bool FEditorViewportSnapping::IsScaleSnappingEnabled() const
+{
+	return IsSnapScaleEnabledInternal();
+}
+
 bool FEditorViewportSnapping::SnapLocationToNearestVertex( FVector& Location, const FVector2D& MouseLocation, TNotNull<FEditorViewportClient*> ViewportClient, FVector& OutVertexNormal, bool bDrawVertHelpers, const TSet<TWeakObjectPtr<AActor>>& ActorsToIgnore )
 {
 	bool bSnapped = false;
@@ -437,16 +491,31 @@ bool FSnappingUtils::IsSnapToGridEnabled()
 	return EditorViewportSnapper.IsValid() && EditorViewportSnapper->IsSnapToGridEnabled();
 }
 
+bool FSnappingUtils::IsTranslationSnappingEnabled()
+{
+	return EditorViewportSnapper.IsValid() && EditorViewportSnapper->IsTranslationSnappingEnabled();
+}
+
 bool FSnappingUtils::IsRotationSnapEnabled()
 {
 	return EditorViewportSnapper.IsValid() && EditorViewportSnapper->IsSnapRotationEnabled();
 }
 
+bool FSnappingUtils::IsRotationSnappingEnabled()
+{
+	return EditorViewportSnapper.IsValid() && EditorViewportSnapper->IsRotationSnappingEnabled();
+}
+
 bool FSnappingUtils::IsScaleSnapEnabled()
 {
 	return EditorViewportSnapper.IsValid() && EditorViewportSnapper->IsSnapScaleEnabled();
 }
 
+bool FSnappingUtils::IsScaleSnappingEnabled()
+{
+	return EditorViewportSnapper.IsValid() && EditorViewportSnapper->IsScaleSnappingEnabled();
+}
+
 bool FSnappingUtils::IsSnapToActorEnabled()
 {
 	return EditorViewportSnapper.IsValid() && EditorViewportSnapper->IsSnapToActorEnabled();
diff --git a/Engine/Source/Editor/UnrealEd/Public/SnappingUtils.h b/Engine/Source/Editor/UnrealEd/Public/SnappingUtils.h
--- a/Engine/Source/Editor/UnrealEd/Public/SnappingUtils.h
+++ b/Engine/Source/Editor/UnrealEd/Public/SnappingUtils.h
@@ -19,15 +19,30 @@ public:
 	static UNREALED_API bool IsSnapToGridEnabled();
 
 	/**
-	 * @return true if orientation snapping is enabled
+	 * @return true if *any* snapping that affects translation is enabled
+	 */
+	static UNREALED_API bool IsTranslationSnappingEnabled();
+
+	/**
+	 * @return true if orientation snapping to angular grid is enabled
 	 */
 	static UNREALED_API bool IsRotationSnapEnabled();
 
 	/**
-	 * @return true if orientation snapping is enabled
+	 * @return true if *any* snapping that affects orientation is enabled
+	 */
+	static UNREALED_API bool IsRotationSnappingEnabled();
+
+	/**
+	 * @return true if scale snapping to grid is enabled
 	 */
 	static UNREALED_API bool IsScaleSnapEnabled();
 
+	/**
+	 * @return true if *any* snapping that affects scale is enabled
+	 */
+	static UNREALED_API bool IsScaleSnappingEnabled();
+
 	/**
 	 * @return true if snapping actors to other actors is enabled
 	 */
diff --git a/Engine/Source/Editor/ViewportSnapping/Private/ViewportSnappingModule.cpp b/Engine/Source/Editor/ViewportSnapping/Private/ViewportSnappingModule.cpp
--- a/Engine/Source/Editor/ViewportSnapping/Private/ViewportSnappingModule.cpp
+++ b/Engine/Source/Editor/ViewportSnapping/Private/ViewportSnappingModule.cpp
@@ -58,6 +58,42 @@ public:
 			(*PolicyIt)->DrawSnappingHelpers(View, PDI);
 		}
 	}
+
+	virtual bool IsRotationSnappingEnabled() const override
+	{
+		for (auto PolicyIt = PolicyList.CreateConstIterator(); PolicyIt; ++PolicyIt)
+		{
+			if ((*PolicyIt)->IsRotationSnappingEnabled())
+			{
+				return true;
+			}
+		}
+		return false;
+	}
+
+	virtual bool IsTranslationSnappingEnabled() const override
+	{
+		for (auto PolicyIt = PolicyList.CreateConstIterator(); PolicyIt; ++PolicyIt)
+		{
+			if ((*PolicyIt)->IsTranslationSnappingEnabled())
+			{
+				return true;
+			}
+		}
+		return false;
+	}
+
+	virtual bool IsScaleSnappingEnabled() const override
+	{
+		for (auto PolicyIt = PolicyList.CreateConstIterator(); PolicyIt; ++PolicyIt)
+		{
+			if ((*PolicyIt)->IsScaleSnappingEnabled())
+			{
+				return true;
+			}
+		}
+		return false;
+	}
 };
 
 //////////////////////////////////////////////////////////////////////////
diff --git a/Engine/Source/Editor/ViewportSnapping/Public/ISnappingPolicy.h b/Engine/Source/Editor/ViewportSnapping/Public/ISnappingPolicy.h
--- a/Engine/Source/Editor/ViewportSnapping/Public/ISnappingPolicy.h
+++ b/Engine/Source/Editor/ViewportSnapping/Public/ISnappingPolicy.h
@@ -49,4 +49,19 @@ public:
 	 * @param PDI	Drawing interface
 	 */
 	virtual void DrawSnappingHelpers(const FSceneView* View, FPrimitiveDrawInterface* PDI)=0;
+
+	/**
+	 * @return true if rotation snapping is enabled
+	 */
+	virtual bool IsRotationSnappingEnabled() const = 0;
+
+	/**
+	 * @return true if translation snapping is enabled
+	 */
+	virtual bool IsTranslationSnappingEnabled() const = 0;
+
+	/**
+	 * @return true if scale snapping is enabled
+	 */
+	virtual bool IsScaleSnappingEnabled() const = 0;
 };
