[ActorFolder] Clean deprecated code

#tests preflight
#rnx

#Synced-CL 50906061

diff --git a/Engine/Source/Editor/UnrealEd/Private/EditorActorFolders.cpp b/Engine/Source/Editor/UnrealEd/Private/EditorActorFolders.cpp
--- a/Engine/Source/Editor/UnrealEd/Private/EditorActorFolders.cpp
+++ b/Engine/Source/Editor/UnrealEd/Private/EditorActorFolders.cpp
@@ -32,14 +32,6 @@
 
 // Static member definitions
 
-//~ Begin Deprecated
-PRAGMA_DISABLE_DEPRECATION_WARNINGS
-FOnActorFolderCreate	FActorFolders::OnFolderCreate;
-FOnActorFolderMove		FActorFolders::OnFolderMove;
-FOnActorFolderDelete	FActorFolders::OnFolderDelete;
-PRAGMA_ENABLE_DEPRECATION_WARNINGS
-//~ End Deprecated
-
 FOnActorFolderCreated	FActorFolders::OnFolderCreated;
 FOnActorFolderMoved		FActorFolders::OnFolderMoved;
 FOnActorFolderDeleted	FActorFolders::OnFolderDeleted;
@@ -100,36 +92,18 @@ void FActorFolders::Housekeeping()
 void FActorFolders::BroadcastOnActorFolderCreated(UWorld& InWorld, const FFolder& InFolder)
 {
 	OnFolderCreated.Broadcast(InWorld, InFolder);
-	PRAGMA_DISABLE_DEPRECATION_WARNINGS
-	if (InFolder.IsRootObjectPersistentLevel())
-	{
-		OnFolderCreate.Broadcast(InWorld, InFolder.GetPath());
-	}
-	PRAGMA_ENABLE_DEPRECATION_WARNINGS
 	BroadcastOnActorEditorContextClientChanged(InWorld);
 }
 
 void FActorFolders::BroadcastOnActorFolderDeleted(UWorld& InWorld, const FFolder& InFolder)
 {
 	OnFolderDeleted.Broadcast(InWorld, InFolder);
-	PRAGMA_DISABLE_DEPRECATION_WARNINGS
-	if (InFolder.IsRootObjectPersistentLevel())
-	{
-		OnFolderDelete.Broadcast(InWorld, InFolder.GetPath());
-	}
-	PRAGMA_ENABLE_DEPRECATION_WARNINGS
 	BroadcastOnActorEditorContextClientChanged(InWorld);
 }
 
 void FActorFolders::BroadcastOnActorFolderMoved(UWorld& InWorld, const FFolder& InSrcFolder, const FFolder& InDstFolder)
 {
 	OnFolderMoved.Broadcast(InWorld, InSrcFolder, InDstFolder);
-	PRAGMA_DISABLE_DEPRECATION_WARNINGS
-	if (InSrcFolder.IsRootObjectPersistentLevel() && InDstFolder.IsRootObjectPersistentLevel())
-	{
-		OnFolderMove.Broadcast(InWorld, InSrcFolder.GetPath(), InDstFolder.GetPath());
-	}
-	PRAGMA_ENABLE_DEPRECATION_WARNINGS
 	BroadcastOnActorEditorContextClientChanged(InWorld);
 }
 
@@ -734,58 +708,4 @@ void FActorFolders::GetWeakActorsFromFolders(UWorld& InWorld, const TSet<FName>&
 	}, FolderRootObject);
 }
 
-////////////////////////////////////////////
-//~ Begin Deprecated
-
-FActorFolderProps* FActorFolders::GetFolderProperties(UWorld& InWorld, FName InPath)
-{
-	return GetFolderProperties(InWorld, FFolder(GetWorldFolderRootObject(InWorld), InPath));
-}
-
-FName FActorFolders::GetDefaultFolderName(UWorld& InWorld, FName ParentPath)
-{
-	FFolder DefaultFolderName = GetDefaultFolderName(InWorld, FFolder(GetWorldFolderRootObject(InWorld), ParentPath));
-	return DefaultFolderName.GetPath();
-}
-
-FName FActorFolders::GetDefaultFolderNameForSelection(UWorld& InWorld)
-{
-	FFolder FolderName = GetDefaultFolderForSelection(InWorld);
-	return FolderName.GetPath();
-}
-
-FName FActorFolders::GetFolderName(UWorld& InWorld, FName InParentPath, FName InFolderName)
-{
-	FFolder FolderName = GetFolderName(InWorld, FFolder(GetWorldFolderRootObject(InWorld), InParentPath), InFolderName);
-	return FolderName.GetPath();
-}
-
-void FActorFolders::CreateFolder(UWorld& InWorld, FName Path)
-{
-	CreateFolder(InWorld, FFolder(GetWorldFolderRootObject(InWorld), Path));
-}
-
-void FActorFolders::CreateFolderContainingSelection(UWorld& InWorld, FName Path)
-{
-	CreateFolderContainingSelection(InWorld, FFolder(GetWorldFolderRootObject(InWorld), Path));
-}
-
-void FActorFolders::SetSelectedFolderPath(FName Path) const
-{
-	SetSelectedFolderPath(FFolder(FFolder::GetInvalidRootObject(), Path));
-}
-
-void FActorFolders::DeleteFolder(UWorld& InWorld, FName FolderToDelete)
-{
-	DeleteFolder(InWorld, FFolder(GetWorldFolderRootObject(InWorld), FolderToDelete));
-}
-
-bool FActorFolders::RenameFolderInWorld(UWorld& InWorld, FName OldPath, FName NewPath)
-{
-	return RenameFolderInWorld(InWorld, FFolder(GetWorldFolderRootObject(InWorld), OldPath), FFolder(GetWorldFolderRootObject(InWorld), NewPath));
-}
-
-//~ End Deprecated
-////////////////////////////////////////////
-
 #undef LOCTEXT_NAMESPACE
\ No newline at end of file
diff --git a/Engine/Source/Editor/UnrealEd/Public/EditorActorFolders.h b/Engine/Source/Editor/UnrealEd/Public/EditorActorFolders.h
--- a/Engine/Source/Editor/UnrealEd/Public/EditorActorFolders.h
+++ b/Engine/Source/Editor/UnrealEd/Public/EditorActorFolders.h
@@ -63,51 +63,6 @@ struct FActorFolders : public FGCObject, public IActorEditorContextClient
 	static UNREALED_API FOnActorFolderCreated OnFolderCreated;
 	static UNREALED_API FOnActorFolderMoved 	OnFolderMoved;
 	static UNREALED_API FOnActorFolderDeleted OnFolderDeleted;
-
-	//~ Begin Deprecated
-
-	UE_DEPRECATED(5.0, "OnFolderCreate has been deprecated. Please use OnFolderCreated.")
-	static UNREALED_API FOnActorFolderCreate OnFolderCreate;
-
-	UE_DEPRECATED(5.0, "OnFolderMove has been deprecated. Please use OnFolderMoved.")
-	static UNREALED_API FOnActorFolderMove 	OnFolderMove;
-
-	UE_DEPRECATED(5.0, "OnFolderDelete has been deprecated. Please use OnFolderDeleted.")
-	static UNREALED_API FOnActorFolderDelete OnFolderDelete;
-
-	UE_DEPRECATED(5.0, "GetFolderProperties using FName has been deprecated. Please use new interface using FFolder.")
-	UNREALED_API FActorFolderProps* GetFolderProperties(UWorld& InWorld, FName InPath);
-
-	UE_DEPRECATED(5.0, "GetDefaultFolderName using FName  has been deprecated. Please use new interface using FFolder.")
-	UNREALED_API FName GetDefaultFolderName(UWorld& InWorld, FName ParentPath = FName());
-	
-	UE_DEPRECATED(5.0, "GetDefaultFolderNameForSelection using FName  has been deprecated. Please use new interface using FFolder.")
-	UNREALED_API FName GetDefaultFolderNameForSelection(UWorld& InWorld);
-
-	UE_DEPRECATED(5.0, "GetFolderName using FName  has been deprecated. Please use new interface using FFolder.")
-	UNREALED_API FName GetFolderName(UWorld& InWorld, FName ParentPath, FName FolderName);
-
-	UE_DEPRECATED(5.0, "CreateFolder using FName  has been deprecated. Please use new interface using FFolder.")
-	UNREALED_API void CreateFolder(UWorld& InWorld, FName Path);
-
-	UE_DEPRECATED(5.0, "CreateFolderContainingSelection using FName  has been deprecated. Please use new interface using FFolder.")
-	UNREALED_API void CreateFolderContainingSelection(UWorld& InWorld, FName Path);
-
-	UE_DEPRECATED(5.0, "SetSelectedFolderPath using FName  has been deprecated. Please use new interface using FFolder.")
-	UNREALED_API void SetSelectedFolderPath(FName Path) const;
-
-	UE_DEPRECATED(5.0, "DeleteFolder using FName  has been deprecated. Please use new interface using FFolder.")
-	UNREALED_API void DeleteFolder(UWorld& InWorld, FName FolderToDelete);
-
-	UE_DEPRECATED(5.0, "RenameFolderInWorld using FName  has been deprecated. Please use new interface using FFolder.")
-	UNREALED_API bool RenameFolderInWorld(UWorld& World, FName OldPath, FName NewPath);
-
-	UE_DEPRECATED(5.4, "Please use ForEachActorDescInstanceInFolders")
-	static void ForEachActorDescInFolders(UWorld& InWorld, const TSet<FName>& InPaths, TFunctionRef<bool(const FWorldPartitionActorDesc*)> Operation, const FFolder::FRootObject& InFolderRootObject = FFolder::GetInvalidRootObject()) {}
-	
-	UE_DEPRECATED(5.4, "Please use GetActorDescInstanceFolder")
-	static FFolder GetActorDescFolder(UWorld& InWorld, const FWorldPartitionActorDesc* InActorDesc) { return FFolder::GetInvalidFolder(); }
-	//~ End Deprecated
 	
 	static UNREALED_API FFolder GetActorDescInstanceFolder(UWorld& InWorld, const FWorldPartitionActorDescInstance* InActorDescInstance);
 	/** Apply an operation to each actor desc in the given list of folders. */
