Relax the duplicate package localization id warning

#Synced-CL 50929914

diff --git a/Engine/Source/Editor/UnrealEd/Private/Commandlets/GatherTextFromAssetsCommandlet.cpp b/Engine/Source/Editor/UnrealEd/Private/Commandlets/GatherTextFromAssetsCommandlet.cpp
--- a/Engine/Source/Editor/UnrealEd/Private/Commandlets/GatherTextFromAssetsCommandlet.cpp
+++ b/Engine/Source/Editor/UnrealEd/Private/Commandlets/GatherTextFromAssetsCommandlet.cpp
@@ -1424,11 +1424,19 @@ void UGatherTextFromAssetsCommandlet::ProcessAndRemoveCachedPackages(TMap<FName,
 				// if this package's localization ID is a duplicate
 				if (const FName* ExistingLongPackageName = AssignedPackageLocalizationIds.Find(PackagePendingGather.PackageLocalizationId))
 				{
-					UE_LOGFMT(LogGatherTextFromAssetsCommandlet, Warning, "Package '{file}' and '{conflictFile}' have the same localization ID ({locKey}). Please reset both of these to avoid conflicts (one will keep the ID and the other one should actually reset). In the Content Browser, Right-Click the assets -> Asset Localization -> Reset Localization ID.", 
-						("file", *FPackageName::LongPackageNameToFilename(*PackageNameStr, TEXT(".uasset"))),
-						("conflictFile", *FPackageName::LongPackageNameToFilename(ExistingLongPackageName->ToString(), TEXT(".uasset"))),
-						("locKey", *PackagePendingGather.PackageLocalizationId), 
-						("id", GatherTextFromAssetsCommandlet::LocalizationLogIdentifier));
+					// Use a simple heuristic to see if we should actually emit a warning (otherwise, in big projects this warning 
+					// might get emitted when no actual problems will happen on the vast majority of duplicate package id).
+					// Heuristic: Only emit the warning if the package's cache cannot be used by the heuristic OR if the package have non-empty localization cache
+					if (PackagePendingGather.PackageLocCacheState != EPackageLocCacheState::Cached ||
+						!PackagePendingGather.GatherableTextDataArray.IsEmpty())
+					{
+						UE_LOGFMT(LogGatherTextFromAssetsCommandlet, Warning, "Package '{file}' and '{conflictFile}' have the same localization ID ({locKey}). Please reset both of these to avoid conflicts (one will probably deterministically keep the ID and the other one should actually reset). In the Content Browser, Right-Click the assets -> Asset Localization -> Reset Localization ID.",
+							("file", *FPackageName::LongPackageNameToFilename(*PackageNameStr, TEXT(".uasset"))),
+							("conflictFile", *FPackageName::LongPackageNameToFilename(ExistingLongPackageName->ToString(), TEXT(".uasset"))),
+							("locKey", *PackagePendingGather.PackageLocalizationId),
+							("id", GatherTextFromAssetsCommandlet::LocalizationLogIdentifier));
+					}
+
 					if (bFixPackageLocalizationIdConflict) 
 					{
 						bThisPackageHasLocIdConflictToFix = true;
