Change CookIncrementalVersion because of shader type issue

#Synced-CL 0

diff --git a/Engine/Source/Editor/UnrealEd/Private/Cooker/CookIncrementalVersion.cpp b/Engine/Source/Editor/UnrealEd/Private/Cooker/CookIncrementalVersion.cpp
--- a/Engine/Source/Editor/UnrealEd/Private/Cooker/CookIncrementalVersion.cpp
+++ b/Engine/Source/Editor/UnrealEd/Private/Cooker/CookIncrementalVersion.cpp
@@ -8,6 +8,6 @@ namespace UE::Cook
 {
 
 // Change CookIncrementalVersion to a new guid when all packages in an incremental cook need to be invalidated
-FGuid CookIncrementalVersion(0xAFB442DC, 0x78124FE2, 0xBEA2667A, 0x2D71F307);
+FGuid CookIncrementalVersion(0x18BDDE20,0xCF584012,0xBCD9FCE6,0x6CFC05A5);
 
 }
