IncrementalCook: Move the ICookArtifactReader interface into CoreUObject to make it available to ICookArtifact. ICookArtifact and now ICookArtifactReader are in WITH_EDITOR in CoreUObject so the code for cooking a type can be put in WITH_EDITOR blocks in runtime modules such as Engine, next to the runtime code that they cook for. Move ICookArtifactReader into the UE::Cook namespace to match ICookArtifact.
#rnx

#Synced-CL 50939324

#Synced-CL 50947327

diff --git a/Engine/Source/Developer/IoStoreUtilities/Private/LayeredCookArtifactReader.cpp b/Engine/Source/Developer/IoStoreUtilities/Private/LayeredCookArtifactReader.cpp
--- a/Engine/Source/Developer/IoStoreUtilities/Private/LayeredCookArtifactReader.cpp
+++ b/Engine/Source/Developer/IoStoreUtilities/Private/LayeredCookArtifactReader.cpp
@@ -5,12 +5,12 @@
 #include "HAL/CriticalSection.h"
 #include "Misc/ScopeRWLock.h"
 
-void FLayeredCookArtifactReader::AddLayer(TSharedRef<ICookArtifactReader> InLayer)
+void FLayeredCookArtifactReader::AddLayer(TSharedRef<UE::Cook::ICookArtifactReader> InLayer)
 {
 	Layers.AddUnique(InLayer);
 }
 
-bool FLayeredCookArtifactReader::RemoveLayer(TSharedRef<ICookArtifactReader> InLayer)
+bool FLayeredCookArtifactReader::RemoveLayer(TSharedRef<UE::Cook::ICookArtifactReader> InLayer)
 {
 	return Layers.Remove(InLayer) > 0;
 }
@@ -22,7 +22,7 @@ void FLayeredCookArtifactReader::EmptyLayers()
 
 bool FLayeredCookArtifactReader::FileExists(const TCHAR* Filename)
 {
-	for (const TSharedRef<ICookArtifactReader>& Layer : Layers)
+	for (const TSharedRef<UE::Cook::ICookArtifactReader>& Layer : Layers)
 	{
 		if (Layer->FileExists(Filename))
 		{
@@ -34,7 +34,7 @@ bool FLayeredCookArtifactReader::FileExists(const TCHAR* Filename)
 
 int64 FLayeredCookArtifactReader::FileSize(const TCHAR* Filename)
 {
-	for (const TSharedRef<ICookArtifactReader>& Layer : Layers)
+	for (const TSharedRef<UE::Cook::ICookArtifactReader>& Layer : Layers)
 	{
 		int64 Size = Layer->FileSize(Filename);
 		if (Size >= 0)
@@ -48,7 +48,7 @@ int64 FLayeredCookArtifactReader::FileSize(const TCHAR* Filename)
 
 IFileHandle* FLayeredCookArtifactReader::OpenRead(const TCHAR* Filename)
 {
-	for (const TSharedRef<ICookArtifactReader>& Layer : Layers)
+	for (const TSharedRef<UE::Cook::ICookArtifactReader>& Layer : Layers)
 	{
 		if (IFileHandle* File = Layer->OpenRead(Filename))
 		{
@@ -90,7 +90,7 @@ bool FLayeredCookArtifactReader::IterateDirectory(const TCHAR* Directory, IPlatf
 	FMergingVisitor MergingVisitor(Entries);
 	bool bRetVal = false;
 
-	for (const TSharedRef<ICookArtifactReader>& Layer : Layers)
+	for (const TSharedRef<UE::Cook::ICookArtifactReader>& Layer : Layers)
 	{
 		bRetVal |= Layer->IterateDirectory(Directory, MergingVisitor);
 	}
diff --git a/Engine/Source/Developer/IoStoreUtilities/Private/ZenStoreWriter.cpp b/Engine/Source/Developer/IoStoreUtilities/Private/ZenStoreWriter.cpp
--- a/Engine/Source/Developer/IoStoreUtilities/Private/ZenStoreWriter.cpp
+++ b/Engine/Source/Developer/IoStoreUtilities/Private/ZenStoreWriter.cpp
@@ -211,7 +211,7 @@ FZenStoreWriter::FZenStoreWriter(
 	const FString& InOutputPath,
 	const FString& InMetadataDirectoryPath,
 	const ITargetPlatform* InTargetPlatform,
-	TSharedRef<ICookArtifactReader> InCookArtifactReader
+	TSharedRef<UE::Cook::ICookArtifactReader> InCookArtifactReader
 )
 	: CookArtifactReader(InCookArtifactReader)
 	, TargetPlatform(*InTargetPlatform)
diff --git a/Engine/Source/Developer/IoStoreUtilities/Public/CookArtifactReader.h b/Engine/Source/Developer/IoStoreUtilities/Public/CookArtifactReader.h
--- a/Engine/Source/Developer/IoStoreUtilities/Public/CookArtifactReader.h
+++ b/Engine/Source/Developer/IoStoreUtilities/Public/CookArtifactReader.h
@@ -2,28 +2,10 @@
 
 #pragma once
 
-#include "Containers/Array.h"
-#include "Containers/UnrealString.h"
-#include "GenericPlatform/GenericPlatformFile.h"
+UE_DEPRECATED_HEADER(5.8, "Include Cooker/CookArtifactReader.h (in CoreUObject)")
 
-class FArchive;
+#include "Cooker/CookArtifactReader.h"
 
-class ICookArtifactReader
-{
-public:
-	virtual ~ICookArtifactReader() = default;
+UE_DEPRECATED(5.8, "Use UE::Cook::ICookArtifactReader.")
+typedef UE::Cook::ICookArtifactReader ICookArtifactReader;
 
-	virtual void Initialize(bool bCleanBuild) {}
-	
-	virtual bool FileExists(const TCHAR* Filename) = 0;
-	virtual int64 FileSize(const TCHAR* Filename) = 0;
-	virtual IFileHandle* OpenRead(const TCHAR* Filename) = 0;
-
-	virtual bool IterateDirectory(const TCHAR* Directory, IPlatformFile::FDirectoryVisitor& Visitor) = 0;
-
-	// Utilities
-	virtual FArchive* CreateFileReader(const TCHAR* Filename) = 0;
-	virtual bool IterateDirectoryRecursively(const TCHAR* Directory, IPlatformFile::FDirectoryVisitor& Visitor) = 0;
-	virtual void FindFiles(TArray<FString>& Result, const TCHAR* Filename, bool Files, bool Directories) = 0;
-	virtual void FindFiles(TArray<FString>& FoundFiles, const TCHAR* Directory, const TCHAR* FileExtension = nullptr) = 0;
-};
diff --git a/Engine/Source/Developer/IoStoreUtilities/Public/CookArtifactReaderCommon.h b/Engine/Source/Developer/IoStoreUtilities/Public/CookArtifactReaderCommon.h
--- a/Engine/Source/Developer/IoStoreUtilities/Public/CookArtifactReaderCommon.h
+++ b/Engine/Source/Developer/IoStoreUtilities/Public/CookArtifactReaderCommon.h
@@ -4,13 +4,13 @@
 
 #include "Containers/Array.h"
 #include "Containers/UnrealString.h"
-#include "CookArtifactReader.h"
+#include "Cooker/CookArtifactReader.h"
 #include "GenericPlatform/GenericPlatformFile.h"
 
 class FArchive;
 
 class FCookArtifactReaderCommon
-	: public ICookArtifactReader
+	: public UE::Cook::ICookArtifactReader
 {
 public:
 	virtual ~FCookArtifactReaderCommon() = default;
diff --git a/Engine/Source/Developer/IoStoreUtilities/Public/LayeredCookArtifactReader.h b/Engine/Source/Developer/IoStoreUtilities/Public/LayeredCookArtifactReader.h
--- a/Engine/Source/Developer/IoStoreUtilities/Public/LayeredCookArtifactReader.h
+++ b/Engine/Source/Developer/IoStoreUtilities/Public/LayeredCookArtifactReader.h
@@ -12,10 +12,16 @@ class FLayeredCookArtifactReader
 public:
 	FLayeredCookArtifactReader() = default;
 
-	void Initialize(bool bCleanBuild) override { for (TSharedRef<ICookArtifactReader>& Layer : Layers) { Layer->Initialize(bCleanBuild); } }
+	void Initialize(bool bCleanBuild) override
+	{
+		for (TSharedRef<UE::Cook::ICookArtifactReader>& Layer : Layers)
+		{
+			Layer->Initialize(bCleanBuild);
+		}
+	}
 
-	IOSTOREUTILITIES_API void AddLayer(TSharedRef<ICookArtifactReader> InLayer);
-	IOSTOREUTILITIES_API bool RemoveLayer(TSharedRef<ICookArtifactReader> InLayer);
+	IOSTOREUTILITIES_API void AddLayer(TSharedRef<UE::Cook::ICookArtifactReader> InLayer);
+	IOSTOREUTILITIES_API bool RemoveLayer(TSharedRef<UE::Cook::ICookArtifactReader> InLayer);
 	IOSTOREUTILITIES_API void EmptyLayers();
 
 	IOSTOREUTILITIES_API bool FileExists(const TCHAR* Filename) override;
@@ -24,5 +30,5 @@ public:
 
 	IOSTOREUTILITIES_API bool IterateDirectory(const TCHAR* Directory, IPlatformFile::FDirectoryVisitor& Visitor) override;
 private:
-	TArray<TSharedRef<ICookArtifactReader>> Layers;
+	TArray<TSharedRef<UE::Cook::ICookArtifactReader>> Layers;
 };
diff --git a/Engine/Source/Developer/IoStoreUtilities/Public/ZenStoreWriter.h b/Engine/Source/Developer/IoStoreUtilities/Public/ZenStoreWriter.h
--- a/Engine/Source/Developer/IoStoreUtilities/Public/ZenStoreWriter.h
+++ b/Engine/Source/Developer/IoStoreUtilities/Public/ZenStoreWriter.h
@@ -36,20 +36,16 @@ class FLargeMemoryWriter;
 class FSharedBuffer;
 class ITargetPlatform;
 struct FFileRegion;
-template <typename FuncType> class TFunction;
-
-namespace UE {
-	class FZenStoreHttpClient;
-}
-
 class FCbAttachment;
 class FCbPackage;
 class FCbWriter;
 class FPackageStoreOptimizer;
 class FPackageStorePackage;
 class FZenFileSystemManifest;
+namespace UE { class FZenStoreHttpClient; }
+namespace UE::Cook { class ICookArtifactReader; }
 struct FZenFileSystemManifestEntry;
-class ICookArtifactReader;
+template <typename FuncType> class TFunction;
 
 /** 
  * A PackageStoreWriter that saves cooked packages for use by IoStore, and stores them in the Zen storage service.
@@ -61,7 +57,7 @@ public:
 	IOSTOREUTILITIES_API FZenStoreWriter(	const FString& OutputPath, 
 											const FString& MetadataDirectoryPath, 
 											const ITargetPlatform* TargetPlatform,
-											TSharedRef<ICookArtifactReader> CookArtifactReader);
+											TSharedRef<UE::Cook::ICookArtifactReader> CookArtifactReader);
 
 	IOSTOREUTILITIES_API ~FZenStoreWriter();
 
@@ -261,7 +257,7 @@ private:
 	TUniquePtr<UE::FZenStoreHttpClient>	HttpClient;
 	bool IsLocalConnection = true;
 
-	TSharedRef<ICookArtifactReader>		CookArtifactReader;
+	TSharedRef<UE::Cook::ICookArtifactReader>		CookArtifactReader;
 	const ITargetPlatform&				TargetPlatform;
 	const FName							TargetPlatformFName;
 	FString								ProjectId;
diff --git a/Engine/Source/Editor/UnrealEd/Classes/CookOnTheSide/CookOnTheFlyServer.h b/Engine/Source/Editor/UnrealEd/Classes/CookOnTheSide/CookOnTheFlyServer.h
--- a/Engine/Source/Editor/UnrealEd/Classes/CookOnTheSide/CookOnTheFlyServer.h
+++ b/Engine/Source/Editor/UnrealEd/Classes/CookOnTheSide/CookOnTheFlyServer.h
@@ -38,19 +38,19 @@ class FLooseFilesCookArtifactReader;
 class FReferenceCollector;
 class FSavePackageContext;
 class IAssetRegistry;
-class ICookArtifactReader;
 class IPlugin;
 class ITargetPlatform;
 enum class EPackageWriterResult : uint8;
 enum class ODSCRecompileCommand;
+namespace UE::Cook { class ICookArtifactReader; }
+namespace UE::Cook { void CalculateGlobalDependenciesHash(const ITargetPlatform* Platform, const UCookOnTheFlyServer& COTFS); }
+namespace UE::LinkerLoad { enum class EImportBehavior : uint8; }
 struct FBeginCookContext;
 struct FCrashContextExtendedWriter;
 struct FGenericMemoryStats;
 struct FPropertyChangedEvent;
 struct FResourceSizeEx;
 
-namespace UE::Cook { void CalculateGlobalDependenciesHash(const ITargetPlatform* Platform, const UCookOnTheFlyServer& COTFS); }
-namespace UE::LinkerLoad { enum class EImportBehavior : uint8; }
 
 enum class ECookInitializationFlags
 {
@@ -1511,8 +1511,8 @@ private:
 	UE::Cook::EPollStatus TryPreSaveGeneratedPackage(UE::Cook::FGenerationHelper& GenerationHelper,
 		UE::Cook::FCookGenerationInfo& GeneratedInfo);
 
-	ICookArtifactReader& FindOrCreateCookArtifactReader(const ITargetPlatform* TargetPlatform);
-	const ICookArtifactReader* FindCookArtifactReader(const ITargetPlatform* TargetPlatform) const;
+	UE::Cook::ICookArtifactReader& FindOrCreateCookArtifactReader(const ITargetPlatform* TargetPlatform);
+	const UE::Cook::ICookArtifactReader* FindCookArtifactReader(const ITargetPlatform* TargetPlatform) const;
 	ICookedPackageWriter& FindOrCreatePackageWriter(const ITargetPlatform* TargetPlatform);
 	const ICookedPackageWriter* FindPackageWriter(const ITargetPlatform* TargetPlatform) const;
 	void FindOrCreateSaveContexts(TConstArrayView<const ITargetPlatform*> TargetPlatforms);
diff --git a/Engine/Source/Editor/UnrealEd/Private/CookOnTheFlyServer.cpp b/Engine/Source/Editor/UnrealEd/Private/CookOnTheFlyServer.cpp
--- a/Engine/Source/Editor/UnrealEd/Private/CookOnTheFlyServer.cpp
+++ b/Engine/Source/Editor/UnrealEd/Private/CookOnTheFlyServer.cpp
@@ -10460,7 +10460,9 @@ void UCookOnTheFlyServer::LoadBeginCookIncrementalFlagsLocal(FBeginCookContext&
 						= PlatformContext.CookArtifactsCurrentSettings.FindChecked(Artifact.GetReference());
 					FString PreviousFileName = GetCookSettingsFileName(TargetPlatform, ArtifactName);
 
-					UE::Cook::Artifact::FCompareSettingsContext Context(*this, Previous, Current.Settings, PreviousFileName);
+					UE::Cook::Artifact::FCompareSettingsContext Context(*this,
+						FindOrCreateCookArtifactReader(TargetPlatform),
+						Previous, Current.Settings, PreviousFileName);
 					Context.TargetPlatform = TargetPlatform;
 					Artifact->CompareSettings(Context);
 
@@ -10521,6 +10523,8 @@ void UCookOnTheFlyServer::LoadBeginCookIncrementalFlagsLocal(FBeginCookContext&
 
 void UCookOnTheFlyServer::BeginCookSandbox(FBeginCookContext& BeginContext)
 {
+	using namespace UE::Cook;
+
 #if OUTPUT_COOKTIMING
 	double CleanSandboxTime = 0.0;
 #endif
@@ -10532,8 +10536,8 @@ void UCookOnTheFlyServer::BeginCookSandbox(FBeginCookContext& BeginContext)
 		for (FBeginCookContextPlatform& PlatformContext : BeginContext.PlatformContexts)
 		{
 			const ITargetPlatform* TargetPlatform = PlatformContext.TargetPlatform;
-			UE::Cook::FPlatformData* PlatformData = PlatformContext.PlatformData;
-			UE::Cook::FCookSavePackageContext& SavePackageContext = FindOrCreateSaveContext(TargetPlatform);
+			FPlatformData* PlatformData = PlatformContext.PlatformData;
+			FCookSavePackageContext& SavePackageContext = FindOrCreateSaveContext(TargetPlatform);
 			ICookedPackageWriter& PackageWriter = *SavePackageContext.PackageWriter;
 			ICookedPackageWriter::FCookInfo CookInfo;
 			CookInfo.CookMode = IsDirectorCookOnTheFly() ? ICookedPackageWriter::FCookInfo::CookOnTheFlyMode : ICookedPackageWriter::FCookInfo::CookByTheBookMode;
@@ -10570,9 +10574,9 @@ void UCookOnTheFlyServer::BeginCookSandbox(FBeginCookContext& BeginContext)
 			if (!PlatformContext.bWorkerOnSharedSandbox)
 			{
 				check(!IsCookWorkerMode());
-				for (const TRefCountPtr<UE::Cook::ICookArtifact>& Artifact: CookArtifacts)
+				for (const TRefCountPtr<ICookArtifact>& Artifact: CookArtifacts)
 				{
-					UE::Cook::FCookArtifactsCurrentSettings& Current
+					FCookArtifactsCurrentSettings& Current
 						= PlatformContext.CookArtifactsCurrentSettings.FindChecked(Artifact.GetReference());
 					FString ArtifactName = Artifact->GetArtifactName();
 
@@ -10612,7 +10616,7 @@ void UCookOnTheFlyServer::BeginCookSandbox(FBeginCookContext& BeginContext)
 			// TODO: Initializing NumPackagesIncrementallyModifiedByClass is not yet implemented in this path (cooking
 			// multiple times in editor process) because we plan to deprecate this path soon.
 			const ITargetPlatform* TargetPlatform = AlreadyCookedPlatforms[0];
-			PackageDatas->LockAndEnumeratePackageDatas([TargetPlatform](UE::Cook::FPackageData* PackageData)
+			PackageDatas->LockAndEnumeratePackageDatas([TargetPlatform](FPackageData* PackageData)
 			{
 				if (PackageData->HasCookedPlatform(TargetPlatform, true /* bIncludeFailed */))
 				{
@@ -12505,12 +12509,14 @@ bool UCookOnTheFlyServer::GetAllPackageFilenamesFromAssetRegistry(const FString&
 	return false;
 }
 
-ICookArtifactReader& UCookOnTheFlyServer::FindOrCreateCookArtifactReader(const ITargetPlatform* TargetPlatform)
+UE::Cook::ICookArtifactReader& UCookOnTheFlyServer::FindOrCreateCookArtifactReader(
+	const ITargetPlatform* TargetPlatform)
 {
 	return *FindOrCreateSaveContext(TargetPlatform).ArtifactReader;
 }
 
-const ICookArtifactReader* UCookOnTheFlyServer::FindCookArtifactReader(const ITargetPlatform* TargetPlatform) const
+const UE::Cook::ICookArtifactReader* UCookOnTheFlyServer::FindCookArtifactReader(
+	const ITargetPlatform* TargetPlatform) const
 {
 	const UE::Cook::FCookSavePackageContext* Context = FindSaveContext(TargetPlatform);
 	return Context ? Context->ArtifactReader.Get() : nullptr;
diff --git a/Engine/Source/Editor/UnrealEd/Private/Cooker/CookTypes.cpp b/Engine/Source/Editor/UnrealEd/Private/Cooker/CookTypes.cpp
--- a/Engine/Source/Editor/UnrealEd/Private/Cooker/CookTypes.cpp
+++ b/Engine/Source/Editor/UnrealEd/Private/Cooker/CookTypes.cpp
@@ -3,8 +3,8 @@
 #include "CookTypes.h"
 
 #include "Containers/StringView.h"
-#include "CookArtifactReader.h"
 #include "Cooker/CompactBinaryTCP.h"
+#include "Cooker/CookArtifactReader.h"
 #include "Cooker/CookDeterminismManager.h"
 #include "Cooker/CookPackageData.h"
 #include "Cooker/PackageTracker.h"
diff --git a/Engine/Source/Editor/UnrealEd/Private/Cooker/CookTypes.h b/Engine/Source/Editor/UnrealEd/Private/Cooker/CookTypes.h
--- a/Engine/Source/Editor/UnrealEd/Private/Cooker/CookTypes.h
+++ b/Engine/Source/Editor/UnrealEd/Private/Cooker/CookTypes.h
@@ -29,10 +29,10 @@
 
 class FCbFieldView;
 class FCbWriter;
-class ICookArtifactReader;
 class ITargetPlatform;
 struct FWeakObjectPtr;
 namespace UE::DerivedData { class FBuildDefinition; }
+namespace UE::Cook { class ICookArtifactReader; }
 namespace UE::Cook { enum class EReachability : uint8; }
 namespace UE::Cook { struct FBeginCookConfigSettings; }
 namespace UE::Cook { struct FCookByTheBookOptions; }
diff --git a/Engine/Source/Runtime/CoreUObject/Public/Cooker/CookArtifact.h b/Engine/Source/Runtime/CoreUObject/Public/Cooker/CookArtifact.h
--- a/Engine/Source/Runtime/CoreUObject/Public/Cooker/CookArtifact.h
+++ b/Engine/Source/Runtime/CoreUObject/Public/Cooker/CookArtifact.h
@@ -11,6 +11,7 @@
 
 class ITargetPlatform;
 class UCookOnTheFlyServer;
+namespace UE::Cook { class ICookArtifactReader; }
 namespace UE::Cook { class ICookInfo; }
 namespace UE::Cook::Artifact { struct FUpdateOplogPackagesContext; }
 
@@ -22,6 +23,8 @@ struct FCompareSettingsContext
 {
 public:
 	ICookInfo& GetCookInfo() const;
+	UE::Cook::ICookArtifactReader& GetArtifactReader() const;
+
 	/** The TargetPlatform being tested; invalidation requests are allowed to differ between platforms. Non-null. */
 	const ITargetPlatform* GetTargetPlatform() const;
 	const FConfigFile& GetPrevious() const;
@@ -43,10 +46,11 @@ public:
 	void RequestFullRecook(bool bValue);
 
 private:
-	FCompareSettingsContext(ICookInfo& InCookInfo, const FConfigFile& InPrevious, const FConfigFile& InCurrent,
-		const FString& InPreviousFileName);
+	FCompareSettingsContext(ICookInfo& InCookInfo, UE::Cook::ICookArtifactReader& InArtifactReader,
+		const FConfigFile& InPrevious, const FConfigFile& InCurrent, const FString& InPreviousFileName);
 
 	ICookInfo& CookInfo;
+	UE::Cook::ICookArtifactReader& ArtifactReader;
 	const ITargetPlatform* TargetPlatform = nullptr;
 	const FConfigFile& Previous;
 	const FConfigFile& Current;
@@ -196,8 +200,10 @@ namespace UE::Cook::Artifact
 {
 
 inline FCompareSettingsContext::FCompareSettingsContext(ICookInfo& InCookInfo,
-	const FConfigFile& InPrevious, const FConfigFile& InCurrent, const FString& InPreviousFileName)
+	UE::Cook::ICookArtifactReader& InArtifactReader, const FConfigFile& InPrevious, const FConfigFile& InCurrent,
+	const FString& InPreviousFileName)
 	: CookInfo(InCookInfo)
+	, ArtifactReader(InArtifactReader)
 	, Previous(InPrevious)
 	, Current(InCurrent)
 	, PreviousFileName(InPreviousFileName)
@@ -208,6 +214,10 @@ inline ICookInfo& FCompareSettingsContext::GetCookInfo() const
 {
 	return CookInfo;
 }
+inline UE::Cook::ICookArtifactReader& FCompareSettingsContext::GetArtifactReader() const
+{
+	return ArtifactReader;
+}
 inline const ITargetPlatform* FCompareSettingsContext::GetTargetPlatform() const
 {
 	return TargetPlatform;
diff --git a/Engine/Source/Runtime/CoreUObject/Public/Cooker/CookArtifactReader.h b/Engine/Source/Runtime/CoreUObject/Public/Cooker/CookArtifactReader.h
new file mode 100644
--- /dev/null
+++ b/Engine/Source/Runtime/CoreUObject/Public/Cooker/CookArtifactReader.h
@@ -0,0 +1,43 @@
+// Copyright Epic Games, Inc. All Rights Reserved.
+
+#pragma once
+
+#if WITH_EDITOR || IS_PROGRAM
+
+#include "Containers/Array.h"
+#include "Containers/UnrealString.h"
+#include "GenericPlatform/GenericPlatformFile.h"
+
+class FArchive;
+
+namespace UE::Cook
+{
+
+/**
+ * An interface for reading files that can be provided either from disk or from ZenServer storage of incremental cook
+ * artifacts. Paths passed to this interface should be the same as the on-disk location of the artifacts during a cook,
+ * e.g. ICookInfo.GetCookMetadataOutputFolder(TargetPlatform).
+ */
+class ICookArtifactReader
+{
+public:
+	virtual ~ICookArtifactReader() = default;
+
+	virtual void Initialize(bool bCleanBuild) {}
+	
+	virtual bool FileExists(const TCHAR* Filename) = 0;
+	virtual int64 FileSize(const TCHAR* Filename) = 0;
+	virtual IFileHandle* OpenRead(const TCHAR* Filename) = 0;
+
+	virtual bool IterateDirectory(const TCHAR* Directory, IPlatformFile::FDirectoryVisitor& Visitor) = 0;
+
+	// Utilities
+	virtual FArchive* CreateFileReader(const TCHAR* Filename) = 0;
+	virtual bool IterateDirectoryRecursively(const TCHAR* Directory, IPlatformFile::FDirectoryVisitor& Visitor) = 0;
+	virtual void FindFiles(TArray<FString>& Result, const TCHAR* Filename, bool Files, bool Directories) = 0;
+	virtual void FindFiles(TArray<FString>& FoundFiles, const TCHAR* Directory, const TCHAR* FileExtension = nullptr) = 0;
+};
+
+} // namespace UE::Cook
+
+#endif // WITH_EDITOR
\ No newline at end of file
diff --git a/Engine/Source/Runtime/RenderCore/Private/ShaderLibrary/EditorShaderStableInfo.cpp b/Engine/Source/Runtime/RenderCore/Private/ShaderLibrary/EditorShaderStableInfo.cpp
--- a/Engine/Source/Runtime/RenderCore/Private/ShaderLibrary/EditorShaderStableInfo.cpp
+++ b/Engine/Source/Runtime/RenderCore/Private/ShaderLibrary/EditorShaderStableInfo.cpp
@@ -31,7 +31,7 @@
 #include "Serialization/CompactBinarySerialization.h"
 #include "Serialization/CompactBinaryWriter.h"
 #include "RHIStrings.h"
-#include "CookArtifactReader.h"
+#include "Cooker/CookArtifactReader.h"
 #include "Misc/App.h"
 #include "Misc/ConfigCacheIni.h"
 #include "Misc/EngineVersion.h"
diff --git a/Engine/Source/Runtime/RenderCore/Private/ShaderLibrary/ShaderCodeLibrary.cpp b/Engine/Source/Runtime/RenderCore/Private/ShaderLibrary/ShaderCodeLibrary.cpp
--- a/Engine/Source/Runtime/RenderCore/Private/ShaderLibrary/ShaderCodeLibrary.cpp
+++ b/Engine/Source/Runtime/RenderCore/Private/ShaderLibrary/ShaderCodeLibrary.cpp
@@ -61,8 +61,8 @@ ShaderCodeLibrary.cpp: Bound shader state cache implementation.
 #include "Serialization/CompactBinarySerialization.h"
 #include "Serialization/CompactBinaryWriter.h"
 #include "RHIStrings.h"
-#include "CookArtifactReader.h"
 #include "Cooker/CookArtifact.h"
+#include "Cooker/CookArtifactReader.h"
 #endif
 
 // allow introspection (e.g. dumping the contents) for easier debugging
@@ -2414,7 +2414,7 @@ void FShaderCodeLibrary::CloseLibrary(FString const& Name)
 
 #if WITH_EDITOR
 // for now a lot of FShaderLibraryCooker code is aliased with the runtime code, but this will be refactored (UE-103486)
-void FShaderLibraryCooker::InitForCooking(bool bNativeFormat, ICookArtifactReader* InCookArtifactReader)
+void FShaderLibraryCooker::InitForCooking(bool bNativeFormat, UE::Cook::ICookArtifactReader* InCookArtifactReader)
 {
 	CookArtifactReader = InCookArtifactReader;
 	FShaderLibrariesCollection::Impl = new FShaderLibrariesCollection(SP_NumPlatforms, bNativeFormat);
diff --git a/Engine/Source/Runtime/RenderCore/Private/ShaderLibrary/ShaderCodeLibraryUtilities.cpp b/Engine/Source/Runtime/RenderCore/Private/ShaderLibrary/ShaderCodeLibraryUtilities.cpp
--- a/Engine/Source/Runtime/RenderCore/Private/ShaderLibrary/ShaderCodeLibraryUtilities.cpp
+++ b/Engine/Source/Runtime/RenderCore/Private/ShaderLibrary/ShaderCodeLibraryUtilities.cpp
@@ -14,7 +14,7 @@
 #include "ShaderCompilerCore.h" // for GetBuildMachineArtifactBasePath()
 
 #if WITH_EDITOR
-#include "CookArtifactReader.h"
+#include "Cooker/CookArtifactReader.h"
 #endif
 
 namespace UE::ShaderLibrary::Private
@@ -73,7 +73,7 @@ namespace UE::ShaderLibrary::Private
 	);
 
 #if WITH_EDITOR
-	ICookArtifactReader* CookArtifactReader = nullptr;
+	UE::Cook::ICookArtifactReader* CookArtifactReader = nullptr;
 #endif // WITH_EDITOR
 
 	FString GetShaderCodeIntermediatePath()
diff --git a/Engine/Source/Runtime/RenderCore/Private/ShaderLibrary/ShaderCodeLibraryUtilities.h b/Engine/Source/Runtime/RenderCore/Private/ShaderLibrary/ShaderCodeLibraryUtilities.h
--- a/Engine/Source/Runtime/RenderCore/Private/ShaderLibrary/ShaderCodeLibraryUtilities.h
+++ b/Engine/Source/Runtime/RenderCore/Private/ShaderLibrary/ShaderCodeLibraryUtilities.h
@@ -12,7 +12,7 @@
 #include "UObject/NameTypes.h"
 #include "ShaderLibrary/ShaderCodeArchiveTypes.h"
 
-class ICookArtifactReader;
+namespace UE::Cook { class ICookArtifactReader; }
 
 namespace UE::ShaderLibrary::Private
 {
@@ -44,7 +44,7 @@ namespace UE::ShaderLibrary::Private
 
 	// Shader library intermediates:
 #if WITH_EDITOR
-	extern ICookArtifactReader* CookArtifactReader;
+	extern UE::Cook::ICookArtifactReader* CookArtifactReader;
 #endif // WITH_EDITOR
 
 	// Shader library utility functions:
diff --git a/Engine/Source/Runtime/RenderCore/Public/ShaderCodeLibrary.h b/Engine/Source/Runtime/RenderCore/Public/ShaderCodeLibrary.h
--- a/Engine/Source/Runtime/RenderCore/Public/ShaderCodeLibrary.h
+++ b/Engine/Source/Runtime/RenderCore/Public/ShaderCodeLibrary.h
@@ -42,7 +42,7 @@ class UObject;
 #if WITH_EDITOR
 class FCbFieldView;
 class FCbWriter;
-class ICookArtifactReader;
+namespace UE::Cook { class ICookArtifactReader; }
 namespace UE::Cook::Artifact { struct FUpdateOplogPackagesContext; }
 #endif
 
@@ -294,7 +294,7 @@ class ITargetPlatform;
 struct FShaderLibraryCooker
 {
 	// Initialize the library cooker
-	static RENDERCORE_API void InitForCooking(bool bNativeFormat, ICookArtifactReader* CookArtifactReader);
+	static RENDERCORE_API void InitForCooking(bool bNativeFormat, UE::Cook::ICookArtifactReader* CookArtifactReader);
 	// Shutdown the library cooker
 	static RENDERCORE_API void Shutdown();
 
