Additional profiling traces for cook startup

#rnx

#Synced-CL 50969558

diff --git a/Engine/Source/Editor/UnrealEd/Private/Commandlets/AssetRegistryGenerator.cpp b/Engine/Source/Editor/UnrealEd/Private/Commandlets/AssetRegistryGenerator.cpp
--- a/Engine/Source/Editor/UnrealEd/Private/Commandlets/AssetRegistryGenerator.cpp
+++ b/Engine/Source/Editor/UnrealEd/Private/Commandlets/AssetRegistryGenerator.cpp
@@ -220,6 +220,7 @@ private:
 
 void FAssetRegistryGenerator::UpdateAssetManagerDatabase()
 {
+	TRACE_CPUPROFILER_EVENT_SCOPE(FAssetRegistryGenerator::UpdateAssetManagerDatabase);
 	LLM_SCOPE_BYTAG(Cooker_GeneratedAssetRegistry);
 	UAssetManager::Get().UpdateManagementDatabase();
 }
diff --git a/Engine/Source/Editor/UnrealEd/Private/CookOnTheFlyServer.cpp b/Engine/Source/Editor/UnrealEd/Private/CookOnTheFlyServer.cpp
--- a/Engine/Source/Editor/UnrealEd/Private/CookOnTheFlyServer.cpp
+++ b/Engine/Source/Editor/UnrealEd/Private/CookOnTheFlyServer.cpp
@@ -2331,6 +2331,7 @@ static void TickAssetRegistry()
 
 void UCookOnTheFlyServer::InitializePollables()
 {
+	UE_SCOPED_COOKTIMER(InitializePollables);
 	using namespace UE::Cook;
 
 	Pollables.Reset();
@@ -6586,6 +6587,7 @@ bool UCookOnTheFlyServer::TryInitializeCookWorker()
 
 void UCookOnTheFlyServer::InitializeSession()
 {
+	UE_SCOPED_COOKTIMER(InitializeSession);
 	if (!bFirstCookInThisProcessInitialized)
 	{
 		// This is the first cook; set bFirstCookInThisProcess=true for the entire cook until SetBeginCookConfigSettings is called to mark the second cook
@@ -10153,6 +10155,7 @@ void UCookOnTheFlyServer::OnTargetPlatformChangedSupportedFormats(const ITargetP
 
 void UCookOnTheFlyServer::CreateSandboxFile(FBeginCookContext& BeginContext)
 {
+	UE_SCOPED_COOKTIMER(CreateSandboxFile);
 	// Output directory override. This directory depends on whether we are cooking dlc, so we cannot
 	// create the sandbox until after StartCookByTheBook or StartCookOnTheFly
 	FString OutputDirectory = GetOutputDirectoryOverride(BeginContext);
@@ -10178,6 +10181,7 @@ void UCookOnTheFlyServer::CreateSandboxFile(FBeginCookContext& BeginContext)
 
 void UCookOnTheFlyServer::LoadBeginCookConfigSettings(FBeginCookContext& BeginContext)
 {
+	UE_SCOPED_COOKTIMER(LoadBeginCookConfigSettings);
 	UE::Cook::FBeginCookConfigSettings Settings;
 	WorkerRequests->GetBeginCookConfigSettings(*this, BeginContext, Settings);
 	SetBeginCookConfigSettings(BeginContext, MoveTemp(Settings));
@@ -10296,6 +10300,7 @@ void UCookOnTheFlyServer::SetNeverCookPackageConfigSettings(FBeginCookContext& B
 
 void UCookOnTheFlyServer::LoadBeginCookIncrementalFlags(FBeginCookContext& BeginContext)
 {
+	UE_SCOPED_COOKTIMER(LoadBeginCookIncrementalFlags);
 	WorkerRequests->GetBeginCookIncrementalFlags(*this, BeginContext);
 }
 
@@ -10523,6 +10528,7 @@ void UCookOnTheFlyServer::LoadBeginCookIncrementalFlagsLocal(FBeginCookContext&
 
 void UCookOnTheFlyServer::BeginCookSandbox(FBeginCookContext& BeginContext)
 {
+	UE_SCOPED_COOKTIMER(BeginCookSandbox);
 	using namespace UE::Cook;
 
 #if OUTPUT_COOKTIMING
@@ -11161,6 +11167,7 @@ const UCookOnTheFlyServer::FCookByTheBookStartupOptions& UCookOnTheFlyServer::Bl
 
 FBeginCookContext UCookOnTheFlyServer::CreateBeginCookByTheBookContext(const FCookByTheBookStartupOptions& StartupOptions)
 {
+	UE_SCOPED_COOKTIMER(CreateBeginCookByTheBookContext);
 	FBeginCookContext BeginContext(*this);
 
 	BeginContext.StartupOptions = &StartupOptions;
@@ -11505,6 +11512,7 @@ FBeginCookContext UCookOnTheFlyServer::CreateCookWorkerContext()
 
 void UCookOnTheFlyServer::GenerateInitialRequests(FBeginCookContext& BeginContext)
 {
+	UE_SCOPED_COOKTIMER(GenerateInitialRequests);
 	TArray<ITargetPlatform*>& TargetPlatforms = BeginContext.TargetPlatforms;
 	TMap<FName, FName> StartupSoftObjectPackageReferencers;
 	if (!CookByTheBookOptions->bSkipSoftReferences)
@@ -11775,6 +11783,7 @@ void UCookOnTheFlyServer::GenerateInitialRequests(FBeginCookContext& BeginContex
 
 void UCookOnTheFlyServer::RecordDLCPackagesFromBaseGame(FBeginCookContext& BeginContext)
 {
+	UE_SCOPED_COOKTIMER(RecordDLCPackagesFromBaseGame);
 	using namespace UE::Cook;
 
 	if (!IsCookingDLC())
@@ -11972,6 +11981,7 @@ void UCookOnTheFlyServer::RecordDLCPackagesFromBaseGame(FBeginCookContext& Begin
 
 void UCookOnTheFlyServer::BeginCookPackageWriters(FBeginCookContext& BeginContext)
 {
+	UE_SCOPED_COOKTIMER(BeginCookPackageWriters);
 	for (FBeginCookContextPlatform& Context : BeginContext.PlatformContexts)
 	{
 		ICookedPackageWriter::FCookInfo CookInfo;
@@ -11986,6 +11996,7 @@ void UCookOnTheFlyServer::BeginCookPackageWriters(FBeginCookContext& BeginContex
 
 void UCookOnTheFlyServer::SelectSessionPlatforms(FBeginCookContext& BeginContext)
 {
+	UE_SCOPED_COOKTIMER(SelectSessionPlatforms);
 	PlatformManager->SelectSessionPlatforms(*this, BeginContext.TargetPlatforms);
 
 	FindOrCreateSaveContexts(BeginContext.TargetPlatforms);
@@ -11998,6 +12009,7 @@ void UCookOnTheFlyServer::SelectSessionPlatforms(FBeginCookContext& BeginContext
 
 void UCookOnTheFlyServer::CalculateGlobalDependenciesHashes()
 {
+	UE_SCOPED_COOKTIMER(CalculateGlobalDependenciesHashes);
 	const TArray<const ITargetPlatform*>& Platforms = PlatformManager->GetSessionPlatforms();
 	for (const ITargetPlatform* Platform : Platforms)
 	{
@@ -12007,6 +12019,7 @@ void UCookOnTheFlyServer::CalculateGlobalDependenciesHashes()
 
 void UCookOnTheFlyServer::BeginCookEditorSystems()
 {
+	UE_SCOPED_COOKTIMER(BeginCookEditorSystems);
 	if (!IsCookingInEditor())
 	{
 		return;
@@ -12036,6 +12049,7 @@ void UCookOnTheFlyServer::BeginCookEditorSystems()
 
 void UCookOnTheFlyServer::BeginCookDirector(FBeginCookContext& BeginContext)
 {
+	UE_SCOPED_COOKTIMER(BeginCookDirector);
 	if (CookDirector)
 	{
 		CookDirector->StartCook(BeginContext);
@@ -12086,6 +12100,7 @@ bool UCookOnTheFlyServer::ShouldVerifyEDLCookInfo() const
 
 void UCookOnTheFlyServer::BeginCookEDLCookInfo(FBeginCookContext& BeginContext)
 {
+	UE_SCOPED_COOKTIMER(BeginCookEDLCookInfo);
 	if (IsCookingInEditor())
 	{
 		return;
@@ -12103,6 +12118,7 @@ void UCookOnTheFlyServer::BeginCookEDLCookInfo(FBeginCookContext& BeginContext)
 
 void UCookOnTheFlyServer::RegisterCookByTheBookDelegates()
 {
+	UE_SCOPED_COOKTIMER(RecordCookByTheBookDelegates);
 	if (!IsCookingInEditor())
 	{
 		FCoreUObjectDelegates::PackageCreatedForLoad.AddUObject(this, &UCookOnTheFlyServer::MaybeMarkPackageAsAlreadyLoaded);
@@ -12567,6 +12583,7 @@ const UE::Cook::FCookSavePackageContext* UCookOnTheFlyServer::FindSaveContext(co
 
 void UCookOnTheFlyServer::InitializeAllCulturesToCook(TConstArrayView<FString> CookCultures)
 {
+	UE_SCOPED_COOKTIMER(InitializeAllCulturesToCook);
 	CookByTheBookOptions->AllCulturesToCook.Reset();
 
 	TArray<FString> AllCulturesToCook(CookCultures);
@@ -12585,6 +12602,7 @@ void UCookOnTheFlyServer::InitializeAllCulturesToCook(TConstArrayView<FString> C
 
 void UCookOnTheFlyServer::CompileDLCLocalization(FBeginCookContext& BeginContext)
 {
+	UE_SCOPED_COOKTIMER(CompileDLCLocalization);
 	if (!IsCookingDLC() || !GetDefault<UUserGeneratedContentLocalizationSettings>()->bCompileDLCLocalizationDuringCook)
 	{
 		return;
@@ -12624,6 +12642,7 @@ void UCookOnTheFlyServer::CompileDLCLocalization(FBeginCookContext& BeginContext
 
 void UCookOnTheFlyServer::GenerateLocalizationReferences()
 {
+	UE_SCOPED_COOKTIMER(GenerateLocalizationReferences);
 	CookByTheBookOptions->SourceToLocalizedPackageVariants.Reset();
 
 	// Find all the localized packages and map them back to their source package
@@ -13182,6 +13201,7 @@ void ConditionalWaitOnCommandFile(FStringView GateName, TFunctionRef<void (FStri
 
 void UCookOnTheFlyServer::BroadcastCookStarted()
 {
+	UE_SCOPED_COOKTIMER(BroadcastCookStarted);
 	if (IsDirectorCookByTheBook())
 	{
 		PRAGMA_DISABLE_DEPRECATION_WARNINGS;
diff --git a/Engine/Source/Editor/UnrealEd/Private/Cooker/ShaderLibraryCooking.cpp b/Engine/Source/Editor/UnrealEd/Private/Cooker/ShaderLibraryCooking.cpp
--- a/Engine/Source/Editor/UnrealEd/Private/Cooker/ShaderLibraryCooking.cpp
+++ b/Engine/Source/Editor/UnrealEd/Private/Cooker/ShaderLibraryCooking.cpp
@@ -241,6 +241,7 @@ void FShaderLibraryCookArtifact::UpdateOplogPackages(UE::Cook::Artifact::FUpdate
 
 void UCookOnTheFlyServer::RegisterShaderLibraryIncrementalCookArtifact(FBeginCookContext& BeginContext)
 {
+	UE_SCOPED_COOKTIMER(RegisterShaderLibraryIncrementalCookArtifact);
 	if (IsCookWorkerMode() || !IsUsingShaderCodeLibrary())
 	{
 		return;
@@ -250,6 +251,7 @@ void UCookOnTheFlyServer::RegisterShaderLibraryIncrementalCookArtifact(FBeginCoo
 
 void UCookOnTheFlyServer::BeginCookStartShaderCodeLibrary(FBeginCookContext& BeginContext)
 {
+	UE_SCOPED_COOKTIMER(BeginCookStartShaderCodeLibrary);
 	const UProjectPackagingSettings* const PackagingSettings = GetDefault<UProjectPackagingSettings>();
 	if (IsUsingShaderCodeLibrary())
 	{
@@ -321,6 +323,7 @@ void UCookOnTheFlyServer::BeginCookStartShaderCodeLibrary(FBeginCookContext& Beg
 
 void UCookOnTheFlyServer::BeginCookFinishShaderCodeLibrary(FBeginCookContext& BeginContext)
 {
+	UE_SCOPED_COOKTIMER(BeginCookFinishShaderCodeLibrary);
 	check(IsDirectorCookByTheBook()); // CookByTheBook only for now
 
 	// Generate the AutogenShaderHeaders.usd files for every shader platforms we are about to cook. Those files can be hashed and used
diff --git a/Engine/Source/Editor/UnrealEd/Private/EditorDomain/EditorDomainUtils.cpp b/Engine/Source/Editor/UnrealEd/Private/EditorDomain/EditorDomainUtils.cpp
--- a/Engine/Source/Editor/UnrealEd/Private/EditorDomain/EditorDomainUtils.cpp
+++ b/Engine/Source/Editor/UnrealEd/Private/EditorDomain/EditorDomainUtils.cpp
@@ -679,6 +679,7 @@ private:
 
 FClassDigestData* FPrecacheClassDigest::GetRecursive(FTopLevelAssetPath ClassPath, bool bAllowRedirects)
 {
+	SCOPED_EDITORDOMAINTIMER(FPrecacheClassDigest::GetRecursive);
 	check(!ClassPath.IsNull());
 
 	static FString ClassName;
@@ -957,6 +958,7 @@ FClassDigestData* FPrecacheClassDigest::GetRecursive(FTopLevelAssetPath ClassPat
 /** Try to add the FClassDigestData for each given class into the GetClassDigests map */
 void PrecacheClassDigests(TConstArrayView<FTopLevelAssetPath> ClassPaths)
 {
+	SCOPED_EDITORDOMAINTIMER(EditorDomain::PrecacheClassDigests);
 	FPrecacheClassDigest Digester;
 	for (const FTopLevelAssetPath& ClassPath : ClassPaths)
 	{
@@ -1085,6 +1087,7 @@ private:
 /** Collect the CustomVersions that can be used by the given Class when it is saved */
 TArray<FGuid> GetCustomVersions(UClass& Class)
 {
+	SCOPED_EDITORDOMAINTIMER(EditorDomain::GetCustomVersions);
 	FCustomVersionCollectorArchive Ar;
 	Class.CallDeclareCustomVersions(Ar);
 	// Default objects of blueprint classes are serialized during SavePackage with a special call to
@@ -1223,6 +1226,7 @@ TSet<FTopLevelAssetPath> ConstructTargetIncrementalClassBlockList()
 
 void ConstructTargetIterativeClassAllowList()
 {
+	SCOPED_EDITORDOMAINTIMER(EditorDomain::ConstructTargetIterativeClassAllowList);
 	TSet<FTopLevelAssetPath> BlockListClassPaths = ConstructTargetIncrementalClassBlockList();
 
 	// AllowList elements implicitly allow all parent classes, so instead of consulting a list and propagating
@@ -1573,6 +1577,7 @@ void UtilsInitialize()
 
 void UtilsCookInitialize()
 {
+	SCOPED_EDITORDOMAINTIMER(EditorDomain::UtilsCookInitialize);
 	if (bGUtilsCookInitialized)
 	{
 		return;
diff --git a/Engine/Source/Runtime/AssetRegistry/Private/AssetRegistry.cpp b/Engine/Source/Runtime/AssetRegistry/Private/AssetRegistry.cpp
--- a/Engine/Source/Runtime/AssetRegistry/Private/AssetRegistry.cpp
+++ b/Engine/Source/Runtime/AssetRegistry/Private/AssetRegistry.cpp
@@ -9897,6 +9897,7 @@ static EAssetSetManagerAssetFlags ConstructAssetManagerAssetFlags(const FDepends
 
 void FAssetRegistryImpl::SetManageReferences(FSetManageReferencesContext& Context)
 {
+	TRACE_CPUPROFILER_EVENT_SCOPE(FAssetRegistryImpl::SetManageReferences);
 	if (!Context.ManagerMap)
 	{
 		UE_LOG(LogAssetRegistry, Error, TEXT("SetManageReferences: Context.ManagerMap must not be null."));
diff --git a/Engine/Source/Runtime/CoreUObject/Private/UObject/UObjectBase.cpp b/Engine/Source/Runtime/CoreUObject/Private/UObject/UObjectBase.cpp
--- a/Engine/Source/Runtime/CoreUObject/Private/UObject/UObjectBase.cpp
+++ b/Engine/Source/Runtime/CoreUObject/Private/UObject/UObjectBase.cpp
@@ -892,30 +892,34 @@ static void UObjectLoadAllCompiledInDefaultProperties(TArray<UClass*>& OutAllNew
 		TArray<UClass*> NewClasses;
 		TArray<UClass*> NewClassesInCoreUObject;
 		TArray<UClass*> NewClassesInEngine;
-		ClassRegistry.DoPendingOuterRegistrations(true, [&OutAllNewClasses, &NewClasses, &NewClassesInCoreUObject, &NewClassesInEngine](const TCHAR* PackageName, UClass& Class) -> void
-			{
-				UE_LOG(LogUObjectBootstrap, Verbose, TEXT("UObjectLoadAllCompiledInDefaultProperties After Registrant %s %s"), PackageName, *Class.GetName());
-
-				if (Class.GetOutermost()->GetFName() == GLongCoreUObjectPackageName)
-				{
-					NewClassesInCoreUObject.Add(&Class);
-				}
-				else if (Class.GetOutermost()->GetFName() == LongEnginePackageName)
-				{
-					NewClassesInEngine.Add(&Class);
-				}
-				else
+		{
+			SCOPED_BOOT_TIMING("DoPendingOuterRegistrations");
+			ClassRegistry.DoPendingOuterRegistrations(true, [&OutAllNewClasses, &NewClasses, &NewClassesInCoreUObject, &NewClassesInEngine](const TCHAR* PackageName, UClass& Class) -> void
 				{
-					NewClasses.Add(&Class);
-				}
+					UE_LOG(LogUObjectBootstrap, Verbose, TEXT("UObjectLoadAllCompiledInDefaultProperties After Registrant %s %s"), PackageName, *Class.GetName());
+
+					if (Class.GetOutermost()->GetFName() == GLongCoreUObjectPackageName)
+					{
+						NewClassesInCoreUObject.Add(&Class);
+					}
+					else if (Class.GetOutermost()->GetFName() == LongEnginePackageName)
+					{
+						NewClassesInEngine.Add(&Class);
+					}
+					else
+					{
+						NewClasses.Add(&Class);
+					}
 
-				OutAllNewClasses.Add(&Class);
-			}); 
+					OutAllNewClasses.Add(&Class);
+				});
+		}		
 
 		// Sort classes by depth to allow creating CDOs by parent type first. Otherwise, complex PostCDOConstruct calls 
 		// might rely on their parent having been initialized already, which might not be possible if there is a cycle in 
 		// the GetDefaultObject -> PostCDOConstruct loading graph.
 		{
+			SCOPED_BOOT_TIMING("SortClassesByDepth");
 			// Memo to avoid re-calculating depth needlessly as we sort. It's important to note that since
 			// two FindOrAdds will occur back to back that we reserve our memory upfront to avoid re-allocing
 			// in-between FindOrAdd calls which will invalidate the first reference returned.
