Eliminate the use of "Global Local" when describing DDC settings.

#Synced-CL 50960226

diff --git a/Engine/Source/Editor/DerivedDataWidgets/Private/DerivedDataInformation.cpp b/Engine/Source/Editor/DerivedDataWidgets/Private/DerivedDataInformation.cpp
--- a/Engine/Source/Editor/DerivedDataWidgets/Private/DerivedDataInformation.cpp
+++ b/Engine/Source/Editor/DerivedDataWidgets/Private/DerivedDataInformation.cpp
@@ -181,12 +181,12 @@ void FDerivedDataInformation::UpdateRemoteCacheState()
 			else if (DDCProjectSettings->RecommendEveryoneSetupAGlobalLocalDDCPath && EditorSettings->bNotifySetupDDCPath && EditorSettings->GlobalLocalDDCPath.Path.IsEmpty())
 			{
 				RemoteCacheState = ERemoteCacheState::Warning;
-				RemoteCacheWarningMessage = FText(LOCTEXT("GlobalLocalDDCPathWarning", "It is recommended that you set up a valid Global Local DDC Path.\nDisable this notification or set up a valid Global Local DDC Path in the settings."));
+				RemoteCacheWarningMessage = FText(LOCTEXT("GlobalLocalDDCPathWarning", "It is recommended that you set up a valid Local DDC Path.\nDisable this notification or set up a valid Local DDC Path in the settings."));
 			}
 			else if (DDCProjectSettings->RecommendEveryoneSetupAGlobalSharedDDCPath && EditorSettings->bNotifySetupDDCPath && EditorSettings->GlobalSharedDDCPath.Path.IsEmpty())
 			{
 				RemoteCacheState = ERemoteCacheState::Warning;
-				RemoteCacheWarningMessage = FText(LOCTEXT("GlobalSharedDDCPathWarning", "It is recommended that you set up a valid Global Shared DDC Path.\nDisable this notification or set up a valid Global Shared DDC Path in the settings."));
+				RemoteCacheWarningMessage = FText(LOCTEXT("GlobalSharedDDCPathWarning", "It is recommended that you set up a valid Shared DDC Path.\nDisable this notification or set up a valid Shared DDC Path in the settings."));
 			}
 			else if (DDCProjectSettings->RecommendEveryoneEnableS3DDC && EditorSettings->bNotifyEnableS3DD && !EditorSettings->bEnableS3DDC)
 			{
@@ -196,7 +196,7 @@ void FDerivedDataInformation::UpdateRemoteCacheState()
 			else if (DDCProjectSettings->RecommendEveryoneSetupAGlobalS3DDCPath && EditorSettings->bNotifySetupDDCPath && EditorSettings->GlobalS3DDCPath.Path.IsEmpty())
 			{
 				RemoteCacheState = ERemoteCacheState::Warning;
-				RemoteCacheWarningMessage = FText(LOCTEXT("S3GloblaLocalPathWarning", "It is recommended that you set up a valid Global Local S3 DDC Path.\nDisable this notification or set up a valid Global Local S3 DDC Path in the settings."));
+				RemoteCacheWarningMessage = FText(LOCTEXT("S3GloblaLocalPathWarning", "It is recommended that you set up a valid Local S3 DDC Path.\nDisable this notification or set up a valid Local S3 DDC Path in the settings."));
 			}
 		}
 	}
diff --git a/Engine/Source/Editor/UnrealEd/Public/Settings/EditorSettings.h b/Engine/Source/Editor/UnrealEd/Public/Settings/EditorSettings.h
--- a/Engine/Source/Editor/UnrealEd/Public/Settings/EditorSettings.h
+++ b/Engine/Source/Editor/UnrealEd/Public/Settings/EditorSettings.h
@@ -62,29 +62,29 @@ class UEditorSettings : public UObject
 	// =====================================================================
 
 	/**
-	 * Adjusts the Local Cache location. This affects every project on your computer that uses the UE-LocalDataCachePath environment variable override.
+	 * Adjusts the Local Cache location. This affects every project and workspace used from your user profile that accepts the UE-LocalDataCachePath environment variable override.
 	 * This is usually the first location to query for previously built data.
 	 */
-	UPROPERTY(EditAnywhere, Category = DerivedDataCache, meta = (DisplayName = "Global Local DDC Path", ConfigRestartRequired = true))
+	UPROPERTY(EditAnywhere, Category = DerivedDataCache, meta = (DisplayName = "Local DDC Path", ConfigRestartRequired = true))
 	FDirectoryPath GlobalLocalDDCPath;
 
 	/**
-	 * Adjusts the Shared cache location. This affects every project on your computer that uses the UE-SharedDataCachePath environment variable override.
+	 * Adjusts the Shared Cache location. This affects every project and workspace used from your user profile that accepts the UE-SharedDataCachePath environment variable override.
 	 * The Shared Cache location is usually queried if we don't find previously built data in the Local cache. Colleagues should point to the same shared location so that work can be distributed. 
 	 */
-	UPROPERTY(EditAnywhere, Category = DerivedDataCache, meta = (DisplayName = "Global Shared DDC Path", ConfigRestartRequired = true))
+	UPROPERTY(EditAnywhere, Category = DerivedDataCache, meta = (DisplayName = "Shared DDC Path", ConfigRestartRequired = true))
 	FDirectoryPath GlobalSharedDDCPath;
 
 	/**
 	 * Project specific override for the Local Cache location. The editor must be restarted for changes to take effect.
-	 * This will override the 'Global Local DDC Path'.
+	 * This will override the 'Local DDC Path'.
 	 */
 	UPROPERTY(EditAnywhere, config, Category= DerivedDataCache, AdvancedDisplay, meta = (DisplayName = "Project Local DDC Path", ConfigRestartRequired = true))
 	FDirectoryPath LocalDerivedDataCache;
 
 	/**
 	 * Project specific override for the Shared Cache location. The editor must be restarted for changes to take effect.
-	 * This will override the 'Global Shared DDC Path'.
+	 * This will override the 'Shared DDC Path'.
 	 */
 	UPROPERTY(EditAnywhere, config, Category= DerivedDataCache, AdvancedDisplay, meta = (DisplayName = "Project Shared DDC Path", ConfigRestartRequired = true))
 	FDirectoryPath SharedDerivedDataCache;
@@ -111,9 +111,9 @@ class UEditorSettings : public UObject
 
 	/**
 	 * Adjusts the Local Cache location for AWS/S3 downloaded package bundles.
-	 * This affects every project on your computer that uses the UE-S3DataCachePath environment variable override.
+	 * This affects every project and workspace used from your user profile that accepts the UE-S3DataCachePath environment variable override.
 	 */
-	UPROPERTY(EditAnywhere, Category="Derived Data Cache S3", meta = (DisplayName = "Global Local S3DDC Path", ConfigRestartRequired = true, EditCondition = "bEnableS3DDC"))
+	UPROPERTY(EditAnywhere, Category="Derived Data Cache S3", meta = (DisplayName = "Local S3DDC Path", ConfigRestartRequired = true, EditCondition = "bEnableS3DDC"))
 	FDirectoryPath GlobalS3DDCPath;
 
 	// Horde Settings
