[Component Visualizers] Ensure component visualizer doesn't lose component info (Crash when right clicking water spline)

#Synced-CL 50978000

diff --git a/Engine/Source/Editor/UnrealEd/Private/ComponentVisualizerManager.cpp b/Engine/Source/Editor/UnrealEd/Private/ComponentVisualizerManager.cpp
--- a/Engine/Source/Editor/UnrealEd/Private/ComponentVisualizerManager.cpp
+++ b/Engine/Source/Editor/UnrealEd/Private/ComponentVisualizerManager.cpp
@@ -89,9 +89,33 @@ bool FComponentVisualizerManager::HandleProxyForComponentVis(FEditorViewportClie
 						if (ElementSelectionSet->CanSelectElement(SelectedComponentHandle, FTypedElementSelectionOptions())
 							&& !ElementSelectionSet->IsElementSelected(SelectedComponentHandle, FTypedElementIsSelectedOptions()))
 						{
-							// Clear the selection and just select the component
-							ElementSelectionSet->ClearSelection(FTypedElementSelectionOptions());
+							// Ideally, we would just perform the selection via a ClearSelection followed by a SelectElement call.
+							// Unfortunately, actors currently get notified of deselection immediately in ClearSelection
+							// instead of in the subsequent NotifyPendingChanges call, which causes some visualizers (spline
+							// visualizer in particular) to lose state, as they see their actor as deselected. So instead,
+							// we select the component and deselect everything except what would normally be added to an
+							// empty selection: the component and its parent actor.
+							// That way the component/actor do not receive a spurious deselection notification.
+
+							// We will remove Component and its Owner handles from this list.
+							TArray<FTypedElementHandle> CachedSelection = ElementSelectionSet->GetSelectedElementHandles();
+
+							// Select the Component
 							ElementSelectionSet->SelectElement(SelectedComponentHandle, FTypedElementSelectionOptions());
+
+							if (const UActorComponent* const Component = ComponentElementDataUtil::GetComponentFromHandle(SelectedComponentHandle))
+							{
+								if (const AActor* const Owner = Component->GetOwner())
+								{
+									FTypedElementHandle OwningActorHandle = UEngineElementsLibrary::AcquireEditorActorElementHandle(Owner);
+									CachedSelection.Remove(OwningActorHandle);
+								}
+
+								CachedSelection.Remove(SelectedComponentHandle);
+
+								// Remove all element handles from list, except from Component and Owner handles.
+								ElementSelectionSet->DeselectElements(CachedSelection, FTypedElementSelectionOptions());
+							}
 						}
 					}
 
