Cooker: AssetRegistry: Reduce GeneratedAssetRegistry memory usage during cook by removing deny-listed AssetRegistry tags (via Engine.ini:[AssetRegistry] settings) from the cooked AssetDatas. They were accidentally added when we started collecting up-to-date AssetRegistry tags during the cook's SavePackage calls. These tags were already being filtered later, immediately before writing to disk in FAssetRegistryGenerator::SaveAssetRegistry.
#rnx

#Synced-CL 50997486

diff --git a/Engine/Source/Editor/UnrealEd/Classes/CookOnTheSide/CookOnTheFlyServer.h b/Engine/Source/Editor/UnrealEd/Classes/CookOnTheSide/CookOnTheFlyServer.h
--- a/Engine/Source/Editor/UnrealEd/Classes/CookOnTheSide/CookOnTheFlyServer.h
+++ b/Engine/Source/Editor/UnrealEd/Classes/CookOnTheSide/CookOnTheFlyServer.h
@@ -1130,6 +1130,7 @@ private:
 	/** Called at the end of CookByTheBook to write aggregated data such as AssetRegistry and shaders. */
 	void CookByTheBookFinished();
 	void CookByTheBookFinishedInternal();
+	bool IsForceNoFilterAssetsFromAssetRegistry();
 
 	/** Clears session-lifetime data from COTFS and CookPackageDatas. */
 	void ShutdownCookSession();
diff --git a/Engine/Source/Editor/UnrealEd/Private/Commandlets/AssetRegistryGenerator.cpp b/Engine/Source/Editor/UnrealEd/Private/Commandlets/AssetRegistryGenerator.cpp
--- a/Engine/Source/Editor/UnrealEd/Private/Commandlets/AssetRegistryGenerator.cpp
+++ b/Engine/Source/Editor/UnrealEd/Private/Commandlets/AssetRegistryGenerator.cpp
@@ -1364,14 +1364,11 @@ void FAssetRegistryGenerator::UpdateCollectionAssetData()
 	}
 }
 
-void FAssetRegistryGenerator::Initialize(bool bInitializeFromExisting)
+void FAssetRegistryGenerator::Initialize(bool bInitializeFromExisting,
+	const FAssetRegistrySerializationOptions& DevelopmentSaveOptions)
 {
 	LLM_SCOPE_BYTAG(Cooker_GeneratedAssetRegistry);
 
-	FAssetRegistrySerializationOptions SaveOptions;
-
-	AssetRegistry.InitializeSerializationOptions(SaveOptions, TargetPlatform, UE::AssetRegistry::ESerializationTarget::ForGame);
-
 	if (bInitializeFromExisting)
 	{
 		// If the asset registry is still doing its background scan, we need to wait for it to finish and tick it so that the results are flushed out
@@ -1379,17 +1376,15 @@ void FAssetRegistryGenerator::Initialize(bool bInitializeFromExisting)
 		ensureMsgf(!AssetRegistry.IsLoadingAssets(), TEXT("Cannot initialize asset registry generator while asset registry is still scanning source assets "));
 
 		bClonedGlobalAssetRegistry = true;
-		AssetRegistry.InitializeTemporaryAssetRegistryState(State, SaveOptions);
+		AssetRegistry.InitializeTemporaryAssetRegistryState(State, DevelopmentSaveOptions);
 	}
 
 	FGameDelegates::Get().GetAssignLayerChunkDelegate() = FAssignLayerChunkDelegate::CreateStatic(AssignLayerChunkDelegate);
 }
 
-void FAssetRegistryGenerator::CloneGlobalAssetRegistryFilteredByPreviousState(const FAssetRegistryState& PreviousState)
+void FAssetRegistryGenerator::CloneGlobalAssetRegistryFilteredByPreviousState(const FAssetRegistryState& PreviousState,
+	const FAssetRegistrySerializationOptions& DevelopmentSaveOptions)
 {
-	FAssetRegistrySerializationOptions SaveOptions;
-	AssetRegistry.InitializeSerializationOptions(SaveOptions, TargetPlatform, UE::AssetRegistry::ESerializationTarget::ForGame);
-
 	TSet<FName> KeepPackages;
 	const TMap<FName, const FAssetPackageData*> &PreviousPackageDatas = PreviousState.GetAssetPackageDataMap();
 	KeepPackages.Reserve(PreviousPackageDatas.Num());
@@ -1397,7 +1392,8 @@ void FAssetRegistryGenerator::CloneGlobalAssetRegistryFilteredByPreviousState(co
 	{
 		KeepPackages.Add(Pair.Key);
 	}
-	AssetRegistry.InitializeTemporaryAssetRegistryState(State, SaveOptions, false /* bRefreshExisting */, KeepPackages);
+	AssetRegistry.InitializeTemporaryAssetRegistryState(State, DevelopmentSaveOptions,
+		false /* bRefreshExisting */, KeepPackages);
 }
 
 static UE::Cook::ECookResult DiskSizeToCookResult(int64 DiskSize)
@@ -1813,7 +1809,8 @@ void FAssetRegistryGenerator::ComputePackageRemovals(const FAssetRegistryState&
 
 void FAssetRegistryGenerator::FinalizeChunkIDs(const TSet<FName>& CookedPackages,
 	const TSet<FName>& DevelopmentOnlyPackages, UE::Cook::FCookSandbox& InSandboxFile,
-	bool bGenerateStreamingInstallManifest, const TSet<FName>& StartupPackages)
+	bool bGenerateStreamingInstallManifest, const TSet<FName>& StartupPackages,
+	const FAssetRegistrySerializationOptions& DevelopmentSaveOptions)
 {
 	LLM_SCOPE_BYTAG(Cooker_GeneratedAssetRegistry);
 	// bGenerateNoChunks overrides bGenerateStreamingInstallManifest overrides ShouldPlatformGenerateStreamingInstallManifest
@@ -1841,9 +1838,6 @@ void FAssetRegistryGenerator::FinalizeChunkIDs(const TSet<FName>& CookedPackages
 
 	// Prune our asset registry to union of cooked + dev only + (previously cooked or dev only)
 	{
-		FAssetRegistrySerializationOptions DevelopmentSaveOptions;
-		AssetRegistry.InitializeSerializationOptions(DevelopmentSaveOptions, TargetPlatform, UE::AssetRegistry::ESerializationTarget::ForDevelopment);
-		DevelopmentSaveOptions.bKeepDevelopmentAssetRegistryTags = true;
 		State.PruneAssetData(AllPackages, TSet<FName>(), TSet<int32>(), DevelopmentSaveOptions);
 	}
 
@@ -1985,29 +1979,24 @@ void FAssetRegistryGenerator::AddAssetToFileOrderRecursive(const FName& InPackag
 }
 
 bool FAssetRegistryGenerator::SaveAssetRegistry(const FString& SandboxPath, bool bSerializeDevelopmentAssetRegistry, 
-	bool bForceNoFilter, uint64& OutDevelopmentAssetRegistryHash, const TSet<FName>& CookedPackages)
+	bool bForceNoFilter, uint64& OutDevelopmentAssetRegistryHash, const TSet<FName>& CookedPackages,
+	const FAssetRegistrySerializationOptions& DevelopmentSaveOptions)
 {
 	LLM_SCOPE_BYTAG(Cooker_GeneratedAssetRegistry);
 	UE_LOG(LogAssetRegistryGenerator, Display, TEXT("Saving asset registry v%d."), FAssetRegistryVersion::Type::LatestVersion);
 	
-	// Write development first, this will always write
-	FAssetRegistrySerializationOptions DevelopmentSaveOptions;
-	AssetRegistry.InitializeSerializationOptions(DevelopmentSaveOptions, TargetPlatform, UE::AssetRegistry::ESerializationTarget::ForDevelopment);
-	DevelopmentSaveOptions.bKeepDevelopmentAssetRegistryTags = true;
-
-	// Write runtime registry, this can be excluded per game/platform
 	FAssetRegistrySerializationOptions SaveOptions;
 	AssetRegistry.InitializeSerializationOptions(SaveOptions, TargetPlatform, UE::AssetRegistry::ESerializationTarget::ForGame);
 	SaveOptions.bKeepDevelopmentAssetRegistryTags = FParse::Param(FCommandLine::Get(), TEXT("ARKeepDevTags"));
 
 	if (bForceNoFilter)
 	{
-		DevelopmentSaveOptions.DisableFilters();
 		SaveOptions.DisableFilters();
 	}
 
 	UpdateCollectionAssetData();
 
+	// Write development first
 	if (DevelopmentSaveOptions.bSerializeAssetRegistry)
 	{
 		FString PlatformSandboxPath = SandboxPath.Replace(TEXT("[Platform]"), *TargetPlatform->PlatformName());
@@ -2059,6 +2048,8 @@ bool FAssetRegistryGenerator::SaveAssetRegistry(const FString& SandboxPath, bool
 		}
 	}
 
+	// Write runtime registry, filtering further the state that we wrote out mostly unfiltered for the
+	// DevelopmentAssetRegistry above.
 	if (SaveOptions.bSerializeAssetRegistry)
 	{
 		TMap<int32, FString> ChunkBucketNames;
diff --git a/Engine/Source/Editor/UnrealEd/Private/CookOnTheFlyServer.cpp b/Engine/Source/Editor/UnrealEd/Private/CookOnTheFlyServer.cpp
--- a/Engine/Source/Editor/UnrealEd/Private/CookOnTheFlyServer.cpp
+++ b/Engine/Source/Editor/UnrealEd/Private/CookOnTheFlyServer.cpp
@@ -618,7 +618,7 @@ bool UCookOnTheFlyServer::StartCookOnTheFly(FCookOnTheFlyStartupOptions InCookOn
 	ITargetPlatformManagerModule& TPM = GetTargetPlatformManagerRef();
 	for (const ITargetPlatform* TargetPlatform : TPM.GetTargetPlatforms())
 	{
-		PlatformManager->CreatePlatformData(TargetPlatform);
+		PlatformManager->CreatePlatformData(*this, TargetPlatform);
 	}
 	PlatformManager->SetArePlatformsPrepopulated(true);
 
@@ -809,7 +809,8 @@ void UCookOnTheFlyServer::AddCookOnTheFlyPlatformFromGameThread(ITargetPlatform*
 	InitializeShadersForCookOnTheFly(BeginContext.TargetPlatforms);
 	// SaveAssetRegistry is done in CookByTheBookFinished for CBTB, but we need at the start of CookOnTheFly to send as startup information to connecting clients
 	uint64 DevelopmentAssetRegistryHash = 0;
-	PlatformData->RegistryGenerator->SaveAssetRegistry(GetSandboxAssetRegistryFilename(), true, false, DevelopmentAssetRegistryHash, TSet<FName>());
+	PlatformData->RegistryGenerator->SaveAssetRegistry(GetSandboxAssetRegistryFilename(), true, false,
+		DevelopmentAssetRegistryHash, TSet<FName>(), PlatformData->AssetRegistryDevelopmentSaveOptions);
 
 	// Initialize systems that nothing in AddCookOnTheFlyPlatformFromGameThread references
 	// Functions in this section are not dependent upon each other and can be ordered arbitrarily or for async performance
@@ -869,7 +870,7 @@ void UCookOnTheFlyServer::OnTargetPlatformsInvalidated()
 	{
 		for (const ITargetPlatform* TargetPlatform : GetTargetPlatformManager()->GetTargetPlatforms())
 		{
-			PlatformManager->CreatePlatformData(TargetPlatform);
+			PlatformManager->CreatePlatformData(*this, TargetPlatform);
 		}
 	}
 }
@@ -7825,7 +7826,8 @@ void UCookOnTheFlyServer::PopulateCookedPackages(TArrayView<const ITargetPlatfor
 	COOK_STAT(DetailedCookStats::NumPackagesIncrementallyModifiedByClass.Empty());
 	for (const ITargetPlatform* TargetPlatform : TargetPlatforms)
 	{
-		FAssetRegistryGenerator& PlatformAssetRegistry = *(PlatformManager->GetPlatformData(TargetPlatform)->RegistryGenerator);
+		FPlatformData* PlatformData = PlatformManager->GetPlatformData(TargetPlatform);
+		FAssetRegistryGenerator& PlatformAssetRegistry = *(PlatformData->RegistryGenerator);
 		FCookSavePackageContext& CookSavePackageContext = FindOrCreateSaveContext(TargetPlatform);
 		ICookedPackageWriter& PackageWriter = *CookSavePackageContext.PackageWriter;
 		UE_LOG(LogCook, Display, TEXT("Populating cooked package(s) from %s package store on platform '%s'"),
@@ -7841,7 +7843,8 @@ void UCookOnTheFlyServer::PopulateCookedPackages(TArrayView<const ITargetPlatfor
 
 		if (!PlatformAssetRegistry.HasClonedGlobalAssetRegistry())
 		{
-			PlatformAssetRegistry.CloneGlobalAssetRegistryFilteredByPreviousState(*PreviousAssetRegistry);
+			PlatformAssetRegistry.CloneGlobalAssetRegistryFilteredByPreviousState(*PreviousAssetRegistry,
+				PlatformData->AssetRegistryDevelopmentSaveOptions);
 		}
 
 		TMap<FName, FAssetRegistryGenerator::FGeneratorPackageInfo> PreviousGeneratorPackages;
@@ -8214,7 +8217,8 @@ void UCookOnTheFlyServer::RefreshPlatformAssetRegistries(const TArrayView<const
 				bInitializeFromExisting = false;
 			}
 
-			PlatformData->RegistryGenerator->Initialize(bInitializeFromExisting);
+			PlatformData->RegistryGenerator->Initialize(bInitializeFromExisting,
+				PlatformData->AssetRegistryDevelopmentSaveOptions);
 		}
 	}
 }
@@ -9325,7 +9329,7 @@ void UCookOnTheFlyServer::CookByTheBookFinishedInternal()
 	// if we are cooking DLC, the DevelopmentAR isn't needed - it's used when making DLC against shipping, so there's no need to make it
 	// again, as we don't make DLC against DLC (but allow an override just in case)
 	bool bSaveDevelopmentAssetRegistry = !FParse::Param(FCommandLine::Get(), TEXT("NoSaveDevAR"));
-	bool bForceNoFilterAssetsFromAssetRegistry = IsCookingDLC() ? !FParse::Param(FCommandLine::Get(), TEXT("ApplyFiltersToDLCAssetRegistry")) : false;
+	bool bForceNoFilterAssetsFromAssetRegistry = IsForceNoFilterAssetsFromAssetRegistry();
 	bool bSaveManifests = true;
 	bool bSaveIniSettings = !FParse::Param(FCommandLine::Get(), TEXT("SkipSaveCookSettings"));
 	bool bSaveCookerOpenOrder = true;
@@ -9561,7 +9565,8 @@ void UCookOnTheFlyServer::CookByTheBookFinishedInternal()
 			{
 				UE_SCOPED_HIERARCHICAL_COOKTIMER(BuildChunkManifest);
 				Generator.FinalizeChunkIDs(CookedPackageNames, IgnorePackageNames, *SandboxFile,
-					CookByTheBookOptions->bGenerateStreamingInstallManifests, CookByTheBookOptions->StartupPackages);
+					CookByTheBookOptions->bGenerateStreamingInstallManifests, CookByTheBookOptions->StartupPackages,
+					PlatformData->AssetRegistryDevelopmentSaveOptions);
 			}
 			{
 				UE_SCOPED_HIERARCHICAL_COOKTIMER(SaveManifests);
@@ -9589,7 +9594,8 @@ void UCookOnTheFlyServer::CookByTheBookFinishedInternal()
 		{
 			UE_SCOPED_HIERARCHICAL_COOKTIMER(SaveAssetRegistry);
 			Generator.SaveAssetRegistry(SandboxRegistryFilename, bSaveDevelopmentAssetRegistry,
-				bForceNoFilterAssetsFromAssetRegistry, DevelopmentAssetRegistryHash, CookedPackageNames);
+				bForceNoFilterAssetsFromAssetRegistry, DevelopmentAssetRegistryHash, CookedPackageNames,
+				PlatformData->AssetRegistryDevelopmentSaveOptions);
 		}
 
 		if (bSaveManifests || bSaveAssetRegistry)
@@ -9720,6 +9726,18 @@ void UCookOnTheFlyServer::CookByTheBookFinishedInternal()
 	UE_LOG(LogCookStatus, Display, TEXT("FinishedInternal: Done"));
 }
 
+bool UCookOnTheFlyServer::IsForceNoFilterAssetsFromAssetRegistry()
+{
+	if (IsDirectorCookByTheBook())
+	{
+		return IsCookingDLC() ? !FParse::Param(FCommandLine::Get(), TEXT("ApplyFiltersToDLCAssetRegistry")) : false;
+	}
+	else
+	{
+		return false;
+	}
+}
+
 void UCookOnTheFlyServer::ShutdownCookSession()
 {
 	ODSCClientData.Reset();
@@ -11237,7 +11255,7 @@ FBeginCookContext UCookOnTheFlyServer::CreateAddPlatformContext(ITargetPlatform*
 
 	FBeginCookContextPlatform& PlatformContext = BeginContext.PlatformContexts.Emplace_GetRef();
 	PlatformContext.TargetPlatform = TargetPlatform;
-	PlatformContext.PlatformData = &PlatformManager->CreatePlatformData(TargetPlatform);
+	PlatformContext.PlatformData = &PlatformManager->CreatePlatformData(*this, TargetPlatform);
 
 	return BeginContext;
 }
diff --git a/Engine/Source/Editor/UnrealEd/Private/Cooker/CookPlatformManager.cpp b/Engine/Source/Editor/UnrealEd/Private/Cooker/CookPlatformManager.cpp
--- a/Engine/Source/Editor/UnrealEd/Private/Cooker/CookPlatformManager.cpp
+++ b/Engine/Source/Editor/UnrealEd/Private/Cooker/CookPlatformManager.cpp
@@ -127,7 +127,7 @@ void FPlatformManager::SelectSessionPlatforms(UCookOnTheFlyServer& COTFS,
 	checkf(IsSchedulerThread(), TEXT("Writing to SessionPlatforms is only allowed from the scheduler thread."));
 	for (FPlatformId TargetPlatform : TargetPlatforms)
 	{
-		CreatePlatformData(TargetPlatform);
+		CreatePlatformData(COTFS, TargetPlatform);
 	}
 
 	FRWScopeLock ScopeLock(SessionLock, FRWScopeLockType::SLT_Write);
@@ -149,7 +149,7 @@ void FPlatformManager::ClearSessionPlatforms(UCookOnTheFlyServer& COTFS)
 void FPlatformManager::AddSessionPlatform(UCookOnTheFlyServer& COTFS, FPlatformId TargetPlatform)
 {
 	checkf(IsSchedulerThread(), TEXT("Writing to SessionPlatforms is only allowed from the scheduler thread."));
-	CreatePlatformData(TargetPlatform);
+	CreatePlatformData(COTFS, TargetPlatform);
 
 	{
 		FRWScopeLock ScopeLock(SessionLock, FRWScopeLockType::SLT_ReadOnly);
@@ -197,7 +197,7 @@ FPlatformData* FPlatformManager::GetPlatformDataByName(FName PlatformName)
 	return Existing ? *Existing : nullptr;
 }
 
-FPlatformData& FPlatformManager::CreatePlatformData(const ITargetPlatform* Platform)
+FPlatformData& FPlatformManager::CreatePlatformData(UCookOnTheFlyServer& COTFS, const ITargetPlatform* Platform)
 {
 	check(Platform != nullptr);
 	FPlatformData** ExistingPlatformData = PlatformDatas.Find(Platform);
@@ -221,6 +221,20 @@ FPlatformData& FPlatformManager::CreatePlatformData(const ITargetPlatform* Platf
 		ITargetPlatform* NonConstPlatform = const_cast<ITargetPlatform*>(Platform);
 		PlatformData->TargetPlatform = NonConstPlatform;
 		PlatformData->PlatformName = PlatformName;
+
+		IAssetRegistry& AssetRegistry = IAssetRegistry::GetChecked();
+		AssetRegistry.InitializeSerializationOptions(PlatformData->AssetRegistryDevelopmentSaveOptions, Platform,
+			// This set of filters applies to the DevelopmentAssetRegistry created by AssetRegistryGenerator, so pass
+			// in ForDevelopment. The AssetRegistryState will be filtered further as ForGame when the
+			// AssetRegistryGenerator creates the RuntimeAssetRegistry from the DevelopmentAssetRegistry during
+			// FAssetRegistryGenerator::SaveAssetRegistry.
+			UE::AssetRegistry::ESerializationTarget::ForDevelopment);
+		PlatformData->AssetRegistryDevelopmentSaveOptions.bKeepDevelopmentAssetRegistryTags = true;
+		if (COTFS.IsForceNoFilterAssetsFromAssetRegistry())
+		{
+			PlatformData->AssetRegistryDevelopmentSaveOptions.DisableFilters();
+		}
+
 		return *PlatformData;
 	}
 }
diff --git a/Engine/Source/Editor/UnrealEd/Private/Cooker/CookPlatformManager.h b/Engine/Source/Editor/UnrealEd/Private/Cooker/CookPlatformManager.h
--- a/Engine/Source/Editor/UnrealEd/Private/Cooker/CookPlatformManager.h
+++ b/Engine/Source/Editor/UnrealEd/Private/Cooker/CookPlatformManager.h
@@ -2,6 +2,7 @@
 
 #pragma once
 
+#include "AssetRegistry/AssetRegistryState.h"
 #include "Commandlets/AssetRegistryGenerator.h"
 #include "Containers/Array.h"
 #include "Containers/ArrayView.h"
@@ -50,6 +51,7 @@ namespace UE::Cook
 		*/
 		TUniquePtr<FAssetRegistryGenerator> RegistryGenerator;
 		TUniquePtr<UE::Cook::IAssetRegistryReporter> RegistryReporter;
+		FAssetRegistrySerializationOptions AssetRegistryDevelopmentSaveOptions;
 
 		/*
 		* Whether BeginCookSandbox has been called for this platform.  If we have already initialized the sandbox we
@@ -135,7 +137,7 @@ namespace UE::Cook
 		FPlatformData* GetPlatformDataByName(FName PlatformName);
 
 		/** Create if not already created the necessary platform-specific data for the given platform. */
-		FPlatformData& CreatePlatformData(const ITargetPlatform* Platform);
+		FPlatformData& CreatePlatformData(UCookOnTheFlyServer& COTFS, const ITargetPlatform* Platform);
 
 		/** Return whether platform-specific setup steps have run in the current UCookOnTheFlyServer. */
 		bool IsPlatformInitialized(FPlatformId Platform) const;
diff --git a/Engine/Source/Editor/UnrealEd/Private/Cooker/CookSavePackage.cpp b/Engine/Source/Editor/UnrealEd/Private/Cooker/CookSavePackage.cpp
--- a/Engine/Source/Editor/UnrealEd/Private/Cooker/CookSavePackage.cpp
+++ b/Engine/Source/Editor/UnrealEd/Private/Cooker/CookSavePackage.cpp
@@ -519,11 +519,18 @@ void FSaveCookedPackageContext::FinishPlatform()
 	// Update asset registry
 	if (COTFS.IsDirectorCookByTheBook())
 	{
-		IAssetRegistryReporter& Reporter = *(COTFS.PlatformManager->GetPlatformData(TargetPlatform)->RegistryReporter);
+		FPlatformData* PlatformData = COTFS.PlatformManager->GetPlatformData(TargetPlatform);
+		IAssetRegistryReporter& Reporter = *(PlatformData->RegistryReporter);
 		TOptional<TArray<FAssetData>> AssetDatasFromSave;
 		if (bSuccessful)
 		{
 			AssetDatasFromSave.Emplace(MoveTemp(SavePackageResult.SavedAssets));
+			// Apply the cook's tag filtering settings to the tags from SavePackage, e.g. remove tags that are allowed
+			// or denied from runtime AssetRegistry by the AssetRegistry's configuration settings.
+			for (FAssetData& AssetData : *AssetDatasFromSave)
+			{
+				FAssetRegistryState::FilterTags(AssetData, PlatformData->AssetRegistryDevelopmentSaveOptions);
+			}
 		}
 		Reporter.UpdateAssetRegistryData(Package->GetFName(), Package, CookResult, &SavePackageResult,
 			MoveTemp(AssetDatasFromSave), MoveTemp(OverrideAssetPackageData), MoveTemp(OverridePackageDependencies),
diff --git a/Engine/Source/Editor/UnrealEd/Public/Commandlets/AssetRegistryGenerator.h b/Engine/Source/Editor/UnrealEd/Public/Commandlets/AssetRegistryGenerator.h
--- a/Engine/Source/Editor/UnrealEd/Public/Commandlets/AssetRegistryGenerator.h
+++ b/Engine/Source/Editor/UnrealEd/Public/Commandlets/AssetRegistryGenerator.h
@@ -51,12 +51,13 @@ public:
 	/**
 	 * Initializes manifest generator - creates manifest lists, hooks up delegates.
 	 */
-	void Initialize(bool bInitializeFromExisting);
+	void Initialize(bool bInitializeFromExisting, const FAssetRegistrySerializationOptions& DevelopmentSaveOptions);
 	/**
 	 * Used by DLC which does not Clone the entire global AR. Copies data from the global assetregistry
 	 * for all packages present in the previous AssetRegistry
 	 */
-	void CloneGlobalAssetRegistryFilteredByPreviousState(const FAssetRegistryState& PreviousState);
+	void CloneGlobalAssetRegistryFilteredByPreviousState(const FAssetRegistryState& PreviousState,
+		const FAssetRegistrySerializationOptions& DevelopmentSaveOptions);
 
 	const ITargetPlatform* GetTargetPlatform() const { return TargetPlatform; }
 
@@ -170,7 +171,7 @@ public:
 	 */
 	void FinalizeChunkIDs(const TSet<FName>& CookedPackages, const TSet<FName>& DevelopmentOnlyPackages,
 		UE::Cook::FCookSandbox& InSandboxFile, bool bGenerateStreamingInstallManifest,
-		const TSet<FName>& StartupPackages);
+		const TSet<FName>& StartupPackages, const FAssetRegistrySerializationOptions& DevelopmentSaveOptions);
 
 	/**
 	 * Register a chunk data generator with this generator.
@@ -225,7 +226,8 @@ public:
 	* Saves generated asset registry data for each platform.
 	*/
 	bool SaveAssetRegistry(const FString& SandboxPath, bool bSerializeDevelopmentAssetRegistry,
-		bool bForceNoFilterAssets, uint64& OutDevelopmentAssetRegistryHash, const TSet<FName>& CookedPackages);
+		bool bForceNoFilterAssets, uint64& OutDevelopmentAssetRegistryHash, const TSet<FName>& CookedPackages,
+		const FAssetRegistrySerializationOptions& DevelopmentSaveOptions);
 
 	/** 
 	 * Writes out CookerOpenOrder.log file 
diff --git a/Engine/Source/Runtime/AssetRegistry/Private/AssetRegistryState.cpp b/Engine/Source/Runtime/AssetRegistry/Private/AssetRegistryState.cpp
--- a/Engine/Source/Runtime/AssetRegistry/Private/AssetRegistryState.cpp
+++ b/Engine/Source/Runtime/AssetRegistry/Private/AssetRegistryState.cpp
@@ -3429,18 +3429,9 @@ void FAssetRegistryState::FilterTags(const FAssetRegistrySerializationOptions& O
 	CachedAssets.Enumerate([this, &Options](FAssetData& AssetData, FAssetDataPtrIndex AssetIndex)
 	{
 #endif
-		FAssetDataTagMap LocalTagsAndValues;
-		FAssetRegistryState::FilterTags(AssetData.TagsAndValues, LocalTagsAndValues,
-			Options.CookFilterlistTagsByClass.Find(AssetData.AssetClassPath), Options);
-		if (LocalTagsAndValues != AssetData.TagsAndValues)
-		{
-			AssetData.TagsAndValues = FAssetDataTagMapSharedView(MoveTemp(LocalTagsAndValues));
-		}
-
-		AssetData.TaggedAssetBundles = FAssetRegistryState::FilterBundles(AssetData.TaggedAssetBundles, Options);
-
-		// Add the AssetData to all its CachedAssetsByTag/CachedClassesByTag keys even if nothing changed, because
-		// we are reconstructing all CachedAssetsByTag.
+		FilterTags(AssetData, Options);
+		// Add the AssetData to all its CachedAssetsByTag/CachedClassesByTag keys even if FilterTags did not change
+		// anything, because we are reconstructing all CachedAssetsByTag.
 		for (auto TagIt = AssetData.TagsAndValues.CreateConstIterator(); TagIt; ++TagIt)
 		{
 #if UE_ASSETREGISTRY_CACHEDASSETSBYTAG
@@ -3457,6 +3448,18 @@ void FAssetRegistryState::FilterTags(const FAssetRegistrySerializationOptions& O
 #endif
 }
 
+void FAssetRegistryState::FilterTags(FAssetData& InOutAssetData, const FAssetRegistrySerializationOptions& Options)
+{
+	FAssetDataTagMap LocalTagsAndValues;
+	FilterTags(InOutAssetData.TagsAndValues, LocalTagsAndValues,
+		Options.CookFilterlistTagsByClass.Find(InOutAssetData.AssetClassPath), Options);
+	if (LocalTagsAndValues != InOutAssetData.TagsAndValues)
+	{
+		InOutAssetData.TagsAndValues = FAssetDataTagMapSharedView(MoveTemp(LocalTagsAndValues));
+	}
+	InOutAssetData.TaggedAssetBundles = FAssetRegistryState::FilterBundles(InOutAssetData.TaggedAssetBundles, Options);
+}
+
 void FAssetRegistryState::SetTagsOnExistingAsset(FAssetData* AssetData, FAssetDataTagMap&& NewTags)
 {
 	// Update the tag cache map to remove deleted tags
diff --git a/Engine/Source/Runtime/AssetRegistry/Public/AssetRegistry/AssetRegistryState.h b/Engine/Source/Runtime/AssetRegistry/Public/AssetRegistry/AssetRegistryState.h
--- a/Engine/Source/Runtime/AssetRegistry/Public/AssetRegistry/AssetRegistryState.h
+++ b/Engine/Source/Runtime/AssetRegistry/Public/AssetRegistry/AssetRegistryState.h
@@ -568,9 +568,11 @@ public:
 		const TSet<FName>& RequiredPackages, const TSet<FName>& RemovePackages, const TSet<int32> ChunksToKeep,
 		const FAssetRegistrySerializationOptions& Options);
 
-	/** Edit every AssetData's Tags to remove Tags that are filtered out by the filtering rules in Options */
+	/** Edit every AssetData's Tags to remove Tags that are filtered out by the filtering rules in Options. */
 	ASSETREGISTRY_API void FilterTags(const FAssetRegistrySerializationOptions& Options);
-
+	/** Filter a given AssetData's Tags to remove Tags that are filtered out by the filtering rules in Options. */
+	ASSETREGISTRY_API static void FilterTags(FAssetData& InOutAssetData,
+		const FAssetRegistrySerializationOptions& Options);
 
 	/** Serialize the registry to/from a file, skipping editor only data */
 	ASSETREGISTRY_API bool Serialize(FArchive& Ar, const FAssetRegistrySerializationOptions& Options);
