Switched ConfigAccessTracking's SectionNames over to using NAME_NO_NUMBER_INTERNAL instead of NAME_NO_NUMBER.

#Synced-CL 50977583

diff --git a/Engine/Source/Editor/UnrealEd/Private/Cooker/CookConfigAccessTracker.cpp b/Engine/Source/Editor/UnrealEd/Private/Cooker/CookConfigAccessTracker.cpp
--- a/Engine/Source/Editor/UnrealEd/Private/Cooker/CookConfigAccessTracker.cpp
+++ b/Engine/Source/Editor/UnrealEd/Private/Cooker/CookConfigAccessTracker.cpp
@@ -371,7 +371,7 @@ void FCookConfigAccessTracker::RecordValuesFromFile(const FConfigAccessData& Fil
 	TArray<FName> ValueNames;
 	for (const TPair<FString, FConfigSection>& SectionPair : ConfigFile)
 	{
-		FullPathData.SectionName = FName(FStringView(SectionPair.Key), NAME_NO_NUMBER).GetComparisonIndex();
+		FullPathData.SectionName = FName(FStringView(SectionPair.Key), NAME_NO_NUMBER_INTERNAL).GetComparisonIndex();
 		const FConfigSection& Section = SectionPair.Value;
 		ValueNames.Reset();
 		SectionPair.Value.GetKeys(ValueNames);
@@ -535,7 +535,7 @@ void FConfigAccessTrackingCollector::ServerReceiveMessage(UE::Cook::FMPCollector
 		++RecordIt;
 		Record.FileName = FName(RecordIt.AsString()).GetComparisonIndex();
 		++RecordIt;
-		Record.SectionName = FName(RecordIt.AsString(), NAME_NO_NUMBER).GetComparisonIndex();
+		Record.SectionName = FName(RecordIt.AsString(), NAME_NO_NUMBER_INTERNAL).GetComparisonIndex();
 		++RecordIt;
 		Record.ValueName = FMinimalName(FName(RecordIt.AsString()));
 		++RecordIt;
diff --git a/Engine/Source/Editor/UnrealEd/Private/Cooker/CookIncrementalVersion.cpp b/Engine/Source/Editor/UnrealEd/Private/Cooker/CookIncrementalVersion.cpp
--- a/Engine/Source/Editor/UnrealEd/Private/Cooker/CookIncrementalVersion.cpp
+++ b/Engine/Source/Editor/UnrealEd/Private/Cooker/CookIncrementalVersion.cpp
@@ -8,6 +8,6 @@ namespace UE::Cook
 {
 
 // Change CookIncrementalVersion to a new guid when all packages in an incremental cook need to be invalidated
-FGuid CookIncrementalVersion(0xAFB442DC, 0x88124FE2, 0xBEA2667A, 0x2D71D307);
+FGuid CookIncrementalVersion(0x7EADB3D2, 0x5A2546D6, 0xA2E79FB3, 0xE00B4D47);
 
 }
diff --git a/Engine/Source/Runtime/Core/Private/Misc/ConfigAccessData.cpp b/Engine/Source/Runtime/Core/Private/Misc/ConfigAccessData.cpp
--- a/Engine/Source/Runtime/Core/Private/Misc/ConfigAccessData.cpp
+++ b/Engine/Source/Runtime/Core/Private/Misc/ConfigAccessData.cpp
@@ -136,7 +136,7 @@ FConfigAccessData FConfigAccessData::Parse(FStringView Text)
 	// We write out sectionnames with [] for readability; remove them if they exist.
 	if (SectionNameView.StartsWith('[')) SectionNameView.RightChopInline(1);
 	if (SectionNameView.EndsWith(']')) SectionNameView.LeftChopInline(1);
-	FName LocalSectionName = SectionNameView.Len() == 0 ? NAME_None : FName(SectionNameView, NAME_NO_NUMBER);
+	FName LocalSectionName = SectionNameView.Len() == 0 ? NAME_None : FName(SectionNameView, NAME_NO_NUMBER_INTERNAL);
 	FName LocalValueName = ValueNameStr.Len() == 0 ? NAME_None : FName(ValueNameStr);
 
 	Result.ConfigPlatform = LocalConfigPlatform.GetComparisonIndex();
diff --git a/Engine/Source/Runtime/Core/Private/Misc/ConfigAccessTracking.cpp b/Engine/Source/Runtime/Core/Private/Misc/ConfigAccessTracking.cpp
--- a/Engine/Source/Runtime/Core/Private/Misc/ConfigAccessTracking.cpp
+++ b/Engine/Source/Runtime/Core/Private/Misc/ConfigAccessTracking.cpp
@@ -151,7 +151,7 @@ void FFile::SetAsLoadTypeConfigSystem(FConfigCacheIni& ConfigSystem, FConfigFile
 
 FSection::FSection(FFile& InFileAccess, FStringView InSectionName)
 	: FileAccess(&InFileAccess)
-	, SectionName(FName(InSectionName, NAME_NO_NUMBER).GetComparisonIndex())
+	, SectionName(FName(InSectionName, NAME_NO_NUMBER_INTERNAL).GetComparisonIndex())
 {
 }
 
