[Viewport ITF] Allow editor modes to opt-in to the new gizmos. When editor modes change, the ITF Transform Gizmo and ITF viewport controls will be disabled if the any active mode requires the legacy system. Editor modes already compatible with the new system have been updated to not require legacy systems. Add "ViewportInteractions.IgnoreModeRequirements" CVAR to bypass this system and always use ITF viewport interactions when a mode manager has it enabled.

#rnx

#Synced-CL 51055698

diff --git a/Engine/Plugins/Animation/MLDeformer/MLDeformerFramework/Source/MLDeformerFrameworkEditor/Public/MLDeformerEditorMode.h b/Engine/Plugins/Animation/MLDeformer/MLDeformerFramework/Source/MLDeformerFrameworkEditor/Public/MLDeformerEditorMode.h
--- a/Engine/Plugins/Animation/MLDeformer/MLDeformerFramework/Source/MLDeformerFrameworkEditor/Public/MLDeformerEditorMode.h
+++ b/Engine/Plugins/Animation/MLDeformer/MLDeformerFramework/Source/MLDeformerFrameworkEditor/Public/MLDeformerEditorMode.h
@@ -42,6 +42,7 @@ namespace UE::MLDeformer
 		virtual bool ShouldDrawWidget() const override { return false; }
 		virtual bool UsesTransformWidget() const override { return false; }
 		virtual bool UsesTransformWidget(UE::Widget::EWidgetMode CheckMode) const override { return false; }
+		virtual bool RequiresLegacyViewportInteractions() const override { return false; }
 		// ~END FEdMode overrides.
 
 	protected:
diff --git a/Engine/Plugins/Animation/MLDeformer/MLDeformerFramework/Source/MLDeformerFrameworkEditor/Public/MLDeformerPaintMode.h b/Engine/Plugins/Animation/MLDeformer/MLDeformerFramework/Source/MLDeformerFrameworkEditor/Public/MLDeformerPaintMode.h
--- a/Engine/Plugins/Animation/MLDeformer/MLDeformerFramework/Source/MLDeformerFrameworkEditor/Public/MLDeformerPaintMode.h
+++ b/Engine/Plugins/Animation/MLDeformer/MLDeformerFramework/Source/MLDeformerFrameworkEditor/Public/MLDeformerPaintMode.h
@@ -32,6 +32,7 @@ public:
 	virtual bool UsesTransformWidget() const override	{ return false; };
 	virtual bool UsesPropertyWidgets() const override	{ return false; };
 	virtual bool UsesToolkits() const override			{ return true; }
+	virtual bool RequiresLegacyViewportInteractions() const override	{ return false; }
 	// ~END UEdMode overrides
 
 	UE_API void SetMLDeformerEditor(UE::MLDeformer::FMLDeformerEditorToolkit* Editor);
diff --git a/Engine/Plugins/Animation/PoseSearch/Source/Editor/Public/PoseSearchDatabaseEdMode.h b/Engine/Plugins/Animation/PoseSearch/Source/Editor/Public/PoseSearchDatabaseEdMode.h
--- a/Engine/Plugins/Animation/PoseSearch/Source/Editor/Public/PoseSearchDatabaseEdMode.h
+++ b/Engine/Plugins/Animation/PoseSearch/Source/Editor/Public/PoseSearchDatabaseEdMode.h
@@ -43,6 +43,7 @@ namespace UE::PoseSearch
 		virtual FVector GetWidgetLocation() const override;
 		virtual bool GetCustomDrawingCoordinateSystem(FMatrix& InMatrix, void* InData) override;
 		virtual bool GetCustomInputCoordinateSystem(FMatrix& InMatrix, void* InData) override;
+		virtual bool RequiresLegacyViewportInteractions() const override { return false; }
 
 	private:
 		FDatabaseViewModel* ViewModel = nullptr;
diff --git a/Engine/Plugins/Animation/PoseSearch/Source/Editor/Public/PoseSearchInteractionAssetEditor.h b/Engine/Plugins/Animation/PoseSearch/Source/Editor/Public/PoseSearchInteractionAssetEditor.h
--- a/Engine/Plugins/Animation/PoseSearch/Source/Editor/Public/PoseSearchInteractionAssetEditor.h
+++ b/Engine/Plugins/Animation/PoseSearch/Source/Editor/Public/PoseSearchInteractionAssetEditor.h
@@ -97,6 +97,7 @@ public:
 	virtual bool AllowWidgetMove() override;
 	virtual bool ShouldDrawWidget() const override;
 	virtual bool GetCustomInputCoordinateSystem(FMatrix& InMatrix, void* InData) override;
+	virtual bool RequiresLegacyViewportInteractions() const override { return false; }
 
 private:
 	FInteractionAssetViewModel* ViewModel = nullptr;
diff --git a/Engine/Plugins/Animation/SkeletalMeshModelingTools/Source/SkeletalMeshModelingTools/Private/SkeletalMeshModelingToolsEditorMode.h b/Engine/Plugins/Animation/SkeletalMeshModelingTools/Source/SkeletalMeshModelingTools/Private/SkeletalMeshModelingToolsEditorMode.h
--- a/Engine/Plugins/Animation/SkeletalMeshModelingTools/Source/SkeletalMeshModelingTools/Private/SkeletalMeshModelingToolsEditorMode.h
+++ b/Engine/Plugins/Animation/SkeletalMeshModelingTools/Source/SkeletalMeshModelingTools/Private/SkeletalMeshModelingToolsEditorMode.h
@@ -110,6 +110,7 @@ public:
 	virtual bool GetCustomDrawingCoordinateSystem(FMatrix& InMatrix, void* InData) override;
 	virtual bool GetCustomInputCoordinateSystem(FMatrix& InMatrix, void* InData) override;
 	virtual bool HandleClick(FEditorViewportClient* InViewportClient, HHitProxy *HitProxy, const FViewportClick &Click) override;
+	virtual bool RequiresLegacyViewportInteractions() const override { return false; }
 
 	virtual void PostUndo() override;
 
diff --git a/Engine/Plugins/Editor/ModelingToolsEditorMode/Source/ModelingToolsEditorMode/Public/ModelingToolsEditorMode.h b/Engine/Plugins/Editor/ModelingToolsEditorMode/Source/ModelingToolsEditorMode/Public/ModelingToolsEditorMode.h
--- a/Engine/Plugins/Editor/ModelingToolsEditorMode/Source/ModelingToolsEditorMode/Public/ModelingToolsEditorMode.h
+++ b/Engine/Plugins/Editor/ModelingToolsEditorMode/Source/ModelingToolsEditorMode/Public/ModelingToolsEditorMode.h
@@ -52,6 +52,8 @@ public:
 	virtual bool CanAutoSave() const override;
 
 	virtual bool GetPivotForOrbit(FVector& OutPivot) const override;
+	
+	virtual bool RequiresLegacyViewportInteractions() const override { return false; }
 
 	/*
 	 * focus events
diff --git a/Engine/Plugins/Editor/ScriptableToolsEditorMode/Source/ScriptableToolsEditorMode/Public/ScriptableToolsEditorMode.h b/Engine/Plugins/Editor/ScriptableToolsEditorMode/Source/ScriptableToolsEditorMode/Public/ScriptableToolsEditorMode.h
--- a/Engine/Plugins/Editor/ScriptableToolsEditorMode/Source/ScriptableToolsEditorMode/Public/ScriptableToolsEditorMode.h
+++ b/Engine/Plugins/Editor/ScriptableToolsEditorMode/Source/ScriptableToolsEditorMode/Public/ScriptableToolsEditorMode.h
@@ -49,6 +49,8 @@ public:
 	UE_API virtual bool ComputeBoundingBoxForViewportFocus(AActor* Actor, UPrimitiveComponent* PrimitiveComponent, FBox& InOutBox) const override;
 
 	UE_API virtual bool GetPivotForOrbit(FVector& OutPivot) const override;
+	
+	virtual bool RequiresLegacyViewportInteractions() const override { return false; }
 
 	/*
 	 * focus events
diff --git a/Engine/Plugins/Editor/SequencerAnimTools/Source/SequencerAnimTools/Public/SequencerToolsEditMode.h b/Engine/Plugins/Editor/SequencerAnimTools/Source/SequencerAnimTools/Public/SequencerToolsEditMode.h
--- a/Engine/Plugins/Editor/SequencerAnimTools/Source/SequencerAnimTools/Public/SequencerToolsEditMode.h
+++ b/Engine/Plugins/Editor/SequencerAnimTools/Source/SequencerAnimTools/Public/SequencerToolsEditMode.h
@@ -27,6 +27,7 @@ public:
 	virtual bool InputKey(FEditorViewportClient* InViewportClient, FViewport* InViewport, FKey InKey, EInputEvent InEvent) override;
 	virtual bool UsesTransformWidget() const override { return false; }
 	virtual bool IsCompatibleWith(FEditorModeID OtherModeID) const override;
+	virtual bool RequiresLegacyViewportInteractions() const override { return false; }
 private:
 
 	void LocalOnModeActivated(const FEditorModeID& InID, bool bIsActive); //UEdMode has this already defined but it's not virtual
diff --git a/Engine/Plugins/Editor/StaticMeshEditorModeling/Source/StaticMeshEditorModeling/Public/StaticMeshEditorModelingMode.h b/Engine/Plugins/Editor/StaticMeshEditorModeling/Source/StaticMeshEditorModeling/Public/StaticMeshEditorModelingMode.h
--- a/Engine/Plugins/Editor/StaticMeshEditorModeling/Source/StaticMeshEditorModeling/Public/StaticMeshEditorModelingMode.h
+++ b/Engine/Plugins/Editor/StaticMeshEditorModeling/Source/StaticMeshEditorModeling/Public/StaticMeshEditorModelingMode.h
@@ -27,5 +27,5 @@ protected:
 	virtual bool UsesToolkits() const override;
 	virtual void CreateToolkit() override;
 	
-
+	virtual bool RequiresLegacyViewportInteractions() const override { return false; }
 };
diff --git a/Engine/Plugins/Enterprise/LidarPointCloud/Source/LidarPointCloudEditor/Private/LidarPointCloudEdMode.h b/Engine/Plugins/Enterprise/LidarPointCloud/Source/LidarPointCloudEditor/Private/LidarPointCloudEdMode.h
--- a/Engine/Plugins/Enterprise/LidarPointCloud/Source/LidarPointCloudEditor/Private/LidarPointCloudEdMode.h
+++ b/Engine/Plugins/Enterprise/LidarPointCloud/Source/LidarPointCloudEditor/Private/LidarPointCloudEdMode.h
@@ -29,7 +29,7 @@ public:
 	virtual void BindCommands() override;
 	virtual void OnToolStarted(UInteractiveToolManager* Manager, UInteractiveTool* Tool) override;
 	virtual void OnToolEnded(UInteractiveToolManager* Manager, UInteractiveTool* Tool) override;
-
+	virtual bool RequiresLegacyViewportInteractions() const override { return false; }
 private:
 	void CancelActiveToolAction();
 };
\ No newline at end of file
diff --git a/Engine/Plugins/Experimental/Animation/AnimDatabase/Source/AnimDatabaseEditor/Public/AnimDatabaseEditorMode.h b/Engine/Plugins/Experimental/Animation/AnimDatabase/Source/AnimDatabaseEditor/Public/AnimDatabaseEditorMode.h
--- a/Engine/Plugins/Experimental/Animation/AnimDatabase/Source/AnimDatabaseEditor/Public/AnimDatabaseEditorMode.h
+++ b/Engine/Plugins/Experimental/Animation/AnimDatabase/Source/AnimDatabaseEditor/Public/AnimDatabaseEditorMode.h
@@ -19,6 +19,7 @@ namespace UE::AnimDatabase::Editor
 
 		virtual void Tick(FEditorViewportClient* ViewportClient, float DeltaTime) override;
 		virtual void Render(const FSceneView* View, FViewport* Viewport, FPrimitiveDrawInterface* PDI) override;
+		virtual bool RequiresLegacyViewportInteractions() const override { return false; }
 
 	private:
 
diff --git a/Engine/Plugins/Experimental/WidgetEditorToolPalette/Source/WidgetEditorToolPalette/Public/WidgetEditorToolPaletteMode.h b/Engine/Plugins/Experimental/WidgetEditorToolPalette/Source/WidgetEditorToolPalette/Public/WidgetEditorToolPaletteMode.h
--- a/Engine/Plugins/Experimental/WidgetEditorToolPalette/Source/WidgetEditorToolPalette/Public/WidgetEditorToolPaletteMode.h
+++ b/Engine/Plugins/Experimental/WidgetEditorToolPalette/Source/WidgetEditorToolPalette/Public/WidgetEditorToolPaletteMode.h
@@ -26,5 +26,6 @@ protected:
 	bool UsesToolkits() const override; 
 	void CreateToolkit() override;
 	
+	virtual bool RequiresLegacyViewportInteractions() const override { return false; }
 
 };
diff --git a/Engine/Plugins/Runtime/HairStrands/Source/HairStrandsEditor/Public/GroomEditorMode.h b/Engine/Plugins/Runtime/HairStrands/Source/HairStrandsEditor/Public/GroomEditorMode.h
--- a/Engine/Plugins/Runtime/HairStrands/Source/HairStrandsEditor/Public/GroomEditorMode.h
+++ b/Engine/Plugins/Runtime/HairStrands/Source/HairStrandsEditor/Public/GroomEditorMode.h
@@ -23,6 +23,8 @@ public:
 	// these disable the standard gizmo, which is probably want we want in
 	// these tools as we can't hit-test the standard gizmo...
 	virtual bool AllowWidgetMove() override;
+	
+	virtual bool RequiresLegacyViewportInteractions() const override { return false; }
 
 	/*
 	 * focus events
diff --git a/Engine/Source/Editor/ActorPickerMode/Private/EditorModeActorPicker.h b/Engine/Source/Editor/ActorPickerMode/Private/EditorModeActorPicker.h
--- a/Engine/Source/Editor/ActorPickerMode/Private/EditorModeActorPicker.h
+++ b/Engine/Source/Editor/ActorPickerMode/Private/EditorModeActorPicker.h
@@ -42,6 +42,7 @@ public:
 	virtual bool IsCompatibleWith(FEditorModeID OtherModeID) const override;
 	virtual void Enter() override;
 	virtual void Exit() override;
+	virtual bool RequiresLegacyViewportInteractions() const override { return false; }
 	/** End FEdMode interface */
 
 	/** Delegate used to display information about picking near the cursor */
diff --git a/Engine/Source/Editor/AnimGraph/Public/AnimNodeEditMode.h b/Engine/Source/Editor/AnimGraph/Public/AnimNodeEditMode.h
--- a/Engine/Source/Editor/AnimGraph/Public/AnimNodeEditMode.h
+++ b/Engine/Source/Editor/AnimGraph/Public/AnimNodeEditMode.h
@@ -74,6 +74,7 @@ public:
 	UE_API virtual bool ShouldDrawWidget() const override;
 	UE_API virtual void Tick(FEditorViewportClient* ViewportClient, float DeltaTime) override;
 	UE_API virtual void Exit() override;
+	virtual bool RequiresLegacyViewportInteractions() const override { return false; }
 
 	UE_API virtual void RegisterPoseWatchedNode(UAnimGraphNode_Base* InEditorNode, FAnimNode_Base* InRuntimeNode);
 
diff --git a/Engine/Source/Editor/AnimationEditMode/Public/AnimationEditMode.h b/Engine/Source/Editor/AnimationEditMode/Public/AnimationEditMode.h
--- a/Engine/Source/Editor/AnimationEditMode/Public/AnimationEditMode.h
+++ b/Engine/Source/Editor/AnimationEditMode/Public/AnimationEditMode.h
@@ -53,6 +53,8 @@ public:
 	UE_API virtual void Enter() override;
 	UE_API virtual void Exit() override;
 	UE_API virtual void AddReferencedObjects(FReferenceCollector& Collector) override;
+	
+	virtual bool RequiresLegacyViewportInteractions() const override { return false; }
 private:
 	TObjectPtr<UAnimationEditModeContext> AnimationEditModeContext;
 };
diff --git a/Engine/Source/Editor/Experimental/EditorInteractiveToolsFramework/Private/EditorInteractiveGizmoManager.cpp b/Engine/Source/Editor/Experimental/EditorInteractiveToolsFramework/Private/EditorInteractiveGizmoManager.cpp
--- a/Engine/Source/Editor/Experimental/EditorInteractiveToolsFramework/Private/EditorInteractiveGizmoManager.cpp
+++ b/Engine/Source/Editor/Experimental/EditorInteractiveToolsFramework/Private/EditorInteractiveGizmoManager.cpp
@@ -275,9 +275,26 @@ bool UEditorInteractiveGizmoManager::GetShowEditorGizmosForView(IToolsContextRen
 	return bShowEditorGizmos && bEngineShowFlagsModeWidget;
 }
 
+void UEditorInteractiveGizmoManager::SetAllowEditorGizmos(bool bInAllowed)
+{
+	if (bAllowEditorGizmos != bInAllowed)
+	{
+		bAllowEditorGizmos = bInAllowed;
+		
+		OnUsesNewTRSGizmosChangedDelegate().Broadcast(bAllowEditorGizmos && UsesNewTRSGizmos());
+		
+		if (bAllowEditorGizmos && UsesNewTRSGizmos())
+		{
+			
+		}
+		
+		UpdateActiveEditorGizmos();
+	}
+}
+
 void UEditorInteractiveGizmoManager::UpdateActiveEditorGizmos()
 {
-	if (!UsesNewTRSGizmos())
+	if (!bAllowEditorGizmos || !UsesNewTRSGizmos())
 	{
 		if (UTransformGizmo* Gizmo = FindDefaultTransformGizmo())
 		{
diff --git a/Engine/Source/Editor/Experimental/EditorInteractiveToolsFramework/Public/EditorInteractiveGizmoManager.h b/Engine/Source/Editor/Experimental/EditorInteractiveToolsFramework/Public/EditorInteractiveGizmoManager.h
--- a/Engine/Source/Editor/Experimental/EditorInteractiveToolsFramework/Public/EditorInteractiveGizmoManager.h
+++ b/Engine/Source/Editor/Experimental/EditorInteractiveToolsFramework/Public/EditorInteractiveGizmoManager.h
@@ -113,6 +113,11 @@ public:
 		return bSearchLocalBuildersOnly;
 	}
 
+	/**
+	 * Sets whether editor gizmos are allowed
+	 */
+	UE_API void SetAllowEditorGizmos(bool bInAllowed);
+
 	/**
 	 * Shutdown and remove a selection-based Gizmo
 	 * @param Gizmo the Gizmo to shutdown and remove
@@ -213,6 +218,9 @@ protected:
 	/** @todo: remove when GetShowEditorGizmos() is moved to gizmo context object */
 	FEditorModeTools* EditorModeManager = nullptr;
 
+	/** Whether Editor gizmos are ever allowed in this context. Layered on top of UsesNewTRSGizmos() */
+	bool bAllowEditorGizmos = true;
+	
 	/** Whether Editor gizmos are enabled. UpdateActiveEditorGizmos() determines this value each tick and updates if it has changed. */
 	bool bShowEditorGizmos = false;
 
diff --git a/Engine/Source/Editor/GameplayDebugger/Public/GameplayDebuggerEdMode.h b/Engine/Source/Editor/GameplayDebugger/Public/GameplayDebuggerEdMode.h
--- a/Engine/Source/Editor/GameplayDebugger/Public/GameplayDebuggerEdMode.h
+++ b/Engine/Source/Editor/GameplayDebugger/Public/GameplayDebuggerEdMode.h
@@ -22,6 +22,7 @@ public:
 	virtual bool LostFocus(FEditorViewportClient* ViewportClient, FViewport* Viewport) override;
 	virtual bool InputKey(FEditorViewportClient* InViewportClient, FViewport* InViewport, FKey InKey, EInputEvent InEvent) override;
 	virtual void Tick(FEditorViewportClient* ViewportClient, float DeltaTime);
+	virtual bool RequiresLegacyViewportInteractions() const override { return false; }
 	// End of FEdMode interface
 
 	static void SafeOpenMode();
diff --git a/Engine/Source/Editor/LandscapeEditor/Private/LandscapeEdMode.h b/Engine/Source/Editor/LandscapeEditor/Private/LandscapeEdMode.h
--- a/Engine/Source/Editor/LandscapeEditor/Private/LandscapeEdMode.h
+++ b/Engine/Source/Editor/LandscapeEditor/Private/LandscapeEdMode.h
@@ -474,6 +474,8 @@ public:
 
 	/** FEdMode: Returns true if this mode uses the transform widget */
 	virtual bool UsesTransformWidget() const override;
+	
+	virtual bool RequiresLegacyViewportInteractions() const override { return false; }
 
 	virtual EAxisList::Type GetWidgetAxisToDraw(UE::Widget::EWidgetMode InWidgetMode) const override;
 
diff --git a/Engine/Source/Editor/LevelEditor/Private/ViewportToolbar/LevelEditorSubmenus.cpp b/Engine/Source/Editor/LevelEditor/Private/ViewportToolbar/LevelEditorSubmenus.cpp
--- a/Engine/Source/Editor/LevelEditor/Private/ViewportToolbar/LevelEditorSubmenus.cpp
+++ b/Engine/Source/Editor/LevelEditor/Private/ViewportToolbar/LevelEditorSubmenus.cpp
@@ -3440,7 +3440,7 @@ void ExtendTransformSubmenu(FName InTransformSubmenuName)
 			MarqueeMode.Visibility = TAttribute<bool>::CreateLambda([]()
 			{
 				// ITF Drag Tools are automatically enabled by new TRS Gizmos
-				return UEditorInteractiveGizmoManager::UsesNewTRSGizmos() && GLevelEditorModeTools().GetInteractiveToolsContext()->UsesViewportInteractions();
+				return GLevelEditorModeTools().GetInteractiveToolsContext()->UsesViewportInteractions();
 			});
 		}
 	}
diff --git a/Engine/Source/Editor/MovieSceneTools/Private/EditModes/SkeletalAnimationTrackEditMode.h b/Engine/Source/Editor/MovieSceneTools/Private/EditModes/SkeletalAnimationTrackEditMode.h
--- a/Engine/Source/Editor/MovieSceneTools/Private/EditModes/SkeletalAnimationTrackEditMode.h
+++ b/Engine/Source/Editor/MovieSceneTools/Private/EditModes/SkeletalAnimationTrackEditMode.h
@@ -53,6 +53,7 @@ public:
 	virtual bool GetCustomDrawingCoordinateSystem(FMatrix& OutMatrix, void* InData) override;
 	virtual bool GetCustomInputCoordinateSystem(FMatrix& OutMatrix, void* InData) override;
 	virtual bool IsCompatibleWith(FEditorModeID OtherModeID) const override;
+	virtual bool RequiresLegacyViewportInteractions() const override { return false; }
 	
 	//GCObject 
 	virtual void AddReferencedObjects(FReferenceCollector& Collector) override;
diff --git a/Engine/Source/Editor/MovieSceneTools/Private/EditModes/SubTrackEditorMode.h b/Engine/Source/Editor/MovieSceneTools/Private/EditModes/SubTrackEditorMode.h
--- a/Engine/Source/Editor/MovieSceneTools/Private/EditModes/SubTrackEditorMode.h
+++ b/Engine/Source/Editor/MovieSceneTools/Private/EditModes/SubTrackEditorMode.h
@@ -28,6 +28,7 @@ public:
 	virtual bool GetCustomDrawingCoordinateSystem(FMatrix& OutMatrix, void* InData) override;
 	virtual bool GetCustomInputCoordinateSystem(FMatrix& OutMatrix, void* InData) override;
 	virtual bool IsCompatibleWith(FEditorModeID OtherModeID) const override;
+	virtual bool RequiresLegacyViewportInteractions() const override { return false; }
 
 	void SetSequencer(const TSharedPtr<ISequencer>& InSequencer) { WeakSequencer = InSequencer; }
 	FOnOriginValueChanged& GetOnOriginValueChanged() { return OnOriginValueChanged; };
diff --git a/Engine/Source/Editor/SceneDepthPickerMode/Private/EditorModeSceneDepthPicker.h b/Engine/Source/Editor/SceneDepthPickerMode/Private/EditorModeSceneDepthPicker.h
--- a/Engine/Source/Editor/SceneDepthPickerMode/Private/EditorModeSceneDepthPicker.h
+++ b/Engine/Source/Editor/SceneDepthPickerMode/Private/EditorModeSceneDepthPicker.h
@@ -36,6 +36,7 @@ public:
 	virtual bool UsesToolkits() const override;
 	virtual bool IsCompatibleWith(FEditorModeID OtherModeID) const override;
 	virtual void Initialize() override;
+	virtual bool RequiresLegacyViewportInteractions() const override { return false; }
 	/** End FEdMode interface */
 
 	/** End the mode */
diff --git a/Engine/Source/Editor/Sequencer/Private/SequencerEdMode.h b/Engine/Source/Editor/Sequencer/Private/SequencerEdMode.h
--- a/Engine/Source/Editor/Sequencer/Private/SequencerEdMode.h
+++ b/Engine/Source/Editor/Sequencer/Private/SequencerEdMode.h
@@ -91,6 +91,7 @@ public:
 	virtual bool FrustumSelect(const FConvexVolume& InFrustum, FEditorViewportClient* InViewportClient, bool InSelect) override;
 	virtual bool InputDelta(FEditorViewportClient* InViewportClient, FViewport* InViewport, FVector& InDrag, FRotator& InRot, FVector& InScale) override;
 	virtual void Tick(FEditorViewportClient* ViewportClient,float DeltaTime) override;
+	virtual bool RequiresLegacyViewportInteractions() const override { return false; }
 
 	bool IsPressingMoveTimeSlider(FViewport* InViewport) const;
 	bool IsDoingDrag(FViewport* InViewport) const;
diff --git a/Engine/Source/Editor/UnrealEd/Private/EditorModeManager.cpp b/Engine/Source/Editor/UnrealEd/Private/EditorModeManager.cpp
--- a/Engine/Source/Editor/UnrealEd/Private/EditorModeManager.cpp
+++ b/Engine/Source/Editor/UnrealEd/Private/EditorModeManager.cpp
@@ -1808,6 +1808,28 @@ bool FEditorModeTools::DisallowMouseDeltaTracking() const
 	return bDisallow;
 }
 
+bool FEditorModeTools::RequiresLegacyViewportInteractions() const
+{
+	bool bRequiresGizmos = false;
+	ForEachEdMode([&bRequiresGizmos, this](UEdMode* EdMode)
+		{
+			if (ILegacyEdModeViewportInterface* Interface = Cast<ILegacyEdModeViewportInterface>(EdMode))
+			{
+				if (Interface->RequiresLegacyViewportInteractions())
+				{
+					if (bViewportITF)
+					{
+						UE_LOG(LogEditorModes, Warning, TEXT("ITF Gizmos and Controls have been disabled. The \"%s\" Editor Mode requires the legacy systems."), *EdMode->GetModeInfo().Name.ToString())
+					}
+					
+					bRequiresGizmos = true;
+				}
+			}
+			return true;
+		});
+	return bRequiresGizmos;	
+}
+
 bool FEditorModeTools::GetCursor(EMouseCursor::Type& OutCursor) const
 {
 	bool bHandled = false;
@@ -2102,7 +2124,7 @@ UModeManagerInteractiveToolsContext* FEditorModeTools::GetInteractiveToolsContex
 
 bool FEditorModeTools::UsesViewportITF() const
 {
-	return bViewportITF && UEditorInteractiveGizmoManager::UsesNewTRSGizmos();
+	return bViewportITF && InteractiveToolsContext && InteractiveToolsContext->UsesViewportInteractions() && UEditorInteractiveGizmoManager::UsesNewTRSGizmos();
 }
 
 void FEditorModeTools::SetSupportsViewportITF(bool bIsSupported)
diff --git a/Engine/Source/Editor/UnrealEd/Private/Tools/DefaultEdMode.cpp b/Engine/Source/Editor/UnrealEd/Private/Tools/DefaultEdMode.cpp
--- a/Engine/Source/Editor/UnrealEd/Private/Tools/DefaultEdMode.cpp
+++ b/Engine/Source/Editor/UnrealEd/Private/Tools/DefaultEdMode.cpp
@@ -21,6 +21,8 @@ public:
 		}
 		return false;
 	}
+	
+	virtual bool RequiresLegacyViewportInteractions() const override { return false; }
 };
 
 UEdModeDefault::UEdModeDefault()
@@ -59,6 +61,8 @@ public:
 	{
 		return true;
 	}
+	
+	virtual bool RequiresLegacyViewportInteractions() const override { return false; }
 };
 
 UAssetEdModeDefault::UAssetEdModeDefault()
diff --git a/Engine/Source/Editor/UnrealEd/Private/Tools/EdModeInteractiveToolsContext.cpp b/Engine/Source/Editor/UnrealEd/Private/Tools/EdModeInteractiveToolsContext.cpp
--- a/Engine/Source/Editor/UnrealEd/Private/Tools/EdModeInteractiveToolsContext.cpp
+++ b/Engine/Source/Editor/UnrealEd/Private/Tools/EdModeInteractiveToolsContext.cpp
@@ -61,6 +61,12 @@ namespace UE::Editor::Private
 		TEXT("0: Do not allow mouse drag operations to continue beyond the viewport bounds (default)\n")
 		TEXT("1: Allow mouse drag operations to continue beyond the viewport bounds"),
 		ECVF_Default);
+		
+	static TAutoConsoleVariable<bool> CVarIgnoreModeLegacyGizmoRequirements(
+		TEXT("ViewportInteractions.IgnoreModeRequirements"),
+		false,
+		TEXT("Whether to force on viewport interactions, ignoring legacy requests from editor modes.")
+	);
 }
 
 class FEdModeToolsContextQueriesImpl : public IToolsContextQueriesAPI
@@ -939,7 +945,24 @@ void UModeManagerInteractiveToolsContext::Tick(FEditorViewportClient* ViewportCl
 	if (bViewportInteractionsDirty)
 	{
 		bViewportInteractionsDirty = false;
-		RegisterViewportInteractions();
+		
+		if (const FEditorModeTools* ModeTools = GetParentEditorModeManager())
+		{
+			bModesSupportViewportInteractions = !ModeTools->RequiresLegacyViewportInteractions() || UE::Editor::Private::CVarIgnoreModeLegacyGizmoRequirements.GetValueOnGameThread();
+			if (UEditorInteractiveGizmoManager* EditorGizmoManager = Cast<UEditorInteractiveGizmoManager>(GizmoManager))
+			{
+				EditorGizmoManager->SetAllowEditorGizmos(bModesSupportViewportInteractions);
+			}
+		}
+		
+		if (UsesViewportInteractions())
+		{
+			RegisterViewportInteractions();	
+		}
+		else
+		{
+			UnregisterViewportInteractions();
+		}
 	}
 
 	if (ViewportInteractionsBehaviorSource)
@@ -1547,7 +1570,7 @@ void UModeManagerInteractiveToolsContext::SetViewportInteractionsEnabled(bool bI
 
 bool UModeManagerInteractiveToolsContext::UsesViewportInteractions() const
 {
-	return bUsesViewportInteractions && UE::Editor::ViewportInteractions::UseEditorViewportInteractions();
+	return bUsesViewportInteractions && bModesSupportViewportInteractions && UE::Editor::ViewportInteractions::UseEditorViewportInteractions();
 }
 
 UViewportInteractionsBehaviorSource* UModeManagerInteractiveToolsContext::GetViewportInteractionsBehaviorSource()
diff --git a/Engine/Source/Editor/UnrealEd/Private/Tools/LegacyEdMode.cpp b/Engine/Source/Editor/UnrealEd/Private/Tools/LegacyEdMode.cpp
--- a/Engine/Source/Editor/UnrealEd/Private/Tools/LegacyEdMode.cpp
+++ b/Engine/Source/Editor/UnrealEd/Private/Tools/LegacyEdMode.cpp
@@ -151,6 +151,11 @@ bool ULegacyEdModeWrapper::DisallowMouseDeltaTracking() const
 	return LegacyEditorMode->DisallowMouseDeltaTracking();
 }
 
+bool ULegacyEdModeWrapper::RequiresLegacyViewportInteractions() const
+{
+	return LegacyEditorMode->RequiresLegacyViewportInteractions();
+}
+
 bool ULegacyEdModeWrapper::GetCursor(EMouseCursor::Type& OutCursor) const
 {
 	return LegacyEditorMode->GetCursor(OutCursor);
diff --git a/Engine/Source/Editor/UnrealEd/Private/Tools/LegacyEdMode.h b/Engine/Source/Editor/UnrealEd/Private/Tools/LegacyEdMode.h
--- a/Engine/Source/Editor/UnrealEd/Private/Tools/LegacyEdMode.h
+++ b/Engine/Source/Editor/UnrealEd/Private/Tools/LegacyEdMode.h
@@ -55,6 +55,7 @@ public:
 	UNREALED_API virtual bool GetPivotForOrbit(FVector& OutPivot) const override;
 	UNREALED_API virtual void PostUndo() override;
 	UNREALED_API virtual bool DisallowMouseDeltaTracking() const override;
+	UNREALED_API virtual bool RequiresLegacyViewportInteractions() const override;
 	UNREALED_API virtual bool GetCursor(EMouseCursor::Type& OutCursor) const override;
 	UNREALED_API virtual bool GetOverrideCursorVisibility(bool& bWantsOverride, bool& bHardwareCursorVisible, bool bSoftwareCursorVisible) const;
 	UNREALED_API virtual bool CanAutoSave() const override;
diff --git a/Engine/Source/Editor/UnrealEd/Private/Tools/LegacyEdModeWidgetHelpers.cpp b/Engine/Source/Editor/UnrealEd/Private/Tools/LegacyEdModeWidgetHelpers.cpp
--- a/Engine/Source/Editor/UnrealEd/Private/Tools/LegacyEdModeWidgetHelpers.cpp
+++ b/Engine/Source/Editor/UnrealEd/Private/Tools/LegacyEdModeWidgetHelpers.cpp
@@ -557,6 +557,11 @@ bool FLegacyEdModeWidgetHelper::AllowsViewportDragTool() const
 	return true;
 }
 
+bool FLegacyEdModeWidgetHelper::RequiresLegacyViewportInteractions() const
+{
+	return true;
+}
+
 bool FLegacyEdModeWidgetHelper::InputDelta(FEditorViewportClient* InViewportClient, FViewport* InViewport, FVector& InDrag, FRotator& InRot, FVector& InScale)
 {
 	if (UsesPropertyWidgets())
@@ -941,3 +946,8 @@ bool UBaseLegacyWidgetEdMode::AllowsViewportDragTool() const
 	return WidgetHelper->AllowsViewportDragTool();
 }
 
+bool UBaseLegacyWidgetEdMode::RequiresLegacyViewportInteractions() const
+{
+	return WidgetHelper->RequiresLegacyViewportInteractions();
+}
+
diff --git a/Engine/Source/Editor/UnrealEd/Public/EditorModeManager.h b/Engine/Source/Editor/UnrealEd/Public/EditorModeManager.h
--- a/Engine/Source/Editor/UnrealEd/Public/EditorModeManager.h
+++ b/Engine/Source/Editor/UnrealEd/Public/EditorModeManager.h
@@ -301,6 +301,9 @@ public:
 
 	/** True if we should disallow mouse delta tracking. */
 	UNREALED_API bool DisallowMouseDeltaTracking() const;
+	
+	/** True if legacy gizmos are required */
+	UNREALED_API bool RequiresLegacyViewportInteractions() const;
 
 	/** Get a cursor to override the default with, if any */
 	UNREALED_API bool GetCursor(EMouseCursor::Type& OutCursor) const;
diff --git a/Engine/Source/Editor/UnrealEd/Public/Tools/EdModeInteractiveToolsContext.h b/Engine/Source/Editor/UnrealEd/Public/Tools/EdModeInteractiveToolsContext.h
--- a/Engine/Source/Editor/UnrealEd/Public/Tools/EdModeInteractiveToolsContext.h
+++ b/Engine/Source/Editor/UnrealEd/Public/Tools/EdModeInteractiveToolsContext.h
@@ -517,7 +517,11 @@ protected:
 	FViewportInteractionsSourceConstMessage PostBuildViewportInteractionsEvent;
 	PRAGMA_ENABLE_EXPERIMENTAL_WARNINGS
 
+	/** Whether viewport interactions can ever be supported */
 	bool bUsesViewportInteractions = false;
+	/** Whether the currently active editor modes currently support viewport interactions */
+	bool bModesSupportViewportInteractions = false;
+	/** A flag that allows de/composing viewport interactions once per set of mode changes */
 	bool bViewportInteractionsDirty = false;
 };
 
diff --git a/Engine/Source/Editor/UnrealEd/Public/Tools/LegacyEdModeInterfaces.h b/Engine/Source/Editor/UnrealEd/Public/Tools/LegacyEdModeInterfaces.h
--- a/Engine/Source/Editor/UnrealEd/Public/Tools/LegacyEdModeInterfaces.h
+++ b/Engine/Source/Editor/UnrealEd/Public/Tools/LegacyEdModeInterfaces.h
@@ -214,4 +214,9 @@ public:
 
 	/** If the Edmode is handling its own mouse deltas, it can disable the MouseDeltaTacker */
 	virtual bool DisallowMouseDeltaTracking() const { return false; }
+	
+	/**
+	 * @return True if this mode requires the non-ITF gizmo and viewport input system.
+	 */
+	virtual bool RequiresLegacyViewportInteractions() const { return true; }
 };
diff --git a/Engine/Source/Editor/UnrealEd/Public/Tools/LegacyEdModeWidgetHelpers.h b/Engine/Source/Editor/UnrealEd/Public/Tools/LegacyEdModeWidgetHelpers.h
--- a/Engine/Source/Editor/UnrealEd/Public/Tools/LegacyEdModeWidgetHelpers.h
+++ b/Engine/Source/Editor/UnrealEd/Public/Tools/LegacyEdModeWidgetHelpers.h
@@ -49,6 +49,7 @@ public:
 	
 	UNREALED_API virtual void ActorSelectionChangeNotify();
 	UNREALED_API virtual bool AllowsViewportDragTool() const;
+	UNREALED_API virtual bool RequiresLegacyViewportInteractions() const;
 
 	// Property Widgets
 	/** Structure that holds info about our optional property widget */
@@ -169,6 +170,7 @@ public:
 	UNREALED_API virtual bool GetCustomInputCoordinateSystem(FMatrix& InMatrix, void* InData) override;
 	UNREALED_API virtual void ActorSelectionChangeNotify() override;
 	UNREALED_API virtual bool AllowsViewportDragTool() const override;
+	UNREALED_API virtual bool RequiresLegacyViewportInteractions() const override;
 
 protected:
 	UNREALED_API virtual TSharedRef<FLegacyEdModeWidgetHelper> CreateWidgetHelper();
