Cooker: CookArtifacts: Add ICookInfo.GetCookMetadataOutputFolder for use by CookArtifacts. It is not a simple fixed subdirectory of GetCookOutputFolder when cooking DLC.
And const specifier to the ICookInfo functions that are const, since UCookOnTheFlyServer is supposed to have const correctness.
#rnx

#Synced-CL 51045425

diff --git a/Engine/Source/Editor/UnrealEd/Classes/CookOnTheSide/CookOnTheFlyServer.h b/Engine/Source/Editor/UnrealEd/Classes/CookOnTheSide/CookOnTheFlyServer.h
--- a/Engine/Source/Editor/UnrealEd/Classes/CookOnTheSide/CookOnTheFlyServer.h
+++ b/Engine/Source/Editor/UnrealEd/Classes/CookOnTheSide/CookOnTheFlyServer.h
@@ -697,19 +697,20 @@ public:
 	ECookInitializationFlags GetCookFlags() const { return CookFlags; }
 
 	// ICookInfo interface
-	UNREALED_API virtual UE::Cook::FInstigator GetInstigator(FName PackageName) override;
-	UNREALED_API virtual TArray<UE::Cook::FInstigator> GetInstigatorChain(FName PackageName) override;
-	UNREALED_API virtual UE::Cook::ECookType GetCookType() override;
-	UNREALED_API virtual UE::Cook::ECookingDLC GetCookingDLC() override;
-	UNREALED_API virtual FString GetDLCName() override;
-	UNREALED_API virtual FString GetBasedOnReleaseVersion() override;
-	UNREALED_API virtual FString GetCreateReleaseVersion() override;
-	UNREALED_API virtual UE::Cook::EProcessType GetProcessType() override;
-	UNREALED_API virtual UE::Cook::ECookValidationOptions GetCookValidationOptions() override;
-	UNREALED_API virtual bool IsIncremental() override;
-	UNREALED_API virtual TArray<const ITargetPlatform*> GetSessionPlatforms() override;
-	UNREALED_API virtual FString GetCookOutputFolder(const ITargetPlatform* TargetPlatform) override;
-	UNREALED_API virtual const TSet<IPlugin*>* GetEnabledPlugins(const ITargetPlatform* TargetPlatform) override;
+	UNREALED_API virtual UE::Cook::FInstigator GetInstigator(FName PackageName) const override;
+	UNREALED_API virtual TArray<UE::Cook::FInstigator> GetInstigatorChain(FName PackageName) const override;
+	UNREALED_API virtual UE::Cook::ECookType GetCookType() const override;
+	UNREALED_API virtual UE::Cook::ECookingDLC GetCookingDLC() const override;
+	UNREALED_API virtual FString GetDLCName() const override;
+	UNREALED_API virtual FString GetBasedOnReleaseVersion() const override;
+	UNREALED_API virtual FString GetCreateReleaseVersion() const override;
+	UNREALED_API virtual UE::Cook::EProcessType GetProcessType() const override;
+	UNREALED_API virtual UE::Cook::ECookValidationOptions GetCookValidationOptions() const override;
+	UNREALED_API virtual bool IsIncremental() const override;
+	UNREALED_API virtual TArray<const ITargetPlatform*> GetSessionPlatforms() const override;
+	UNREALED_API virtual FString GetCookOutputFolder(const ITargetPlatform* TargetPlatform) const override;
+	UNREALED_API virtual FString GetCookMetadataOutputFolder(const ITargetPlatform* TargetPlatform) const override;	
+	UNREALED_API virtual const TSet<IPlugin*>* GetEnabledPlugins(const ITargetPlatform* TargetPlatform) const override;
 
 	UNREALED_API virtual void RegisterCollector(UE::Cook::IMPCollector* Collector,
 		UE::Cook::EProcessType ProcessType = UE::Cook::EProcessType::AllMPCook) override;
@@ -1006,7 +1007,7 @@ protected:
 
 private:
 
-	UE::Cook::FInstigator GetInstigator(FName PackageName, UE::Cook::EReachability Reachability);
+	UE::Cook::FInstigator GetInstigator(FName PackageName, UE::Cook::EReachability Reachability) const;
 	bool PumpHasExceededMaxMemory(uint32& OutResultFlags);
 
 	/**
@@ -1372,7 +1373,7 @@ private:
 	FString GetCookedAssetRegistryFilename(const FString& PlatformName);
 
 	// Return the filename to use for the cook metadata file, adjusted for DLC, sandbox, and platform.
-	FString GetCookedCookMetadataFilename(const FString& PlatformName);
+	FString GetCookedCookMetadataFilename(const ITargetPlatform* TargetPlatform);
 	void WriteCookMetadata(const ITargetPlatform* InTargetPlatform, uint64 InDevelopmentAssetRegistryHash);
 	void WriteReferencedSet(const ITargetPlatform* InTargetPlatform, TArray<FName>&& CookedPackageNames);
 
diff --git a/Engine/Source/Editor/UnrealEd/Private/CookOnTheFlyServer.cpp b/Engine/Source/Editor/UnrealEd/Private/CookOnTheFlyServer.cpp
--- a/Engine/Source/Editor/UnrealEd/Private/CookOnTheFlyServer.cpp
+++ b/Engine/Source/Editor/UnrealEd/Private/CookOnTheFlyServer.cpp
@@ -6694,12 +6694,12 @@ bool UCookOnTheFlyServer::Exec_Editor(class UWorld* InWorld, const TCHAR* Cmd, F
 	return false;
 }
 
-UE::Cook::FInstigator UCookOnTheFlyServer::GetInstigator(FName PackageName)
+UE::Cook::FInstigator UCookOnTheFlyServer::GetInstigator(FName PackageName) const
 {
 	return GetInstigator(PackageName, UE::Cook::EReachability::All);
 }
 
-UE::Cook::FInstigator UCookOnTheFlyServer::GetInstigator(FName PackageName, UE::Cook::EReachability Reachability)
+UE::Cook::FInstigator UCookOnTheFlyServer::GetInstigator(FName PackageName, UE::Cook::EReachability Reachability) const
 {
 	using namespace UE::Cook;
 
@@ -6711,7 +6711,7 @@ UE::Cook::FInstigator UCookOnTheFlyServer::GetInstigator(FName PackageName, UE::
 	return PackageData->GetInstigator(Reachability);
 }
 
-TArray<UE::Cook::FInstigator> UCookOnTheFlyServer::GetInstigatorChain(FName PackageName)
+TArray<UE::Cook::FInstigator> UCookOnTheFlyServer::GetInstigatorChain(FName PackageName) const
 {
 	using namespace UE::Cook;
 	TArray<FInstigator> Result;
@@ -6761,7 +6761,7 @@ TArray<UE::Cook::FInstigator> UCookOnTheFlyServer::GetInstigatorChain(FName Pack
 	}
 }
 
-UE::Cook::ECookType UCookOnTheFlyServer::GetCookType()
+UE::Cook::ECookType UCookOnTheFlyServer::GetCookType() const
 {
 	if (IsDirectorCookByTheBook())
 	{
@@ -6774,7 +6774,7 @@ UE::Cook::ECookType UCookOnTheFlyServer::GetCookType()
 	}
 }
 
-UE::Cook::ECookingDLC UCookOnTheFlyServer::GetCookingDLC()
+UE::Cook::ECookingDLC UCookOnTheFlyServer::GetCookingDLC() const
 {
 	using namespace UE::Cook;
 	if (!IsCookingDLC())
@@ -6792,12 +6792,12 @@ UE::Cook::ECookingDLC UCookOnTheFlyServer::GetCookingDLC()
 	return ECookingDLC::Yes;
 }
 
-FString UCookOnTheFlyServer::GetDLCName()
+FString UCookOnTheFlyServer::GetDLCName() const
 {
 	return CookByTheBookOptions->DlcName;
 }
 
-FString UCookOnTheFlyServer::GetBasedOnReleaseVersion()
+FString UCookOnTheFlyServer::GetBasedOnReleaseVersion() const
 {
 	if (!IsCookingDLC() && CookByTheBookOptions->CreateReleaseVersion.IsEmpty())
 	{
@@ -6806,12 +6806,12 @@ FString UCookOnTheFlyServer::GetBasedOnReleaseVersion()
 	return CookByTheBookOptions->BasedOnReleaseVersion;
 }
 
-FString UCookOnTheFlyServer::GetCreateReleaseVersion()
+FString UCookOnTheFlyServer::GetCreateReleaseVersion() const
 {
 	return CookByTheBookOptions->CreateReleaseVersion;
 }
 
-UE::Cook::EProcessType UCookOnTheFlyServer::GetProcessType()
+UE::Cook::EProcessType UCookOnTheFlyServer::GetProcessType() const
 {
 	using namespace UE::Cook;
 	if (IsCookWorkerMode())
@@ -6828,7 +6828,7 @@ UE::Cook::EProcessType UCookOnTheFlyServer::GetProcessType()
 	}
 }
 
-UE::Cook::ECookValidationOptions UCookOnTheFlyServer::GetCookValidationOptions()
+UE::Cook::ECookValidationOptions UCookOnTheFlyServer::GetCookValidationOptions() const
 {
 	using namespace UE::Cook;
 	ECookValidationOptions ValidationOptions = ECookValidationOptions::None;
@@ -6847,7 +6847,7 @@ UE::Cook::ECookValidationOptions UCookOnTheFlyServer::GetCookValidationOptions()
 	return ValidationOptions;
 }
 
-bool UCookOnTheFlyServer::IsIncremental()
+bool UCookOnTheFlyServer::IsIncremental() const
 {
 	// TODO: For simplicity, we provide a single bool for all platforms in multiprocess cooks
 	// But it is not currently guaranteed that they all have the same value; add enforcement of
@@ -6861,12 +6861,12 @@ bool UCookOnTheFlyServer::IsIncremental()
 	return !PlatformData->bFullBuild;
 }
 
-TArray<const ITargetPlatform*> UCookOnTheFlyServer::GetSessionPlatforms()
+TArray<const ITargetPlatform*> UCookOnTheFlyServer::GetSessionPlatforms() const
 {
 	return PlatformManager ? PlatformManager->GetSessionPlatforms() : TArray<const ITargetPlatform*>();
 }
 
-FString UCookOnTheFlyServer::GetCookOutputFolder(const ITargetPlatform* TargetPlatform)
+FString UCookOnTheFlyServer::GetCookOutputFolder(const ITargetPlatform* TargetPlatform) const
 {
 	if (!SandboxFile || !PlatformManager || !TargetPlatform)
 	{
@@ -6881,6 +6881,19 @@ FString UCookOnTheFlyServer::GetCookOutputFolder(const ITargetPlatform* TargetPl
 	return Result;
 }
 
+FString UCookOnTheFlyServer::GetCookMetadataOutputFolder(const ITargetPlatform* TargetPlatform) const
+{
+	if (!SandboxFile || !PlatformManager || !TargetPlatform)
+	{
+		return FString();
+	}
+	if (!PlatformManager->HasSessionPlatform(TargetPlatform))
+	{
+		return FString();
+	}
+	return ConvertToFullSandboxPath(GetMetadataDirectory(), true, TargetPlatform->PlatformName());
+}
+
 void UCookOnTheFlyServer::RegisterCollector(UE::Cook::IMPCollector* Collector, UE::Cook::EProcessType ProcessType)
 {
 	using namespace UE::Cook;
@@ -7486,9 +7499,7 @@ bool UCookOnTheFlyServer::GetCurrentIniVersionStrings( const ITargetPlatform* Ta
 
 bool UCookOnTheFlyServer::GetCookedIniVersionStrings(const ITargetPlatform* TargetPlatform, UE::Cook::FIniSettingContainer& OutIniSettings, TMap<FString,FString>& OutAdditionalSettings) const
 {
-	const FString EditorIni = GetMetadataDirectory() / TEXT("CookedIniVersion.txt");
-	const FString SandboxEditorIni = ConvertToFullSandboxPath(*EditorIni, true);
-	const FString PlatformSandboxEditorIni = SandboxEditorIni.Replace(TEXT("[Platform]"), *TargetPlatform->PlatformName());
+	const FString PlatformSandboxEditorIni = GetCookMetadataOutputFolder(TargetPlatform) / TEXT("CookedIniVersion.txt"); 
 
 	FConfigFile ConfigFile;
 	ConfigFile.Read(*PlatformSandboxEditorIni);
@@ -7555,8 +7566,7 @@ FString UCookOnTheFlyServer::GetCookSettingsFileName(const ITargetPlatform* Targ
 	FString LeafName = ArtifactName != TEXT("global")
 		? FString::Printf(TEXT("CookedSettings_%s.txt"), *ArtifactName)
 		: TEXT("CookedSettings.txt");
-	FString UncookedPath = GetMetadataDirectory() / LeafName;
-	return ConvertToFullSandboxPath(*UncookedPath, true /* bForWrite */, TargetPlatform->PlatformName());
+	return GetCookMetadataOutputFolder(TargetPlatform) / LeafName;
 }
 
 FConfigFile UCookOnTheFlyServer::LoadCookSettings(const ITargetPlatform* TargetPlatform, const FString& ArtifactName)
@@ -7734,12 +7744,7 @@ bool UCookOnTheFlyServer::SaveCurrentIniSettings(const ITargetPlatform* TargetPl
 	UE::Cook::FIniSettingContainer CurrentIniSettings;
 	GetCurrentIniVersionStrings(TargetPlatform, CurrentIniSettings);
 
-	const FString EditorIni = GetMetadataDirectory() / TEXT("CookedIniVersion.txt");
-	const FString SandboxEditorIni = ConvertToFullSandboxPath(*EditorIni, true);
-
-
-	const FString PlatformSandboxEditorIni = SandboxEditorIni.Replace(TEXT("[Platform]"), *TargetPlatform->PlatformName());
-
+	const FString PlatformSandboxEditorIni = GetCookMetadataOutputFolder(TargetPlatform) / TEXT("CookedIniVersion.txt");
 
 	FConfigFile ConfigFile;
 	// ConfigFile.Read(*PlatformSandboxEditorIni);
@@ -8904,16 +8909,15 @@ FString UCookOnTheFlyServer::GetSandboxAssetRegistryFilename()
 	return ConvertToFullSandboxPath(*RegistryFilename, true);
 }
 
-FString UCookOnTheFlyServer::GetCookedAssetRegistryFilename(const FString& PlatformName )
+FString UCookOnTheFlyServer::GetCookedAssetRegistryFilename(const FString& PlatformName)
 {
 	const FString CookedAssetRegistryFilename = GetSandboxAssetRegistryFilename().Replace(TEXT("[Platform]"), *PlatformName);
 	return CookedAssetRegistryFilename;
 }
 
-FString UCookOnTheFlyServer::GetCookedCookMetadataFilename(const FString& PlatformName )
+FString UCookOnTheFlyServer::GetCookedCookMetadataFilename(const ITargetPlatform* TargetPlatform)
 {
-	const FString MetadataFilename = GetMetadataDirectory() / UE::Cook::GetCookMetadataFilename();
-	return ConvertToFullSandboxPath(*MetadataFilename, true).Replace(TEXT("[Platform]"), *PlatformName);
+	return GetCookMetadataOutputFolder(TargetPlatform) / UE::Cook::GetCookMetadataFilename();
 }
 
 FString UCookOnTheFlyServer::GetSandboxCachedEditorThumbnailsFilename()
@@ -9205,15 +9209,13 @@ void UCookOnTheFlyServer::WriteCookMetadata(const ITargetPlatform* InTargetPlatf
 		MetadataState.SetPlatformAndBuildVersion(PlatformNameString, FApp::GetBuildVersion());
 		MetadataState.SetHordeJobId(FPlatformMisc::GetEnvironmentVariable(TEXT("UE_HORDE_JOBID")));
 
-		MetadataState.SaveToFile(GetCookedCookMetadataFilename(PlatformNameString));
+		MetadataState.SaveToFile(GetCookedCookMetadataFilename(InTargetPlatform));
 	}
 }
 
 void UCookOnTheFlyServer::WriteReferencedSet(const ITargetPlatform* InTargetPlatform, TArray<FName>&& CookedPackageNames)
 {
-	const FString MetadataPlatformAgnosticFilename = GetMetadataDirectory() / UE::Cook::GetReferencedSetFilename();
-	const FString MetadataFilename = ConvertToFullSandboxPath(*MetadataPlatformAgnosticFilename, true)
-		.Replace(TEXT("[Platform]"), *InTargetPlatform->PlatformName());
+	const FString MetadataFilename = GetCookMetadataOutputFolder(InTargetPlatform) / UE::Cook::GetReferencedSetFilename();
 
 	CookedPackageNames.Sort(FNameLexicalLess());
 
@@ -9677,8 +9679,7 @@ void UCookOnTheFlyServer::CookByTheBookFinishedInternal()
 				}
 				else
 				{
-					DiagnosticsDir = GetMetadataDirectory();
-					DiagnosticsDir = ConvertToFullSandboxPath(DiagnosticsDir, true, TargetPlatform->PlatformName());
+					DiagnosticsDir = GetCookMetadataOutputFolder(TargetPlatform);
 				}
 				IncrementallyModifiedDiagnostics->WriteToFile(DiagnosticsDir);
 			}
@@ -10667,10 +10668,9 @@ UE::Cook::FCookSavePackageContext* UCookOnTheFlyServer::CreateSaveContext(const
 	checkf(SandboxFile, TEXT("SaveContexts cannot be created until after CreateSandboxFile has been called from a StartCook function."));
 
 	const FString RootPathSandbox = ConvertToFullSandboxPath(FPaths::RootDir(), true);
-	FString MetadataPathSandbox = ConvertToFullSandboxPath(GetMetadataDirectory(), true);
 	const FString PlatformString = TargetPlatform->PlatformName();
 	const FString ResolvedRootPath = RootPathSandbox.Replace(TEXT("[Platform]"), *PlatformString);
-	const FString ResolvedMetadataPath = MetadataPathSandbox.Replace(TEXT("[Platform]"), *PlatformString);
+	const FString ResolvedMetadataPath = GetCookMetadataOutputFolder(TargetPlatform);
 
 	TUniquePtr<FDeterminismManager> DeterminismManager;
 	TSharedPtr<ICookArtifactReader> CookArtifactReader = nullptr;
@@ -10807,7 +10807,7 @@ UE::Cook::FCookSavePackageContext* UCookOnTheFlyServer::CreateSaveContext(const
 
 }
 
-const TSet<IPlugin*>* UCookOnTheFlyServer::GetEnabledPlugins(const ITargetPlatform* TargetPlatform)
+const TSet<IPlugin*>* UCookOnTheFlyServer::GetEnabledPlugins(const ITargetPlatform* TargetPlatform) const
 {
 	using namespace UE::Cook;
 
diff --git a/Engine/Source/Runtime/CoreUObject/Public/UObject/ICookInfo.h b/Engine/Source/Runtime/CoreUObject/Public/UObject/ICookInfo.h
--- a/Engine/Source/Runtime/CoreUObject/Public/UObject/ICookInfo.h
+++ b/Engine/Source/Runtime/CoreUObject/Public/UObject/ICookInfo.h
@@ -203,31 +203,31 @@ public:
 	 * Return the instigator that first caused the package to be requested by the cook.
 	 * Returns category EInstigator::NotYetRequested if package is not yet known to the cook.
 	 */
-	virtual FInstigator GetInstigator(FName PackageName) = 0;
+	virtual FInstigator GetInstigator(FName PackageName) const = 0;
 	/**
 	 * Return the chain of instigators that caused the package to be requested by the cook.
 	 * First element is the direct instigator of the package, last is the root instigator that started the chain.
 	 */
-	virtual TArray<FInstigator> GetInstigatorChain(FName PackageName) = 0;
+	virtual TArray<FInstigator> GetInstigatorChain(FName PackageName) const = 0;
 
 	/** The type (e.g. CookByTheBook) of the running cook. This function will not return ECookType::Unknown. */
-	virtual UE::Cook::ECookType GetCookType() = 0;
+	virtual UE::Cook::ECookType GetCookType() const = 0;
 	/** Whether DLC is being cooked (e.g. via "-dlcname=<PluginName>"). This function will not return ECookingDLC::Unknown. */
-	virtual UE::Cook::ECookingDLC GetCookingDLC() = 0;
+	virtual UE::Cook::ECookingDLC GetCookingDLC() const = 0;
 	/** When DLC is being cooked (@see GetCookingDLC), this will return the name of the DLC plugin. */
-	virtual FString GetDLCName() = 0;
+	virtual FString GetDLCName() const = 0;
 	/**
 	 * The release version of the basegame the current DLC is based on, when cooking DLC or when making a patch.
 	 * Empty string when not cooking DLC or making a patch.
 	 */
-	virtual FString GetBasedOnReleaseVersion() = 0;
+	virtual FString GetBasedOnReleaseVersion() const = 0;
 	/** The release version being created when making a patch. */
-	virtual FString GetCreateReleaseVersion() = 0;
+	virtual FString GetCreateReleaseVersion() const = 0;
 
 	/** The role the current process plays in its MPCook session, or EProcessType::SingleProcess if it is running standalone. */
-	virtual UE::Cook::EProcessType GetProcessType() = 0;
+	virtual UE::Cook::EProcessType GetProcessType() const = 0;
 	/** Get the validation options used by the running cook, if any. */
-	virtual UE::Cook::ECookValidationOptions GetCookValidationOptions() = 0;
+	virtual UE::Cook::ECookValidationOptions GetCookValidationOptions() const = 0;
 	/**
 	 * Returns true if the cooker is cooking after a previous cook session and is cooking only the changed files.
 	 * Returns false if the cooker is doing a recook of all packages discovered in the session.
@@ -236,9 +236,9 @@ public:
 	 * When IsIncremental is true, systems that write artifacts to the cook output should load/update/resave their
 	 * artifacts.
 	 */
-	virtual bool IsIncremental() = 0;
+	virtual bool IsIncremental() const = 0;
 	UE_DEPRECATED(5.6, "Use IsIncremental instead.")
-	bool IsIterative()
+	bool IsIterative() const
 	{
 		return IsIncremental();
 	}
@@ -249,15 +249,19 @@ public:
 	 * no platforms have been requested yet. During cook by the book, this list will not change throughout the cook,
 	 * during CookOnTheFly it can be added to or removed from when platforms are requested or go idle and are dropped.
 	 */
-	virtual TArray<const ITargetPlatform*> GetSessionPlatforms() = 0;
+	virtual TArray<const ITargetPlatform*> GetSessionPlatforms() const = 0;
 	/**
 	 * Returns the output folder being used by the cooker for the given platform in the given session.
 	 * Returns empty string if not in a session or the given platform is not in GetSessionPlatforms().
-	 * Returns the path to the root folder of the output, so e.g. GetCookOutputFolder()/<ProjectName>/Metadata
-	 * is the path to the output metadata for the cook.
 	 * Returns the path in FPaths::MakeStandardFilename format ("../../../<ProjectName>/Saved/Cooked/<PlatformName>")
 	 */
-	virtual FString GetCookOutputFolder(const ITargetPlatform* TargetPlatform) = 0;
+	virtual FString GetCookOutputFolder(const ITargetPlatform* TargetPlatform) const = 0;
+	/**
+	 * Returns the path within the GetCookOutputFolder where the Metadata folder for the cook is written. For basegame
+	 * cooks this is GetCookOutputFolder()/<ProjectName>/Metadata, but for DLC cooks it is in the relative path to the
+	 * DLC Folder: GetCookOutputFolder()/<RelPathToDLC>/Metadata.
+	 */
+	virtual FString GetCookMetadataOutputFolder(const ITargetPlatform* TargetPlatform) const = 0;
 
 	/**
 	 * Return the set of plugins that the cooker has determined are enabled on the current platform. This is different
@@ -266,7 +270,7 @@ public:
 	 * 
 	 * Returns nullptr if the given Platform is not a session platform in the current cook.
 	 */
-	virtual const TSet<IPlugin*>* GetEnabledPlugins(const ITargetPlatform* TargetPlatform) = 0;
+	virtual const TSet<IPlugin*>* GetEnabledPlugins(const ITargetPlatform* TargetPlatform) const = 0;
 
 	/**
 	 * MPCook: register in the current process a collector that replicates system-specific and package-specific
