[Cook/ALT]
- Add support for preloading when async loading thread is enabled
- Increase the preload window as it improves preloading efficiency a lot

#Synced-CL 51088034

diff --git a/Engine/Source/Editor/UnrealEd/Private/CookOnTheFlyServer.cpp b/Engine/Source/Editor/UnrealEd/Private/CookOnTheFlyServer.cpp
--- a/Engine/Source/Editor/UnrealEd/Private/CookOnTheFlyServer.cpp
+++ b/Engine/Source/Editor/UnrealEd/Private/CookOnTheFlyServer.cpp
@@ -6034,13 +6034,22 @@ void UCookOnTheFlyServer::Initialize( ECookMode::Type DesiredCookMode, ECookInit
 	IdleStatusStartTime = FPlatformTime::Seconds();
 	IdleStatusLastReportTime = IdleStatusStartTime;
 
-	if (!IsCookOnTheFlyMode() && !IsCookingInEditor() &&
-		FPlatformMisc::SupportsMultithreadedFileHandles() && // Preloading moves file handles between threads
-		!GAllowCookedDataInEditorBuilds // // Use of preloaded files is not yet implemented when GAllowCookedDataInEditorBuilds is on, see FLinkerLoad::CreateLoader
-		)
+	// It is perfectly safe to allow preloading in any environment whenever async loading is active
+	// since we won't try to deal with linkers directly, we'll use proper calls to LoadPackageAsync in the preloader.
+	if (IsAsyncLoadingMultithreaded())
 	{
 		bPreloadingEnabled = true;
-		FLinkerLoad::SetPreloadingEnabled(true);
+	}
+	else 
+	{
+		if (!IsCookOnTheFlyMode() && !IsCookingInEditor() &&
+			FPlatformMisc::SupportsMultithreadedFileHandles() && // Preloading moves file handles between threads
+			!GAllowCookedDataInEditorBuilds // // Use of preloaded files is not yet implemented when GAllowCookedDataInEditorBuilds is on, see FLinkerLoad::CreateLoader
+			)
+		{
+			bPreloadingEnabled = true;
+			FLinkerLoad::SetPreloadingEnabled(true);
+		}
 	}
 
 	// Prepare a map of SplitDataClass to FRegisteredCookPackageSplitter* for TryGetRegisteredCookPackageSplitter to use
@@ -6288,11 +6297,25 @@ void UCookOnTheFlyServer::SetInitializeConfigSettings(UE::Cook::FInitializeConfi
 	// For preload to actually be able to pipeline with load/batch, we need both RequestBatchSize
 	// and MaxPreloadAllocated to be bigger than LoadBatchSize so that we won't consume all preloads
 	// for every iteration.
-	MaxPreloadAllocated = 32;
 	DesiredSaveQueueLength = 8;
-	DesiredLoadQueueLength = 8;
-	LoadBatchSize = 16;
-	RequestBatchSize = 32;
+	
+	// Best preloading performance is achieved when all the loading parameters have the same value
+	// and that value is high enough to get multithreaded async loading using all cores.
+	if (IsAsyncLoadingMultithreaded())
+	{
+		DesiredLoadQueueLength = 256;
+		LoadBatchSize = DesiredLoadQueueLength;
+		RequestBatchSize = DesiredLoadQueueLength;
+		MaxPreloadAllocated = DesiredLoadQueueLength;
+	}
+	else
+	{
+		MaxPreloadAllocated = 32;
+		DesiredLoadQueueLength = 8;
+		LoadBatchSize = 16;
+		RequestBatchSize = 32;
+	}
+
 	WaitForAsyncSleepSeconds = 1.0f;
 
 
