[SoundWave] FixTimeCodeInfo hydrates bulk data.
Removes code from post load and moves it into a commandlet.
void USoundWave::FixTimeCodeInfo() has been updated to return a bool that is true if the TimecodeInfo was fixed.
[RN] Adds commandlet to fix LogSlate warnings from USoundWave::TimecodeInfo with incorrect symbols

#Synced-CL 51092343

diff --git a/Engine/Source/Editor/UnrealEd/Classes/Commandlets/FixSoundWaveTimecodeInfoCommandlet.h b/Engine/Source/Editor/UnrealEd/Classes/Commandlets/FixSoundWaveTimecodeInfoCommandlet.h
new file mode 100644
--- /dev/null
+++ b/Engine/Source/Editor/UnrealEd/Classes/Commandlets/FixSoundWaveTimecodeInfoCommandlet.h
@@ -0,0 +1,21 @@
+// Copyright Epic Games, Inc. All Rights Reserved.
+
+#pragma once
+
+#include "Commandlets/Commandlet.h"
+
+#include "FixSoundWaveTimecodeInfoCommandlet.generated.h"
+
+/**
+ * Internal commandlet to fix any USoundWaves with corrupt TimecodeInfo.
+ */
+UCLASS()
+class UFixSoundWaveTimecodeInfoCommandlet : public UCommandlet
+{
+	GENERATED_BODY()
+
+public:
+	UFixSoundWaveTimecodeInfoCommandlet();
+
+	virtual int32 Main(const FString& Params);
+};
diff --git a/Engine/Source/Editor/UnrealEd/Private/Commandlets/FixSoundWaveTimecodeInfoCommandlet.cpp b/Engine/Source/Editor/UnrealEd/Private/Commandlets/FixSoundWaveTimecodeInfoCommandlet.cpp
new file mode 100644
--- /dev/null
+++ b/Engine/Source/Editor/UnrealEd/Private/Commandlets/FixSoundWaveTimecodeInfoCommandlet.cpp
@@ -0,0 +1,77 @@
+// Copyright Epic Games, Inc. All Rights Reserved.
+
+#include "Commandlets/FixSoundWaveTimecodeInfoCommandlet.h"
+
+#include "AssetRegistry/AssetData.h"
+#include "AssetRegistry/AssetRegistryModule.h"
+#include "FileHelpers.h"
+#include "Sound/SoundWave.h"
+
+DEFINE_LOG_CATEGORY_STATIC(LogFixSoundWaveTimecodeInfoCommandlet, Log, All);
+
+UFixSoundWaveTimecodeInfoCommandlet::UFixSoundWaveTimecodeInfoCommandlet()
+{
+	LogToConsole = true;
+	IsClient = false;
+	IsEditor = true;
+	IsServer = false;
+}
+
+int32 UFixSoundWaveTimecodeInfoCommandlet::Main(const FString& Params)
+{
+	UE_LOG(LogFixSoundWaveTimecodeInfoCommandlet, Display, TEXT("Fixing soundwave timecodeInfo. Note: Params are ignored."));
+
+	// Return 0 on success
+	int32 ReturnCode = 0;
+
+#if WITH_EDITORONLY_DATA
+	// Load the asset registry module
+	FAssetRegistryModule& AssetRegistryModule = FModuleManager::LoadModuleChecked<FAssetRegistryModule>(TEXT("AssetRegistry"));
+
+	IAssetRegistry& AssetRegistry = AssetRegistryModule.Get();
+
+	// Update Registry Module
+	AssetRegistry.SearchAllAssets(true);
+
+	TArray<FAssetData> SoundWaveAssetDatas;
+	AssetRegistry.GetAssetsByClass(USoundWave::StaticClass()->GetClassPathName(), SoundWaveAssetDatas);
+
+	for (const FAssetData& SoundWaveAssetData : SoundWaveAssetDatas)
+	{
+		if (TObjectPtr<USoundWave> SoundWave = Cast<USoundWave>(SoundWaveAssetData.GetAsset()))
+		{
+			if (SoundWave->FixTimeCodeInfo())
+			{
+				check(SoundWave->GetTimecodeInfo().IsSet());
+				UE_LOG(LogFixSoundWaveTimecodeInfoCommandlet, Display, TEXT("Fixed: %s's TimecodeInfo. TimecodeInfo::Description: %s"), 
+					*SoundWave->GetName(), *SoundWave->GetTimecodeInfo()->Description);
+				
+				if (UPackage* SoundWavePackage = SoundWave->GetPackage())
+				{
+					TArray<UPackage*> PackagesToSave;
+					PackagesToSave.Add(SoundWavePackage);
+
+					const FString PackageFileName = FPackageName::LongPackageNameToFilename(SoundWavePackage->GetName(), FPackageName::GetAssetPackageExtension());
+					const bool bSuccess = UEditorLoadingAndSavingUtils::SavePackages(PackagesToSave, false);
+
+					if (bSuccess)
+					{
+						UE_LOG(LogFixSoundWaveTimecodeInfoCommandlet, Display, TEXT("Fixed SoundWave: %s saved."), *SoundWave->GetName());
+					}
+					else
+					{
+						UE_LOG(LogFixSoundWaveTimecodeInfoCommandlet, Warning, TEXT("Fixed SoundWave: %s failed to save."), *SoundWave->GetName());
+
+						ReturnCode = 1;
+					}
+				}
+			}
+			else
+			{
+				UE_LOG(LogFixSoundWaveTimecodeInfoCommandlet, Verbose, TEXT("%s's TimecodeInfo is valid and no fix was needed."), *SoundWave->GetName());
+			}
+		}
+	}
+#endif //WITH_EDITORONLY_DATA
+	return ReturnCode;
+}
\ No newline at end of file
diff --git a/Engine/Source/Runtime/Engine/Classes/Sound/SoundWave.h b/Engine/Source/Runtime/Engine/Classes/Sound/SoundWave.h
--- a/Engine/Source/Runtime/Engine/Classes/Sound/SoundWave.h
+++ b/Engine/Source/Runtime/Engine/Classes/Sound/SoundWave.h
@@ -1231,8 +1231,8 @@ public:
 	ENGINE_API void SetTimecodeInfo(const FSoundWaveTimecodeInfo& InTimecode);
 	ENGINE_API TOptional<FSoundWaveTimecodeInfo> GetTimecodeInfo() const;
 
-protected:
-	ENGINE_API void FixTimeCodeInfo();
+	// Returns true if the TimecodeInfo needed to be fixed
+	ENGINE_API bool FixTimeCodeInfo();
 #endif //WITH_EDITORONLY_DATA	
 
 public:
diff --git a/Engine/Source/Runtime/Engine/Private/SoundWave.cpp b/Engine/Source/Runtime/Engine/Private/SoundWave.cpp
--- a/Engine/Source/Runtime/Engine/Private/SoundWave.cpp
+++ b/Engine/Source/Runtime/Engine/Private/SoundWave.cpp
@@ -1080,13 +1080,13 @@ TOptional<FSoundWaveTimecodeInfo> USoundWave::GetTimecodeInfo() const
 	return TimecodeInfo;
 }
 
-void USoundWave::FixTimeCodeInfo()
+bool USoundWave::FixTimeCodeInfo()
 {
 	if (TimecodeInfo.OriginatorDate.IsEmpty() && TimecodeInfo.Description.IsEmpty() && TimecodeInfo.OriginatorReference.IsEmpty()
 		&& TimecodeInfo.OriginatorDescription.IsEmpty() && TimecodeInfo.OriginatorTime.IsEmpty())
 	{
 		// TimecodeInfo data is empty and does not need to be fixed
-		return;
+		return false;
 	}
 
 	FSharedBuffer BulkDataBuffer;
@@ -1097,13 +1097,25 @@ void USoundWave::FixTimeCodeInfo()
 	{
 		UE_LOG(LogAudio, Error, TEXT("Sound wave payload in asset %s failed to parse as a wave in FixTimeCodeInfo!"), *GetName());
 
-		return;
+		return false;
 	}
 
-	if (WaveInfo.TimecodeInfo.IsValid())
+	if (WaveInfo.TimecodeInfo.IsValid() && 
+		(TimecodeInfo != *WaveInfo.TimecodeInfo.Get() || 
+			TimecodeInfo.OriginatorDate != WaveInfo.TimecodeInfo->OriginatorDate ||
+			TimecodeInfo.Description != WaveInfo.TimecodeInfo->Description ||
+			TimecodeInfo.OriginatorReference != WaveInfo.TimecodeInfo->OriginatorReference ||
+			TimecodeInfo.OriginatorDescription != WaveInfo.TimecodeInfo->OriginatorDescription || 
+			TimecodeInfo.OriginatorTime != WaveInfo.TimecodeInfo->OriginatorTime))
 	{
+		Modify();
 		SetTimecodeInfo(*WaveInfo.TimecodeInfo);
+		MarkPackageDirty();
+
+		return true;
 	}
+
+	return false;
 }
 
 #endif //WITH_EDITORONLY_DATA
@@ -2255,9 +2267,6 @@ void USoundWave::PostLoad()
 				ImportedSampleRate = OriginalSampleRate;
 			}
 		}
-
-		// Does not mark the asset dirty as it is EDITOR_ONLY_DATA but will always appear fixed to the user
-		FixTimeCodeInfo();
 	}
 	
 	if (TotalSamples == 0 && ImportedSampleRate != 0)
