- Exposed on option on UFontFileImportFactory to create a matching UFont for the UFontFace when programmatically/scripting importing a font

#Synced-CL 50913554

diff --git a/Engine/Source/Editor/UnrealEd/Classes/Factories/FontFileImportFactory.h b/Engine/Source/Editor/UnrealEd/Classes/Factories/FontFileImportFactory.h
--- a/Engine/Source/Editor/UnrealEd/Classes/Factories/FontFileImportFactory.h
+++ b/Engine/Source/Editor/UnrealEd/Classes/Factories/FontFileImportFactory.h
@@ -12,6 +12,19 @@
 #include "EditorReimportHandler.h"
 #include "FontFileImportFactory.generated.h"
 
+UENUM(BlueprintType)
+enum class EBatchCreateFontAsset : uint8
+{
+	// Internal do not use
+	Unknown,
+	// Internal do not use
+	Cancel,
+	// Create a UFont to match the new UFontFace
+	Yes,
+	// Do not create a UFont to match the new UFontFace
+	No,
+};
+
 UCLASS(MinimalAPI)
 class UFontFileImportFactory : public UFactory, public FReimportHandler
 {
@@ -29,13 +42,6 @@ class UFontFileImportFactory : public UFactory, public FReimportHandler
 	UNREALED_API virtual int32 GetPriority() const override;
 	//~ End FReimportHandler Interface
 
-private:
-	enum class EBatchCreateFontAsset : uint8
-	{
-		Unknown,
-		Cancel,
-		Yes,
-		No,
-	};
+	UPROPERTY(BlueprintReadWrite, Category = Misc)
 	EBatchCreateFontAsset BatchCreateFontAsset;
 };
