Allow to turn off alpha in SColorGradientEditor and LinearColorRampCustomization that makes use of it, as a part of the efforts to use a linear color ramp for bone weight colors in the skeletal mesh editor

#rnx

#Synced-CL 51156678

diff --git a/Engine/Source/Editor/DetailCustomizations/Private/LinearColorRampCustomization.cpp b/Engine/Source/Editor/DetailCustomizations/Private/LinearColorRampCustomization.cpp
--- a/Engine/Source/Editor/DetailCustomizations/Private/LinearColorRampCustomization.cpp
+++ b/Engine/Source/Editor/DetailCustomizations/Private/LinearColorRampCustomization.cpp
@@ -7,9 +7,19 @@
 #include "Curves/LinearColorRamp.h"
 #include "DetailWidgetRow.h"
 
+FLinearColorRampCustomization::FLinearColorRampCustomization(const FLinearColorRampCustomizationArgs InArgs)
+	: Args(InArgs)
+{
+}
+
 TSharedRef<IPropertyTypeCustomization> FLinearColorRampCustomization::MakeInstance()
 {
-	return MakeShareable(new FLinearColorRampCustomization);
+	return MakeShared<FLinearColorRampCustomization>();
+}
+
+TSharedRef<IPropertyTypeCustomization> FLinearColorRampCustomization::MakeInstance(const FLinearColorRampCustomizationArgs InArgs)
+{
+	return MakeShared<FLinearColorRampCustomization>(InArgs);
 }
 
 void FLinearColorRampCustomization::CustomizeHeader(TSharedRef<IPropertyHandle> PropertyHandle, FDetailWidgetRow& HeaderRow, IPropertyTypeCustomizationUtils& CustomizationUtils)
@@ -32,6 +42,7 @@ void FLinearColorRampCustomization::CustomizeHeader(TSharedRef<IPropertyHandle>
 							.ViewMinInput(0.0f)
 							.ViewMaxInput(1.0f)
 							.ClampStopsToViewRange(true)
+							.WithAlphaChannel(Args.bWithAlphaChannel)
 					];
 
 				GradientEditor->SetCurveOwner(ColorRamp);
diff --git a/Engine/Source/Editor/DetailCustomizations/Public/Customizations/LinearColorRampCustomization.h b/Engine/Source/Editor/DetailCustomizations/Public/Customizations/LinearColorRampCustomization.h
--- a/Engine/Source/Editor/DetailCustomizations/Public/Customizations/LinearColorRampCustomization.h
+++ b/Engine/Source/Editor/DetailCustomizations/Public/Customizations/LinearColorRampCustomization.h
@@ -4,15 +4,31 @@
 
 #include "IPropertyTypeCustomization.h"
 
+/** Arguments for the linear color ramp customization */
+struct FLinearColorRampCustomizationArgs
+{
+	/** When true, the alpha channel is displayed and editable */
+	bool bWithAlphaChannel = true;
+};
+
 class FLinearColorRampCustomization : public IPropertyTypeCustomization
 {
 public:
+	FLinearColorRampCustomization() = default;
+	FLinearColorRampCustomization(const FLinearColorRampCustomizationArgs InArgs);
 
 	/** Makes a new instance of this detail layout class for a specific detail view requesting it */
 	DETAILCUSTOMIZATIONS_API static TSharedRef<IPropertyTypeCustomization> MakeInstance();
 
+	/** Makes a new instance of this detail layout class for a specific detail view requesting it, with arguments */
+	DETAILCUSTOMIZATIONS_API static TSharedRef<IPropertyTypeCustomization> MakeInstance(const FLinearColorRampCustomizationArgs InArgs);
+
 	//~ Begin - IPropertyTypeCustomization interface
 	DETAILCUSTOMIZATIONS_API virtual void CustomizeHeader(TSharedRef<IPropertyHandle> PropertyHandle, FDetailWidgetRow& HeaderRow, IPropertyTypeCustomizationUtils& CustomizationUtils) override;
 	virtual void CustomizeChildren(TSharedRef<IPropertyHandle> /*StructPropertyHandle*/, IDetailChildrenBuilder& /*ChildBuilder*/, IPropertyTypeCustomizationUtils& /*CustomizationUtils*/) override {}
 	//~ End - IPropertyTypeCustomization interface
+
+private:
+	/** Arguments for this customization */
+	const FLinearColorRampCustomizationArgs Args;
 };
diff --git a/Engine/Source/Editor/UnrealEd/Private/SColorGradientEditor.cpp b/Engine/Source/Editor/UnrealEd/Private/SColorGradientEditor.cpp
--- a/Engine/Source/Editor/UnrealEd/Private/SColorGradientEditor.cpp
+++ b/Engine/Source/Editor/UnrealEd/Private/SColorGradientEditor.cpp
@@ -28,59 +28,82 @@
 
 #define LOCTEXT_NAMESPACE "SColorGradientEditor"
 
-namespace FColorGradientEditorGeometry
+struct FColorGradientEditorGeometrySolver
 {
-	const float HandleWidth = 13.0f;
-	const float HandleHeight = 16.0f;
-	const float HandleColorWidth = 7.0f;
-	const float HandleColorHeight = 7.0f;
-	const float HandleColorXOffset = (HandleWidth - HandleColorWidth) / 2;
-	const float HandleColorYOffsetForColor = HandleColorXOffset + 1;
-	const float HandleColorYOffsetForAlpha = HandleHeight - HandleColorHeight - HandleColorXOffset - 1;
+	/** Constructs the solver with or without expecting an alpha channel */
+	FColorGradientEditorGeometrySolver(const bool bInWithAlphaChannel)
+		: bWithAlphaChannel(bInWithAlphaChannel)
+	{}
+
+	static constexpr float HandleWidth = 13.0f;
+	static constexpr float HandleHeight = 16.0f;
+	static constexpr float HandleColorWidth = 7.0f;
+	static constexpr float HandleColorHeight = 7.0f;
+	static constexpr float HandleColorXOffset = (HandleWidth - HandleColorWidth) / 2;
+	static constexpr float HandleColorYOffsetForColor = HandleColorXOffset + 1;
+	static constexpr float HandleColorYOffsetForAlpha = HandleHeight - HandleColorHeight - HandleColorXOffset - 1;
 
 	const FSlateRect HandleRect = FSlateRect(
-		-FColorGradientEditorGeometry::HandleWidth / 2.0f,
+		-HandleWidth / 2.0f,
 		0.0f,
-		FColorGradientEditorGeometry::HandleWidth / 2.0f,
-		FColorGradientEditorGeometry::HandleHeight);
+		HandleWidth / 2.0f,
+		HandleHeight);
 
-	FGeometry GetColorMarkAreaPaintGeometry( const FGeometry& WidgetGeometry )
+	FGeometry GetColorMarkAreaPaintGeometry(const FGeometry& WidgetGeometry)
 	{
+		const float ClampedHandleWidth = FMath::Max(WidgetGeometry.GetLocalSize().X - HandleWidth, 0.f);
+		
 		return WidgetGeometry.MakeChild(
-			FVector2D( WidgetGeometry.GetLocalSize().X - HandleWidth, HandleHeight ),
-			FSlateLayoutTransform(FVector2D(HandleWidth / 2, 0)) );
+			FVector2D(ClampedHandleWidth, HandleHeight),
+			FSlateLayoutTransform(FVector2D(HandleWidth / 2, 0)));
 	}
 
 	FGeometry GetColorMarkAreaMouseGeometry(const FGeometry& WidgetGeometry)
 	{
 		return WidgetGeometry.MakeChild(
 			FVector2D(WidgetGeometry.GetLocalSize().X, HandleHeight),
-			FSlateLayoutTransform() );
+			FSlateLayoutTransform());
 	}
 
-	FGeometry GetAlphaMarkAreaPaintGeometry( const FGeometry& WidgetGeometry)
+	TOptional<FGeometry> GetAlphaMarkAreaPaintGeometry(const FGeometry& WidgetGeometry)
 	{
-		return WidgetGeometry.MakeChild(
-			FVector2D(WidgetGeometry.GetLocalSize().X - HandleWidth, HandleHeight ),
-			FSlateLayoutTransform(FVector2D(HandleWidth / 2, WidgetGeometry.GetLocalSize().Y - HandleHeight)) );
+		if (bWithAlphaChannel)
+		{
+			return WidgetGeometry.MakeChild(
+				FVector2D(WidgetGeometry.GetLocalSize().X - HandleWidth, HandleHeight),
+				FSlateLayoutTransform(FVector2D(HandleWidth / 2, WidgetGeometry.GetLocalSize().Y - HandleHeight)));
+		}
+
+		return TOptional<FGeometry>();
 	}
 
-	FGeometry GetAlphaMarkAreaMouseGeometry(const FGeometry& WidgetGeometry)
+	TOptional<FGeometry> GetAlphaMarkAreaMouseGeometry(const FGeometry& WidgetGeometry)
 	{
-		return WidgetGeometry.MakeChild(
-			FVector2D(WidgetGeometry.GetLocalSize().X, HandleHeight ),
-			FSlateLayoutTransform(FVector2D(0.0f, WidgetGeometry.GetLocalSize().Y - HandleHeight)) );
+		if (bWithAlphaChannel)
+		{
+			return WidgetGeometry.MakeChild(
+				FVector2D(WidgetGeometry.GetLocalSize().X, HandleHeight),
+				FSlateLayoutTransform(FVector2D(0.0f, WidgetGeometry.GetLocalSize().Y - HandleHeight)));
+		}
+
+		return TOptional<FGeometry>();
 	}
 
-	FGeometry GetGradientAreaPaintGeometry( const FGeometry& WidgetGeometry, bool bHideStops = false)
+	FGeometry GetGradientAreaPaintGeometry(const FGeometry& WidgetGeometry, const bool bHideStops)
 	{
+		const float HandleWith = bHideStops ? 0.f : HandleWidth;
+		const float GradientWidth = WidgetGeometry.GetLocalSize().X - HandleWith;
+
+		const float HandlesHeight = bHideStops ? 0.f : bWithAlphaChannel ? 2.f * HandleHeight : HandleHeight;
+		const float GradientHeight = WidgetGeometry.GetLocalSize().Y - HandlesHeight;
+
 		return WidgetGeometry.MakeChild(
-			FVector2D( WidgetGeometry.GetLocalSize().X - (bHideStops ? 0: HandleWidth), WidgetGeometry.GetLocalSize().Y - (bHideStops ? 0 : (2 * HandleHeight)) ),
-			bHideStops? FSlateLayoutTransform(): FSlateLayoutTransform(FVector2D(HandleWidth / 2, HandleHeight)) 
+			FVector2D(GradientWidth, GradientHeight),
+			bHideStops ? FSlateLayoutTransform() : FSlateLayoutTransform(FVector2D(HandleWidth / 2, HandleHeight))
 		);
 	}
 
-	FTrackScaleInfo GetTimeScaleForWidgetGeometry(const FGeometry& WidgetGeometry, float ViewMinInput, float ViewMaxInput, bool bHideStops = false)
+	FTrackScaleInfo GetTimeScaleForWidgetGeometry(const FGeometry& WidgetGeometry, const float ViewMinInput, const float ViewMaxInput, const bool bHideStops)
 	{
 		return FTrackScaleInfo(ViewMinInput, ViewMaxInput, 0.0f, 1.0f, GetGradientAreaPaintGeometry(WidgetGeometry, bHideStops).GetLocalSize());
 	}
@@ -90,13 +113,16 @@ namespace FColorGradientEditorGeometry
 		return FTrackScaleInfo(ViewMinInput, ViewMaxInput, 0.0f, 1.0f, AreaGeometry.GetLocalSize());
 	}
 
-	float GetTimeForAbsolutePosition(const FGeometry& WidgetGeometry, float ViewMinInput, float ViewMaxInput, bool bHideStops, const FVector2D& AbsolutePosition)
+	float GetTimeForAbsolutePosition(const FGeometry& WidgetGeometry, const float ViewMinInput, const float ViewMaxInput, bool bHideStops, const FVector2D& AbsolutePosition)
 	{
 		FGeometry GradientAreaPaintGeometry = GetGradientAreaPaintGeometry(WidgetGeometry, bHideStops);
 		FTrackScaleInfo TimeScaleInfo = FTrackScaleInfo(ViewMinInput, ViewMaxInput, 0.0f, 1.0f, GradientAreaPaintGeometry.GetLocalSize());
 		return TimeScaleInfo.LocalXToInput(GradientAreaPaintGeometry.AbsoluteToLocal(AbsolutePosition).X);
 	}
-}
+
+private:
+	const bool bWithAlphaChannel;
+};
 
 FGradientStopMark::FGradientStopMark()
 {}
@@ -217,12 +243,14 @@ void SColorGradientEditor::Construct( const FArguments& InArgs )
 	ViewMaxInput = InArgs._ViewMaxInput;
 	bClampStopsToViewRange = InArgs._ClampStopsToViewRange;
 	bDraggingAlphaValue = false;
-	bDrawColorAndAlphaSeparate = InArgs._DrawColorAndAlphaSeparate;
+	bDrawColorAndAlphaSeparate = InArgs._WithAlphaChannel && InArgs._DrawColorAndAlphaSeparate;
 	bDraggingStop = false;
 	DistanceDragged = 0.0f;
 	ContextMenuPosition = FVector2D::ZeroVector;
 	DesiredSizeOverride.Assign(*this, InArgs._DesiredSizeOverride);
 	bRequestHideStops = InArgs._HideStops;
+
+	GeometrySolver = MakePimpl<FColorGradientEditorGeometrySolver>(InArgs._WithAlphaChannel);
 }
 
 int32 SColorGradientEditor::OnPaint( const FPaintArgs& Args, const FGeometry& AllottedGeometry, const FSlateRect& MyCullingRect, FSlateWindowElementList& OutDrawElements, int32 LayerId, const FWidgetStyle& InWidgetStyle, bool bParentEnabled ) const
@@ -234,15 +262,16 @@ int32 SColorGradientEditor::OnPaint( const FPaintArgs& Args, const FGeometry& Al
 		const bool bHideStops = ShouldHideStops();
 
 		// Split the geometry into areas for stops and the gradient
-		FGeometry ColorMarkAreaGeometry = FColorGradientEditorGeometry::GetColorMarkAreaPaintGeometry( AllottedGeometry );
-		FGeometry AlphaMarkAreaGeometry = FColorGradientEditorGeometry::GetAlphaMarkAreaPaintGeometry( AllottedGeometry );
-		FGeometry GradientAreaGeometry = FColorGradientEditorGeometry::GetGradientAreaPaintGeometry( AllottedGeometry , bHideStops);
+		const FGeometry ColorMarkAreaGeometry = GeometrySolver->GetColorMarkAreaPaintGeometry( AllottedGeometry );
+		const FGeometry GradientAreaGeometry = GeometrySolver->GetGradientAreaPaintGeometry(AllottedGeometry, bHideStops);
+		const TOptional<FGeometry> AlphaMarkAreaGeometry = GeometrySolver->GetAlphaMarkAreaPaintGeometry( AllottedGeometry );
 
 		bool bEnabled = ShouldBeEnabled( bParentEnabled );
 		ESlateDrawEffect DrawEffects = bEnabled ? ESlateDrawEffect::None : ESlateDrawEffect::DisabledEffect;
 
 		// Pixel to value input converter
-		FTrackScaleInfo ScaleInfo = FColorGradientEditorGeometry::GetTimeScaleForWidgetGeometry(AllottedGeometry, ViewMinInput.Get(), ViewMaxInput.Get());
+		constexpr bool bWithHideStops = false;
+		const FTrackScaleInfo ScaleInfo = GeometrySolver->GetTimeScaleForWidgetGeometry(AllottedGeometry, ViewMinInput.Get(), ViewMaxInput.Get(), bWithHideStops);
 
 		// The start and end location in slate units of the area to draw
 		int32 Start = 0;
@@ -272,14 +301,14 @@ int32 SColorGradientEditor::OnPaint( const FPaintArgs& Args, const FGeometry& Al
 				);
 			++LayerId;
 		}
-		else if (bAlphaAreaHovered)
+		else if (AlphaMarkAreaGeometry.IsSet() && bAlphaAreaHovered)
 		{
 			// Draw a checkerboard behind there is any transparency visible
 			FSlateDrawElement::MakeBox
 			(
 				OutDrawElements,
 				LayerId,
-				AlphaMarkAreaGeometry.ToPaintGeometry(),
+				AlphaMarkAreaGeometry.GetValue().ToPaintGeometry(),
 				WhiteBrush,
 				DrawEffects,
 				FLinearColor(.5f, .5f, .5f, .15f)
@@ -395,18 +424,21 @@ int32 SColorGradientEditor::OnPaint( const FPaintArgs& Args, const FGeometry& Al
 				}
 			}
 
-			// Draw each alpha stop
-			for (int32 ColorIndex = 0; ColorIndex < AlphaMarks.Num(); ++ColorIndex)
+			if (AlphaMarkAreaGeometry.IsSet())
 			{
-				const FGradientStopMark& Mark = AlphaMarks[ColorIndex];
+				// Draw each alpha stop
+				for (int32 ColorIndex = 0; ColorIndex < AlphaMarks.Num(); ++ColorIndex)
+				{
+					const FGradientStopMark& Mark = AlphaMarks[ColorIndex];
 
-				float XVal = ScaleInfo.InputToLocalX(Mark.Time);
+					float XVal = ScaleInfo.InputToLocalX(Mark.Time);
 
-				// Dont draw stops which are not visible
-				if (XVal >= 0 && XVal <= AlphaMarkAreaGeometry.GetLocalSize().X)
-				{
-					float Alpha = CurveOwner->GetLinearColorValue(Mark.Time).A;
-					DrawGradientStopMark(Mark, AlphaMarkAreaGeometry, XVal, FLinearColor(Alpha, Alpha, Alpha, 1.0f), OutDrawElements, LayerId, MyCullingRect, DrawEffects, false, InWidgetStyle);
+					// Dont draw stops which are not visible
+					if (XVal >= 0 && XVal <= AlphaMarkAreaGeometry.GetValue().GetLocalSize().X)
+					{
+						float Alpha = CurveOwner->GetLinearColorValue(Mark.Time).A;
+						DrawGradientStopMark(Mark, AlphaMarkAreaGeometry.GetValue(), XVal, FLinearColor(Alpha, Alpha, Alpha, 1.0f), OutDrawElements, LayerId, MyCullingRect, DrawEffects, false, InWidgetStyle);
+					}
 				}
 			}
 
@@ -434,9 +466,10 @@ int32 SColorGradientEditor::OnPaint( const FPaintArgs& Args, const FGeometry& Al
 				}
 
 				// Draw the text centered in the alpha region
+				if (AlphaMarkAreaGeometry.IsSet())
 				{
 					FVector2D StringSize = FontMeasureService->Measure(GradientAlphaMessage, FCoreStyle::GetDefaultFontStyle("Regular", 8));
-					FPaintGeometry PaintGeom = AlphaMarkAreaGeometry.ToPaintGeometry(FSlateLayoutTransform(FVector2D((AlphaMarkAreaGeometry.GetLocalSize().X - StringSize.X) * 0.5f, 1.0f)));
+					const FPaintGeometry PaintGeom = AlphaMarkAreaGeometry.GetValue().ToPaintGeometry(FSlateLayoutTransform(FVector2D((AlphaMarkAreaGeometry.GetValue().GetLocalSize().X - StringSize.X) * 0.5f, 1.0f)));
 
 					FSlateDrawElement::MakeText
 					(
@@ -504,15 +537,16 @@ FReply SColorGradientEditor::OnMouseButtonDoubleClick( const FGeometry& InMyGeom
 
 FReply SColorGradientEditor::OnMouseMove( const FGeometry& MyGeometry, const FPointerEvent& MouseEvent )
 {
-	FGeometry ColorMarkAreaMouseGeometry = FColorGradientEditorGeometry::GetColorMarkAreaMouseGeometry(MyGeometry);
-	FGeometry AlphaMarkAreaMouseGeometry = FColorGradientEditorGeometry::GetAlphaMarkAreaMouseGeometry(MyGeometry);
+	const FGeometry ColorMarkAreaMouseGeometry = GeometrySolver->GetColorMarkAreaMouseGeometry(MyGeometry);
+	const TOptional<FGeometry> AlphaMarkAreaMouseGeometry = GeometrySolver->GetAlphaMarkAreaMouseGeometry(MyGeometry);
 
 	if (ColorMarkAreaMouseGeometry.IsUnderLocation(MouseEvent.GetScreenSpacePosition()))
 	{
 		bColorAreaHovered = true;
 		bAlphaAreaHovered = false;
 	}
-	else if (AlphaMarkAreaMouseGeometry.IsUnderLocation(MouseEvent.GetScreenSpacePosition()))
+	else if (AlphaMarkAreaMouseGeometry.IsSet() &&
+		AlphaMarkAreaMouseGeometry.GetValue().IsUnderLocation(MouseEvent.GetScreenSpacePosition()))
 	{
 		bColorAreaHovered = false;
 		bAlphaAreaHovered = true;
@@ -544,13 +578,26 @@ FReply SColorGradientEditor::OnMouseMove( const FGeometry& MyGeometry, const FPo
 			}
 			else
 			{
-				// Already dragging a stop, move it
-				const bool bHideStops = ShouldHideStops();
-				float MouseTime = FColorGradientEditorGeometry::GetTimeForAbsolutePosition(MyGeometry, ViewMinInput.Get(), ViewMaxInput.Get(), bHideStops, MouseEvent.GetScreenSpacePosition());
-				if (bClampStopsToViewRange)
-				{
-					MouseTime = FMath::Clamp(MouseTime, ViewMinInput.Get(), ViewMaxInput.Get());
-				}
+				const float MouseTime = [&]()
+					{
+						// Already dragging a stop, move it
+						const bool bHideStops = ShouldHideStops();
+
+						float Result = GeometrySolver->GetTimeForAbsolutePosition(
+							MyGeometry, 
+							ViewMinInput.Get(), 
+							ViewMaxInput.Get(), 
+							bHideStops, 
+							MouseEvent.GetScreenSpacePosition());
+
+						if (bClampStopsToViewRange)
+						{
+							Result = FMath::Clamp(Result, ViewMinInput.Get(), ViewMaxInput.Get());
+						}
+
+						return Result;
+					}();
+
 				MoveStop( SelectedStop, MouseTime );
 
 				return FReply::Handled();
@@ -576,8 +623,8 @@ FReply SColorGradientEditor::OnMouseButtonUp( const FGeometry& MyGeometry, const
 			}
 			else if( DistanceDragged < DragThresholdDist && !SelectedStop.IsValid( *CurveOwner ) )
 			{
-				FGeometry ColorMarkAreaMouseGeometry = FColorGradientEditorGeometry::GetColorMarkAreaMouseGeometry( MyGeometry );
-				FGeometry AlphaMarkAreaMouseGeometry = FColorGradientEditorGeometry::GetAlphaMarkAreaMouseGeometry( MyGeometry );
+				const FGeometry ColorMarkAreaMouseGeometry = GeometrySolver->GetColorMarkAreaMouseGeometry( MyGeometry );
+				const TOptional<FGeometry> AlphaMarkAreaMouseGeometry = GeometrySolver->GetAlphaMarkAreaMouseGeometry( MyGeometry );
 
 				TOptional<FLinearColor> StopColor;
 				if (MouseEvent.IsControlDown())
@@ -598,7 +645,9 @@ FReply SColorGradientEditor::OnMouseButtonUp( const FGeometry& MyGeometry, const
 					return FReply::Handled().CaptureMouse( SharedThis(this) );
 
 				}
-				else if( AlphaMarkAreaMouseGeometry.IsUnderLocation( MouseEvent.GetScreenSpacePosition() ) )
+				else if( 
+					AlphaMarkAreaMouseGeometry.IsSet() &&
+					AlphaMarkAreaMouseGeometry.GetValue().IsUnderLocation(MouseEvent.GetScreenSpacePosition()))
 				{
 					// Add a new alpha mark
 					bool bColorStop = false;
@@ -626,7 +675,7 @@ FReply SColorGradientEditor::OnMouseButtonUp( const FGeometry& MyGeometry, const
 				else
 				{
 					const bool bHideStops = ShouldHideStops();
-					const FGeometry GradientAreaPaintGeometry = FColorGradientEditorGeometry::GetGradientAreaPaintGeometry(MyGeometry, bHideStops);
+					const FGeometry GradientAreaPaintGeometry = GeometrySolver->GetGradientAreaPaintGeometry(MyGeometry, bHideStops);
 					if (GradientAreaPaintGeometry.IsUnderLocation(MouseEvent.GetScreenSpacePosition()))
 					{
 						OpenGradientOptionsMenu(MouseEvent);
@@ -956,8 +1005,8 @@ void SColorGradientEditor::DrawGradientStopMark( const FGradientStopMark& Mark,
 	}
 
 	// Draw a box for the non colored area
-	FVector2f HandleSize = FColorGradientEditorGeometry::HandleRect.GetSize();
-	FVector2f HandlePosition = FColorGradientEditorGeometry::HandleRect.GetTopLeft2f() + FVector2f(XPos, 0);
+	const FVector2f HandleSize = GeometrySolver->HandleRect.GetSize();
+	const FVector2f HandlePosition = GeometrySolver->HandleRect.GetTopLeft2f() + FVector2f(XPos, 0);
 	FSlateDrawElement::MakeBox
 	( 
 		OutDrawElements,
@@ -970,12 +1019,12 @@ void SColorGradientEditor::DrawGradientStopMark( const FGradientStopMark& Mark,
 
 	// Draw a box with the gradient stop color
 	//Slate MakeGradient call expects the linear colors to be pre-converted to sRGB
-	FVector2f HandleColorSize = FVector2f(FColorGradientEditorGeometry::HandleColorWidth, FColorGradientEditorGeometry::HandleColorHeight);
-	FVector2f HandleColorPosition = HandlePosition + FVector2f(
-		FColorGradientEditorGeometry::HandleColorXOffset, 
+	const FVector2f HandleColorSize = FVector2f(GeometrySolver->HandleColorWidth, GeometrySolver->HandleColorHeight);
+	const FVector2f HandleColorPosition = HandlePosition + FVector2f(
+		GeometrySolver->HandleColorXOffset, 
 		bColor 
-			? FColorGradientEditorGeometry::HandleColorYOffsetForColor 
-			: FColorGradientEditorGeometry::HandleColorYOffsetForAlpha);
+			? GeometrySolver->HandleColorYOffsetForColor 
+			: GeometrySolver->HandleColorYOffsetForAlpha);
 	FSlateDrawElement::MakeBox
 	( 
 		OutDrawElements,
@@ -990,10 +1039,10 @@ void SColorGradientEditor::DrawGradientStopMark( const FGradientStopMark& Mark,
 FGradientStopMark SColorGradientEditor::GetGradientStopAtPoint( const FVector2D& MousePos, const FGeometry& MyGeometry )
 {
 	FGradientStopMark StopUnderMouse;
-	FGeometry ColorMarkAreaMouseGeometry = FColorGradientEditorGeometry::GetColorMarkAreaMouseGeometry( MyGeometry );
-	FGeometry AlphaMarkAreaMouseGeometry = FColorGradientEditorGeometry::GetAlphaMarkAreaMouseGeometry( MyGeometry );
-	bool bUnderColorMarkAreaGeometry = ColorMarkAreaMouseGeometry.IsUnderLocation(MousePos);
-	bool bUnderAlphaMarkAreaGeometry = AlphaMarkAreaMouseGeometry.IsUnderLocation(MousePos);
+	const FGeometry ColorMarkAreaMouseGeometry = GeometrySolver->GetColorMarkAreaMouseGeometry( MyGeometry );
+	const TOptional<FGeometry> AlphaMarkAreaMouseGeometry = GeometrySolver->GetAlphaMarkAreaMouseGeometry( MyGeometry );
+	const bool bUnderColorMarkAreaGeometry = ColorMarkAreaMouseGeometry.IsUnderLocation(MousePos);
+	const bool bUnderAlphaMarkAreaGeometry = AlphaMarkAreaMouseGeometry.IsSet() && AlphaMarkAreaMouseGeometry.GetValue().IsUnderLocation(MousePos);
 
 	if( bUnderColorMarkAreaGeometry || bUnderAlphaMarkAreaGeometry )
 	{
@@ -1003,7 +1052,7 @@ FGradientStopMark SColorGradientEditor::GetGradientStopAtPoint( const FVector2D&
 
 		auto GetMarkUnderMouse = [this](const FVector2D& MousePos, const TArray<FGradientStopMark>& Marks, const FGeometry& AreaPaintGeometry)
 		{
-			FTrackScaleInfo TimeScaleInfo = FColorGradientEditorGeometry::GetTimeScaleForAreaPaintGeomtry(AreaPaintGeometry, ViewMinInput.Get(), ViewMaxInput.Get());
+			const FTrackScaleInfo TimeScaleInfo = GeometrySolver->GetTimeScaleForAreaPaintGeomtry(AreaPaintGeometry, ViewMinInput.Get(), ViewMaxInput.Get());
 			for (int32 MarkIndex = 0; MarkIndex < Marks.Num(); ++MarkIndex)
 			{
 				const FGradientStopMark& Mark = Marks[MarkIndex];
@@ -1013,9 +1062,9 @@ FGradientStopMark SColorGradientEditor::GetGradientStopAtPoint( const FVector2D&
 
 				if (XVal >= 0)
 				{
-					FVector2f HandleSize = FColorGradientEditorGeometry::HandleRect.GetSize2f();
-					FVector2f HandlePosition = FColorGradientEditorGeometry::HandleRect.GetTopLeft2f() + FVector2f(XVal, 0.0f);
-					FGeometry HandleGeometry = AreaPaintGeometry.MakeChild(HandleSize, FSlateLayoutTransform(HandlePosition));
+					const FVector2f HandleSize = GeometrySolver->HandleRect.GetSize2f();
+					const FVector2f HandlePosition = GeometrySolver->HandleRect.GetTopLeft2f() + FVector2f(XVal, 0.0f);
+					const FGeometry HandleGeometry = AreaPaintGeometry.MakeChild(HandleSize, FSlateLayoutTransform(HandlePosition));
 					if (HandleGeometry.IsUnderLocation(MousePos))
 					{
 						return Mark;
@@ -1027,13 +1076,16 @@ FGradientStopMark SColorGradientEditor::GetGradientStopAtPoint( const FVector2D&
 
 		if (bUnderColorMarkAreaGeometry)
 		{
-			FGeometry ColorMarkAreaPaintGeometry = FColorGradientEditorGeometry::GetColorMarkAreaPaintGeometry(MyGeometry);
+			const FGeometry ColorMarkAreaPaintGeometry = GeometrySolver->GetColorMarkAreaPaintGeometry(MyGeometry);
 			StopUnderMouse = GetMarkUnderMouse(MousePos, ColorMarks, ColorMarkAreaPaintGeometry);
 		}
 		else // bUnderAlphaMarkAreaGeometry
 		{
-			FGeometry AlphaMarkAreaPaintGeometry = FColorGradientEditorGeometry::GetAlphaMarkAreaPaintGeometry(MyGeometry);
-			StopUnderMouse = GetMarkUnderMouse(MousePos, AlphaMarks, AlphaMarkAreaPaintGeometry);
+			const TOptional<FGeometry> AlphaMarkAreaPaintGeometry = GeometrySolver->GetAlphaMarkAreaPaintGeometry(MyGeometry);
+			if (AlphaMarkAreaPaintGeometry.IsSet())
+			{
+				StopUnderMouse = GetMarkUnderMouse(MousePos, AlphaMarks, AlphaMarkAreaPaintGeometry.GetValue());
+			}
 		}
 	}
 	return StopUnderMouse;
@@ -1116,7 +1168,7 @@ FGradientStopMark SColorGradientEditor::AddStop( const FVector2D& Position, cons
 	CurveOwner->ModifyOwner();
 
 	const bool bHideStops = ShouldHideStops();
-	float NewStopTime = FColorGradientEditorGeometry::GetTimeForAbsolutePosition(MyGeometry, ViewMinInput.Get(), ViewMaxInput.Get(), bHideStops, Position);
+	const float NewStopTime = GeometrySolver->GetTimeForAbsolutePosition(MyGeometry, ViewMinInput.Get(), ViewMaxInput.Get(), bHideStops, Position);
 
 	TArray<FRichCurveEditInfo> Curves = CurveOwner->GetCurves();
 
@@ -1129,7 +1181,7 @@ FGradientStopMark SColorGradientEditor::AddStop( const FVector2D& Position, cons
 		FRealCurve* GreenCurve = Curves[1].CurveToEdit;
 		FRealCurve* BlueCurve = Curves[2].CurveToEdit;
 
-		const float RValue = OptionalColor.IsSet() ? OptionalColor.GetValue().A : RedCurve->Eval(NewStopTime, 1.0f);
+		const float RValue = OptionalColor.IsSet() ? OptionalColor.GetValue().R : RedCurve->Eval(NewStopTime, 1.0f);
 		const float GValue = OptionalColor.IsSet() ? OptionalColor.GetValue().G : GreenCurve->Eval(NewStopTime, 1.0f);
 		const float BValue = OptionalColor.IsSet() ? OptionalColor.GetValue().B : BlueCurve->Eval(NewStopTime, 1.0f);
 
diff --git a/Engine/Source/Editor/UnrealEd/Public/SColorGradientEditor.h b/Engine/Source/Editor/UnrealEd/Public/SColorGradientEditor.h
--- a/Engine/Source/Editor/UnrealEd/Public/SColorGradientEditor.h
+++ b/Engine/Source/Editor/UnrealEd/Public/SColorGradientEditor.h
@@ -16,6 +16,7 @@
 class FCurveOwnerInterface;
 class FPaintArgs;
 class FSlateWindowElementList;
+struct FColorGradientEditorGeometrySolver;
 
 struct FGradientStopMark
 {
@@ -51,6 +52,7 @@ public:
 		, _ViewMaxInput(0.0f)
 		, _IsEditingEnabled( true )
 		, _ClampStopsToViewRange( false )
+		, _WithAlphaChannel( true )
 		, _DrawColorAndAlphaSeparate( true )
 	{}
 		SLATE_ATTRIBUTE( float, ViewMinInput )
@@ -58,6 +60,7 @@ public:
 		SLATE_ATTRIBUTE( bool, IsEditingEnabled )
 		SLATE_ATTRIBUTE(bool, HideStops)
 		SLATE_ARGUMENT( bool, ClampStopsToViewRange )
+		SLATE_ARGUMENT( bool, WithAlphaChannel )
 		SLATE_ARGUMENT(bool, DrawColorAndAlphaSeparate )
 		SLATE_ATTRIBUTE(TOptional<FVector2D>, DesiredSizeOverride)
 	SLATE_END_ARGS()
@@ -236,4 +239,7 @@ private:
 	TAttribute<bool> bRequestHideStops;
 	/** desired size override */
 	TSlateAttribute<TOptional<FVector2D>, EInvalidateWidgetReason::Layout> DesiredSizeOverride;
+
+	/** Solver for the geometry in this widget */
+	TPimplPtr<FColorGradientEditorGeometrySolver> GeometrySolver;
 };
