[ITF Gizmos] fix misuses of FEditorModeTools::UsesViewportITF() where an equivalent of UsesNewTRSGizmos is actually expected
- animation gizmos are currently disabled because of that and new trs gizmos aren't available in all asset editors yet

#rnx

#Synced-CL 51156304

diff --git a/Engine/Plugins/Animation/ControlRig/Source/ControlRigEditor/Private/EditMode/ControlRigEditMode.cpp b/Engine/Plugins/Animation/ControlRig/Source/ControlRigEditor/Private/EditMode/ControlRigEditMode.cpp
--- a/Engine/Plugins/Animation/ControlRig/Source/ControlRigEditor/Private/EditMode/ControlRigEditMode.cpp
+++ b/Engine/Plugins/Animation/ControlRig/Source/ControlRigEditor/Private/EditMode/ControlRigEditMode.cpp
@@ -2155,7 +2155,7 @@ bool FControlRigEditMode::InputDelta(FEditorViewportClient* InViewportClient, FV
 	//button down if left and ctrl and right is down, needed for indirect posting
 
 	// enable MMB with the new TRS gizmos
-	const bool bEnableMMB = InViewportClient->GetModeTools()->UsesViewportITF();
+	const bool bEnableMMB = InViewportClient->GetModeTools()->SupportsViewportITF() && UEditorInteractiveGizmoManager::UsesNewTRSGizmos();
 	
 	const bool bMouseButtonDown =
 		InViewport->KeyState(EKeys::LeftMouseButton) ||
diff --git a/Engine/Source/Editor/Experimental/EditorInteractiveToolsFramework/Private/EditorGizmos/EditorTransformGizmoSource.cpp b/Engine/Source/Editor/Experimental/EditorInteractiveToolsFramework/Private/EditorGizmos/EditorTransformGizmoSource.cpp
--- a/Engine/Source/Editor/Experimental/EditorInteractiveToolsFramework/Private/EditorGizmos/EditorTransformGizmoSource.cpp
+++ b/Engine/Source/Editor/Experimental/EditorInteractiveToolsFramework/Private/EditorGizmos/EditorTransformGizmoSource.cpp
@@ -98,7 +98,7 @@ bool UEditorTransformGizmoSource::CanInteract(const EViewportContext InViewportC
 			bool bUseLegacyWidget = (WidgetMode == UE::Widget::WM_TranslateRotateZ || WidgetMode == UE::Widget::WM_2D);
 			if (!bUseLegacyWidget)
 			{
-				bUseLegacyWidget = !ModeTools.UsesViewportITF();
+				bUseLegacyWidget = !(ModeTools.SupportsViewportITF() && UEditorInteractiveGizmoManager::UsesNewTRSGizmos());
 			}
 			return !bUseLegacyWidget;
 		}
diff --git a/Engine/Source/Editor/UnrealEd/Private/EditorModeManager.cpp b/Engine/Source/Editor/UnrealEd/Private/EditorModeManager.cpp
--- a/Engine/Source/Editor/UnrealEd/Private/EditorModeManager.cpp
+++ b/Engine/Source/Editor/UnrealEd/Private/EditorModeManager.cpp
@@ -146,7 +146,7 @@ void FEditorModeTools::LoadConfig(void)
 	if (CoordinateSystem == COORD_Parent || CoordinateSystem == COORD_Explicit)
 	{
 		// parent mode is only supported with new trs gizmos for now
-		if (!UsesViewportITF())
+		if (!(bViewportITF && UEditorInteractiveGizmoManager::UsesNewTRSGizmos()))
 		{
 			CoordSystemAsInt = static_cast<int32>(COORD_Local);
 		}
diff --git a/Engine/Source/Editor/UnrealEd/Private/LevelEditorViewport.cpp b/Engine/Source/Editor/UnrealEd/Private/LevelEditorViewport.cpp
--- a/Engine/Source/Editor/UnrealEd/Private/LevelEditorViewport.cpp
+++ b/Engine/Source/Editor/UnrealEd/Private/LevelEditorViewport.cpp
@@ -3320,7 +3320,7 @@ bool FLevelEditorViewportClient::InputWidgetDelta(FViewport* InViewport, EAxisLi
 				// enable MMB for New TRS Gizmos
 				const bool bModifyTransform = [this, MiddleMouseButtonDown]() -> bool
 				{
-					if (!ModeTools->UsesViewportITF())
+					if (!(ModeTools->SupportsViewportITF() && UEditorInteractiveGizmoManager::UsesNewTRSGizmos()))
 					{
 						return !MiddleMouseButtonDown;
 					}
diff --git a/Engine/Source/Editor/UnrealEd/Private/SEditorViewport.cpp b/Engine/Source/Editor/UnrealEd/Private/SEditorViewport.cpp
--- a/Engine/Source/Editor/UnrealEd/Private/SEditorViewport.cpp
+++ b/Engine/Source/Editor/UnrealEd/Private/SEditorViewport.cpp
@@ -1198,7 +1198,8 @@ void SEditorViewport::OnCycleCoordinateSystem()
 
 	// parent mode is only supported with new trs gizmos for now
 	int CoordMax = COORD_Parent;
-	if (Client->GetModeTools()->UsesViewportITF())
+	
+	if (Client->GetModeTools()->SupportsViewportITF() && UEditorInteractiveGizmoManager::UsesNewTRSGizmos())
 	{
 		CoordMax = UEditorInteractiveGizmoManager::IsExplicitModeEnabled() ? COORD_Max : COORD_Explicit;
 	}
diff --git a/Engine/Source/Editor/UnrealEd/Private/UnrealWidgetRender.cpp b/Engine/Source/Editor/UnrealEd/Private/UnrealWidgetRender.cpp
--- a/Engine/Source/Editor/UnrealEd/Private/UnrealWidgetRender.cpp
+++ b/Engine/Source/Editor/UnrealEd/Private/UnrealWidgetRender.cpp
@@ -190,7 +190,7 @@ void FWidget::Render( const FSceneView* View,FPrimitiveDrawInterface* PDI, FEdit
 	bool bUseLegacyWidget = (WidgetMode == UE::Widget::EWidgetMode::WM_TranslateRotateZ || WidgetMode == UE::Widget::EWidgetMode::WM_2D);
 	if (!bUseLegacyWidget)
 	{
-		bUseLegacyWidget = !ViewportClient->GetModeTools()->UsesViewportITF();
+		bUseLegacyWidget = !(ViewportClient->GetModeTools()->SupportsViewportITF() && UEditorInteractiveGizmoManager::UsesNewTRSGizmos());
 	}
 
 	bool bDrawWidget;
diff --git a/Engine/Source/Editor/UnrealEd/Private/ViewportToolbar/UnrealEdViewportToolbar.cpp b/Engine/Source/Editor/UnrealEd/Private/ViewportToolbar/UnrealEdViewportToolbar.cpp
--- a/Engine/Source/Editor/UnrealEd/Private/ViewportToolbar/UnrealEdViewportToolbar.cpp
+++ b/Engine/Source/Editor/UnrealEd/Private/ViewportToolbar/UnrealEdViewportToolbar.cpp
@@ -1064,7 +1064,7 @@ FToolMenuEntry CreateTransformsSubmenu()
 
 													if (TSharedPtr<FEditorViewportClient> Client = UUnrealEdViewportToolbarContext::GetEditorViewportClient(InSubmenu->Context))
 													{
-														if (FEditorModeTools* ModeTools = Client->GetModeTools(); ModeTools && ModeTools->UsesViewportITF())
+														if (FEditorModeTools* ModeTools = Client->GetModeTools(); ModeTools && ModeTools->SupportsViewportITF() && UEditorInteractiveGizmoManager::UsesNewTRSGizmos())
 														{
 															FToolMenuEntry ParentCoords = FToolMenuEntry::InitMenuEntry(
 																FEditorViewportCommands::Get().RelativeCoordinateSystem_Parent
