As for legacy import, invoke asset import before level import when monitoring folders

#rnx

#Synced-CL 51095113

diff --git a/Engine/Source/Editor/UnrealEd/Private/AutoReimport/ContentDirectoryMonitor.cpp b/Engine/Source/Editor/UnrealEd/Private/AutoReimport/ContentDirectoryMonitor.cpp
--- a/Engine/Source/Editor/UnrealEd/Private/AutoReimport/ContentDirectoryMonitor.cpp
+++ b/Engine/Source/Editor/UnrealEd/Private/AutoReimport/ContentDirectoryMonitor.cpp
@@ -312,30 +312,18 @@ void FContentDirectoryMonitor::ProcessAdditions(const DirectoryWatcher::FTimeLim
 				ImportAssetParameters.bReplaceExisting = false;
 				ImportAssetParameters.DestinationName = NewAssetName;
 
-				// Check if this is the Scene import or Asset Import
+				// Priority is given to asset import. Try it first 
+				bImportHandledThroughInterchange = InterchangeManager.ImportAsset(FPaths::GetPath(PackagePath), ScopedSourceData.GetSourceData(), ImportAssetParameters);
+				if (!bImportHandledThroughInterchange)
 				{
-					UE::Interchange::FScopedTranslator ScopedTranslator(ScopedSourceData.GetSourceData());
-					bSceneImport = ScopedTranslator.GetTranslator()->GetTranslatorType() == EInterchangeTranslatorType::Scenes;
-				}
-
-				if (bSceneImport)
-				{
-					if (!InterchangeManager.ImportScene(FPaths::GetPath(PackagePath), ScopedSourceData.GetSourceData(), ImportAssetParameters))
+					// Try a level import then
+					bImportHandledThroughInterchange = InterchangeManager.ImportScene(FPaths::GetPath(PackagePath), ScopedSourceData.GetSourceData(), ImportAssetParameters);
+					if (!bImportHandledThroughInterchange)
 					{
-						// Add the message to the message log letting the user know that the last scene import didn't work.
-						Context.GetMessageLog().Message(EMessageSeverity::Info, FText::Format(LOCTEXT("Info_FailedToImportAsset", "Failed to import file {0}."), FText::FromString(FullFilename)));
+						// Add the message to the message log letting the user know that the import didn't work.
+						Context.GetMessageLog().Message(EMessageSeverity::Info, FText::Format(LOCTEXT("Info_FailedToImportAsset", "Interchange: Failed to import file {0}."), FText::FromString(FullFilename)));
 					}
 				}
-				else
-				{
-					if (!InterchangeManager.ImportAsset(FPaths::GetPath(PackagePath), ScopedSourceData.GetSourceData(), ImportAssetParameters))
-					{
-						// Add the message to the message log letting the user know that the last scene import didn't work.
-						Context.GetMessageLog().Message(EMessageSeverity::Info, FText::Format(LOCTEXT("Info_FailedToImportAsset", "Failed to import file {0}."), FText::FromString(FullFilename)));
-					}
-				}
-
-				bImportHandledThroughInterchange = true;
 			}
 		}
 		
