[Lightmass] Remove shadowed variables to avoid a crash.
#tests build lighting in VR template
#rnx

#Synced-CL 50782841

diff --git a/Engine/Source/Editor/UnrealEd/Private/Lightmass/LightmassRender.cpp b/Engine/Source/Editor/UnrealEd/Private/Lightmass/LightmassRender.cpp
--- a/Engine/Source/Editor/UnrealEd/Private/Lightmass/LightmassRender.cpp
+++ b/Engine/Source/Editor/UnrealEd/Private/Lightmass/LightmassRender.cpp
@@ -995,8 +995,6 @@ private:
 	/** The material interface for this proxy */
 	UMaterialInterface* MaterialInterface;
 	UMaterial* Material;
-	TArray<TObjectPtr<UObject>> ReferencedTextures;
-	TArray<TObjectPtr<UTextureCollection>> ReferencedTextureCollections;
 	/** The property to compile for rendering the sample */
 	EMaterialProperty PropertyToCompile;
 	/** Stores which exported attribute this proxy is compiling for. */
