Snapping - global IsSnappingEnabled renamed to IsSnappingActive

#Synced-CL 51038616

diff --git a/Engine/Plugins/Editor/UVEditor/Source/UVEditor/Private/SUVEditor2DViewport.cpp b/Engine/Plugins/Editor/UVEditor/Source/UVEditor/Private/SUVEditor2DViewport.cpp
--- a/Engine/Plugins/Editor/UVEditor/Source/UVEditor/Private/SUVEditor2DViewport.cpp
+++ b/Engine/Plugins/Editor/UVEditor/Source/UVEditor/Private/SUVEditor2DViewport.cpp
@@ -84,7 +84,7 @@ void SUVEditor2DViewport::BindCommands()
 		{
 			if (TSharedPtr<FUVEditor2DViewportClient> Client = StaticCastSharedPtr<FUVEditor2DViewportClient>(WeakClient.Pin()))
 			{
-				Client->SetSnappingEnabled(!Client->GetSnappingEnabled());
+				Client->SetSnappingActive(!Client->IsSnappingActive());
 			}
 		}),
 		FCanExecuteAction(),
@@ -92,7 +92,7 @@ void SUVEditor2DViewport::BindCommands()
 		{
 			if (TSharedPtr<FUVEditor2DViewportClient> Client = StaticCastSharedPtr<FUVEditor2DViewportClient>(WeakClient.Pin()))
 			{
-				return Client->GetSnappingEnabled();
+				return Client->IsSnappingActive();
 			}
 			return false;
 		}));
diff --git a/Engine/Plugins/Editor/UVEditor/Source/UVEditor/Private/UVEditor2DViewportClient.cpp b/Engine/Plugins/Editor/UVEditor/Source/UVEditor/Private/UVEditor2DViewportClient.cpp
--- a/Engine/Plugins/Editor/UVEditor/Source/UVEditor/Private/UVEditor2DViewportClient.cpp
+++ b/Engine/Plugins/Editor/UVEditor/Source/UVEditor/Private/UVEditor2DViewportClient.cpp
@@ -335,17 +335,17 @@ bool FUVEditor2DViewportClient::AreWidgetButtonsEnabled() const
 	return ViewportButtonsAPI && ViewportButtonsAPI->AreGizmoButtonsEnabled();
 }
 
-void FUVEditor2DViewportClient::SetSnappingEnabled(bool bEnabled)
+void FUVEditor2DViewportClient::SetSnappingActive(bool bActive)
 {
 	if (ViewportButtonsAPI)
 	{
-		ViewportButtonsAPI->SetSnappingEnabled(bEnabled);	
+		ViewportButtonsAPI->SetSnappingActive(bActive);	
 	}
 }
 
-bool FUVEditor2DViewportClient::GetSnappingEnabled()
+bool FUVEditor2DViewportClient::IsSnappingActive()
 {
-	return ViewportButtonsAPI && ViewportButtonsAPI->IsSnappingEnabled();
+	return ViewportButtonsAPI && ViewportButtonsAPI->IsSnappingActive();
 }
 
 void FUVEditor2DViewportClient::SetLocationGridSnapEnabled(bool bEnabled)
diff --git a/Engine/Plugins/Editor/UVEditor/Source/UVEditor/Public/UVEditor2DViewportClient.h b/Engine/Plugins/Editor/UVEditor/Source/UVEditor/Public/UVEditor2DViewportClient.h
--- a/Engine/Plugins/Editor/UVEditor/Source/UVEditor/Public/UVEditor2DViewportClient.h
+++ b/Engine/Plugins/Editor/UVEditor/Source/UVEditor/Public/UVEditor2DViewportClient.h
@@ -31,8 +31,8 @@ public:
 	UE_API UUVToolViewportButtonsAPI::ESelectionMode GetSelectionMode() const;
 	UE_API bool AreWidgetButtonsEnabled() const;
 
-	UE_API void SetSnappingEnabled(bool bEnabled);
-	UE_API bool GetSnappingEnabled();
+	UE_API void SetSnappingActive(bool bActive);
+	UE_API bool IsSnappingActive();
 	UE_API void SetLocationGridSnapEnabled(bool bEnabled);
 	UE_API bool GetLocationGridSnapEnabled();
 	UE_API void SetLocationGridSnapValue(float SnapValue);
diff --git a/Engine/Plugins/Editor/UVEditor/Source/UVEditorTools/Public/ContextObjects/UVToolViewportButtonsAPI.h b/Engine/Plugins/Editor/UVEditor/Source/UVEditorTools/Public/ContextObjects/UVToolViewportButtonsAPI.h
--- a/Engine/Plugins/Editor/UVEditor/Source/UVEditorTools/Public/ContextObjects/UVToolViewportButtonsAPI.h
+++ b/Engine/Plugins/Editor/UVEditor/Source/UVEditorTools/Public/ContextObjects/UVToolViewportButtonsAPI.h
@@ -103,14 +103,14 @@ public:
 	}
 
 	// Sets snapping on/off, rather than the explicit type to enable/disable
-	void SetSnappingEnabled(bool bInEnabled)
+	void SetSnappingActive(bool bInEnabled)
 	{
-		bSnappingEnabled = bInEnabled;
+		bSnappingActive = bInEnabled;
 	}
 
-	bool IsSnappingEnabled() const
+	bool IsSnappingActive() const
 	{
-		return bSnappingEnabled;
+		return bSnappingActive;
 	}
 
 	void ToggleSnapEnabled(ESnapTypeFlag SnapMode)
@@ -126,7 +126,7 @@ public:
 	// Whether to perform the snapping, not just get it's checkbox state
 	bool IsSnapActive(ESnapTypeFlag SnapMode) const
 	{
-		return bSnappingEnabled && GetSnapEnabled(SnapMode);
+		return bSnappingActive && GetSnapEnabled(SnapMode);
 	}
 
 	void SetSnapValue(ESnapTypeFlag SnapMode, float SnapValue)
@@ -167,7 +167,7 @@ public:
 protected:
 
 	// Toggle for all snapping, can be toggled so that whatever options were stored aren't cleared on enable (vs. ESnapTypeFlag::NoSnap)
-	bool bSnappingEnabled = true;
+	bool bSnappingActive = true;
 	int32 SnapEnabled = (uint8)ESnapTypeFlag::NoSnap;
 	float LocationSnap = 1.0;
 	float RotationSnap = 5.0;
diff --git a/Engine/Source/Editor/LevelEditor/Private/LevelEditorActions.cpp b/Engine/Source/Editor/LevelEditor/Private/LevelEditorActions.cpp
--- a/Engine/Source/Editor/LevelEditor/Private/LevelEditorActions.cpp
+++ b/Engine/Source/Editor/LevelEditor/Private/LevelEditorActions.cpp
@@ -3452,7 +3452,7 @@ void FLevelEditorActionCallbacks::OnEnableVertexSnap()
 bool FLevelEditorActionCallbacks::OnIsVertexSnapEnabled()
 {
 	const ULevelEditorViewportSettings* ViewportSettings = GetDefault<ULevelEditorViewportSettings>();
-	return ViewportSettings->bSnappingEnabled && ViewportSettings->bSnapVertices;
+	return ViewportSettings->bSnappingActive && ViewportSettings->bSnapVertices;
 }
 
 FText FLevelEditorActionCallbacks::GetActorSnapTooltip()
diff --git a/Engine/Source/Editor/UnrealEd/Classes/Settings/LevelEditorViewportSettings.h b/Engine/Source/Editor/UnrealEd/Classes/Settings/LevelEditorViewportSettings.h
--- a/Engine/Source/Editor/UnrealEd/Classes/Settings/LevelEditorViewportSettings.h
+++ b/Engine/Source/Editor/UnrealEd/Classes/Settings/LevelEditorViewportSettings.h
@@ -449,50 +449,50 @@ public:
 
 	/** If disabled, any and all snapping is disabled */
 	UPROPERTY(EditAnywhere, config, AdvancedDisplay, Category=GridSnapping)
-	bool bSnappingEnabled = true;
+	bool bSnappingActive = true;
 
 	/** If enabled will use power of 2 grid settings (e.g, 1,2,4,8,16,...,1024) instead of decimal grid sizes */
-	UPROPERTY(EditAnywhere, config, Category=GridSnapping, meta=(DisplayName = "Use Power of Two Snap Size", EditCondition="bSnappingEnabled"))
+	UPROPERTY(EditAnywhere, config, Category=GridSnapping, meta=(DisplayName = "Use Power of Two Snap Size", EditCondition="bSnappingActive"))
 	bool bUsePowerOf2SnapSize;
 
 	/** Decimal grid sizes (for translation snapping and grid rendering) */
-	UPROPERTY(EditAnywhere, config, AdvancedDisplay, Category=GridSnapping, meta=(EditCondition="bSnappingEnabled"))
+	UPROPERTY(EditAnywhere, config, AdvancedDisplay, Category=GridSnapping, meta=(EditCondition="bSnappingActive"))
 	TArray<float> DecimalGridSizes;
 
 	/** The number of lines between each major line interval for decimal grids */
-	UPROPERTY(EditAnywhere, config, AdvancedDisplay, Category=GridSnapping, meta=(EditCondition="bSnappingEnabled"))
+	UPROPERTY(EditAnywhere, config, AdvancedDisplay, Category=GridSnapping, meta=(EditCondition="bSnappingActive"))
 	TArray<float> DecimalGridIntervals;	
 
 	/** Power of 2 grid sizes (for translation snapping and grid rendering) */
-	UPROPERTY(EditAnywhere, config, AdvancedDisplay, Category=GridSnapping, meta=(EditCondition="bSnappingEnabled"))
+	UPROPERTY(EditAnywhere, config, AdvancedDisplay, Category=GridSnapping, meta=(EditCondition="bSnappingActive"))
 	TArray<float> Pow2GridSizes;
 
 	/** The number of lines between each major line interval for pow2 grids */
-	UPROPERTY(EditAnywhere, config, AdvancedDisplay, Category=GridSnapping, meta=(EditCondition="bSnappingEnabled"))
+	UPROPERTY(EditAnywhere, config, AdvancedDisplay, Category=GridSnapping, meta=(EditCondition="bSnappingActive"))
 	TArray<float> Pow2GridIntervals;
 
 	/** User defined grid intervals for rotations */
-	UPROPERTY(EditAnywhere, config, AdvancedDisplay, Category=GridSnapping, meta=(EditCondition="bSnappingEnabled"))
+	UPROPERTY(EditAnywhere, config, AdvancedDisplay, Category=GridSnapping, meta=(EditCondition="bSnappingActive"))
 	TArray<float> CommonRotGridSizes;
 
 	/** Preset grid intervals for rotations */
-	UPROPERTY(EditAnywhere, config, AdvancedDisplay, Category=GridSnapping, meta=(EditCondition="bSnappingEnabled"))
+	UPROPERTY(EditAnywhere, config, AdvancedDisplay, Category=GridSnapping, meta=(EditCondition="bSnappingActive"))
 	TArray<float> DivisionsOf360RotGridSizes;
 
 	/** Grid sizes for scaling */
-	UPROPERTY(EditAnywhere, config, AdvancedDisplay, Category=GridSnapping, meta=(EditCondition="bSnappingEnabled"))
+	UPROPERTY(EditAnywhere, config, AdvancedDisplay, Category=GridSnapping, meta=(EditCondition="bSnappingActive"))
 	TArray<float> ScalingGridSizes;
 
 	/** If enabled, actor positions will snap to the grid. */
-	UPROPERTY(EditAnywhere, config, Category=GridSnapping, meta=(DisplayName = "Enable Grid Snapping", EditCondition="bSnappingEnabled"))
+	UPROPERTY(EditAnywhere, config, Category=GridSnapping, meta=(DisplayName = "Enable Grid Snapping", EditCondition="bSnappingActive"))
 	uint32 GridEnabled:1;
 	
 	/** If enabled, actor rotations will snap to the grid. */
-	UPROPERTY(EditAnywhere, config, Category=GridSnapping, meta=(DisplayName = "Enable Rotation Snapping", EditCondition="bSnappingEnabled"))
+	UPROPERTY(EditAnywhere, config, Category=GridSnapping, meta=(DisplayName = "Enable Rotation Snapping", EditCondition="bSnappingActive"))
 	uint32 RotGridEnabled:1;
 
 	/** If enabled, actor sizes will snap to the grid. */
-	UPROPERTY(EditAnywhere, config, Category=GridSnapping, meta=(DisplayName = "Enable Scale Snapping", EditCondition="bSnappingEnabled"))
+	UPROPERTY(EditAnywhere, config, Category=GridSnapping, meta=(DisplayName = "Enable Scale Snapping", EditCondition="bSnappingActive"))
 	uint32 SnapScaleEnabled:1;
 
 	/** If enabled, actors will snap to surfaces in the viewport when dragged around */
@@ -508,7 +508,7 @@ private:
 public:
 
 	/** If enabled, new Actors that you drag into the viewport snap to the active 2D layer. */
-	UPROPERTY(EditAnywhere, config, Category=GridSnapping, meta=(DisplayName = "Enable 2D Layer Snapping", EditCondition="bSnappingEnabled"))
+	UPROPERTY(EditAnywhere, config, Category=GridSnapping, meta=(DisplayName = "Enable 2D Layer Snapping", EditCondition="bSnappingActive"))
 	uint32 bEnableLayerSnap:1;
 
 	/** The index of the snap plane to use when 2D Layer Snapping is enabled, from the array of Snap Layers set for the project. */
@@ -516,11 +516,11 @@ public:
 	int32 ActiveSnapLayerIndex;
 
 	/** If true actor snap will be enabled in the editor **/
-	UPROPERTY(config, Category=GridSnapping, VisibleDefaultsOnly,AdvancedDisplay)
+	UPROPERTY(config, Category=GridSnapping, VisibleDefaultsOnly,AdvancedDisplay, meta=(EditCondition="bSnappingActive"))
 	uint32 bEnableActorSnap:1;
 
 	/** Global actor snap scale for the editor */
-	UPROPERTY(config, Category=GridSnapping, VisibleDefaultsOnly,AdvancedDisplay)
+	UPROPERTY(config, Category=GridSnapping, VisibleDefaultsOnly,AdvancedDisplay, meta=(EditCondition="bSnappingActive"))
 	float ActorSnapScale;
 
 	/** Global actor snap distance setting for the editor */
diff --git a/Engine/Source/Editor/UnrealEd/Private/SnappingUtils.cpp b/Engine/Source/Editor/UnrealEd/Private/SnappingUtils.cpp
--- a/Engine/Source/Editor/UnrealEd/Private/SnappingUtils.cpp
+++ b/Engine/Source/Editor/UnrealEd/Private/SnappingUtils.cpp
@@ -49,7 +49,7 @@ public:
 	virtual bool IsScaleSnappingEnabled() const override;
 	// End of FEditorViewportSnapping interface
 
-	bool IsSnappingEnabled() const;
+	bool IsSnappingActive() const;
 
 	bool IsSnapToGridEnabled();
 	bool IsSnapRotationEnabled();
@@ -148,7 +148,7 @@ private:
 
 bool FEditorViewportSnapping::IsSnapToGridEnabled()
 {
-	return IsSnappingEnabled() && GetDefault<ULevelEditorViewportSettings>()->GridEnabled && !IsSnapToVertexEnabledInternal();
+	return IsSnappingActive() && GetDefault<ULevelEditorViewportSettings>()->GridEnabled && !IsSnapToVertexEnabledInternal();
 }
 
 bool FEditorViewportSnapping::IsSnapRotationEnabled()
@@ -163,11 +163,6 @@ bool FEditorViewportSnapping::IsSnapScaleEnabled()
 
 bool FEditorViewportSnapping::IsSnapRotationEnabledInternal() const
 {
-	if (!IsSnappingEnabled())
-	{
-		return false;
-	}
-
 	// Ask Current Editor Mode if Rotation Snap is enabled
 	const ULevelEditorViewportSettings* Settings = GetDefault<ULevelEditorViewportSettings>();
 	return GLevelEditorModeTools().IsSnapRotationEnabled() || (Settings->SnapToSurface.bEnabled && Settings->SnapToSurface.bSnapRotation);
@@ -175,7 +170,7 @@ bool FEditorViewportSnapping::IsSnapRotationEnabledInternal() const
 
 bool FEditorViewportSnapping::IsSnapScaleEnabledInternal() const
 {
-	return IsSnappingEnabled() &&  GetDefault<ULevelEditorViewportSettings>()->SnapScaleEnabled;
+	return IsSnappingActive() && GetDefault<ULevelEditorViewportSettings>()->SnapScaleEnabled;
 }
 
 bool FEditorViewportSnapping::IsSnapToVertexEnabled(bool bIsPivot)
@@ -185,7 +180,7 @@ bool FEditorViewportSnapping::IsSnapToVertexEnabled(bool bIsPivot)
 
 bool FEditorViewportSnapping::IsSnapToVertexEnabledInternal(bool bIsPivot) const
 {
-	if (!IsSnappingEnabled())
+	if (!IsSnappingActive())
 	{
 		return false;
 	}
@@ -213,7 +208,7 @@ bool FEditorViewportSnapping::IsSnapToActorEnabled()
 
 bool FEditorViewportSnapping::IsSnapToActorEnabledInternal() const
 {
-	return IsSnappingEnabled() && (GetDefault<ULevelEditorViewportSettings>()->bEnableActorSnap && !IsSnapToVertexEnabledInternal());
+	return IsSnappingActive() && (GetDefault<ULevelEditorViewportSettings>()->bEnableActorSnap && !IsSnapToVertexEnabledInternal());
 }
 
 void FEditorViewportSnapping::EnableActorSnap(bool bEnable)
@@ -442,19 +437,14 @@ void FEditorViewportSnapping::DrawSnappingHelpers(const FSceneView* View, FPrimi
 	VertexSnappingImpl.DrawSnappingHelpers( View, PDI );
 }
 
-bool FEditorViewportSnapping::IsSnappingEnabled() const
+bool FEditorViewportSnapping::IsSnappingActive() const
 {
 	const ULevelEditorViewportSettings* ViewportSettings = GetDefault<ULevelEditorViewportSettings>();
-	return ViewportSettings->bSnappingEnabled;
+	return ViewportSettings->bSnappingActive;
 }
 
 bool FEditorViewportSnapping::IsTranslationSnappingEnabled() const
 {
-	if (!IsSnappingEnabled())
-	{
-		return false;
-	}
-
 	const ULevelEditorViewportSettings* Settings = GetDefault<ULevelEditorViewportSettings>();
 	return Settings->GridEnabled || Settings->SnapToSurface.bEnabled || IsSnapToVertexEnabledInternal() || IsSnapToActorEnabledInternal();
 }
@@ -509,14 +499,21 @@ bool FEditorViewportSnapping::SnapDragLocationToNearestVertex( const FVector& Ba
 
 TSharedPtr<FEditorViewportSnapping> FSnappingUtils::EditorViewportSnapper;
 
+bool FSnappingUtils::IsSnappingActive()
+{
+	return EditorViewportSnapper.IsValid() && EditorViewportSnapper->IsSnappingActive();
+}
+
 bool FSnappingUtils::IsSnapToGridEnabled()
 {
 	return EditorViewportSnapper.IsValid() && EditorViewportSnapper->IsSnapToGridEnabled();
 }
 
-bool FSnappingUtils::IsTranslationSnappingEnabled()
+bool FSnappingUtils::IsTranslationSnappingEnabled(const bool bInRequireSnappingActive)
 {
-	return EditorViewportSnapper.IsValid() && EditorViewportSnapper->IsTranslationSnappingEnabled();
+	return EditorViewportSnapper.IsValid()
+		&& (!bInRequireSnappingActive || EditorViewportSnapper->IsSnappingActive()) 
+		&& EditorViewportSnapper->IsTranslationSnappingEnabled();
 }
 
 bool FSnappingUtils::IsRotationSnapEnabled()
@@ -524,9 +521,11 @@ bool FSnappingUtils::IsRotationSnapEnabled()
 	return EditorViewportSnapper.IsValid() && EditorViewportSnapper->IsSnapRotationEnabled();
 }
 
-bool FSnappingUtils::IsRotationSnappingEnabled()
+bool FSnappingUtils::IsRotationSnappingEnabled(const bool bInRequireSnappingActive)
 {
-	return EditorViewportSnapper.IsValid() && EditorViewportSnapper->IsRotationSnappingEnabled();
+	return EditorViewportSnapper.IsValid() 
+		&& (!bInRequireSnappingActive || EditorViewportSnapper->IsSnappingActive()) 
+		&& EditorViewportSnapper->IsRotationSnappingEnabled();
 }
 
 bool FSnappingUtils::IsScaleSnapEnabled()
@@ -534,9 +533,11 @@ bool FSnappingUtils::IsScaleSnapEnabled()
 	return EditorViewportSnapper.IsValid() && EditorViewportSnapper->IsSnapScaleEnabled();
 }
 
-bool FSnappingUtils::IsScaleSnappingEnabled()
+bool FSnappingUtils::IsScaleSnappingEnabled(const bool bInRequireSnappingActive)
 {
-	return EditorViewportSnapper.IsValid() && EditorViewportSnapper->IsScaleSnappingEnabled();
+	return EditorViewportSnapper.IsValid()
+		&& (!bInRequireSnappingActive || EditorViewportSnapper->IsSnappingActive()) 
+		&& EditorViewportSnapper->IsScaleSnappingEnabled();
 }
 
 bool FSnappingUtils::IsSnapToActorEnabled()
diff --git a/Engine/Source/Editor/UnrealEd/Private/Tools/EdModeInteractiveToolsContext.cpp b/Engine/Source/Editor/UnrealEd/Private/Tools/EdModeInteractiveToolsContext.cpp
--- a/Engine/Source/Editor/UnrealEd/Private/Tools/EdModeInteractiveToolsContext.cpp
+++ b/Engine/Source/Editor/UnrealEd/Private/Tools/EdModeInteractiveToolsContext.cpp
@@ -192,7 +192,7 @@ public:
 		const ULevelEditorViewportSettings* ViewportSettings = GetDefault<ULevelEditorViewportSettings>();
 
 		FToolContextSnappingConfiguration Config;
-		Config.bEnableSnapping = ViewportSettings->bSnappingEnabled;
+		Config.bEnableSnapping = ViewportSettings->bSnappingActive;
 		Config.bEnablePositionGridSnapping = (ViewportSettings->GridEnabled != 0);
 		float EditorGridSize = GEditor->GetGridSize();
 		Config.PositionGridDimensions = FVector(EditorGridSize, EditorGridSize, EditorGridSize);
diff --git a/Engine/Source/Editor/UnrealEd/Public/SnappingUtils.h b/Engine/Source/Editor/UnrealEd/Public/SnappingUtils.h
--- a/Engine/Source/Editor/UnrealEd/Public/SnappingUtils.h
+++ b/Engine/Source/Editor/UnrealEd/Public/SnappingUtils.h
@@ -13,6 +13,11 @@ class FSceneView;
 class FSnappingUtils
 {
 public:
+	/**
+	 * @return true if any snapping is active and should be applied
+	 */
+	static UNREALED_API bool IsSnappingActive();
+
 	/**
 	 * @return true if snapping (translation) to the grid is enabled
 	 */
@@ -21,7 +26,7 @@ public:
 	/**
 	 * @return true if *any* snapping that affects translation is enabled
 	 */
-	static UNREALED_API bool IsTranslationSnappingEnabled();
+	static UNREALED_API bool IsTranslationSnappingEnabled(const bool bInRequireSnappingActive = true);
 
 	/**
 	 * @return true if orientation snapping to angular grid is enabled
@@ -31,7 +36,7 @@ public:
 	/**
 	 * @return true if *any* snapping that affects orientation is enabled
 	 */
-	static UNREALED_API bool IsRotationSnappingEnabled();
+	static UNREALED_API bool IsRotationSnappingEnabled(const bool bInRequireSnappingActive = true);
 
 	/**
 	 * @return true if scale snapping to grid is enabled
@@ -41,7 +46,7 @@ public:
 	/**
 	 * @return true if *any* snapping that affects scale is enabled
 	 */
-	static UNREALED_API bool IsScaleSnappingEnabled();
+	static UNREALED_API bool IsScaleSnappingEnabled(const bool bInRequireSnappingActive = true);
 
 	/**
 	 * @return true if snapping actors to other actors is enabled
