- Added EBatchCreateFontAsset::CreateIfNoFontExists to help with automated importing not creating Foo_Font1, Foo_Font2 unncessarily

#Synced-CL 50913554

diff --git a/Engine/Source/Editor/UnrealEd/Classes/Factories/FontFileImportFactory.h b/Engine/Source/Editor/UnrealEd/Classes/Factories/FontFileImportFactory.h
--- a/Engine/Source/Editor/UnrealEd/Classes/Factories/FontFileImportFactory.h
+++ b/Engine/Source/Editor/UnrealEd/Classes/Factories/FontFileImportFactory.h
@@ -23,6 +23,8 @@ enum class EBatchCreateFontAsset : uint8
 	Yes,
 	// Do not create a UFont to match the new UFontFace
 	No,
+	// Only create a new font if one doesn't already exist, this will help with automated/script imports
+	CreateIfNoFontExists,
 };
 
 UCLASS(MinimalAPI)
diff --git a/Engine/Source/Editor/UnrealEd/Private/Factories/EditorFactories.cpp b/Engine/Source/Editor/UnrealEd/Private/Factories/EditorFactories.cpp
--- a/Engine/Source/Editor/UnrealEd/Private/Factories/EditorFactories.cpp
+++ b/Engine/Source/Editor/UnrealEd/Private/Factories/EditorFactories.cpp
@@ -5158,6 +5158,8 @@ bool UFontFileImportFactory::ConfigureProperties()
 
 UObject* UFontFileImportFactory::FactoryCreateBinary(UClass* InClass, UObject* InParent, FName InName, EObjectFlags InFlags, UObject* InContext, const TCHAR* InType, const uint8*& InBuffer, const uint8* InBufferEnd, FFeedbackContext* InWarn)
 {
+	const FString CreatedFontPackageName = FString::Printf(TEXT("%s/%s_Font"), *FPackageName::GetLongPackagePath(InParent->GetOutermost()->GetName()), *InName.ToString());
+
 	// Should we create a font asset alongside our font face?
 	bool bCreateFontAsset = false;
 	{
@@ -5189,6 +5191,15 @@ UObject* UFontFileImportFactory::FactoryCreateBinary(UClass* InClass, UObject* I
 		else
 		{
 			bCreateFontAsset = BatchCreateFontAsset == EBatchCreateFontAsset::Yes;
+			if (BatchCreateFontAsset == EBatchCreateFontAsset::CreateIfNoFontExists)
+			{
+				FAssetData Unused;
+				const FString CreatedFontObjectName = FString::Printf(TEXT("%s.%s_Font"), *CreatedFontPackageName, *InName.ToString());
+				if (IAssetRegistryInterface::GetPtr()->TryGetAssetByObjectPath(FSoftObjectPath(CreatedFontObjectName), Unused) != UE::AssetRegistry::EExists::Exists)
+				{
+					bCreateFontAsset = true;
+				}
+			}
 		}
 	}
 
@@ -5221,7 +5232,7 @@ UObject* UFontFileImportFactory::FactoryCreateBinary(UClass* InClass, UObject* I
 
 		FString FontPackageName;
 		FString FontAssetName;
-		AssetToolsModule.Get().CreateUniqueAssetName(FString::Printf(TEXT("%s/%s_Font"), *FPackageName::GetLongPackagePath(InParent->GetOutermost()->GetName()), *InName.ToString()), FString(), FontPackageName, FontAssetName);
+		AssetToolsModule.Get().CreateUniqueAssetName(CreatedFontPackageName, FString(), FontPackageName, FontAssetName);
 
 		UFontFactory* FontFactory = NewObject<UFontFactory>();
 		FontFactory->bEditAfterNew = false;
