IncrementalCook: AssetRegistry: Change the way incremental cook loads the previous AssetRegistry.
Always load all the AssetRegistry data for previous cooked packages from the previous AssetRegistry. Writing AssetRegistry data occurs for every cooked package and is more up to date than reading data from the editor's global AssetRegistry. If AssetRegistry data in the global editor AssetRegistry changes, this already causes the package to be recooked, maintaining the invariant that the cooked AssetRegistry data is more up to date than the editor's global AssetRegistry. This fixes the issue that incremental cooks having different AssetRegistry data than full cooks, which was frequently manifesting as different captialization of package names in assetregistry data.
Move the AssetRegistry loading code into shared code between ZenStoreWriter and LoosePackageWriter since it no longer needs the writer-specific invalidation. Move the decision to cause a full recook if the AssetRegistry is missing from that special case PackageWriter code into the GlobalCookArtifact.
Remove the error case handling that removes packages from the oplog if they are not in the AssetRegistry. This has never come up as an issue, and if it does then it should be handled with a generic artifact interface to invalidate missing packages rather than a special case.
#rnx
#Synced-CL 51086784

#Synced-CL 51086784

diff --git a/Engine/Source/Developer/IoStoreUtilities/Private/ZenStoreWriter.cpp b/Engine/Source/Developer/IoStoreUtilities/Private/ZenStoreWriter.cpp
--- a/Engine/Source/Developer/IoStoreUtilities/Private/ZenStoreWriter.cpp
+++ b/Engine/Source/Developer/IoStoreUtilities/Private/ZenStoreWriter.cpp
@@ -1123,63 +1123,26 @@ void FZenStoreWriter::GetEntries(TFunction<void(TArrayView<const FPackageStoreEn
 	Callback(PackageStoreEntries, CookedPackagesInfo);
 }
 
-TUniquePtr<FAssetRegistryState> FZenStoreWriter::LoadPreviousAssetRegistry() 
+void FZenStoreWriter::PopulateOplog(const FAssetRegistryState& /* PreviousState*/, int32& OutNumPackagesInOplog)
 {
-	// Load the previous asset registry to return to CookOnTheFlyServer, and set the packages enumerated in both *this and
-	// the returned asset registry to the intersection of the oplog and the previous asset registry;
-	// to report a package as already cooked we have to have the information from both sources.
-	FString PreviousAssetRegistryFile = FPaths::Combine(MetadataDirectoryPath, GetDevelopmentAssetRegistryFilename());
-	TUniquePtr<FArchive> Reader(CookArtifactReader->CreateFileReader(*PreviousAssetRegistryFile));
-
-	if (!Reader)
-	{
-		RemoveCookedPackages();
-		return TUniquePtr<FAssetRegistryState>();
-	}
-
-	TUniquePtr<FAssetRegistryState> PreviousState = MakeUnique<FAssetRegistryState>();
-	PreviousState->Load(*Reader);
-
-	TSet<FName> RemoveSet;
-	TArray<FName> RemoveArray;
-	const TMap<FName, const FAssetPackageData*>& PreviousStatePackages = PreviousState->GetAssetPackageDataMap();
-	{
-		FReadScopeLock _(EntriesLock);
-		for (const TPair<FName, const FAssetPackageData*>& Pair : PreviousStatePackages)
-		{
-			FName PackageName = Pair.Key;
-			if (Pair.Value->DiskSize < 0)
-			{
-				// Keep the FailedSave previous cook packages; some of them (NeverCookPlaceholders) are not expected to exist
-				// in the package store
-				continue;
-			}
-			if (PackageNameToIndex.Find(PackageName))
-			{
-				continue;
-			}
-			RemoveSet.Add(PackageName);
-		}
-		if (RemoveSet.Num())
-		{
-			PreviousState->PruneAssetData(TSet<FName>(), RemoveSet, FAssetRegistrySerializationOptions());
-		}
+	// We already populated the oplog during FZenStoreWriter::Initialize. Just return the number of packages.
+	FReadScopeLock _(EntriesLock);
+	OutNumPackagesInOplog = PackageNameToIndex.Num();
+}
 
-		for (TPair<FName, int32>& Pair : PackageNameToIndex)
-		{
-			FName PackageName = Pair.Key;
-			if (!PreviousStatePackages.Find(PackageName))
-			{
-				RemoveArray.Add(PackageName);
-			}
-		}
-	}
-	if (RemoveArray.Num())
-	{
-		RemoveCookedPackages(RemoveArray);
-	}
+void FZenStoreWriter::UpdateLastReferenceDateAndPruneStaleOps(
+	UE::Cook::Artifact::FUpdateOplogPackagesContext& OplogContext)
+{
+	// Not yet implemented. We need to store a LastReferenced timestamp in oplog entries, and this timestamp has to be
+	// understood by the ZenServer so it can update the value for all ops that we pass to it in this function.
+}
 
-	return PreviousState;
+TArray<FName> FZenStoreWriter::GetOplogPackageNames()
+{
+	FReadScopeLock _(EntriesLock);
+	TArray<FName> Result;
+	PackageNameToIndex.GenerateKeyArray(Result);
+	return Result;
 }
 
 FCbObject FZenStoreWriter::GetOplogAttachment(FName PackageName, FUtf8StringView AttachmentKey)
diff --git a/Engine/Source/Developer/IoStoreUtilities/Public/ZenStoreWriter.h b/Engine/Source/Developer/IoStoreUtilities/Public/ZenStoreWriter.h
--- a/Engine/Source/Developer/IoStoreUtilities/Public/ZenStoreWriter.h
+++ b/Engine/Source/Developer/IoStoreUtilities/Public/ZenStoreWriter.h
@@ -151,7 +151,11 @@ public:
 
 	IOSTOREUTILITIES_API void WriteIoStorePackageData(const FPackageInfo& Info, const FIoBuffer& PackageData, const FPackageStoreEntryResource& PackageStoreEntry, const TArray<FFileRegion>& FileRegions);
 
-	IOSTOREUTILITIES_API TUniquePtr<FAssetRegistryState> LoadPreviousAssetRegistry() override;
+	IOSTOREUTILITIES_API virtual void PopulateOplog(const FAssetRegistryState& PreviousState,
+		int32& OutNumPackagesInOplog) override;
+	IOSTOREUTILITIES_API virtual void UpdateLastReferenceDateAndPruneStaleOps(
+		UE::Cook::Artifact::FUpdateOplogPackagesContext& OplogContext) override;
+	IOSTOREUTILITIES_API virtual TArray<FName> GetOplogPackageNames() override;
 	IOSTOREUTILITIES_API virtual FCbObject GetOplogAttachment(FName PackageName, FUtf8StringView AttachmentKey) override;
 	IOSTOREUTILITIES_API virtual void GetOplogAttachments(TArrayView<FName> PackageNames,
 		TArrayView<FUtf8StringView> AttachmentKeys,
diff --git a/Engine/Source/Editor/UnrealEd/Private/Commandlets/AssetRegistryGenerator.cpp b/Engine/Source/Editor/UnrealEd/Private/Commandlets/AssetRegistryGenerator.cpp
--- a/Engine/Source/Editor/UnrealEd/Private/Commandlets/AssetRegistryGenerator.cpp
+++ b/Engine/Source/Editor/UnrealEd/Private/Commandlets/AssetRegistryGenerator.cpp
@@ -229,7 +229,6 @@ FAssetRegistryGenerator::FAssetRegistryGenerator(const ITargetPlatform* InPlatfo
 	: AssetRegistry(FModuleManager::LoadModuleChecked<FAssetRegistryModule>(TEXT("AssetRegistry")).Get())
 	, TargetPlatform(InPlatform)
 	, bGenerateChunks(false)
-	, bClonedGlobalAssetRegistry(false)
 	, HighestChunkId(0)
 	, DependencyInfo(*GetMutableDefault<UChunkDependencyInfo>())
 {
@@ -1039,39 +1038,8 @@ bool FAssetRegistryGenerator::CleanManifestDirectories(UE::Cook::FCookSandbox& I
 void FAssetRegistryGenerator::SetPreviousAssetRegistry(TUniquePtr<FAssetRegistryState>&& InPreviousState)
 {
 	LLM_SCOPE_BYTAG(Cooker_GeneratedAssetRegistry);
-	PreviousPackagesToUpdate.Empty();
-	if (InPreviousState)
-	{
-		const TMap<FName, const FAssetPackageData*>& PreviousPackageDataMap = InPreviousState->GetAssetPackageDataMap();
-		PreviousPackagesToUpdate.Reserve(PreviousPackageDataMap.Num());
-		for (const TPair<FName, const FAssetPackageData*>& Pair : PreviousPackageDataMap)
-		{
-			FName PackageName = Pair.Key;
-			if (FPackageName::IsScriptPackage(WriteToString<256>(PackageName)))
-			{
-				continue;
-			}
-			FIterativelySkippedPackageUpdateData& UpdateData = PreviousPackagesToUpdate.FindOrAdd(PackageName);
-			bool bGenerated = false;
-
-			UpdateData.AssetDatas.Reserve(InPreviousState->NumAssetsByPackageName(PackageName));
-			InPreviousState->EnumerateAssetsByPackageName(PackageName, [&bGenerated, &UpdateData](const FAssetData* AssetData)
-				{
-					bGenerated |= (AssetData->PackageFlags & PKG_CookGenerated) != 0;
-					UpdateData.AssetDatas.Emplace(*AssetData);
-					return true; // Keep iterating
-				});
-			UpdateData.PackageData = *Pair.Value;
-
-			// Keep the dependencies and referencers of generated packages
-			if (bGenerated)
-			{
-				FAssetIdentifier PackageIdentifier(PackageName);
-				InPreviousState->GetDependencies(PackageIdentifier, UpdateData.PackageDependencies, UE::AssetRegistry::EDependencyCategory::Package);
-				InPreviousState->GetReferencers(PackageIdentifier, UpdateData.PackageReferencers, UE::AssetRegistry::EDependencyCategory::Package);
-			}
-		}
-	}
+	check(InPreviousState);
+	this->State = MoveTemp(*InPreviousState);
 }
 
 void FAssetRegistryGenerator::InjectEncryptionData(FAssetRegistryState& TargetState)
@@ -1208,40 +1176,6 @@ FAssetPackageData* FAssetRegistryGenerator::GetAssetPackageData(const FName& Pac
 	return State.CreateOrGetAssetPackageData(PackageName);
 }
 
-void FAssetRegistryGenerator::UpdateKeptPackages(TSet<FName>& InOutAllPackages)
-{
-	for (TPair<FName, FIterativelySkippedPackageUpdateData>& Pair : PreviousPackagesToUpdate)
-	{
-		for (FAssetData& PreviousAssetData : Pair.Value.AssetDatas)
-		{
-			State.UpdateAssetData(MoveTemp(PreviousAssetData), true /* bCreateIfNotExists */);
-		}
-		*State.CreateOrGetAssetPackageData(Pair.Key) = MoveTemp(Pair.Value.PackageData);
-
-		if (!Pair.Value.PackageDependencies.IsEmpty())
-		{
-			State.AddDependencies(FAssetIdentifier(Pair.Key), Pair.Value.PackageDependencies);
-		}
-		if (!Pair.Value.PackageReferencers.IsEmpty())
-		{
-			State.AddReferencers(FAssetIdentifier(Pair.Key), Pair.Value.PackageReferencers);
-		}
-
-		// Keep each of currently unreferenced but previously referenced (or previously development-only)
-		// packages from an earlier incremental cook in our development AssetRegistry, but not in the list
-		// of referenced packages for the current cook.
-		InOutAllPackages.Add(Pair.Key); // May already be present
-	}
-	PreviousPackagesToUpdate.Empty();
-}
-
-TArray<FName> FAssetRegistryGenerator::GetPreviousIncrementalCookPackages()
-{
-	TArray<FName> PackageNames;
-	PreviousPackagesToUpdate.GetKeys(PackageNames);
-	return PackageNames;
-}
-
 TMap<int32, FString> FAssetRegistryGenerator::GetPakChunkIdToStringOverrideMap()
 {
 	TMap<int32, FString> ChunkIdStringOverride;
@@ -1364,38 +1298,11 @@ void FAssetRegistryGenerator::UpdateCollectionAssetData()
 	}
 }
 
-void FAssetRegistryGenerator::Initialize(bool bInitializeFromExisting,
-	const FAssetRegistrySerializationOptions& DevelopmentSaveOptions)
+void FAssetRegistryGenerator::Initialize()
 {
-	LLM_SCOPE_BYTAG(Cooker_GeneratedAssetRegistry);
-
-	if (bInitializeFromExisting)
-	{
-		// If the asset registry is still doing its background scan, we need to wait for it to finish and tick it so that the results are flushed out
-		AssetRegistry.WaitForCompletion();
-		ensureMsgf(!AssetRegistry.IsLoadingAssets(), TEXT("Cannot initialize asset registry generator while asset registry is still scanning source assets "));
-
-		bClonedGlobalAssetRegistry = true;
-		AssetRegistry.InitializeTemporaryAssetRegistryState(State, DevelopmentSaveOptions);
-	}
-
 	FGameDelegates::Get().GetAssignLayerChunkDelegate() = FAssignLayerChunkDelegate::CreateStatic(AssignLayerChunkDelegate);
 }
 
-void FAssetRegistryGenerator::CloneGlobalAssetRegistryFilteredByPreviousState(const FAssetRegistryState& PreviousState,
-	const FAssetRegistrySerializationOptions& DevelopmentSaveOptions)
-{
-	TSet<FName> KeepPackages;
-	const TMap<FName, const FAssetPackageData*> &PreviousPackageDatas = PreviousState.GetAssetPackageDataMap();
-	KeepPackages.Reserve(PreviousPackageDatas.Num());
-	for (const TPair<FName, const FAssetPackageData*>& Pair : PreviousPackageDatas)
-	{
-		KeepPackages.Add(Pair.Key);
-	}
-	AssetRegistry.InitializeTemporaryAssetRegistryState(State, DevelopmentSaveOptions,
-		false /* bRefreshExisting */, KeepPackages);
-}
-
 static UE::Cook::ECookResult DiskSizeToCookResult(int64 DiskSize)
 {
 	if (DiskSize >= 0)
@@ -1424,14 +1331,17 @@ static int64 CookResultToDiskSize(UE::Cook::ECookResult CookResult)
 void FAssetRegistryGenerator::ComputePackageDifferences(const FComputeDifferenceOptions& Options,
 	const FAssetRegistryState& PreviousState, FAssetRegistryDifference& OutDifference)
 {
+	LLM_SCOPE_BYTAG(Cooker_GeneratedAssetRegistry);
+
 	TArray<FName> ModifiedScriptPackages;
 	const TMap<FName, const FAssetPackageData*>& PreviousAssetPackageDataMap = PreviousState.GetAssetPackageDataMap();
 	OutDifference.Packages.Reserve(PreviousAssetPackageDataMap.Num());
 
-	for (const TPair<FName, const FAssetPackageData*>& PackagePair : State.GetAssetPackageDataMap())
+	AssetRegistry.EnumerateAllPackages(
+		[this, &Options, &PreviousAssetPackageDataMap, &OutDifference]
+		(FName PackageName, const FAssetPackageData& CurrentPackageDataRef)
 	{
-		FName PackageName = PackagePair.Key;
-		const FAssetPackageData* CurrentPackageData = PackagePair.Value;
+		const FAssetPackageData* CurrentPackageData = &CurrentPackageDataRef;
 		const FAssetPackageData* PreviousPackageData = PreviousAssetPackageDataMap.FindRef(PackageName);
 
 		if (!PreviousPackageData)
@@ -1485,13 +1395,14 @@ void FAssetRegistryGenerator::ComputePackageDifferences(const FComputeDifference
 				OutDifference.Packages.Add(PackageName, EDifference::ModifiedCooked);
 			}
 		}
-	}
+	});
 
 	for (const TPair<FName, const FAssetPackageData*>& PackagePair : PreviousAssetPackageDataMap)
 	{
 		FName PackageName = PackagePair.Key;
 		const FAssetPackageData* PreviousPackageData = PackagePair.Value;
-		const FAssetPackageData* CurrentPackageData = State.GetAssetPackageData(PackageName);
+		TOptional<FAssetPackageData> CurrentPackageDataCopy = AssetRegistry.GetAssetPackageDataCopy(PackageName);
+		const FAssetPackageData* CurrentPackageData = CurrentPackageDataCopy.GetPtrOrNull();
 
 		if (!CurrentPackageData)
 		{
@@ -1503,7 +1414,8 @@ void FAssetRegistryGenerator::ComputePackageDifferences(const FComputeDifference
 			if (!GeneratorName.IsNone())
 			{
 				// Keep track of all generators and their list of generated
-				const FAssetPackageData* GeneratorData = State.GetAssetPackageData(GeneratorName);
+				TOptional<FAssetPackageData> GeneratorDataCopy = AssetRegistry.GetAssetPackageDataCopy(GeneratorName);
+				const FAssetPackageData* GeneratorData = GeneratorDataCopy.GetPtrOrNull();
 				if (!GeneratorData)
 				{
 					// Mark it as removed; it will be regenerated when the Generator cooks
@@ -1573,7 +1485,7 @@ void FAssetRegistryGenerator::ComputePackageDifferences(const FComputeDifference
 		}
 
 		TSet<FName> Visited;
-		Visited.Reserve(State.GetNumPackages());
+		Visited.Reserve(PreviousState.GetNumPackages());
 		for (FName PackageName : VisitStack)
 		{
 			Visited.Add(PackageName);
@@ -1583,12 +1495,13 @@ void FAssetRegistryGenerator::ComputePackageDifferences(const FComputeDifference
 			FName ModifiedPackage = VisitStack.Pop(EAllowShrinking::No);
 			FString ModifiedPackageLeafName;
 
-			// Read referencers from the current state. If there are referencers in the old state that are not in the
-			// current state, then they must have changed themselves and so are already in the modified set.
+			// Read referencers from the current global assetregistry. If there are referencers in the old state that are not in the
+			// current assetregistry, then they must have changed themselves and so are already in the modified set.
 			TArray<FAssetIdentifier> Referencers;
-			State.GetReferencers(ModifiedPackage, Referencers, UE::AssetRegistry::EDependencyCategory::Package,
+			FAssetIdentifier ModifiedPackageAssetId(ModifiedPackage);
+			AssetRegistry.GetReferencers(ModifiedPackageAssetId, Referencers, UE::AssetRegistry::EDependencyCategory::Package,
 				UE::AssetRegistry::EDependencyQuery::Hard);
-			State.GetReferencers(ModifiedPackage, Referencers, UE::AssetRegistry::EDependencyCategory::Package,
+			AssetRegistry.GetReferencers(ModifiedPackageAssetId, Referencers, UE::AssetRegistry::EDependencyCategory::Package,
 				UE::AssetRegistry::EDependencyQuery::Build);
 
 			for (const FAssetIdentifier& Referencer : Referencers)
@@ -1643,7 +1556,7 @@ void FAssetRegistryGenerator::ComputePackageDifferences(const FComputeDifference
 				}
 
 				// Convert this Referencer to modified. We do not need to handle Removed differences, because we found
-				// it in the current state's dependency tree and so it can not be a Removed difference.
+				// it in the current assetregistry's dependency tree and so it can not be a Removed difference.
 				switch (*ReferencerDifference)
 				{
 				case EDifference::IdenticalCooked: *ReferencerDifference = EDifference::ModifiedCooked; break;
@@ -1736,48 +1649,52 @@ FName FAssetRegistryGenerator::GetGeneratorPackage(FName PackageName, const FAss
 	return Dependency.AssetId.PackageName;
 }
 
-void FAssetRegistryGenerator::ComputePackageRemovals(const FAssetRegistryState& PreviousState, TArray<FName>& OutRemovedPackages,
+void FAssetRegistryGenerator::ReportGeneratorPackages(const FAssetRegistryState& PreviousState,
 	TMap<FName, FGeneratorPackageInfo>& OutGeneratorPackages, int32& OutNumNeverCookPlaceHolderPackages)
 {
+	LLM_SCOPE_BYTAG(Cooker_GeneratedAssetRegistry);
+
 	OutGeneratorPackages.Reset();
 	OutNumNeverCookPlaceHolderPackages = 0;
-	TSet<FName> RemovedPackageSet;
 
 	for (const TPair<FName, const FAssetPackageData*>& PackagePair : PreviousState.GetAssetPackageDataMap())
 	{
 		FName PackageName = PackagePair.Key;
 		const FAssetPackageData* PreviousPackageData = PackagePair.Value;
-		const FAssetPackageData* CurrentPackageData = State.GetAssetPackageData(PackageName);
-		if (!CurrentPackageData)
+		bool bGenerated = false;
+		PreviousState.EnumerateAssetsByPackageName(PackageName, [&bGenerated](const FAssetData* AssetData)
+			{
+				bGenerated = (AssetData->PackageFlags & PKG_CookGenerated) != 0;
+				return false; // stop iterating
+			});
+		if (bGenerated)
 		{
-			// If it's a generated package, never mark it as removed (that can only be handled by the generator)
-			// Mark it as modified if its generator or any of its dependencies are modified.
-			bool bGenerated = false;
-			PreviousState.EnumerateAssetsByPackageName(PackageName, [&bGenerated](const FAssetData* AssetData)
-				{
-					bGenerated = (AssetData->PackageFlags & PKG_CookGenerated) != 0;
-					return false; // stop iterating
-				});
-			if (bGenerated)
+			// A generated package is always referenced by its generator, and it might be referenced by its fellow generated
+			// packages. Look in its referencers for a non-generated package, and assume that is the generator.
+			TArray<FAssetIdentifier> Referencers;
+			PreviousState.GetReferencers(FAssetIdentifier(PackageName), Referencers,
+				UE::AssetRegistry::EDependencyCategory::Package);
+			for (TArray<FAssetIdentifier>::TIterator ReferencerIter(Referencers); ReferencerIter; ++ReferencerIter)
 			{
-				TArray<FAssetIdentifier> Referencers;
-				PreviousState.GetReferencers(FAssetIdentifier(PackageName), Referencers, UE::AssetRegistry::EDependencyCategory::Package);
-				FName GeneratorName = Referencers.Num() == 1 ? Referencers[0].PackageName : NAME_None;
-				const FAssetPackageData* GeneratorData = !GeneratorName.IsNone() ? State.GetAssetPackageData(GeneratorName) : nullptr;
-				if (!GeneratorData)
-				{
-					// Mark it as removed; it will be regenerated when the Generator cooks
-					RemovedPackageSet.Add(PackageName);
-				}
-				else
+				bool bReferencerIsGenerated = false;
+				PreviousState.EnumerateAssetsByPackageName(ReferencerIter->PackageName,
+					[&bReferencerIsGenerated](const FAssetData* AssetData)
+					{
+						bReferencerIsGenerated = (AssetData->PackageFlags & PKG_CookGenerated) != 0;
+						return false; // stop iterating
+					});
+				if (bReferencerIsGenerated)
 				{
-					OutGeneratorPackages.FindOrAdd(GeneratorName).Generated.Add(PackageName,
-						CopyAssetPackageDataForIncrementalCook(*PreviousPackageData));
+					ReferencerIter.RemoveCurrent();
 				}
 			}
-			else
+			FName GeneratorName = Referencers.Num() == 1 ? Referencers[0].PackageName : NAME_None;
+			const FAssetPackageData* GeneratorData = !GeneratorName.IsNone()
+				? PreviousState.GetAssetPackageData(GeneratorName) : nullptr;
+			if (GeneratorData)
 			{
-				RemovedPackageSet.Add(PackageName);
+				OutGeneratorPackages.FindOrAdd(GeneratorName).Generated.Add(PackageName,
+					CopyAssetPackageDataForIncrementalCook(*PreviousPackageData));
 			}
 		}
 		if (PreviousPackageData->DiskSize < 0)
@@ -1789,22 +1706,159 @@ void FAssetRegistryGenerator::ComputePackageRemovals(const FAssetRegistryState&
 			}
 		}
 	}
+}
 
-	// If a Generator package has been removed, remove all its generated packages
-	for (TMap<FName, FGeneratorPackageInfo>::TIterator Iter(OutGeneratorPackages); Iter; ++Iter)
+void FAssetRegistryGenerator::RemoveStalePackages(
+	const TMap<FName, UE::Cook::Artifact::FOplogPackageData>& OplogPackages)
+{
+	using namespace UE::Cook;
+
+	TStringBuilder<256> PackageNameStr;
+
+	// Delete stale script packages. We know which script packages are stale immediately, no LastReferencedDate timeout
+	// required,by looking for them in the current editor binary.
+	TSet<FAssetIdentifier> StaleScriptPackages;
+	for (const TPair<FName, const FAssetPackageData*>& Pair : State.GetAssetPackageDataMap())
 	{
-		FName GeneratorName = Iter->Key;
-		if (RemovedPackageSet.Contains(GeneratorName))
+		PackageNameStr.Reset();
+		PackageNameStr << Pair.Key;
+		if (FPackageName::IsScriptPackage(PackageNameStr))
 		{
-			for (const TPair<FName, FAssetPackageData>& Generated : Iter->Value.Generated)
+			if (!FindPackage(nullptr, *PackageNameStr))
 			{
-				RemovedPackageSet.Add(Generated.Key);
+				StaleScriptPackages.Add(FAssetIdentifier(Pair.Key));
 			}
-			Iter.RemoveCurrent();
 		}
 	}
+	if (!StaleScriptPackages.IsEmpty())
+	{
+		State.RemovePackageDatas(StaleScriptPackages);
+		State.RemoveDependsNodes(StaleScriptPackages);
+	}
+
+	// Delete stale content packages: if the PackageWriter removed them from the oplog for staleness then we need
+	// to remove them from AssetRegistry is well.
+	TArray<FAssetData*> RemoveDatas;
+	for (const TPair<FName, const FAssetPackageData*>& Pair : State.GetAssetPackageDataMap())
+	{
+		FName PackageName = Pair.Key;
+		const FAssetPackageData* PackageData = Pair.Value;
 
-	OutRemovedPackages = RemovedPackageSet.Array();
+		// TODO: NeverCookPlaceholders were accidentally stripped from the AssetRegistry by CL 38685901 to support
+		// build dependencies in the oplog; that changelist dropped them from the list of IgnorePackageDatas passed
+		// into the AssetRegistryGenerator to keep for the DevelopmentAssetRegistry.
+		// Culling of data from the DevelopmentAssetRegistry now only occurs here, so if we change this function to
+		// keep NeverCookPlaceholders, we will change the behavior back to keeping them in the AssetRegistry. This
+		// might no longer be needed. Investigate and enable or delete this block.
+		static volatile bool bKeepNeverCookPlaceholders = false;
+		if (bKeepNeverCookPlaceholders)
+		{
+			// Special case: NeverCookPlaceholders are not culled, because they are never put into the Oplog.
+			// TODO: Need to add some way to cull stale NeverCookPlaceholders; add NeverCookPlaceholders to
+			// the oplog?
+			if (PackageData->DiskSize < 0)
+			{
+				ECookResult CookResult = DiskSizeToCookResult(PackageData->DiskSize);
+				if (CookResult == ECookResult::NeverCookPlaceholder)
+				{
+					continue;
+				}
+			}
+		}
+		if (!OplogPackages.Contains(PackageName))
+		{
+			State.EnumerateMutableAssetsByPackageName(PackageName, [&RemoveDatas](FAssetData* AssetData)
+				{
+					RemoveDatas.Add(AssetData);
+					return true; // keep enumerating
+				});
+		}
+	}
+	if (!RemoveDatas.IsEmpty())
+	{
+		TBitArray<FDefaultBitArrayAllocator> bRemoved;
+		TBitArray<FDefaultBitArrayAllocator> bRemovedDependencyData;
+		bRemoved.SetNumUninitialized(RemoveDatas.Num());
+		bRemovedDependencyData.SetNumUninitialized(RemoveDatas.Num());
+		State.RemoveAssetDatas(RemoveDatas, true /* bRemoveDependencyData */, bRemoved, bRemovedDependencyData);
+	}
+}
+
+void FAssetRegistryGenerator::AddGlobalsToCookedAssetRegistry(UCookOnTheFlyServer& COTFS)
+{
+	const bool bCookingDLC = COTFS.IsCookingDLC();
+	TStringBuilder<256> PackageNameStr;
+
+	// Script packages: add Script packages from global assetregistry.
+	// This addition of script packages is possibly no longer used, we include it for now as a legacy behavior.
+	// We also take advantage of it in RemoveStalePackages to find all the stale /Script depends nodes.
+
+	// Copy /Script PackageDatas from the current global AssetRegistry
+	// DLC cooks skip this behavior to save time.
+	if (!bCookingDLC)
+	{
+		AssetRegistry.EnumerateAllPackages([this, &PackageNameStr](FName PackageName, const FAssetPackageData& PackageData)
+			{
+				PackageNameStr.Reset();
+				PackageNameStr << PackageName;
+				if (FPackageName::IsScriptPackage(PackageNameStr))
+				{
+					*State.CreateOrGetAssetPackageData(PackageName) = PackageData;
+				}
+			});
+	}
+
+	// DependsNodes: 1) Remove dependencies to unrecorded assets. We copied dependencies for recorded assets from the
+	// global assetregistry, but the targets of those dependencies could have included assets that were not cooked
+	// (in this session or even in previous incremental cooks) because they are editoronly or NeverCook.
+	
+	// DependsNodes: 2) Manage Dependencies: add in manage dependenciesfrom the global AssetRegistry, and remove
+	// any stale PrimaryAsset values from previous cooks. Mange dependencies are used by some editor tools that
+	// analyze the cooked assetregistry such as AssetManagerEditorModule.
+	TSet<FAssetIdentifier> RemoveNodes;
+	for (TPair<FAssetIdentifier, FDependsNode*>& Pair : State.CachedDependsNodes)
+	{
+		const FAssetIdentifier& Id = Pair.Key;
+		if (Id.IsValue())
+		{
+			// DependsNode for a SearchableName, keep it.
+		}
+		else if (Id.PrimaryAssetType.IsValid())
+		{
+			// DependsNode for a PrimaryAsset, remove it (to be added back below).
+			RemoveNodes.Add(Pair.Key);
+		}
+		else if (Id.IsPackage())
+		{
+			// DependsNode for a Package. Either /Script or content. Keep the dependsnode for either type if and only
+			// if we have a PackageAssetData for it.
+			if (!State.CachedPackageData.Contains(Id.PackageName))
+			{
+				RemoveNodes.Add(Id);
+			}
+		}
+		else if (!Id.PackageName.IsNone() && !Id.ObjectName.IsNone())
+		{
+			// A DependsNode to an object within a package. Invalid, remove it.
+			RemoveNodes.Add(Id);
+		}
+		else
+		{
+			// Invalid depends node with all fields == NAME_None; remove it
+			RemoveNodes.Add(Id);
+		}
+	}
+	if (!RemoveNodes.IsEmpty())
+	{
+		State.RemoveDependsNodes(RemoveNodes, FAssetRegistryState::ERemoveDependsFlags::RemoveOrphans);
+	}
+
+	// Copy Manage dependencies from the current global AssetRegistry.
+	// DLC cooks skip this behavior to save time.
+	if (!bCookingDLC)
+	{
+		AssetRegistry.CopyManageDependenciesToState(State);
+	}
 }
 
 void FAssetRegistryGenerator::FinalizeChunkIDs(const TSet<FName>& CookedPackages,
@@ -1829,18 +1883,6 @@ void FAssetRegistryGenerator::FinalizeChunkIDs(const TSet<FName>& CookedPackages
 		bGenerateChunks = ShouldPlatformGenerateStreamingInstallManifest(TargetPlatform);
 	}
 
-	TSet<FName> AllPackages;
-	AllPackages.Append(CookedPackages);
-	AllPackages.Append(DevelopmentOnlyPackages);
-
-	// Possibly apply previous AssetData and AssetPackageData for packages kept from a previous cook
-	UpdateKeptPackages(AllPackages);
-
-	// Prune our asset registry to union of cooked + dev only + (previously cooked or dev only)
-	{
-		State.PruneAssetData(AllPackages, TSet<FName>(), TSet<int32>(), DevelopmentSaveOptions);
-	}
-
 	// Development only packages should have been marked via UpdateAssetRegistryData with a negative size indicating why they were not cooked 
 	for (FName DevelopmentOnlyPackage : DevelopmentOnlyPackages)
 	{
@@ -2928,8 +2970,6 @@ UE::Cook::FAssetRegistryPackageMessage FAssetRegistryGenerator::RecordUpdateAsse
 
 void FAssetRegistryGenerator::StoreUpdateAssetRegistryData(FName PackageName, UE::Cook::FAssetRegistryPackageMessage&& Message)
 {
-	PreviousPackagesToUpdate.Remove(PackageName);
-
 	// Copy latest data for all Assets in the package into the cooked registry. This should be done even
 	// if not successful so that editor-only packages are recorded as well.
 	for (FAssetData& AssetData : Message.AssetDatas)
@@ -2942,18 +2982,17 @@ void FAssetRegistryGenerator::StoreUpdateAssetRegistryData(FName PackageName, UE
 
 	// Copy AssetPackageData and Dependencies
 	FAssetPackageData* AssetPackageData = GetAssetPackageData(PackageName);
-	if (!bClonedGlobalAssetRegistry)
+	// Copy the Dependencies and PackageData from the global AssetRegistry
+	TOptional<FAssetPackageData> GlobalPackageData = AssetRegistry.GetAssetPackageDataCopy(PackageName);
+	if (GlobalPackageData)
 	{
-		// Copy the Dependencies and PackageData from the global AssetRegistry
-		TOptional<FAssetPackageData> GlobalPackageData = AssetRegistry.GetAssetPackageDataCopy(PackageName);
-		if (GlobalPackageData)
-		{
-			*AssetPackageData = MoveTemp(*GlobalPackageData);
-		}
-		TArray<FAssetDependency> GlobalDependencies;
-		AssetRegistry.GetDependencies(PackageName, GlobalDependencies, UE::AssetRegistry::EDependencyCategory::All);
-		State.SetDependencies(FAssetIdentifier(PackageName), GlobalDependencies, UE::AssetRegistry::EDependencyCategory::All);
+		*AssetPackageData = MoveTemp(*GlobalPackageData);
 	}
+	TArray<FAssetDependency> GlobalDependencies;
+	AssetRegistry.GetDependencies(PackageName, GlobalDependencies, UE::AssetRegistry::EDependencyCategory::All);
+	State.SetDependencies(FAssetIdentifier(PackageName), GlobalDependencies, UE::AssetRegistry::EDependencyCategory::All);
+
+	// Override the global AssetRegistry's data with the data discovered during save
 	if (Message.OverrideAssetPackageData)
 	{
 		// Generated packages pass in their own AssetPackageData
diff --git a/Engine/Source/Editor/UnrealEd/Private/CookOnTheFlyServer.cpp b/Engine/Source/Editor/UnrealEd/Private/CookOnTheFlyServer.cpp
--- a/Engine/Source/Editor/UnrealEd/Private/CookOnTheFlyServer.cpp
+++ b/Engine/Source/Editor/UnrealEd/Private/CookOnTheFlyServer.cpp
@@ -7861,18 +7861,18 @@ void UCookOnTheFlyServer::PopulateCookedPackages(TArrayView<const ITargetPlatfor
 		UE_LOG(LogCook, Display, TEXT("Populating cooked package(s) from %s package store on platform '%s'"),
 			*CookSavePackageContext.WriterDebugName, *TargetPlatform->PlatformName());
 
-		TUniquePtr<FAssetRegistryState> PreviousAssetRegistry(PackageWriter.LoadPreviousAssetRegistry());
-		int32 NumPreviousPackages = PreviousAssetRegistry ? PreviousAssetRegistry->GetNumPackages() : 0;
-		if (NumPreviousPackages == 0)
+		TUniquePtr<FAssetRegistryState> PreviousAssetRegistry(GlobalArtifact->DetachAssetRegistry(TargetPlatform));
+		if (!PreviousAssetRegistry)
 		{
-			UE_LOG(LogCook, Display, TEXT("Found %d cooked package(s) in package store."), NumPreviousPackages);
+			UE_LOG(LogCook, Error, TEXT("Could not load PreviousAssetRegistry, which is unexpected when cooking incrementally. Abandoning incremental cook and executing a full cook."));
 			continue;
 		}
-
-		if (!PlatformAssetRegistry.HasClonedGlobalAssetRegistry())
+		int32 NumPreviousPackages;
+		PackageWriter.PopulateOplog(*PreviousAssetRegistry, NumPreviousPackages);
+		if (NumPreviousPackages == 0)
 		{
-			PlatformAssetRegistry.CloneGlobalAssetRegistryFilteredByPreviousState(*PreviousAssetRegistry,
-				PlatformData->AssetRegistryDevelopmentSaveOptions);
+			UE_LOG(LogCook, Display, TEXT("Found %d cooked package(s) in package store."), NumPreviousPackages);
+			continue;
 		}
 
 		TMap<FName, FAssetRegistryGenerator::FGeneratorPackageInfo> PreviousGeneratorPackages;
@@ -7881,11 +7881,9 @@ void UCookOnTheFlyServer::PopulateCookedPackages(TArrayView<const ITargetPlatfor
 			// Incremental cook does the equivalent operation of bRecurseModifications=true and bRecurseScriptModifications=true,
 			// but checks for out-of-datedness are done by FRequestCluster using the TargetDomainKey (which is built
 			// from dependencies), so we do not need to check for out-of-datedness here
-
-			TArray<FName> TombstonePackages;
 			int32 NumNeverCookPlaceHolderPackages;
-			PlatformAssetRegistry.ComputePackageRemovals(*PreviousAssetRegistry, TombstonePackages,
-				PreviousGeneratorPackages, NumNeverCookPlaceHolderPackages);
+			PlatformAssetRegistry.ReportGeneratorPackages(*PreviousAssetRegistry, PreviousGeneratorPackages,
+				NumNeverCookPlaceHolderPackages);
 
 			// Do not show NeverCookPlaceholder packages in the package counts
 			NumPreviousPackages -= NumNeverCookPlaceHolderPackages;
@@ -7985,7 +7983,6 @@ void UCookOnTheFlyServer::PopulateCookedPackages(TArrayView<const ITargetPlatfor
 			};
 
 			// Add CookedPackages for any identical packages, delete from disk any modified packages
-			// For legacy paranoia, also delete from disk any packages that were marked as uncooked.
 			for (const TPair<FName, EDifference>& Pair : Difference.Packages)
 			{
 				FName PackageName = Pair.Key;
@@ -8003,7 +8000,6 @@ void UCookOnTheFlyServer::PopulateCookedPackages(TArrayView<const ITargetPlatfor
 					break;
 				case EDifference::IdenticalUncooked:
 					AddPlaceholderPackage(PackageName, ECookResult::Failed, true /* bLegacyIterativelyUnmodified */);
-					PackagesToRemove.Add(PackageName);
 					break;
 				case EDifference::ModifiedUncooked:
 					PackagesToRemove.Add(PackageName);
@@ -8225,6 +8221,7 @@ void UCookOnTheFlyServer::RefreshPlatformAssetRegistries(const TArrayView<const
 			{
 				PlatformData->RegistryGenerator = MakeUnique<FAssetRegistryGenerator>(TargetPlatform);
 				PlatformData->RegistryReporter = MakeUnique<UE::Cook::FAssetRegistryReporterLocal>(*PlatformData->RegistryGenerator);
+				PlatformData->RegistryGenerator->Initialize();
 			}
 			else
 			{
@@ -8232,22 +8229,6 @@ void UCookOnTheFlyServer::RefreshPlatformAssetRegistries(const TArrayView<const
 			}
 			RegistryReporter = PlatformData->RegistryReporter.Get();
 		}
-
-		if (PlatformData->RegistryGenerator)
-		{
-			// if we are cooking DLC, we will just spend a lot of time removing the shipped packages from the AR,
-			// so we don't bother copying them over. can easily save 10 seconds on a large project
-			bool bInitializeFromExisting = !IsCookingDLC();
-			if (EnumHasAnyFlags(CookByTheBookOptions->StartupOptions, ECookByTheBookOptions::SkipHardReferences) &&
-				bCookFastStartup)
-			{
-				// We don't want to wait on the AssetRegistry when just testing the cook of a single file
-				bInitializeFromExisting = false;
-			}
-
-			PlatformData->RegistryGenerator->Initialize(bInitializeFromExisting,
-				PlatformData->AssetRegistryDevelopmentSaveOptions);
-		}
 	}
 }
 
@@ -9479,9 +9460,6 @@ void UCookOnTheFlyServer::CookByTheBookFinishedInternal()
 		// correctly. FShaderLibraryCookArtifact needs UpdateOplogPackages called after FinishPopulateShaderLibrary,
 		// so that it has all of its data loaded and ready to be pruned. But it has to do this pruning before any calls
 		// to FinalizeChunkIds or SaveShaderLibrary; those functions need to write out the pruned version of the data.
-		// Another reason: We call Generator.GetPreviousIncrementalCookPackages to calculate the data we pass to 
-		// UpdateOplogPackages, and the data read by that function is cleared during FinalizeChunkIds.
-		if (!bLegacyBuildDependencies)
 		{
 			UE_SCOPED_HIERARCHICAL_COOKTIMER(CookArtifacts.UpdateOplogPackages);
 			using namespace UE::Cook::Artifact;
@@ -9494,7 +9472,18 @@ void UCookOnTheFlyServer::CookByTheBookFinishedInternal()
 				const FPackagePlatformData& PackagePlatformData = PackageData->FindOrAddPlatformData(TargetPlatform);
 				Data.bIncrementallyCooked = !PackagePlatformData.IsIncrementallySkipped();
 			}
-			for (FName PackageName : Generator.GetPreviousIncrementalCookPackages())
+			for (FPackageData* PackageData : IgnorePackageDatas)
+			{
+				FOplogPackageData& Data = OplogContext.OplogPackages.FindOrAdd(PackageData->GetPackageName());
+				Data.ReferenceType = UE::Cook::EPackageReferenceType::Cook;
+				const FPackagePlatformData& PackagePlatformData = PackageData->FindOrAddPlatformData(TargetPlatform);
+				Data.bIncrementallyCooked = !PackagePlatformData.IsIncrementallySkipped();
+			}
+
+			// Call UpdateLastReferenceDateAndPruneStaleOps to remove any stale ops from the oplog, and then
+			// add all of the ops in the oplog to OplogContext.OplogPackages, if they do not already exist.
+			SaveContext.PackageWriter->UpdateLastReferenceDateAndPruneStaleOps(OplogContext);
+			for (FName PackageName : SaveContext.PackageWriter->GetOplogPackageNames())
 			{
 				// We rely on ReferenceType being initialized to EPackageReferenceType::Unreferenced for newly added
 				// elements. In that case, calling FindOrAdd is sufficient to add the entry we need because we only
@@ -9505,7 +9494,10 @@ void UCookOnTheFlyServer::CookByTheBookFinishedInternal()
 
 			for (const TRefCountPtr<UE::Cook::ICookArtifact>& Artifact : CookArtifacts)
 			{
-				Artifact->UpdateOplogPackages(OplogContext);
+				if (!bLegacyBuildDependencies || Artifact == GlobalArtifact)
+				{
+					Artifact->UpdateOplogPackages(OplogContext);
+				}
 			}
 		}
 	}
diff --git a/Engine/Source/Editor/UnrealEd/Private/Cooker/DiffPackageWriter.h b/Engine/Source/Editor/UnrealEd/Private/Cooker/DiffPackageWriter.h
--- a/Engine/Source/Editor/UnrealEd/Private/Cooker/DiffPackageWriter.h
+++ b/Engine/Source/Editor/UnrealEd/Private/Cooker/DiffPackageWriter.h
@@ -92,9 +92,18 @@ public:
 		Inner->BeginCook(Info);
 	}
 	virtual void EndCook(const FCookInfo& Info) override;
-	virtual TUniquePtr<FAssetRegistryState> LoadPreviousAssetRegistry() override
+	virtual void PopulateOplog(const FAssetRegistryState& PreviousState, int32& OutNumPackagesInOplog) override
 	{
-		return Inner->LoadPreviousAssetRegistry();
+		Inner->PopulateOplog(PreviousState, OutNumPackagesInOplog);
+	}
+	virtual void UpdateLastReferenceDateAndPruneStaleOps(
+		UE::Cook::Artifact::FUpdateOplogPackagesContext& OplogContext) override
+	{
+		Inner->UpdateLastReferenceDateAndPruneStaleOps(OplogContext);
+	}
+	virtual TArray<FName> GetOplogPackageNames() override
+	{
+		return Inner->GetOplogPackageNames();
 	}
 	virtual FCbObject GetOplogAttachment(FName PackageName, FUtf8StringView AttachmentKey) override
 	{
diff --git a/Engine/Source/Editor/UnrealEd/Private/Cooker/GlobalCookArtifact.cpp b/Engine/Source/Editor/UnrealEd/Private/Cooker/GlobalCookArtifact.cpp
--- a/Engine/Source/Editor/UnrealEd/Private/Cooker/GlobalCookArtifact.cpp
+++ b/Engine/Source/Editor/UnrealEd/Private/Cooker/GlobalCookArtifact.cpp
@@ -2,9 +2,12 @@
 
 #include "Cooker/GlobalCookArtifact.h"
 
+#include "AssetRegistry/IAssetRegistry.h"
 #include "Containers/Array.h"
 #include "Containers/Map.h"
 #include "Containers/Set.h"
+#include "Cooker/CookArtifactReader.h"
+#include "Cooker/CookPlatformManager.h"
 #include "Cooker/CookTypes.h"
 #include "CookOnTheSide/CookLog.h"
 #include "CookOnTheSide/CookOnTheFlyServer.h"
@@ -12,6 +15,7 @@
 #include "Hash/xxhash.h"
 #include "Interfaces/ITargetPlatform.h"
 #include "Logging/LogMacros.h"
+#include "Logging/StructuredLog.h"
 #include "Misc/ConfigAccessTracking.h"
 #include "Misc/StringBuilder.h"
 #include "Modules/ModuleManager.h"
@@ -212,6 +216,75 @@ void FGlobalCookArtifact::CompareSettings(UE::Cook::Artifact::FCompareSettingsCo
 			return;
 		}
 	}
+
+	if (!TryLoadPreviousAssetRegistry(Context))
+	{
+		// Log and RequestFullRecook was issued by TryLoadPreviousAssetRegistry.
+		return;
+	}
+}
+
+bool FGlobalCookArtifact::TryLoadPreviousAssetRegistry(UE::Cook::Artifact::FCompareSettingsContext& Context)
+{
+	// Load the previous asset registry; if we cannot load it then we are missing required data for all packages
+	// and will need to do a full recook. After loading the PreviousAssetRegistry we hand it over during
+	// PopulateCookedPackages to the AssetRegistryGenerator for update during the cook and writing at end of cook.
+	const ITargetPlatform* TargetPlatform = Context.GetTargetPlatform();
+	FPlatformData* PlatformData = COTFS.PlatformManager->GetPlatformData(TargetPlatform);
+	check(PlatformData);
+	FString PreviousAssetRegistryFile;
+	if (!PlatformData->bLegacyIterativeSharedBuild)
+	{
+		PreviousAssetRegistryFile = FPaths::Combine(Context.GetCookInfo().GetCookMetadataOutputFolder(TargetPlatform),
+			GetDevelopmentAssetRegistryFilename());
+	}
+	else
+	{
+		PreviousAssetRegistryFile = FPaths::Combine(*FPaths::ProjectSavedDir(), TEXT("SharedIterativeBuild"),
+			*TargetPlatform->PlatformName(), TEXT("Metadata"), GetDevelopmentAssetRegistryFilename());
+	}
+
+	TUniquePtr<FArchive> Reader(Context.GetArtifactReader().CreateFileReader(*PreviousAssetRegistryFile));
+
+	if (!Reader)
+	{
+		UE_LOGFMT(LogCook, Display, "Cook invalidated for platform {Platform} because the previous AssetRegistry is not available. "
+			"Unable to load '{AssetRegistryFile}'. Clearing previously cooked packages.",
+			*TargetPlatform->PlatformName(), *PreviousAssetRegistryFile);
+		Context.RequestFullRecook(true);
+		return false;
+	}
+
+	TUniquePtr<FAssetRegistryState> PreviousState = MakeUnique<FAssetRegistryState>();
+	if (!PreviousState->Load(*Reader))
+	{
+		UE_LOGFMT(LogCook, Display, "Cook invalidated for platform {Platform} because the previous AssetRegistry is not available. "
+			"Error during reading '{AssetRegistryFile}'. Clearing previously cooked packages.",
+			*TargetPlatform->PlatformName(), *PreviousAssetRegistryFile);
+		Context.RequestFullRecook(true);
+		return false;
+	}
+
+	PreviousAssetRegistries.Add(TargetPlatform, MoveTemp(PreviousState));
+	return true;
+}
+
+TUniquePtr<FAssetRegistryState> FGlobalCookArtifact::DetachAssetRegistry(const ITargetPlatform* TargetPlatform)
+{
+	TUniquePtr<FAssetRegistryState> Result;
+	PreviousAssetRegistries.RemoveAndCopyValue(TargetPlatform, Result);
+	return Result;
+}
+
+void FGlobalCookArtifact::UpdateOplogPackages(UE::Cook::Artifact::FUpdateOplogPackagesContext& Context)
+{
+	const ITargetPlatform* TargetPlatform = Context.GetTargetPlatform();
+	FPlatformData* PlatformData = COTFS.PlatformManager->GetPlatformData(TargetPlatform);
+	check(PlatformData && PlatformData->RegistryGenerator);
+	FAssetRegistryGenerator& Generator = *PlatformData->RegistryGenerator;
+
+	Generator.RemoveStalePackages(Context.GetOplogPackages());
+	Generator.AddGlobalsToCookedAssetRegistry(COTFS);
 }
 
 } // namespace UE::Cook
diff --git a/Engine/Source/Editor/UnrealEd/Private/Cooker/GlobalCookArtifact.h b/Engine/Source/Editor/UnrealEd/Private/Cooker/GlobalCookArtifact.h
--- a/Engine/Source/Editor/UnrealEd/Private/Cooker/GlobalCookArtifact.h
+++ b/Engine/Source/Editor/UnrealEd/Private/Cooker/GlobalCookArtifact.h
@@ -2,8 +2,12 @@
 
 #pragma once
 
+#include "AssetRegistry/AssetRegistryState.h"
+#include "Containers/Map.h"
 #include "Cooker/CookArtifact.h"
+#include "Templates/UniquePtr.h"
 
+class ITargetPlatform;
 class UCookOnTheFlyServer;
 
 namespace UE::Cook
@@ -11,7 +15,7 @@ namespace UE::Cook
 
 /**
  * CookArtifact that triggers a recook of all packages if any of the global settings that affect all
- * packages change. This artifact does not store files on disk other than its record of the global CookSettings.
+ * packages change. This artifact is also responsible for incremental cook's load of the AssetRegistry artifact.
  */
 class FGlobalCookArtifact : public ICookArtifact
 {
@@ -21,9 +25,19 @@ public:
 	FString GetArtifactName() const override;
 	FConfigFile CalculateCurrentSettings(ICookInfo& CookInfo, const ITargetPlatform* TargetPlatform) override;
 	void CompareSettings(UE::Cook::Artifact::FCompareSettingsContext& Context) override;
+	void UpdateOplogPackages(UE::Cook::Artifact::FUpdateOplogPackagesContext& Context) override;
+
+	TUniquePtr<FAssetRegistryState> DetachAssetRegistry(const ITargetPlatform* TargetPlatform);
 
 private:
+	bool TryLoadPreviousAssetRegistry(UE::Cook::Artifact::FCompareSettingsContext& Context);
+
 	UCookOnTheFlyServer& COTFS;
+	/**
+	 * This artifact holds the PreviousAssetRegistries loaded from disk for a short time during StartCookByTheBook
+	 * or StartCookOnTheFly. It relinquishes them to the AssetRegistryGenerator during PopulateCookedPackages.
+	 */
+	TMap<const ITargetPlatform*, TUniquePtr<FAssetRegistryState>> PreviousAssetRegistries;
 };
 
 
diff --git a/Engine/Source/Editor/UnrealEd/Private/Cooker/LooseCookedPackageWriter.cpp b/Engine/Source/Editor/UnrealEd/Private/Cooker/LooseCookedPackageWriter.cpp
--- a/Engine/Source/Editor/UnrealEd/Private/Cooker/LooseCookedPackageWriter.cpp
+++ b/Engine/Source/Editor/UnrealEd/Private/Cooker/LooseCookedPackageWriter.cpp
@@ -341,102 +341,68 @@ void FLooseCookedPackageWriter::EndCook(const FCookInfo& Info)
 	}
 }
 
-TUniquePtr<FAssetRegistryState> FLooseCookedPackageWriter::LoadPreviousAssetRegistry()
+void FLooseCookedPackageWriter::PopulateOplog(const FAssetRegistryState& PreviousState, int32& OutNumPackagesInOplog)
 {
-	// Report files from the shared build if the option is set
-	FString PreviousAssetRegistryFile;
-	if (bLegacyIterativeSharedBuild)
-	{
-		// clean the sandbox
-		DeleteSandboxDirectory();
-		PreviousAssetRegistryFile = FPaths::Combine(*FPaths::ProjectSavedDir(), TEXT("SharedIterativeBuild"),
-			*TargetPlatform.PlatformName(), TEXT("Metadata"), GetDevelopmentAssetRegistryFilename());
-	}
-	else
-	{
-		PreviousAssetRegistryFile = FPaths::Combine(MetadataDirectoryPath, GetDevelopmentAssetRegistryFilename());
-	}
-
 	PackageNameToCookedFiles.Reset();
 
-	TUniquePtr<FArchive> Reader(CookArtifactReader->CreateFileReader(*PreviousAssetRegistryFile));
-	if (!Reader)
-	{ 
-		RemoveCookedPackages();
-		return TUniquePtr<FAssetRegistryState>();
-	}
-
-	TUniquePtr<FAssetRegistryState> PreviousState = MakeUnique<FAssetRegistryState>();
-	PreviousState->Load(*Reader);
-
 	// If we are legacyiterative from a shared build the cooked files do not exist in the local cooked directory;
 	// we assume they are packaged in the pak file (which we don't want to extract to confirm) and keep them all.
 	if (!bLegacyIterativeSharedBuild)
 	{
-		// For regular iteration, remove each SuccessfulSave cooked file from the PreviousState if it no longer exists
-		// in the cooked directory. Keep the FailedSave previous cook packages; we don't expect them to exist on disk.
-		// Also, remove from the registry and from disk the cooked files that no longer exist in the editor.
+		// Remove from disk the cooked files that no longer exist in the editor.
 		GetAllCookedFiles();
-		TSet<FName> RemoveFromRegistry;
 		TSet<FName> RemoveFromDisk;
-		RemoveFromDisk.Reserve(PackageNameToCookedFiles.Num());
-		for (TPair<FName, TArray<FString>>& Pair : PackageNameToCookedFiles)
-		{
-			RemoveFromDisk.Add(Pair.Key);
-		}
 		IAssetRegistry& AssetRegistry = IAssetRegistry::GetChecked();
-		for (const TPair<FName, const FAssetPackageData*>& Pair : PreviousState->GetAssetPackageDataMap())
+		for (const TPair<FName, TArray<FString>>& Pair : PackageNameToCookedFiles)
 		{
 			FName PackageName = Pair.Key;
-			bool bCurrentPackageExists = AssetRegistry.DoesPackageExistOnDisk(PackageName);
 
-			bool bNoLongerExistsInEditor = false;
-			bool bIsScriptPackage = FPackageName::IsScriptPackage(WriteToString<256>(PackageName));
-			if (!bCurrentPackageExists)
+			if (!AssetRegistry.DoesPackageExistOnDisk(PackageName))
 			{
-				// Script and generated packages do not exist uncooked
-				// Check that the package is not an exception before removing from cooked
-				bool bIsCookedOnly = bIsScriptPackage;
+				// Check that the package is not an exception before removing from cooked.
+				// Generated packages do not exist uncooked.
+				bool bIsCookedOnly = false;
+				PreviousState.EnumerateAssetsByPackageName(PackageName, [&bIsCookedOnly](const FAssetData* AssetData)
+					{
+						bIsCookedOnly |= !!(AssetData->PackageFlags & PKG_CookGenerated);
+						return true; // Keep iterating
+					});
 				if (!bIsCookedOnly)
 				{
-					PreviousState->EnumerateAssetsByPackageName(PackageName, [&bIsCookedOnly](const FAssetData* AssetData)
-						{
-							bIsCookedOnly |= !!(AssetData->PackageFlags & PKG_CookGenerated);
-							return true; // Keep iterating
-						});
+					RemoveFromDisk.Add(PackageName);
 				}
-				bNoLongerExistsInEditor = !bIsCookedOnly;
-			}
-			if (bNoLongerExistsInEditor)
-			{
-				// Remove package from both disk and registry
-				// Keep its RemoveFromDisk entry
-				RemoveFromRegistry.Add(PackageName); // Add a RemoveFromRegistry entry
 			}
-			else
-			{
-				// Keep package on disk if it exists. Keep package in registry if it exists on disk or was a FailedSave.
-				bool bExistsOnDisk = (RemoveFromDisk.Remove(PackageName) == 1); // Remove its RemoveFromDisk entry
-				const FAssetPackageData* PackageData = Pair.Value;
-				if (!bExistsOnDisk && PackageData->DiskSize >= 0 && !bIsScriptPackage)
-				{
-					// Add RemoveFromRegistry entry if it didn't exist on disk and is a SuccessfulSave package
-					RemoveFromRegistry.Add(PackageName);
-				}
-			}
-		}
-
-		if (RemoveFromRegistry.Num())
-		{
-			PreviousState->PruneAssetData(TSet<FName>(), RemoveFromRegistry, FAssetRegistrySerializationOptions());
 		}
 		if (RemoveFromDisk.Num())
 		{
 			RemoveCookedPackagesByPackageName(RemoveFromDisk.Array(), true /* bRemoveRecords */);
 		}
+		OutNumPackagesInOplog = PackageNameToCookedFiles.Num();
 	}
+	else
+	{
+		// Clean the sandbox, since we are getting files from the SharedBuild location rather than from the local cook
+		// directory.
+		DeleteSandboxDirectory();
+		// If we are legacyiterative from a shared build the cooked files do not exist in the local cooked directory;
+		// we assume they are packaged in the pak file (which we don't want to extract to confirm) and keep them all.
+		OutNumPackagesInOplog = PreviousState.GetNumPackages();
+	}
+}
 
-	return PreviousState;
+void FLooseCookedPackageWriter::UpdateLastReferenceDateAndPruneStaleOps(
+	UE::Cook::Artifact::FUpdateOplogPackagesContext& OplogContext)
+{
+	// LooseCookedPackageWriter does not support setting LastReferenceDate and pruning stale ops; it only does
+	// pruning based on packages that no longer exist in editor. That more limited pruning was already done in
+	// PopulateOplog.
+}
+
+TArray<FName> FLooseCookedPackageWriter::GetOplogPackageNames()
+{
+	TArray<FName> Result;
+	PackageNameToCookedFiles.GenerateKeyArray(Result);
+	return Result;
 }
 
 FCbObject FLooseCookedPackageWriter::GetOplogAttachment(FName PackageName, FUtf8StringView AttachmentKey)
diff --git a/Engine/Source/Editor/UnrealEd/Private/Cooker/LooseCookedPackageWriter.h b/Engine/Source/Editor/UnrealEd/Private/Cooker/LooseCookedPackageWriter.h
--- a/Engine/Source/Editor/UnrealEd/Private/Cooker/LooseCookedPackageWriter.h
+++ b/Engine/Source/Editor/UnrealEd/Private/Cooker/LooseCookedPackageWriter.h
@@ -56,7 +56,10 @@ public:
 	virtual void Initialize(const FCookInfo& Info) override;
 	virtual void BeginCook(const FCookInfo& Info) override;
 	virtual void EndCook(const FCookInfo& Info) override;
-	virtual TUniquePtr<FAssetRegistryState> LoadPreviousAssetRegistry() override;
+	virtual void PopulateOplog(const FAssetRegistryState& PreviousState, int32& OutNumOplogPackages) override;
+	virtual void UpdateLastReferenceDateAndPruneStaleOps(
+		UE::Cook::Artifact::FUpdateOplogPackagesContext& OplogContext) override;
+	virtual TArray<FName> GetOplogPackageNames() override;
 	virtual FCbObject GetOplogAttachment(FName PackageName, FUtf8StringView AttachmentKey) override;
 	virtual void GetOplogAttachments(TArrayView<FName> PackageNames,
 		TArrayView<FUtf8StringView> AttachmentKeys,
diff --git a/Engine/Source/Editor/UnrealEd/Public/Commandlets/AssetRegistryGenerator.h b/Engine/Source/Editor/UnrealEd/Public/Commandlets/AssetRegistryGenerator.h
--- a/Engine/Source/Editor/UnrealEd/Public/Commandlets/AssetRegistryGenerator.h
+++ b/Engine/Source/Editor/UnrealEd/Public/Commandlets/AssetRegistryGenerator.h
@@ -8,6 +8,7 @@
 #include "Containers/Map.h"
 #include "Containers/Set.h"
 #include "Containers/UnrealString.h"
+#include "Cooker/CookArtifact.h"
 #include "Cooker/MPCollector.h"
 #include "Misc/AssetRegistryInterface.h"
 #include "Misc/Optional.h"
@@ -48,25 +49,10 @@ public:
 	 */
 	~FAssetRegistryGenerator();
 
-	/**
-	 * Initializes manifest generator - creates manifest lists, hooks up delegates.
-	 */
-	void Initialize(bool bInitializeFromExisting, const FAssetRegistrySerializationOptions& DevelopmentSaveOptions);
-	/**
-	 * Used by DLC which does not Clone the entire global AR. Copies data from the global assetregistry
-	 * for all packages present in the previous AssetRegistry
-	 */
-	void CloneGlobalAssetRegistryFilteredByPreviousState(const FAssetRegistryState& PreviousState,
-		const FAssetRegistrySerializationOptions& DevelopmentSaveOptions);
+	void Initialize();
 
 	const ITargetPlatform* GetTargetPlatform() const { return TargetPlatform; }
 
-	/**
-	 * Whether *this cloned the global the AssetRegistry during its construction.
-	 * If false, package data must be copied during UpdataAssetRegistryData.
-	 */
-	bool HasClonedGlobalAssetRegistry() const { return bClonedGlobalAssetRegistry; }
-
 	/**
 	 * Options when computing the differences between current and previous state.
 	 */
@@ -142,15 +128,21 @@ public:
 	void ComputePackageDifferences(const FComputeDifferenceOptions& Options, const FAssetRegistryState& PreviousState, 
 		FAssetRegistryDifference& OutDifference);
 
-	/** Computes just the list of packages in the PreviousState that no longer exist in the current state. */
-	void ComputePackageRemovals(const FAssetRegistryState& PreviousState, TArray<FName>& OutRemovedPackages,
+	/** Reads Generator packages and other cooking data from the previous incrementalcook's DevelopmentAssetRegistry. */
+	void ReportGeneratorPackages(const FAssetRegistryState& PreviousState,
 		TMap<FName, FGeneratorPackageInfo>& OutGeneratorPackages, int32& OutNumNeverCookPlaceHolderPackages);
 
 	/**
-	 * Return the list of packages from Previous IncrementalCooks the Generator received from the PackageWriter during
-	 * StartCookByTheBook.
+	 * Given the list of packages still in the oplog, remove any packages from the AssetRegistryState that are no longer
+	 * in the oplog (they were not referenced in the current session and were pruned for being too old or invalid).
 	 */
-	TArray<FName> GetPreviousIncrementalCookPackages();
+	void RemoveStalePackages(const TMap<FName, UE::Cook::Artifact::FOplogPackageData>& OplogPackages);
+
+	/**
+	 * Add or update data in the generated AssetRegistry for non-package data such as management references from
+	 * PrimaryAssets and /Script PackageDatas.
+	 */
+	void AddGlobalsToCookedAssetRegistry(UCookOnTheFlyServer& COTFS);
 
 	/**
 	 * FinalizeChunkIDs 
@@ -294,16 +286,6 @@ private:
 	 */
 	bool UpdateAssetPackageFlags(const FName& PackageName, const uint32 PackageFlags);
 
-	/**
-	 * Updates AssetData with previous TagsAndValues, and updates PackageData values with previous PackageData,
-	 * for all packages kept from a previous cook.
-	 * 
-	 * InOutAllPackages: Initially contains the list of packages that were referenced in the current cook. The function
-	 * appends to it to also include any packages that were present in a previous incremental cook so that we can keep
-	 * the information for them in case they are referenced again in a later cook.
-	 */
-	void UpdateKeptPackages(TSet<FName>& InOutAllPackages);
-
 	static UE::Cook::FAssetRegistryPackageMessage RecordUpdateAssetRegistryData(FName PackageName,
 		const ITargetPlatform* TargetPlatform, UE::Cook::ECookResult CookResult,
 		FSavePackageResultStruct* SavePackageResult,
@@ -320,7 +302,7 @@ private:
 		const FAssetPackageData& CurrentPackageData, const FAssetPackageData& PreviousPackageData);
 	/**
 	 * When copying AssetPackageData from the previous saved development assetregistry into the
-	 * FGeneratorPackageInfo returned from ComputePackageRemovals that are used for incremental cook,
+	 * FGeneratorPackageInfo returned from ComputePackageDifferences that are used for incremental cook,
 	 * we only need a subset of the data; we skip copying the portions we don't need to improve performance.
 	 */
 	static FAssetPackageData CopyAssetPackageDataForIncrementalCook(const FAssetPackageData& Source);
@@ -328,14 +310,6 @@ private:
 	/** State of the asset registry that is being built for this platform */
 	FAssetRegistryState State;
 	
-	struct FIterativelySkippedPackageUpdateData
-	{
-		TArray<FAssetData> AssetDatas;
-		FAssetPackageData PackageData;
-		TArray<FAssetDependency> PackageDependencies;
-		TArray<FAssetDependency> PackageReferencers;
-	};
-	TMap<FName, FIterativelySkippedPackageUpdateData> PreviousPackagesToUpdate;
 	/** Map of Package name to Sandbox Paths */
 	typedef TMap<FName, FString> FChunkPackageSet;
 	/** Holds a reference to asset registry */
@@ -350,7 +324,6 @@ private:
 	TSet<FName> PackagesContainingMaps;
 	/** Should the chunks be generated or only asset registry */
 	bool bGenerateChunks;
-	bool bClonedGlobalAssetRegistry;
 	/** Highest chunk id, being used for geneating dependency tree */
 	int32 HighestChunkId;
 	/** Array of Maps with chunks<->packages assignments */
diff --git a/Engine/Source/Runtime/Core/Public/Serialization/PackageWriter.h b/Engine/Source/Runtime/Core/Public/Serialization/PackageWriter.h
--- a/Engine/Source/Runtime/Core/Public/Serialization/PackageWriter.h
+++ b/Engine/Source/Runtime/Core/Public/Serialization/PackageWriter.h
@@ -3,6 +3,7 @@
 #pragma once
 
 #include "Async/Future.h"
+#include "Containers/ArrayView.h"
 #include "Containers/Set.h"
 #include "Containers/StringView.h"
 #include "IO/IoDispatcher.h"
@@ -25,6 +26,7 @@ class ICookedPackageWriter;
 class IPackageStoreWriter;
 class UObject;
 namespace UE::Cook { class IDeterminismHelper; }
+namespace UE::Cook::Artifact { struct FUpdateOplogPackagesContext; }
 struct FPackageStoreEntryResource;
 struct FSavePackageArgs;
 struct FSavePackageResultStruct;
@@ -414,12 +416,26 @@ public:
 	virtual void EndCook(const FCookInfo& Info) = 0;
 
 	/**
-	 * Returns an AssetRegistry describing the previous cook results. This doesn't mean a cook saved off
-	 * to another directory - it means the AssetRegistry that's living in the directory we are about
-	 * to cook in to.
+	 * Called from UCookOnTheFlyServer::PopulateCookedPackages. Load the oplog of previously cooked packages, if not
+	 * already done, and return the number of them.
 	 */
-	virtual TUniquePtr<FAssetRegistryState> LoadPreviousAssetRegistry() = 0;
-
+	virtual void PopulateOplog(const FAssetRegistryState& PreviousState, int32& OutNumPackagesInOplog) = 0;
+	/**
+	 * Called during CookByTheBookFinished. Remove Ops from the Oplog that have not been referenced recently
+	 * and should be removed to save space in the oplog and in artifacts that cache data about packages.
+	 * 
+	 * @param OplogContext Information that describes which packages were cooked or their previously cooked op
+	 *        inspected (by e.g. incremental cook) in the current cook. The PackageWriter should update the
+	 *        LastReferencedDate for these packages before checking for ops that have not recently been referenced.
+	 */
+	virtual void UpdateLastReferenceDateAndPruneStaleOps(
+		UE::Cook::Artifact::FUpdateOplogPackagesContext& OplogContext) = 0;
+	/**
+	 * Called during CookByTheBookFinished, after UpdateLastReferenceDateAndPruneStaleOps. Return the list of all
+	 * packages that have cooked data from the current or previous cooks. CookArtifacts use this to cull their data
+	 * when previously cooked packages are pruned.
+	 */
+	virtual TArray<FName> GetOplogPackageNames() = 0;
 	/**
 	 * Returns an Attachment that was previously commited for the given PackageName.
 	 * Returns an empty object if not found.
