Cooker: Add commandline option -cook.DependencyTracking=0 to turn off the collection of data used for IncrementalCook. This prevents IncrementalCook but saves time for farm jobs that never want to cook incrementally.
#rnx

#Synced-CL 45140074

diff --git a/Engine/Source/Editor/UnrealEd/Classes/CookOnTheSide/CookOnTheFlyServer.h b/Engine/Source/Editor/UnrealEd/Classes/CookOnTheSide/CookOnTheFlyServer.h
--- a/Engine/Source/Editor/UnrealEd/Classes/CookOnTheSide/CookOnTheFlyServer.h
+++ b/Engine/Source/Editor/UnrealEd/Classes/CookOnTheSide/CookOnTheFlyServer.h
@@ -1631,6 +1631,12 @@ private:
 	 * incrementally in the current mode; see PlatformManager.GetPlatformData().bFullCook for that property.
 	 */
 	bool bLegacyBuildDependencies = false;
+	/**
+	 * Automated DependencyTracking is only used in CookIncremental; legacyiterative does not use its data.
+	 * Additionally, it can be turned off to improve performance in CookIncremental, at the cost of not being able to
+	 * cook incrementally afterwards.
+	 */
+	bool bDependencyTrackingEnabled = true;
 	bool bCookIncrementalAllowAllClasses = false;
 	bool bHiddenDependenciesDebug = false;
 	bool bIncrementallyModifiedDiagnostics = false;
@@ -1810,3 +1816,15 @@ private:
 
 	friend void UE::Cook::CalculateGlobalDependenciesHash(const ITargetPlatform* Platform, const UCookOnTheFlyServer& COTFS);
 };
+
+namespace UE::Cook
+{
+
+/**
+ * Public function for systems that track dependencies for the cooker and need to disable the tracking as early as
+ * possible for performance when commandline or config settings disable CookIncremental dependency tracking. If this
+ * function returns false, systems can turn off dependency tracking or reporting to the cooker.
+ */
+UNREALED_API bool IsDependencyTrackingEnabled();
+
+}
\ No newline at end of file
diff --git a/Engine/Source/Editor/UnrealEd/Private/CookOnTheFlyServer.cpp b/Engine/Source/Editor/UnrealEd/Private/CookOnTheFlyServer.cpp
--- a/Engine/Source/Editor/UnrealEd/Private/CookOnTheFlyServer.cpp
+++ b/Engine/Source/Editor/UnrealEd/Private/CookOnTheFlyServer.cpp
@@ -10224,36 +10224,80 @@ void UCookOnTheFlyServer::LoadBeginCookConfigSettings(FBeginCookContext& BeginCo
 namespace UE::Cook
 {
 
-void FBeginCookConfigSettings::LoadLocal(FBeginCookContext& BeginContext)
+struct FDependencyTrackingConfigParser
+{
+	void Parse();
+
+	// Input parameters
+	bool bInIssueLogs = false;
+	TOptional<bool> bInUsingZenStore;
+
+	// Output environment settings
+	bool bOutLegacyBuildDependencies = false;
+	bool bOutDependencyTrackingEnabled = true;
+};
+
+void FDependencyTrackingConfigParser::Parse()
 {
 	const TCHAR* CommandLine = FCommandLine::Get();
-	bLegacyBuildDependencies = false;
-	GConfig->GetBool(TEXT("CookSettings"), TEXT("LegacyIterative"), bLegacyBuildDependencies, GEditorIni);
+	GConfig->GetBool(TEXT("CookSettings"), TEXT("LegacyIterative"), bOutLegacyBuildDependencies, GEditorIni);
 
-	if (!bLegacyBuildDependencies && FParse::Param(CommandLine, TEXT("ODSC")))
+	if (!bOutLegacyBuildDependencies && FParse::Param(CommandLine, TEXT("ODSC")))
 	{
 		// INCREMENTALCOOK_TODO: implement incremental cook for ODSC
-		UE_LOG(LogCook, Display,
+		UE_CLOG(bInIssueLogs, LogCook, Display,
 			TEXT("INCREMENTAL COOK DEPENDENCIES: Disabled. The cooker is running as ODSC which forces use of legacyiterative."));
-		bLegacyBuildDependencies = true;
+		bOutLegacyBuildDependencies = true;
 	}
-	else if (!bLegacyBuildDependencies && !BeginContext.COTFS.IsUsingZenStore())
+	// TODO: Parse bIsUsingZenStore in this function from config and commandline, so that we can have a more accurate result
+	// for bDependencyTrackingEnabled for callers that don't know the value of bInUsingZenStore.
+	else if (!bOutLegacyBuildDependencies && bInUsingZenStore.IsSet() && !*bInUsingZenStore)
 	{
 		// Incremental cook uses TargetDomain storage of dependencies which is only implemented in ZenStore
-		UE_LOG(LogCook, Display,
+		UE_CLOG(bInIssueLogs, LogCook, Display,
 			TEXT("INCREMENTAL COOK DEPENDENCIES: Disabled. Incremental Cook Dependencies are disabled because the cooker is not using ZenStore, ")
 			TEXT("and we do not yet support storage of Incremental Cook dependency data with loose cooked packages. ")
 			TEXT("Falling back to legacy iterative cook mode."));
-		bLegacyBuildDependencies = true;
+		bOutLegacyBuildDependencies = true;
 	}
-	else
+	else if (!bOutLegacyBuildDependencies)
 	{
-		UE_CLOG(!bLegacyBuildDependencies, LogCook, Display,
+		UE_CLOG(bInIssueLogs, LogCook, Display,
 			TEXT("INCREMENTAL COOK DEPENDENCIES: Enabled. Incremental Cook Dependencies are enabled in Editor.ini:[CookSettings]:LegacyIterative=false."));
-		UE_CLOG(bLegacyBuildDependencies, LogCook, Display,
+	}
+	else
+	{
+		UE_CLOG(bInIssueLogs, LogCook, Display,
 			TEXT("INCREMENTAL COOK DEPENDENCIES: Disabled. Incremental Cook Dependencies are disabled in Editor.ini:[CookSettings]:LegacyIterative=true."));
 	}
-		 
+
+	if (bOutLegacyBuildDependencies)
+	{
+		bOutDependencyTrackingEnabled = false;
+	}
+	else
+	{
+		bOutDependencyTrackingEnabled = true;
+		FParse::Bool(CommandLine, TEXT("-Cook.DependencyTracking="), bOutDependencyTrackingEnabled);
+	}
+}
+bool IsDependencyTrackingEnabled()
+{
+	FDependencyTrackingConfigParser Parser;
+	Parser.Parse();
+	return Parser.bOutDependencyTrackingEnabled;
+}
+
+void FBeginCookConfigSettings::LoadLocal(FBeginCookContext& BeginContext)
+{
+	FDependencyTrackingConfigParser Parser;
+	Parser.bInIssueLogs= true;
+	Parser.bInUsingZenStore.Emplace(BeginContext.COTFS.IsUsingZenStore());
+	Parser.Parse();
+	bLegacyBuildDependencies = Parser.bOutLegacyBuildDependencies;
+	bDependencyTrackingEnabled = Parser.bOutDependencyTrackingEnabled;
+
+	const TCHAR* CommandLine = FCommandLine::Get();
 	bCookIncrementalAllowAllClasses = FParse::Param(CommandLine, TEXT("CookIncrementalAllowAllClasses"));
 
 	FParse::Value(CommandLine, TEXT("-CookShowInstigator="), CookShowInstigator);
@@ -10278,6 +10322,7 @@ void FBeginCookConfigSettings::LoadLocal(FBeginCookContext& BeginContext)
 void UCookOnTheFlyServer::SetBeginCookConfigSettings(FBeginCookContext& BeginContext, UE::Cook::FBeginCookConfigSettings&& Settings)
 {
 	bLegacyBuildDependencies = Settings.bLegacyBuildDependencies;
+	bDependencyTrackingEnabled = Settings.bDependencyTrackingEnabled;
 	bCookIncrementalAllowAllClasses = Settings.bCookIncrementalAllowAllClasses;
 	PackageDatas->SetBeginCookConfigSettings(Settings.CookShowInstigator);
 	UE::Cook::FGenerationHelper::SetBeginCookConfigSettings();
@@ -10483,7 +10528,16 @@ void UCookOnTheFlyServer::LoadBeginCookIncrementalFlagsLocal(FBeginCookContext&
 					);
 				}
 			}
-			else
+
+			if (bIncrementalOrLegacyIterativeAllowed && !bLegacyBuildDependencies && !bDependencyTrackingEnabled)
+			{
+				UE_LOG(LogCook, Display,
+					TEXT("FULL COOK: -cookincremental was specified, but -Cook.DependencyTracking=0 was present on commandline and so cookincremental is not supported. Deleting previously cooked packages for platform %s and recooking all packages discovered in the current cook."),
+					*TargetPlatform->PlatformName());
+				bIncrementalOrLegacyIterativeAllowed = false;
+			}
+
+			if (bIncrementalOrLegacyIterativeAllowed)
 			{
 				UE::ConfigAccessTracking::FIgnoreScope IgnoreScope;
 				for (const TRefCountPtr<UE::Cook::ICookArtifact>& Artifact: CookArtifacts)
diff --git a/Engine/Source/Editor/UnrealEd/Private/Cooker/CookSavePackage.cpp b/Engine/Source/Editor/UnrealEd/Private/Cooker/CookSavePackage.cpp
--- a/Engine/Source/Editor/UnrealEd/Private/Cooker/CookSavePackage.cpp
+++ b/Engine/Source/Editor/UnrealEd/Private/Cooker/CookSavePackage.cpp
@@ -843,7 +843,7 @@ private:
 void FSaveCookedPackageContext::CalculateCookDependencies(FGenerationHelper* GenerationHelper, bool bGenerated,
 	const TArray<FAssetDependency>* ExtraArDependencies, FBuildResultDependenciesMap* ExtraBuildResultDependencies)
 {
-	if (COTFS.bLegacyBuildDependencies)
+	if (!COTFS.bDependencyTrackingEnabled)
 	{
 		return;
 	}
@@ -936,7 +936,7 @@ void FSaveCookedPackageContext::CalculateCookDependencies(FGenerationHelper* Gen
 
 void FSaveCookedPackageContext::RecordPlatformBuildDependencies()
 {
-	if (COTFS.bLegacyBuildDependencies)
+	if (!COTFS.bDependencyTrackingEnabled)
 	{
 		return;
 	}
diff --git a/Engine/Source/Editor/UnrealEd/Private/Cooker/CookTypes.cpp b/Engine/Source/Editor/UnrealEd/Private/Cooker/CookTypes.cpp
--- a/Engine/Source/Editor/UnrealEd/Private/Cooker/CookTypes.cpp
+++ b/Engine/Source/Editor/UnrealEd/Private/Cooker/CookTypes.cpp
@@ -878,6 +878,7 @@ void FInitializeConfigSettings::MoveToLocal(UCookOnTheFlyServer& COTFS)
 void FBeginCookConfigSettings::CopyFromLocal(const UCookOnTheFlyServer& COTFS)
 {
 	bLegacyBuildDependencies = COTFS.bLegacyBuildDependencies;
+	bDependencyTrackingEnabled = COTFS.bDependencyTrackingEnabled;
 	bCookIncrementalAllowAllClasses = COTFS.bCookIncrementalAllowAllClasses;
 	FParse::Value(FCommandLine::Get(), TEXT("-CookShowInstigator="), CookShowInstigator); // We don't store this on COTFS, so reparse it from commandLine
 	TSet<FName> COTFSNeverCookPackageList;
@@ -914,6 +915,7 @@ FCbWriter& operator<<(FCbWriter& Writer, const UE::Cook::FBeginCookConfigSetting
 {
 	Writer.BeginObject();
 	Writer << "LegacyBuildDependencies" << Value.bLegacyBuildDependencies;
+	Writer << "DependencyTrackingEnabled" << Value.bDependencyTrackingEnabled;
 	Writer << "CookIncrementalAllowAllClasses" << Value.bCookIncrementalAllowAllClasses;
 	Writer << "CookShowInstigator" << Value.CookShowInstigator;
 	Writer << "NeverCookPackageList" << Value.NeverCookPackageList;
@@ -936,6 +938,7 @@ bool LoadFromCompactBinary(FCbFieldView Field, UE::Cook::FBeginCookConfigSetting
 {
 	bool bOk = Field.IsObject();
 	bOk = LoadFromCompactBinary(Field["LegacyBuildDependencies"], OutValue.bLegacyBuildDependencies) & bOk;
+	bOk = LoadFromCompactBinary(Field["DependencyTrackingEnabled"], OutValue.bDependencyTrackingEnabled) & bOk;
 	bOk = LoadFromCompactBinary(Field["CookIncrementalAllowAllClasses"], OutValue.bCookIncrementalAllowAllClasses) & bOk;
 	bOk = LoadFromCompactBinary(Field["CookShowInstigator"], OutValue.CookShowInstigator) & bOk;
 	bOk = LoadFromCompactBinary(Field["NeverCookPackageList"], OutValue.NeverCookPackageList) & bOk;
diff --git a/Engine/Source/Editor/UnrealEd/Private/Cooker/CookTypes.h b/Engine/Source/Editor/UnrealEd/Private/Cooker/CookTypes.h
--- a/Engine/Source/Editor/UnrealEd/Private/Cooker/CookTypes.h
+++ b/Engine/Source/Editor/UnrealEd/Private/Cooker/CookTypes.h
@@ -499,6 +499,7 @@ namespace UE::Cook
 
 		FString CookShowInstigator;
 		bool bLegacyBuildDependencies = false;
+		bool bDependencyTrackingEnabled = true;
 		bool bCookIncrementalAllowAllClasses = false;
 		TArray<FName> NeverCookPackageList;
 		TMap<const ITargetPlatform*, TSet<FName>> PlatformSpecificNeverCookPackages;
diff --git a/Engine/Source/Editor/UnrealEd/Private/Cooker/GlobalCookArtifact.cpp b/Engine/Source/Editor/UnrealEd/Private/Cooker/GlobalCookArtifact.cpp
--- a/Engine/Source/Editor/UnrealEd/Private/Cooker/GlobalCookArtifact.cpp
+++ b/Engine/Source/Editor/UnrealEd/Private/Cooker/GlobalCookArtifact.cpp
@@ -56,6 +56,7 @@ FConfigFile FGlobalCookArtifact::CalculateCurrentSettings(ICookInfo& CookInfo, c
 
 	CookSettingStrings.Add(FName(TEXT("Version")), TEXT("21F52B9EDD4D456AB1AF381CA172BD28"));
 	CookSettingStrings.Add(FName(TEXT("LegacyBuildDependencies")), ::LexToString(COTFS.bLegacyBuildDependencies));
+	CookSettingStrings.Add(FName(TEXT("DependencyTrackingEnabled")), ::LexToString(COTFS.bDependencyTrackingEnabled));
 
 	if (!COTFS.bLegacyBuildDependencies)
 	{
diff --git a/Engine/Source/Editor/UnrealEd/Private/UnrealEdMisc.cpp b/Engine/Source/Editor/UnrealEd/Private/UnrealEdMisc.cpp
--- a/Engine/Source/Editor/UnrealEd/Private/UnrealEdMisc.cpp
+++ b/Engine/Source/Editor/UnrealEd/Private/UnrealEdMisc.cpp
@@ -98,6 +98,7 @@
 #include "DeveloperToolSettingsDelegates.h"
 #include "Cooker/CookConfigAccessTracker.h"
 #include "Cooker/PackageBuildDependencyTracker.h"
+#include "CookOnTheSide/CookOnTheFlyServer.h"
 #include "EditorAxisDisplayInfo.h"
 #include "ThumbnailRendering/SceneThumbnailInfo.h"
 #include "Factories/FbxAnimSequenceImportData.h"
@@ -233,7 +234,9 @@ FUnrealEdMisc::FUnrealEdMisc() :
 {
 	//This is an early entry-point into the UnrealEd module to perform some editor-specific configuration
 #if UE_WITH_PACKAGE_ACCESS_TRACKING || UE_WITH_CONFIG_TRACKING
-	const bool bBuildDependencyTrackingNeeded = GIsEditor && (IsRunningCookCommandlet() || !GetDefault<UEditorExperimentalSettings>()->bDisableCookInEditor);
+	const bool bBuildDependencyTrackingNeeded = GIsEditor &&
+		(IsRunningCookCommandlet() || !GetDefault<UEditorExperimentalSettings>()->bDisableCookInEditor)
+		&& UE::Cook::IsDependencyTrackingEnabled();
 	if (!bBuildDependencyTrackingNeeded)
 	{
 #if UE_WITH_PACKAGE_ACCESS_TRACKING
