EditorUsability : SettingsEditor
- Open both editor and project settings in "All settings" by default instead of "Appearances" and "General" since search is now localized
- Auto focus search box when opening settings

#Synced-CL 51208877

diff --git a/Engine/Source/Developer/SettingsEditor/Private/Widgets/SSettingsEditor.cpp b/Engine/Source/Developer/SettingsEditor/Private/Widgets/SSettingsEditor.cpp
--- a/Engine/Source/Developer/SettingsEditor/Private/Widgets/SSettingsEditor.cpp
+++ b/Engine/Source/Developer/SettingsEditor/Private/Widgets/SSettingsEditor.cpp
@@ -231,6 +231,21 @@ void SSettingsEditor::Construct( const FArguments& InArgs, const ISettingsEditor
 	{
 		SettingsEditorModule->OnSearchTermsChanged().AddSP(this, &SSettingsEditor::OnSearchTermsChanged);
 	}
+
+	// Focus search box next tick
+	FTSTicker::GetCoreTicker().AddTicker(
+		FTickerDelegate::CreateSPLambda(
+			this, 
+			[this](float)
+			{
+				if (FSlateApplication::IsInitialized() && SearchBox.IsValid())
+				{
+					FSlateApplication::Get().SetKeyboardFocus(SearchBox, EFocusCause::SetDirectly);
+				}
+				return false; // Stop
+			}
+		)
+	);
 }
 
 /* FNotifyHook interface
diff --git a/Engine/Source/Editor/UnrealEd/Private/Toolkits/GlobalEditorCommonCommands.cpp b/Engine/Source/Editor/UnrealEd/Private/Toolkits/GlobalEditorCommonCommands.cpp
--- a/Engine/Source/Editor/UnrealEd/Private/Toolkits/GlobalEditorCommonCommands.cpp
+++ b/Engine/Source/Editor/UnrealEd/Private/Toolkits/GlobalEditorCommonCommands.cpp
@@ -136,34 +136,26 @@ void FGlobalEditorCommonCommands::MapActions(TSharedRef<FUICommandList>& Toolkit
 		FExecuteAction::CreateStatic(&FGlobalEditorCommonCommands::OnOpenFindInAllBlueprints));
 	
 	// Get editor default category and section
-	FName CategoryName;
-	FName SectionName;
+	FName CategoryName = NAME_None; // "All settings"
+	FName SectionName = NAME_None; // "All settings"
 	if (IMainFrameModule* MainFrameModule = FModuleManager::GetModulePtr<IMainFrameModule>("MainFrame"))
 	{
 		MainFrameModule->GetEditorSettingsDefaultSelectionOverride(CategoryName, SectionName);
-		
-		if (CategoryName.IsNone())
-		{
-			CategoryName = FName("General");
-		}
-		
-		if (SectionName.IsNone())
-		{
-			SectionName = FName("Appearance");
-		}
 	}
 
 	ToolkitCommands->MapAction(
 		GlobalEditorCommands.OpenEditorPreferences,
 		FExecuteAction::CreateStatic(&FGlobalEditorCommonCommands::OpenSettings, FName("Editor"), CategoryName, SectionName));
-	
+
 	ToolkitCommands->MapAction(
 		GlobalEditorCommands.OpenKeyboardShortcuts,
 		FExecuteAction::CreateStatic(&FGlobalEditorCommonCommands::OpenSettings, FName("Editor"), FName("Shortcuts"), FName("InputBindings")));
-	
+
+	CategoryName = NAME_None; // "All settings"
+	SectionName = NAME_None; // "All settings"
 	ToolkitCommands->MapAction(
 		GlobalEditorCommands.OpenProjectSettings,
-		FExecuteAction::CreateStatic(&FGlobalEditorCommonCommands::OpenSettings, FName("Project"), FName("Project"), FName("General")));
+		FExecuteAction::CreateStatic(&FGlobalEditorCommonCommands::OpenSettings, FName("Project"), CategoryName, SectionName));
 }
 
 void FGlobalEditorCommonCommands::OnPressedCtrlTab(TSharedPtr<FUICommandInfo> TriggeringCommand)
