[ITF Viewport Interactions] Fix an issue where alt-based mouse movements were not fully re-enabled in the legacy viewport interaction path.

#rnx

#Synced-CL 51211095

diff --git a/Engine/Source/Editor/UnrealEd/Private/Tools/EdModeInteractiveToolsContext.cpp b/Engine/Source/Editor/UnrealEd/Private/Tools/EdModeInteractiveToolsContext.cpp
--- a/Engine/Source/Editor/UnrealEd/Private/Tools/EdModeInteractiveToolsContext.cpp
+++ b/Engine/Source/Editor/UnrealEd/Private/Tools/EdModeInteractiveToolsContext.cpp
@@ -1547,12 +1547,6 @@ bool UModeManagerInteractiveToolsContext::OnChildEdModeDeactivated(UEdModeIntera
 
 void UModeManagerInteractiveToolsContext::SetBypassAltModifier(bool bInBypassAltModifier)
 {
-	if (bInBypassAltModifier
-		&& (UE::Editor::ViewportInteractions::UseEditorViewportInteractions() || UE::Editor::DragTools::UseEditorDragTools()))
-	{
-		return;
-	}
-
 	bBypassAltModifier = bInBypassAltModifier;
 }
 
