World Partition - HLOD: Fix distributed build revert failure due to path mismatch in CopyFilesFromWorkingDir()

CopyFilesFromWorkingDir() reconstructed file paths using FPaths::ProjectDir() which can be relative, while FSourceControlHelper tracks files using PackageFilename() which always returns absolute paths via ConvertRelativePathToFull(). When an HLOD actor was actually built, Save() would add the absolute path to ModifiedFiles, but the Contains() check against the relative path in FileAdded would fail, causing the file to appear in both FileAdded and FileEdited. The subsequent RevertFiles call would then attempt to revert the same depot file twice, failing with "file(s) not opened on this client".

Fix: normalize the reconstructed path to absolute via ConvertRelativePathToFull(), consistent with what SourceControlHelper uses.

#Synced-CL 51228359

diff --git a/Engine/Source/Editor/UnrealEd/Private/WorldPartition/WorldPartitionHLODsBuilder.cpp b/Engine/Source/Editor/UnrealEd/Private/WorldPartition/WorldPartitionHLODsBuilder.cpp
--- a/Engine/Source/Editor/UnrealEd/Private/WorldPartition/WorldPartitionHLODsBuilder.cpp
+++ b/Engine/Source/Editor/UnrealEd/Private/WorldPartition/WorldPartitionHLODsBuilder.cpp
@@ -1647,7 +1647,7 @@ bool UWorldPartitionHLODsBuilder::CopyFilesFromWorkingDir(const FString& SourceD
 		}
 
 		FPaths::MakePathRelativeTo(PathRelativeToWorkingDir, *(FileActionString / TEXT("")));
-		FString FullPathInProjectDirectory = FPaths::ProjectDir() / PathRelativeToWorkingDir;
+		FString FullPathInProjectDirectory = FPaths::ConvertRelativePathToFull(FPaths::ProjectDir() / PathRelativeToWorkingDir);
 
 		FName FileAction(FileActionString);
 		if (FileAction == FileAction_Add)
@@ -1893,7 +1893,7 @@ EHLODRebuildPolicyDecision UWorldPartitionHLODsBuilder::EvaluateHLODBuildConditi
 		FScopedLogSuppression VerbosityScope(TEXT("LogLinker"));
 
 		FScopedDurationTimer TimerLoad(EvaluationStats.TimeLoadParent);
-		
+
 		const FPackagePath TempPackagePath = FPackagePath::FromLocalPath(LocalFilePathParent);
 		const FPackagePath AssetPath = FPackagePath::FromLocalPath(LocalFilePathCurrent);
 		if (UPackage* TempPackage = DiffUtils::LoadPackageForDiff(TempPackagePath, AssetPath))
