Make "new-process" PIE work like "Preview Platform"
- Mobile "new-process" PIE now respects the Preview Platform set, the Device Profile and the Aspect ratio(flipped, unflipped) of the main editor viewport.

#Synced-CL 51259166

diff --git a/Engine/Source/Editor/LevelEditor/Private/LevelEditorActions.cpp b/Engine/Source/Editor/LevelEditor/Private/LevelEditorActions.cpp
--- a/Engine/Source/Editor/LevelEditor/Private/LevelEditorActions.cpp
+++ b/Engine/Source/Editor/LevelEditor/Private/LevelEditorActions.cpp
@@ -604,7 +604,7 @@ void FLevelEditorActionCallbacks::PreviewJson_Clicked(FName PlatformName, FStrin
 					{
 						FName PreviewShaderPlatformName = FDataDrivenShaderPlatformInfo::GetName(PreviewShaderPlatform);
 
-						FName FriendlyDeviceProfilePreviewName = FName(FDataDrivenShaderPlatformInfo::GetFriendlyName(PreviewShaderPlatform).ToString() + TEXT(" ") + DevileProfileName);
+						FName FriendlyDeviceProfilePreviewName = FName(FDataDrivenShaderPlatformInfo::GetFriendlyName(PreviewShaderPlatform).ToString() + TEXT(" DeviceProfile:") + DevileProfileName);
 
 						UDeviceProfile* DevileProfile = UDeviceProfileManager::Get().FindProfile(DevileProfileName, false);
 						if (DevileProfile)
diff --git a/Engine/Source/Editor/UnrealEd/Private/PlayLevelNewProcess.cpp b/Engine/Source/Editor/UnrealEd/Private/PlayLevelNewProcess.cpp
--- a/Engine/Source/Editor/UnrealEd/Private/PlayLevelNewProcess.cpp
+++ b/Engine/Source/Editor/UnrealEd/Private/PlayLevelNewProcess.cpp
@@ -151,6 +151,17 @@ void UEditorEngine::LaunchNewProcess(const FRequestPlaySessionParams& InParams,
 		// Ensure the executable writes out a differently named config file to avoid multiple instances overwriting each other.
 		// ToDo: Should this be on all multi-client launches?
 		CommandLine += TEXT(" -MultiprocessSaveConfig");
+
+		if(FDataDrivenShaderPlatformInfo::IsValid(PreviewPlatform.GetShaderPlatform()))
+		{
+			// Override the ShaderPlatform if the PreviewPlatform is set in the main editor viewport.
+			CommandLine += FString::Printf(TEXT(" OVERRIDESP=%s"), *PreviewPlatform.PreviewShaderPlatformName.ToString());
+			if (PreviewPlatform.DeviceProfileName != NAME_None)
+			{
+				// Override the DeviceProfile if the DeviceProfile is set in the main editor viewport for preview.
+				CommandLine += FString::Printf(TEXT(" dp.Override=%s"), *PreviewPlatform.DeviceProfileName.ToString());
+			}
+		}
 	}
 	// If we are not launching a mobile preview, them we need to check the editor/project settings for using mouse touch as well.
 	else
@@ -291,6 +302,32 @@ void UEditorEngine::LaunchNewProcess(const FRequestPlaySessionParams& InParams,
 			WindowSize.Y = DisplayMetrics.PrimaryDisplayHeight;
 		}
 
+		// Update the WindowSize to the Constrained Aspect Ratio if its set in the main editor viewport for preview.
+		if (InParams.SessionPreviewTypeOverride.Get(EPlaySessionPreviewType::NoPreview) == EPlaySessionPreviewType::MobilePreview)
+		{
+			if (FDataDrivenShaderPlatformInfo::IsValid(PreviewPlatform.GetShaderPlatform()))
+			{
+				float PreviewConstrainedAspectRatio = PreviewPlatform.GetConstrainedAspectRatio();
+				if (PreviewConstrainedAspectRatio != 0.f)
+				{
+					if (GEditor->GetFlipConstrainedAspectRatioInPreview())
+					{
+						PreviewConstrainedAspectRatio = 1/ PreviewConstrainedAspectRatio;
+					}
+					float CurrentAspectRatio = float(WindowSize.X) / float(WindowSize.Y);
+
+					if (CurrentAspectRatio > PreviewConstrainedAspectRatio)
+					{
+						WindowSize.X = WindowSize.Y * PreviewConstrainedAspectRatio;
+					}
+					else
+					{
+						WindowSize.Y = WindowSize.X / PreviewConstrainedAspectRatio;
+					}
+				}
+			}
+		}
+
 		CommandLine += FString::Printf(TEXT(" -ResX=%d -ResY=%d"), WindowSize.X, WindowSize.Y);
 
 		// If they request a size larger than their display, add -ForceRes to prevent the engine
diff --git a/Engine/Source/Runtime/Engine/Private/GameEngine.cpp b/Engine/Source/Runtime/Engine/Private/GameEngine.cpp
--- a/Engine/Source/Runtime/Engine/Private/GameEngine.cpp
+++ b/Engine/Source/Runtime/Engine/Private/GameEngine.cpp
@@ -525,7 +525,31 @@ TSharedRef<SWindow> UGameEngine::CreateGameWindow()
 	FFormatNamedArguments Args;
 	Args.Add( TEXT("GameName"), FText::FromString( FApp::GetProjectName() ) );
 	Args.Add( TEXT("PlatformArchitecture"), PlatformBits );
-	Args.Add( TEXT("RHIName"), FText::FromName( LegacyShaderPlatformToShaderFormat( GMaxRHIShaderPlatform ) ) );
+
+	FText RHIName = FText::FromName(LegacyShaderPlatformToShaderFormat( GMaxRHIShaderPlatform ));
+#if WITH_EDITOR
+	FString ShaderPlatformName;
+	// If we are overriding the ShaderPlatform, update the RHIName
+	if (FParse::Value(FCommandLine::Get(), TEXT("OverrideSP="), ShaderPlatformName) && !ShaderPlatformName.IsEmpty())
+	{
+		EShaderPlatform OverrideShaderPlatform = FDataDrivenShaderPlatformInfo::GetShaderPlatformFromName(*ShaderPlatformName);
+		if (OverrideShaderPlatform != SP_NumPlatforms)
+		{
+			if (FDataDrivenShaderPlatformInfo::GetIsPreviewPlatform(OverrideShaderPlatform))
+			{
+				FString DeviceProfileName;
+				FString FriendlyDeviceProfilePreviewName = FDataDrivenShaderPlatformInfo::GetFriendlyName(OverrideShaderPlatform).ToString();
+				// If we are overriding the DeviceProfile, append the DP name to the RHIName.
+				if(FParse::Value(FCommandLine::Get(), TEXT("dp.Override="), DeviceProfileName))
+				{
+					FriendlyDeviceProfilePreviewName += TEXT(" DeviceProfile:") + DeviceProfileName;
+				}
+				RHIName = FText::FromName(FName(FriendlyDeviceProfilePreviewName));
+			}
+		}
+	}
+#endif
+	Args.Add( TEXT("RHIName"), RHIName);
 	Args.Add( TEXT("BuildConfiguration"), FText::FromString(LexToString(FApp::GetBuildConfiguration()) ) );
 
 	/************************************************************************/
