Cooker DLC Cooks: Fix path to plugin directory for plugins that are not in a direct subdirectory of <Project>\Plugins;.
#rnx

#Synced-CL 51302343

diff --git a/Engine/Source/Editor/UnrealEd/Private/Commandlets/AssetRegistryGenerator.cpp b/Engine/Source/Editor/UnrealEd/Private/Commandlets/AssetRegistryGenerator.cpp
--- a/Engine/Source/Editor/UnrealEd/Private/Commandlets/AssetRegistryGenerator.cpp
+++ b/Engine/Source/Editor/UnrealEd/Private/Commandlets/AssetRegistryGenerator.cpp
@@ -419,10 +419,10 @@ static void AssignLayerChunkDelegate(const FAssignLayerChunkMap* ChunkManifest,
 }
 
 bool FAssetRegistryGenerator::GenerateStreamingInstallManifest(int64 InOverrideChunkSize, const TCHAR* InManifestSubDir,
-	UE::Cook::FCookSandbox& InSandboxFile)
+	UE::Cook::FCookSandbox& InSandboxFile, UCookOnTheFlyServer& COTFS)
 {
 	const FString Platform = TargetPlatform->PlatformName();
-	FString ChunkManifestDir = GetChunkManifestDirectoryForPlatform(Platform, InSandboxFile);
+	FString ChunkManifestDir = GetChunkManifestDirectory(COTFS);
 	if (InManifestSubDir)
 	{
 		ChunkManifestDir /= InManifestSubDir;
@@ -1010,10 +1010,10 @@ void FAssetRegistryGenerator::CalculateChunkIdsAndAssignToManifest(const FName&
 	}
 }
 
-bool FAssetRegistryGenerator::CleanManifestDirectories(UE::Cook::FCookSandbox& InSandboxFile)
+bool FAssetRegistryGenerator::CleanManifestDirectories(UE::Cook::FCookSandbox& InSandboxFile, UCookOnTheFlyServer& COTFS)
 {
 	LLM_SCOPE_BYTAG(Cooker_GeneratedAssetRegistry);
-	FString ChunkManifestDir = GetChunkManifestDirectoryForPlatform(TargetPlatform->PlatformName(), InSandboxFile);
+	FString ChunkManifestDir = GetChunkManifestDirectory(COTFS);
 	if (IFileManager::Get().DirectoryExists(*ChunkManifestDir))
 	{
 		if (!IFileManager::Get().DeleteDirectory(*ChunkManifestDir, false, true))
@@ -1119,8 +1119,8 @@ void FAssetRegistryGenerator::InjectEncryptionData(FAssetRegistryState& TargetSt
 		});
 }
 
-bool FAssetRegistryGenerator::SaveManifests(UE::Cook::FCookSandbox& InSandboxFile, int64 InOverrideChunkSize,
-	const TCHAR* InManifestSubDir)
+bool FAssetRegistryGenerator::SaveManifests(UE::Cook::FCookSandbox& InSandboxFile, UCookOnTheFlyServer& COTFS,
+	int64 InOverrideChunkSize, const TCHAR* InManifestSubDir)
 {
 	LLM_SCOPE_BYTAG(Cooker_GeneratedAssetRegistry);
 	if (!bGenerateChunks)
@@ -1128,7 +1128,7 @@ bool FAssetRegistryGenerator::SaveManifests(UE::Cook::FCookSandbox& InSandboxFil
 		return true;
 	}
 
-	if (!GenerateStreamingInstallManifest(InOverrideChunkSize, InManifestSubDir, InSandboxFile))
+	if (!GenerateStreamingInstallManifest(InOverrideChunkSize, InManifestSubDir, InSandboxFile, COTFS))
 	{
 		return false;
 	}
@@ -2696,22 +2696,9 @@ void FAssetRegistryGenerator::AddPackageToChunk(FChunkPackageSet& ThisPackageSet
 	ThisPackageSet.Add(InPkgName, InSandboxFile);
 }
 
-FString FAssetRegistryGenerator::GetChunkManifestDirectoryForPlatform(const FString& Platform, UE::Cook::FCookSandbox& InSandboxFile) const
+FString FAssetRegistryGenerator::GetChunkManifestDirectory(UCookOnTheFlyServer& COTFS) const
 {
-	static FString DLCName = []()
-		{
-			FString DLCNameReturn;
-			FParse::Value(FCommandLine::Get(), TEXT("DLCNAME="), DLCNameReturn);
-			return DLCNameReturn;
-		}();
-	if (DLCName.IsEmpty())
-	{
-		return InSandboxFile.GetSandboxDirectory(Platform) / FApp::GetProjectName() / TEXT("Metadata") / TEXT("ChunkManifest");
-	}
-	else
-	{
-		return InSandboxFile.GetSandboxDirectory(Platform) / FApp::GetProjectName() / TEXT("Plugins") / DLCName / TEXT("Metadata") / TEXT("ChunkManifest");
-	}
+	return COTFS.GetCookMetadataOutputFolder(TargetPlatform) / TEXT("ChunkManifest");
 }
 
 void FAssetRegistryGenerator::FixupPackageDependenciesForChunks(UE::Cook::FCookSandbox& InSandboxFile)
diff --git a/Engine/Source/Editor/UnrealEd/Private/CookOnTheFlyServer.cpp b/Engine/Source/Editor/UnrealEd/Private/CookOnTheFlyServer.cpp
--- a/Engine/Source/Editor/UnrealEd/Private/CookOnTheFlyServer.cpp
+++ b/Engine/Source/Editor/UnrealEd/Private/CookOnTheFlyServer.cpp
@@ -9587,7 +9587,7 @@ void UCookOnTheFlyServer::CookByTheBookFinishedInternal()
 			}
 			{
 				UE_SCOPED_HIERARCHICAL_COOKTIMER(SaveManifests);
-				if (!Generator.SaveManifests(*SandboxFile))
+				if (!Generator.SaveManifests(*SandboxFile, *this))
 				{
 					UE_LOG(LogCook, Warning, TEXT("Failed to save chunk manifest"));
 				}
@@ -9598,7 +9598,7 @@ void UCookOnTheFlyServer::CookByTheBookFinishedInternal()
 					// ExtraFlavor is a legacy term for this override; etymology unknown. Override the chunksize specified by the platform,
 					// and write the manifest files created with that chunksize into a separate subdirectory.
 					const TCHAR* ManifestSubDir = TEXT("ExtraFlavor");
-					if (!Generator.SaveManifests(*SandboxFile, ExtraFlavorChunkSize, ManifestSubDir))
+					if (!Generator.SaveManifests(*SandboxFile, *this, ExtraFlavorChunkSize, ManifestSubDir))
 					{
 						UE_LOG(LogCook, Warning, TEXT("Failed to save chunk manifest"));
 					}
@@ -10652,7 +10652,7 @@ void UCookOnTheFlyServer::BeginCookSandbox(FBeginCookContext& BeginContext)
 				checkf(SandboxFile, TEXT("Cannot begin cooking to a sandbox until after CreateSandboxFile has been called from a StartCook function."));
 				// Clean the Manifest directory even on incremental cooks; it is written from scratch each time
 				// But only do this if we own the output directory
-				PlatformData->RegistryGenerator->CleanManifestDirectories(*SandboxFile);
+				PlatformData->RegistryGenerator->CleanManifestDirectories(*SandboxFile, *this);
 			}
 
 			if (PlatformContext.bPopulateMemoryResultsFromDiskResults)
diff --git a/Engine/Source/Editor/UnrealEd/Public/Commandlets/AssetRegistryGenerator.h b/Engine/Source/Editor/UnrealEd/Public/Commandlets/AssetRegistryGenerator.h
--- a/Engine/Source/Editor/UnrealEd/Public/Commandlets/AssetRegistryGenerator.h
+++ b/Engine/Source/Editor/UnrealEd/Public/Commandlets/AssetRegistryGenerator.h
@@ -200,19 +200,18 @@ public:
 	/**
 	 * Deletes temporary manifest directories.
 	 */
-	bool CleanManifestDirectories(UE::Cook::FCookSandbox& InSandboxFile);
+	bool CleanManifestDirectories(UE::Cook::FCookSandbox& InSandboxFile, UCookOnTheFlyServer& COTFS);
 
 	/**
 	 * Saves all generated manifests for each target platform.
 	 * 
-	 * @param InSandboxFile the InSandboxFile used during cook
 	 * @param InOverrideChunkSize the ChunkSize used during chunk division.
 	 *        If greater than 0, this overrides the default chunksize derived from the platform.
 	 * @param InManifestSubDir If non-null, the manifests are written into this subpath
 	 *        of the usual location.
 	 */
-	bool SaveManifests(UE::Cook::FCookSandbox& InSandboxFile, int64 InOverrideChunkSize = 0,
-		const TCHAR* InManifestSubDir = nullptr);
+	bool SaveManifests(UE::Cook::FCookSandbox& InSandboxFile, UCookOnTheFlyServer& COTFS,
+		int64 InOverrideChunkSize = 0, const TCHAR* InManifestSubDir = nullptr);
 
 	/**
 	* Saves generated asset registry data for each platform.
@@ -399,10 +398,8 @@ private:
 	void AddPackageToChunk(FChunkPackageSet& ThisPackageSet, FName InPkgName,
 		const FString& InSandboxFile, int32 PakchunkIndex, UE::Cook::FCookSandbox& SandboxPlatformFile);
 
-	/**
-	 * Returns the path of the chunk manifest directory for the specified platform.
-	 */
-	FString GetChunkManifestDirectoryForPlatform(const FString& Platform, UE::Cook::FCookSandbox& InSandboxFile) const;
+	/** The output path for this generator's platform where chunk manifests are written for use during staging. */
+	FString GetChunkManifestDirectory(UCookOnTheFlyServer& COTFS) const;
 
 	/** Returns the config-driven max size of a chunk for the given platform, or -1 for no limit. */
 	int64 GetMaxChunkSizePerPlatform( const ITargetPlatform* Platform ) const;
@@ -426,7 +423,7 @@ private:
 
 	/** Generates and saves streaming install chunk manifest */
 	bool GenerateStreamingInstallManifest(int64 InOverrideChunkSize, const TCHAR* InManifestSubDir,
-		UE::Cook::FCookSandbox& InSandboxFile);
+		UE::Cook::FCookSandbox& InSandboxFile, UCookOnTheFlyServer& COTFS);
 
 	/** Gather a list of dependencies required by to completely load this package */
 	bool GatherAllPackageDependencies(FName PackageName, TArray<FName>& DependentPackageNames);
