[Analytics][Cook] Adding IsIncrementalCook to the standard cook telemetry attributes.

#Synced-CL 51285544

diff --git a/Engine/Source/Editor/UnrealEd/Private/CookOnTheFlyServer.cpp b/Engine/Source/Editor/UnrealEd/Private/CookOnTheFlyServer.cpp
--- a/Engine/Source/Editor/UnrealEd/Private/CookOnTheFlyServer.cpp
+++ b/Engine/Source/Editor/UnrealEd/Private/CookOnTheFlyServer.cpp
@@ -10605,6 +10605,8 @@ void UCookOnTheFlyServer::LoadBeginCookIncrementalFlagsLocal(FBeginCookContext&
 		PlatformData->bLegacyIterativeSharedBuild = PlatformContext.bLegacyIterativeSharedBuild;
 		PlatformData->bWorkerOnSharedSandbox = PlatformContext.bWorkerOnSharedSandbox;
 
+		COOK_STAT(DetailedCookStats::IsIncrementalCook = IsIncremental());
+		
 		if (PlatformData->bFullBuild && !bFullBuildAllowed)
 		{
 			UE_LOG(LogCook, Fatal,
diff --git a/Engine/Source/Editor/UnrealEd/Private/Cooker/CookProfiling.cpp b/Engine/Source/Editor/UnrealEd/Private/Cooker/CookProfiling.cpp
--- a/Engine/Source/Editor/UnrealEd/Private/Cooker/CookProfiling.cpp
+++ b/Engine/Source/Editor/UnrealEd/Private/Cooker/CookProfiling.cpp
@@ -223,6 +223,7 @@ namespace DetailedCookStats
 	FString TargetPlatforms;
 	bool IsCookAll = false;
 	bool IsCookOnTheFly = false;
+	bool IsIncrementalCook = false;
 	bool IsIterativeCook = false;
 	bool IsFastCook = false;
 	bool IsUnversioned = false;
@@ -275,6 +276,7 @@ namespace DetailedCookStats
 					COOK_PROFILE_STAT_TEXT(TargetPlatforms),
 					COOK_PROFILE_STAT_TEXT(IsCookAll),
 					COOK_PROFILE_STAT_TEXT(IsCookOnTheFly),
+					COOK_PROFILE_STAT_TEXT(IsIncrementalCook),
 					COOK_PROFILE_STAT_TEXT(IsIterativeCook),
 					COOK_PROFILE_STAT_TEXT(IsFastCook),
 					COOK_PROFILE_STAT_TEXT(IsUnversioned)
diff --git a/Engine/Source/Editor/UnrealEd/Private/Cooker/CookProfilingPrivate.h b/Engine/Source/Editor/UnrealEd/Private/Cooker/CookProfilingPrivate.h
--- a/Engine/Source/Editor/UnrealEd/Private/Cooker/CookProfilingPrivate.h
+++ b/Engine/Source/Editor/UnrealEd/Private/Cooker/CookProfilingPrivate.h
@@ -139,6 +139,7 @@ extern FString CookLabel;
 extern FString TargetPlatforms;
 extern bool IsCookAll;
 extern bool IsCookOnTheFly;
+extern bool IsIncrementalCook;
 extern bool IsIterativeCook;
 extern bool IsFastCook;
 extern bool IsUnversioned;
