Cooker: CookPackageSplitter: Add more feedback for cook error "CookPackageSplitter failure: We are demoting a %s package from save...".
#rnx

#Synced-CL 51330552

diff --git a/Engine/Source/Editor/UnrealEd/Private/Cooker/CookGenerationHelper.cpp b/Engine/Source/Editor/UnrealEd/Private/Cooker/CookGenerationHelper.cpp
--- a/Engine/Source/Editor/UnrealEd/Private/Cooker/CookGenerationHelper.cpp
+++ b/Engine/Source/Editor/UnrealEd/Private/Cooker/CookGenerationHelper.cpp
@@ -2006,13 +2006,40 @@ void FGenerationHelper::ResetSaveState(FCookGenerationInfo& Info, UPackage* Pack
 			|| ReleaseSaveReason == EStateChangeReason::DoneForNow || ReleaseSaveReason == EStateChangeReason::Retraction;
 		if (!bDemotionAllowed)
 		{
+			FString DeletedObjectInOuter(TEXT("<None>"));
+			int32 NumNulls = 0;
+			for (const TPair<UObject*, FCachedObjectInOuterGeneratorInfo>& Pair : Info.CachedObjectsInOuterInfo)
+			{
+				bool bStillPresent = Info.PackageData->GetCachedObjectsInOuter().ContainsByPredicate(
+						[RawPtr = Pair.Key](const FCachedObjectInOuter& CachedObject)
+						{
+							return CachedObject.Object.Get() == RawPtr;
+						});
+				if (!bStillPresent)
+				{
+					if (NumNulls == 0)
+					{
+						DeletedObjectInOuter = Pair.Value.FullName;
+					}
+					++NumNulls;
+				}
+			}
+			UPackage* PackagePtr = Info.PackageData->GetPackage();
+
 			UE_LOG(LogCook, Error,
 				TEXT("CookPackageSplitter failure: We are demoting a %s package from save and removing our references that keep its objects loaded.\n")
 				TEXT("This will allow the objects to be garbage collected and cause failures in the splitter which expects them to remain loaded.\n")
-				TEXT("Package=%s, Splitter=%s, ReleaseSaveReason=%s, NewState=%s"),
+				TEXT("    Package=%s, Splitter=%s, ReleaseSaveReason=%s, CurrentState=%s, CurrentSubState=%s, NewState=%s,\n")
+				TEXT("    PackagePtr=%s, Package->IsFullyLoaded()=%s,\n")
+				TEXT("    GetCachedObjectsInOuter->Num()==%d, Num nulls in GetCachedObjectsInOuter==%d, First null=='%s'"),
 				Info.IsGenerator() ? TEXT("generator") : TEXT("generated"),
 				*Info.GetPackageName(),
-				*GetSplitDataObjectName().ToString(), LexToString(ReleaseSaveReason), LexToString(NewState));
+				*GetSplitDataObjectName().ToString(), LexToString(ReleaseSaveReason),
+				LexToString(Info.PackageData->GetState()), LexToString(Info.PackageData->GetSaveSubState()),
+				LexToString(NewState), PackagePtr == nullptr ? TEXT("<null>") : TEXT("<nonnull>"),
+				*(::LexToString(PackagePtr == nullptr ? false : PackagePtr->IsFullyLoaded())),
+				Info.PackageData->GetCachedObjectsInOuter().Num(),
+				NumNulls, *DeletedObjectInOuter);
 			FDebug::DumpStackTraceToLog(ELogVerbosity::Display);
 		}
 		Info.CachedObjectsInOuterInfo.Empty();
