Cooker: CookPackageSplitter: Add OwnerPackage feedback for cook error "CookPackageSplitter failure: We are demoting a %s package from save...".
#rnx

#Synced-CL 51332327

diff --git a/Engine/Source/Editor/UnrealEd/Private/Cooker/CookGenerationHelper.cpp b/Engine/Source/Editor/UnrealEd/Private/Cooker/CookGenerationHelper.cpp
--- a/Engine/Source/Editor/UnrealEd/Private/Cooker/CookGenerationHelper.cpp
+++ b/Engine/Source/Editor/UnrealEd/Private/Cooker/CookGenerationHelper.cpp
@@ -2024,20 +2024,25 @@ void FGenerationHelper::ResetSaveState(FCookGenerationInfo& Info, UPackage* Pack
 					++NumNulls;
 				}
 			}
-			UPackage* PackagePtr = Info.PackageData->GetPackage();
+			UPackage* LocalOwnerPackage = Info.IsGenerator() ? Package : GetOwnerPackage();
 
 			UE_LOG(LogCook, Error,
-				TEXT("CookPackageSplitter failure: We are demoting a %s package from save and removing our references that keep its objects loaded.\n")
-				TEXT("This will allow the objects to be garbage collected and cause failures in the splitter which expects them to remain loaded.\n")
-				TEXT("    Package=%s, Splitter=%s, ReleaseSaveReason=%s, CurrentState=%s, CurrentSubState=%s, NewState=%s,\n")
-				TEXT("    PackagePtr=%s, Package->IsFullyLoaded()=%s,\n")
-				TEXT("    GetCachedObjectsInOuter->Num()==%d, Num nulls in GetCachedObjectsInOuter==%d, First null=='%s'"),
+				TEXT("CookPackageSplitter failure: We are demoting a %s package from save and removing our references that keep its objects loaded.")
+				TEXT("\nThis will allow the objects to be garbage collected and cause failures in the splitter which expects them to remain loaded.")
+				TEXT("\n\tPackage=%s, Splitter=%s, ")
+				TEXT("\n\tReleaseSaveReason=%s, CurrentState=%s, CurrentSubState=%s, NewState=%s,")
+				TEXT("\n\tPackagePtr=%s, Package->IsFullyLoaded()=%s,")
+				TEXT("\n\tOwnerPackagePtr=%s, OwnerPackage->IsFullyLoaded()=%s,")
+				TEXT("\n\tGetCachedObjectsInOuter->Num()==%d, Num nulls in GetCachedObjectsInOuter==%d, First null=='%s'"),
 				Info.IsGenerator() ? TEXT("generator") : TEXT("generated"),
 				*Info.GetPackageName(),
 				*GetSplitDataObjectName().ToString(), LexToString(ReleaseSaveReason),
 				LexToString(Info.PackageData->GetState()), LexToString(Info.PackageData->GetSaveSubState()),
-				LexToString(NewState), PackagePtr == nullptr ? TEXT("<null>") : TEXT("<nonnull>"),
-				*(::LexToString(PackagePtr == nullptr ? false : PackagePtr->IsFullyLoaded())),
+				LexToString(NewState),
+				Package == nullptr ? TEXT("<null>") : TEXT("<nonnull>"),
+				*(::LexToString(Package == nullptr ? false : Package->IsFullyLoaded())),
+				LocalOwnerPackage == nullptr ? TEXT("<null>") : TEXT("<nonnull>"),
+				* (::LexToString(LocalOwnerPackage == nullptr ? false : LocalOwnerPackage->IsFullyLoaded())),
 				Info.PackageData->GetCachedObjectsInOuter().Num(),
 				NumNulls, *DeletedObjectInOuter);
 			FDebug::DumpStackTraceToLog(ELogVerbosity::Display);
@@ -2565,9 +2570,21 @@ void FGenerationHelper::UpdateSaveAfterGarbageCollect(const FPackageData& Packag
 
 	if (!Info->IsGenerator())
 	{
-		UPackage* LocalPackage = OwnerPackage.Get();
-		if (!LocalPackage || !LocalPackage->IsFullyLoaded())
+		UPackage* LocalOwnerPackage = OwnerPackage.Get();
+		if (!LocalOwnerPackage || !LocalOwnerPackage->IsFullyLoaded())
 		{
+			if (IsUseInternalReferenceToAvoidGarbageCollect() && Info->HasTakenOverCachedCookedPlatformData())
+			{
+				UE_LOG(LogCook, Error,
+					TEXT("The generator package of a generated package was unloaded during garbage collection after the generated package started generating. ")
+					TEXT("This will cause errors during save of the generated package.")
+					TEXT("\n\tSplitter=%s, Generated=%s")
+					TEXT("\n\tGeneratorPackage=%s, GeneratorPackage->IsFullyLoaded()=%s"),
+					*GetSplitDataObjectName().ToString(),
+					*Info->PackageData->GetPackageName().ToString(),
+					LocalOwnerPackage == nullptr ? TEXT("<null>") : TEXT("<nonnull>"),
+					*(::LexToString(LocalOwnerPackage == nullptr ? false : LocalOwnerPackage->IsFullyLoaded())));
+			}
 			bInOutDemote = true;
 			return;
 		}
