Cooker: CookPackageSplitter: Fix cook error "CookPackageSplitter failure: We are demoting a %s package from save..." that can occur when the generator package finishes saving before a generated package, and a garbage collect occurs before the generated package finishes saving. We need to find the generator package and keep it in memory, rather than trying to read it from the generatordata's GetPackage() function, since GetPackage() is cleared when the generator package exits the save state.
#rnx

#Synced-CL 51332327

diff --git a/Engine/Source/Editor/UnrealEd/Private/Cooker/CookGenerationHelper.cpp b/Engine/Source/Editor/UnrealEd/Private/Cooker/CookGenerationHelper.cpp
--- a/Engine/Source/Editor/UnrealEd/Private/Cooker/CookGenerationHelper.cpp
+++ b/Engine/Source/Editor/UnrealEd/Private/Cooker/CookGenerationHelper.cpp
@@ -2240,6 +2240,10 @@ void FGenerationHelper::PreGarbageCollect(const TRefCountPtr<FGenerationHelper>&
 		if (&Info != &OwnerInfo)
 		{
 			UPackage* Package = Info.PackageData->GetPackage();
+			if (!Package)
+			{
+				Package = FindObject<UPackage>(nullptr, *Info.PackageData->GetPackageName().ToString());
+			}
 			if (Package)
 			{
 				GCKeepPackages.Add(Package);
@@ -2252,6 +2256,10 @@ void FGenerationHelper::PreGarbageCollect(const TRefCountPtr<FGenerationHelper>&
 	{
 		bCurrentGCHasKeptGeneratorPackage = true;
 		UPackage* Package = OwnerInfo.PackageData->GetPackage();
+		if (!Package)
+		{
+			Package = FindObject<UPackage>(nullptr, *OwnerInfo.PackageData->GetPackageName().ToString());
+		}
 		if (Package)
 		{
 			GCKeepPackages.Add(Package);
