Fix invalid return type for FCompositeDragDropOp::HasSubOp

#Synced-CL 51408661

diff --git a/Engine/Source/Editor/UnrealEd/Public/DragAndDrop/CompositeDragDropOp.h b/Engine/Source/Editor/UnrealEd/Public/DragAndDrop/CompositeDragDropOp.h
--- a/Engine/Source/Editor/UnrealEd/Public/DragAndDrop/CompositeDragDropOp.h
+++ b/Engine/Source/Editor/UnrealEd/Public/DragAndDrop/CompositeDragDropOp.h
@@ -20,7 +20,7 @@ public:
 	}
 
 	template <typename T>
-	TSharedPtr<T> HasSubOp()
+	bool HasSubOp()
 	{
 		for (auto& SubOp : SubOps)
 		{
