UE: Lower verbosity about leaving in stale packages for removed package generators. This is expected in the current code, but notable

diff --git a/Engine/Source/Editor/UnrealEd/Private/CookOnTheFlyServer.cpp b/Engine/Source/Editor/UnrealEd/Private/CookOnTheFlyServer.cpp
--- a/Engine/Source/Editor/UnrealEd/Private/CookOnTheFlyServer.cpp
+++ b/Engine/Source/Editor/UnrealEd/Private/CookOnTheFlyServer.cpp
@@ -8058,7 +8058,7 @@ void UCookOnTheFlyServer::PopulateCookedPackages(TArrayView<const ITargetPlatfor
 			FPackageData* Generator = PackageDatas->TryAddPackageDataByPackageName(Pair.Key);
 			if (!Generator)
 			{
-				UE_LOG(LogCook, Warning,
+				UE_LOG(LogCook, Display,
 					TEXT("Previous cook results returned a record for generator package %s, but that package can no longer be found; its generated packages will not be removed from cook results. Run a full cook to remove them."),
 					*Pair.Key.ToString());
 				continue;
