[ITF Viewport Interactions] Dolly interaction not working while piloting

#Synced-CL 51415372

diff --git a/Engine/Source/Editor/UnrealEd/Private/EditorViewportClient.cpp b/Engine/Source/Editor/UnrealEd/Private/EditorViewportClient.cpp
--- a/Engine/Source/Editor/UnrealEd/Private/EditorViewportClient.cpp
+++ b/Engine/Source/Editor/UnrealEd/Private/EditorViewportClient.cpp
@@ -2439,24 +2439,33 @@ bool FEditorViewportClient::ConsumeExternalCameraMovement()
 	
 	if (IsPerspective())
 	{
+		bool bCameraHasMoved = false;
 		if (!FMath::IsNearlyZero(Orbit.X) || !FMath::IsNearlyZero(Orbit.Y))
 		{
 			FRotator NewRotation = FRotator::ZeroRotator;
             NewRotation.Yaw = Orbit.X;
             NewRotation.Pitch = Orbit.Y;
-    
+
             OrbitPerspectiveCamera(NewRotation);
+			bCameraHasMoved = true;
 		}
 		
 		if (!FMath::IsNearlyZero(Orbit.Z))
 		{
 			DollyPerspectiveCamera(Orbit.Z);
+			bCameraHasMoved = true;
 		}
 		
 		if (!Translation.IsNearlyZero() || !Rotation.IsNearlyZero())
 		{
 			ToggleOrbitCamera(false);
 			MoveViewportPerspectiveCamera(Translation * GetCameraSpeed(), Rotation);
+			bCameraHasMoved = true;
+		}
+
+		if (bCameraHasMoved)
+		{
+			PerspectiveCameraMoved();
 		}
 	}
 	else
