Restore overridden state of subobjects during re-instantiation.

#Synced-CL 51454789

diff --git a/Engine/Source/Editor/UnrealEd/Private/Kismet2/KismetReinstanceUtilities.cpp b/Engine/Source/Editor/UnrealEd/Private/Kismet2/KismetReinstanceUtilities.cpp
--- a/Engine/Source/Editor/UnrealEd/Private/Kismet2/KismetReinstanceUtilities.cpp
+++ b/Engine/Source/Editor/UnrealEd/Private/Kismet2/KismetReinstanceUtilities.cpp
@@ -2402,8 +2402,8 @@ void FBlueprintCompileReinstancer::ReplaceObjectHelper(const FReinstancingState&
 	for (const auto& Pair : OrderedListOfObjectToCopy)	
 	{
 		UEditorEngine::CopyPropertiesForUnrelatedObjects(Pair.Key, Pair.Value, Options);
+		FOverridableManager::Get().RestoreOverrideState(Pair.Key, Pair.Value);
 	}
-	FOverridableManager::Get().RestoreOverrideState(OldObject, NewUObject);
 
 	UWorld* RegisteredWorld = nullptr;
 	bool bWasRegistered = false;
