Cooker: Fix iteration order of RequestQueue to maintain the root to leaf order that it receives from insertions. InsertionOrder == IterationOrder used to partially work with TSets, but had small chances of small changes, and now it does not work at all with TCompactSet implementation of TSet. Create a TMutableSortedSet class to robustly implement the sorted behavior while still being almost as fast as a TSet. TMutableSortedSet is faster than the alternative TSortedSet, which has O(N) removals.
#rnx

#Synced-CL 51332327

diff --git a/Engine/Source/Editor/UnrealEd/Private/Cooker/CookPackageData.cpp b/Engine/Source/Editor/UnrealEd/Private/Cooker/CookPackageData.cpp
--- a/Engine/Source/Editor/UnrealEd/Private/Cooker/CookPackageData.cpp
+++ b/Engine/Source/Editor/UnrealEd/Private/Cooker/CookPackageData.cpp
@@ -3202,16 +3202,14 @@ uint32 FRequestQueue::ReadyRequestsNum() const
 
 FPackageData* FRequestQueue::PopReadyRequest()
 {
-	if (auto Iterator = UrgentRequests.CreateIterator(); Iterator)
+	FPackageData* PackageData = UrgentRequests.Pop();
+	if (PackageData)
 	{
-		FPackageData* PackageData = *Iterator;
-		Iterator.RemoveCurrent();
 		return PackageData;
 	}
-	if (auto Iterator = NormalRequests.CreateIterator(); Iterator)
+	PackageData = NormalRequests.Pop();
+	if (PackageData)
 	{
-		FPackageData* PackageData = *Iterator;
-		Iterator.RemoveCurrent();
 		return PackageData;
 	}
 	return nullptr;
diff --git a/Engine/Source/Editor/UnrealEd/Private/Cooker/CookPackageData.h b/Engine/Source/Editor/UnrealEd/Private/Cooker/CookPackageData.h
--- a/Engine/Source/Editor/UnrealEd/Private/Cooker/CookPackageData.h
+++ b/Engine/Source/Editor/UnrealEd/Private/Cooker/CookPackageData.h
@@ -2,7 +2,6 @@
 
 #pragma once
 
-#include "CookTypes.h"
 #include "Containers/Array.h"
 #include "Containers/ArrayView.h"
 #include "Containers/Map.h"
@@ -10,6 +9,8 @@
 #include "Containers/SortedMap.h"
 #include "Containers/StringFwd.h"
 #include "Containers/UnrealString.h"
+#include "Cooker/CookTypes.h"
+#include "Cooker/FifoSet.h"
 #include "HAL/CriticalSection.h"
 #include "HAL/Platform.h"
 #include "HAL/PlatformMath.h"
@@ -1013,6 +1014,8 @@ struct FPendingCookedPlatformDataCancelManager
 	void Release(FPendingCookedPlatformData& Data);
 };
 
+using FPackageDataSortedSet = TFifoSet<FPackageData*>;
+
 /**
  * The container class for PackageData pointers that are InProgress in a CookOnTheFlyServer. These containers
  * most frequently do queue push/pop operations, but also commonly need to support iteration.
@@ -1133,8 +1136,8 @@ public:
 	TRingBuffer<FDiscoveryQueueElement>& GetDiscoveryQueue();
 	TRingBuffer<FPackageData*>& GetBuildDependencyDiscoveryQueue();
 	TRingBuffer<TUniquePtr<FRequestCluster>>& GetRequestClusters();
-	FPackageDataSet& GetReadyRequestsUrgent();
-	FPackageDataSet& GetReadyRequestsNormal();
+	FPackageDataSortedSet& GetReadyRequestsUrgent();
+	FPackageDataSortedSet& GetReadyRequestsNormal();
 	/**
 	 * Add an FPackageData by name that will be notified when all packages that were present in the DiscoveryQueue
 	 * before the fence was created have been assigned or demoted to idle.
@@ -1149,8 +1152,8 @@ private:
 	TRingBuffer<FPackageData*> BuildDependencyDiscoveryQueue;
 	TRingBuffer<TUniquePtr<FRequestCluster>> RequestClusters; // Not trivially relocatable, so must be TUniquePtr
 	TSet<FName> RequestFencePackageListeners;
-	FPackageDataSet UrgentRequests;
-	FPackageDataSet NormalRequests;
+	FPackageDataSortedSet UrgentRequests;
+	FPackageDataSortedSet NormalRequests;
 };
 
 /** A wrapper around a TRefCountPtr<FPackagePreloader> that defines operator< for FPackagePreloaderPriorityQueue. */
@@ -2025,17 +2028,16 @@ inline TRingBuffer<TUniquePtr<FRequestCluster>>& FRequestQueue::GetRequestCluste
 	return RequestClusters;
 }
 
-inline FPackageDataSet& FRequestQueue::GetReadyRequestsUrgent()
+inline FPackageDataSortedSet& FRequestQueue::GetReadyRequestsUrgent()
 {
 	return UrgentRequests;
 }
 
-inline FPackageDataSet& FRequestQueue::GetReadyRequestsNormal()
+inline FPackageDataSortedSet& FRequestQueue::GetReadyRequestsNormal()
 {
 	return NormalRequests;
 }
 
-
 inline FPackageDataMonitor& FPackageDatas::GetMonitor()
 {
 	return Monitor;
diff --git a/Engine/Source/Editor/UnrealEd/Private/Cooker/FifoSet.h b/Engine/Source/Editor/UnrealEd/Private/Cooker/FifoSet.h
new file mode 100644
--- /dev/null
+++ b/Engine/Source/Editor/UnrealEd/Private/Cooker/FifoSet.h
@@ -0,0 +1,219 @@
+// Copyright Epic Games, Inc. All Rights Reserved.
+
+#pragma once
+
+#include "Concepts/Pointer.h"
+#include "Containers/Map.h"
+#include "Containers/Array.h"
+#include "CoreTypes.h"
+#include "Misc/AssertionMacros.h"
+#include "Misc/NotNull.h"
+#include "Templates/UnrealTemplate.h"
+
+namespace UE::Cook
+{
+
+/**
+ * A container with API similar to TSet<T*> that supports iteration order == insertion order, with extra costs during
+ * Add and Pop to handle the ordering; these costs are small when amortized.
+ * This container only supports pointers, and nullptr is used as a special value and is not permitted in the container.
+ */
+template <UE::CPointer T>
+class TFifoSet
+{
+public:
+	/**
+	 * Add Key to the container if it does not already exist.
+	 * nullptr is a special value and is not permitted in the container.
+	 */
+	void Add(TNotNull<T> Key);
+	/** Remove Key if it exists in the container and return 1 if it does, else 0. */
+	int32 Remove(const T& Key);
+	bool Contains(const T& Key) const;
+	/** Remove and return the oldest element in the container. Returns null if empty. */
+	T Pop();
+	void Empty();
+	int32 Num() const;
+	bool IsEmpty() const;
+
+private:
+	struct FIterator
+	{
+	public:
+		explicit FIterator(const TFifoSet<T>& InSet, int32 InIndex);
+
+		/** Used only in range for. Assumes that Other is an iterator from the same set. */
+		bool operator!=(const FIterator& Other) const;
+		const T& operator*() const;
+		FIterator& operator++();
+	private:
+		const TFifoSet<T>& Set;
+		int32 Index;
+	};
+
+	// Begin and end are implemented as hidden friends to allow them in ranged for but not direct calls.
+	friend FIterator begin(const TFifoSet<T>& Set)
+	{
+		return TFifoSet<T>::FIterator(Set, -1);
+	}
+	friend FIterator end(const TFifoSet<T>& Set)
+	{
+		return TFifoSet<T>::FIterator(Set, Set.PopOrder.Num());
+	}
+	void ReIndex();
+
+	TMap<T, int32> KeyToIndex;
+	TArray<T> PopOrder;
+	int32 NextPop = 0;
+};
+
+template <UE::CPointer T>
+void TFifoSet<T>::Add(TNotNull<T> Key)
+{
+	int32* StoredIndex = &KeyToIndex.FindOrAdd(Key, PopOrder.Num());
+	if (*StoredIndex < PopOrder.Num())
+	{
+		// Already exists
+		return;
+	}
+
+	if (PopOrder.Num() == PopOrder.Max() && PopOrder.Max() > 0)
+	{
+		// We are about to increase the size of PopOrder. Check whether we should instead strip the nulls out of it
+		// and rewrite the stored indices because we have removed enough elements to make that worthwhile.
+		constexpr int32 MinimumRatioToBeWorthReIndex = 2;
+		if (KeyToIndex.Num() * MinimumRatioToBeWorthReIndex < PopOrder.Max())
+		{
+			ReIndex();
+
+			// ReIndex might have modified KeyToIndex, so refetch our pointer to its internal memory.
+			StoredIndex = KeyToIndex.Find(Key);
+			check(StoredIndex != nullptr);
+
+			// Assign the new index to end of PopOrder into the StoredIndex. ReIndex did not already do this
+			// because Key was not yet present in PopOrder.
+			*StoredIndex = PopOrder.Num();
+		}
+	}
+
+	// We already set StoredIndex above, just assign Key into that index (the end of PopOrder).
+	PopOrder.Add(Key);
+}
+
+template <UE::CPointer T>
+int32 TFifoSet<T>::Remove(const T& Key)
+{
+	int32 Index;
+	if (KeyToIndex.RemoveAndCopyValue(Key, Index))
+	{
+		check(PopOrder[Index] == Key);
+		PopOrder[Index] = nullptr;
+		return 1;
+	}
+	return 0;
+}
+
+template <UE::CPointer T>
+bool TFifoSet<T>::Contains(const T& Key) const
+{
+	return KeyToIndex.Contains(Key);
+}
+
+template <UE::CPointer T>
+T TFifoSet<T>::Pop()
+{
+	while (NextPop < PopOrder.Num())
+	{
+		if (!PopOrder[NextPop])
+		{
+			++NextPop;
+			continue;
+		}
+		T Result = MoveTemp(PopOrder[NextPop]);
+		PopOrder[NextPop] = nullptr;
+		int32 StoredIndex = 0;
+		bool bRemoved = KeyToIndex.RemoveAndCopyValue(Result, StoredIndex);
+		check(bRemoved && StoredIndex == NextPop);
+		++NextPop;
+		return Result;
+	}
+	return nullptr;
+}
+
+template <UE::CPointer T>
+void TFifoSet<T>::Empty()
+{
+	KeyToIndex.Empty();
+	PopOrder.Empty();
+	NextPop = 0;
+}
+
+template <UE::CPointer T>
+int32 TFifoSet<T>::Num() const
+{
+	return KeyToIndex.Num();
+}
+
+template <UE::CPointer T>
+bool TFifoSet<T>::IsEmpty() const
+{
+	return KeyToIndex.IsEmpty();
+}
+
+template <UE::CPointer T>
+void TFifoSet<T>::ReIndex()
+{
+	// Note that ReIndex is called during Add and KeyToIndex has one element in it that is not yet
+	// present in PopOrder.
+	TArray<T> NewPopOrder;
+	NewPopOrder.Reserve(KeyToIndex.Num());
+	for (const T& Key : PopOrder)
+	{
+		if (Key)
+		{
+			int32* Index = KeyToIndex.Find(Key);
+			check(Index != nullptr);
+			*Index = NewPopOrder.Num();
+			NewPopOrder.Add(Key);
+		}
+	}
+	PopOrder = MoveTemp(NewPopOrder);
+	NextPop = 0;
+}
+
+template <UE::CPointer T>
+TFifoSet<T>::FIterator::FIterator(const TFifoSet<T>& InSet, int32 InIndex)
+	: Set(InSet)
+	, Index(InIndex)
+{
+	if (Index < 0)
+	{
+		this->operator++();
+	}
+}
+
+template <UE::CPointer T>
+bool TFifoSet<T>::FIterator::operator!=(const TFifoSet<T>::FIterator& Other) const
+{
+	return Index != Other.Index;
+}
+
+template <UE::CPointer T>
+const T& TFifoSet<T>::FIterator::operator*() const
+{
+	return Set.PopOrder[Index];
+}
+
+template <UE::CPointer T>
+TFifoSet<T>::FIterator& TFifoSet<T>::FIterator::operator++()
+{
+	check(Index < Set.PopOrder.Num());
+	++Index;
+	while (Index < Set.PopOrder.Num() && Set.PopOrder[Index] == nullptr)
+	{
+		++Index;
+	}
+	return *this;
+}
+
+} // namespace UE::Cook
\ No newline at end of file
diff --git a/Engine/Source/Editor/UnrealEd/Private/Cooker/FifoSetTest.cpp b/Engine/Source/Editor/UnrealEd/Private/Cooker/FifoSetTest.cpp
new file mode 100644
--- /dev/null
+++ b/Engine/Source/Editor/UnrealEd/Private/Cooker/FifoSetTest.cpp
@@ -0,0 +1,133 @@
+// Copyright Epic Games, Inc. All Rights Reserved.
+
+#include "Misc/AutomationTest.h"
+
+#if WITH_AUTOMATION_TESTS
+
+#include "Cooker/FifoSet.h"
+
+IMPLEMENT_SIMPLE_AUTOMATION_TEST(FFifoSetTest, "System.Core.Cooker.FifoSet",
+	EAutomationTestFlags::EditorContext | EAutomationTestFlags::EngineFilter)
+
+bool FFifoSetTest::RunTest(const FString& Parameters)
+{
+	using namespace UE::Cook;
+
+	int32 NumKeySet = 1000; 
+
+	// KeyBuffer num must be 1 more than NumKeySet. We skip the middle number and use the N/2 values on to each side of it.
+	int32 NumKeyBuffer = NumKeySet + 1;
+	TArray<int32> KeyBuffer;
+	KeyBuffer.SetNum(NumKeyBuffer);
+	for (int32 Index = 0; Index < NumKeyBuffer; ++Index)
+	{
+		KeyBuffer[Index] = Index;
+	}
+
+	// Key pointers are in this sequence:
+	// KeyBuffer.GetData() + [501, 499, 502, 498, ... 1000, 0]
+	// This is an attempt to avoid the possibility that the TMap inside of TFifoSet just happens to sort
+	// the keys in insertion order because insertion order is numerically increasing.
+	TArray<int32*> Keys;
+	Keys.SetNum(NumKeySet);
+	for (int32 Index = 0; Index < NumKeySet; ++Index)
+	{
+		int32 BufferIndex = (NumKeySet / 2) + ((Index % 2) ? -1 : 1) * ((Index + 2) / 2);
+		check(0 <= BufferIndex && BufferIndex < NumKeyBuffer);
+		Keys[Index] = &KeyBuffer[BufferIndex];
+	}
+
+	TFifoSet<int32*> Set;
+
+	// A string of Adds followed by string of pops, make sure we get them out in the expected order
+	int32 NumKeys = 100;
+	for (int32 Index = 0; Index < NumKeys; ++Index)
+	{
+		Set.Add(Keys[Index]);
+	}
+	int32 Index = 0;
+	for (int32* Iter : Set)
+	{
+		TestEqual(TEXT("Adds followed by Pops iteration is in order"), Iter, Keys[Index]);
+		++Index;
+	}
+	for (Index = 0; Index < NumKeys; ++Index)
+	{
+		int32* Key = Set.Pop();
+		TestEqual(TEXT("Adds followed by Pops keep order"), Key, Keys[Index]);
+	}
+	TestTrue(TEXT("Adds followed by Pops return the set to empty"), Set.IsEmpty());
+
+	// Interleave pops with adds, make sure we get them out in the expected order
+	int32 Batch;
+	for (Batch = 0; Batch < NumKeys; ++Batch)
+	{
+		Set.Add(Keys[2 * Batch]);
+		Set.Add(Keys[2 * Batch + 1]);
+		int32* Popped = Set.Pop();
+		TestEqual(TEXT("Interleaved Adds and Pops keep order"), Popped, Keys[Batch]);
+	}
+	Index = NumKeys;
+	for (int32* Iter : Set)
+	{
+		TestEqual(TEXT("Interleaved Adds and Pops iteration is in order"), Iter, Keys[Index]);
+		++Index;
+	}
+	for (Index = NumKeys; Index < 2*NumKeys; ++Index)
+	{
+		int32* Key = Set.Pop();
+		TestEqual(TEXT("Interleaved Adds and Pops keep order"), Key, Keys[Index]);
+	}
+	TestTrue(TEXT("Interleaved Adds and Pops return the set to empty"), Set.IsEmpty());
+
+	// Interleave removes with adds, make sure we get them out in the expected order
+	for (Batch = 0; Batch < NumKeys; ++Batch)
+	{
+		Set.Add(Keys[2 * Batch]);
+		Set.Add(Keys[2 * Batch + 1]);
+		TestEqual(TEXT("Interleaved Adds and Removes remove successfully"), Set.Remove(Keys[Batch]), 1);
+	}
+	Index = NumKeys;
+	for (int32* Iter : Set)
+	{
+		TestEqual(TEXT("Interleaved Adds and Removes iteration is in order"), Iter, Keys[Index]);
+		++Index;
+	}
+	Index = NumKeys;
+	for (Batch = 0; Batch < NumKeys; ++Batch)
+	{
+		int32* Key = Set.Pop();
+		TestEqual(TEXT("Interleaved Adds and Removes keep order"), Key, Keys[Index]);
+		++Index;
+	}
+	TestTrue(TEXT("Interleaved Adds and Removes return the set to empty"), Set.IsEmpty());
+
+	// Adding an element that already exists does not modify its position
+	int32 NumKeysSmall = 5;
+	for (Index = 0; Index < NumKeysSmall; ++Index)
+	{
+		Set.Add(Keys[Index]);
+	}
+	Set.Add(Keys[0]);
+	TestEqual(TEXT("Add of an existing item does not change num"), Set.Num(), NumKeysSmall);
+	Index = 0;
+	for (int32* Iter : Set)
+	{
+		TestEqual(TEXT("Add of an existing item does not change order"), Iter, Keys[Index]);
+		++Index;
+	}
+	// Also use this test to test emptying works.
+	Set.Empty();
+	TestEqual(TEXT("Empty clears Num -> 0"), Set.Num(), 0);
+	TestEqual(TEXT("Empty clears IsEmpty -> true"), Set.IsEmpty(), true);
+	Index = 0;
+	for (int32* Iter : Set)
+	{
+		++Index;
+	}
+	TestEqual(TEXT("Empty clears iteration -> empty"), Index == 0, true);
+
+	return true;
+}
+
+#endif // WITH_AUTOMATION_TESTS
