EditorUsability : LevelEditor
- Fixed level editor PlayFromHere action not working with keybinds (set to ALT+SHIFT+P)
- Improved API to play from here based on click, selected actor or custom position.
- Play mode (Floating window or in viewport) will be based on the selected user choice in Play Settings
- Play from here will first check selected actor (must set bCanPlayFromHere to true) then fallback to viewport last clicked location

#Synced-CL 51539177

diff --git a/Engine/Source/Editor/LevelEditor/Private/LevelEditor.cpp b/Engine/Source/Editor/LevelEditor/Private/LevelEditor.cpp
--- a/Engine/Source/Editor/LevelEditor/Private/LevelEditor.cpp
+++ b/Engine/Source/Editor/LevelEditor/Private/LevelEditor.cpp
@@ -1012,7 +1012,12 @@ void FLevelEditorModule::BindGlobalLevelEditorCommands()
 		Commands.GoHere,
 		FExecuteAction::CreateStatic( &FLevelEditorActionCallbacks::GoHere_Clicked, (const FVector*)nullptr )
 		);
-		
+
+	ActionList.MapAction(
+		Commands.PlayFromHere,
+		FExecuteAction::CreateStatic(&FLevelEditorActionCallbacks::PlayFromHere)
+	);
+
 	ActionList.MapAction( 
 		Commands.SnapCameraToObject,
 		FExecuteAction::CreateStatic( &FLevelEditorActionCallbacks::ExecuteExecCommand, FString( TEXT("CAMERA SNAP") ) )
diff --git a/Engine/Source/Editor/LevelEditor/Private/LevelEditorActions.cpp b/Engine/Source/Editor/LevelEditor/Private/LevelEditorActions.cpp
--- a/Engine/Source/Editor/LevelEditor/Private/LevelEditorActions.cpp
+++ b/Engine/Source/Editor/LevelEditor/Private/LevelEditorActions.cpp
@@ -1759,7 +1759,7 @@ bool FLevelEditorActionCallbacks::PlayFromHere_IsVisible()
 {
 	if (GEditor->GetSelectedActorCount() == 1)
 	{
-		if (AActor* Actor = Cast<AActor>(*GEditor->GetSelectedActorIterator()))
+		if (const AActor* Actor = Cast<AActor>(*GEditor->GetSelectedActorIterator()))
 		{
 			return Actor->CanPlayFromHere();
 		}
@@ -1768,6 +1768,44 @@ bool FLevelEditorActionCallbacks::PlayFromHere_IsVisible()
 	return false;
 }
 
+void FLevelEditorActionCallbacks::PlayFromHere()
+{
+	FPlayWorldCommandCallbacks::FStartPlayFromHereOptions PlayOptions;
+
+	// Determine whether to spawn a new window or reuse existing viewport
+	const ULevelEditorPlaySettings* PlaySettings = GetDefault<ULevelEditorPlaySettings>();
+	const bool bSpawnNewWindow = PlaySettings->LastExecutedPlayModeType != EPlayModeType::PlayMode_InViewPort;
+	if (!bSpawnNewWindow)
+	{
+		FLevelEditorModule& LevelEditorModule = FModuleManager::GetModuleChecked<FLevelEditorModule>("LevelEditor");
+		PlayOptions.DestinationSlateViewport = LevelEditorModule.GetFirstActiveViewport();
+	}
+
+	if (GEditor->GetSelectedActorCount() == 1)
+	{
+		if (const AActor* Actor = Cast<AActor>(*GEditor->GetSelectedActorIterator()))
+		{
+			if (Actor->CanPlayFromHere())
+			{
+				// Play from selected actor position
+				PlayOptions.Mode = FPlayWorldCommandCallbacks::SelectedActor;
+				if (FPlayWorldCommandCallbacks::StartPlayFromHere(PlayOptions))
+				{
+					return;
+				}
+			}
+			else
+			{
+				UE_LOG(LevelEditorActions, Log, TEXT("PlayFromHere: Selected actor to play from must set bCanPlayFromHere to true, falling back on last clicked location..."));
+			}
+		}
+	}
+
+	// Fallback to click position
+	PlayOptions.Mode = FPlayWorldCommandCallbacks::Click;
+	FPlayWorldCommandCallbacks::StartPlayFromHere(PlayOptions);
+}
+
 void FLevelEditorActionCallbacks::GoHere_Clicked( const FVector* Point )
 {
 	if( GCurrentLevelEditingViewportClient )
@@ -4345,6 +4383,7 @@ void FLevelEditorCommands::RegisterCommands()
 	UI_COMMAND( OpenSelectionInPropertyMatrix, "Edit Selection in Property Matrix", "Bulk edit the selected assets in the Property Matrix", EUserInterfaceActionType::Button, FInputChord() );
 
 	UI_COMMAND( GoHere, "Go Here", "Moves the camera to the current mouse position", EUserInterfaceActionType::Button, FInputChord() );
+	UI_COMMAND( PlayFromHere, "Play From Here", "Starts a game preview from the selected actor or clicked location as fallback", EUserInterfaceActionType::Button, FInputChord(EKeys::P, EModifierKey::Alt | EModifierKey::Shift) );
 
 	UI_COMMAND( SnapCameraToObject, "Move Camera to Object", "Move the current camera to match the location and rotation of the selected object.", EUserInterfaceActionType::Button, FInputChord() );
 	UI_COMMAND( SnapObjectToCamera, "Move Object to Camera", "Move the selected object to match the location and rotation of the current camera.", EUserInterfaceActionType::Button, FInputChord() );
diff --git a/Engine/Source/Editor/LevelEditor/Private/LevelEditorContextMenu.cpp b/Engine/Source/Editor/LevelEditor/Private/LevelEditorContextMenu.cpp
--- a/Engine/Source/Editor/LevelEditor/Private/LevelEditorContextMenu.cpp
+++ b/Engine/Source/Editor/LevelEditor/Private/LevelEditorContextMenu.cpp
@@ -479,6 +479,10 @@ void FLevelEditorContextMenu::RegisterActorContextMenu()
 								FText::Format(LOCTEXT("PilotActor", "Pilot '{0}'"), FText::FromString(SelectedActors[0]->GetActorLabel()))
 								);
 						}
+
+						// Only allowed on actor with bCanPlayFromHere = true
+						FToolMenuEntry& PlayFromHereEntry = Section.AddMenuEntry(FLevelEditorCommands::Get().PlayFromHere);
+						PlayFromHereEntry.Visibility = FIsActionButtonVisible::CreateStatic(&FLevelEditorActionCallbacks::PlayFromHere_IsVisible);
 					}
 				}
 			}
@@ -605,12 +609,8 @@ void FLevelEditorContextMenu::RegisterActorContextMenu()
 				// Only extend if above PlayFromHere option isn't available
 				if (!FLevelEditorActionCallbacks::PlayFromHere_IsVisible())
 				{
-					// Note: not using a command for play from here since it requires a mouse click
-					FUIAction PlayFromHereAction(
-						FExecuteAction::CreateStatic(&FPlayWorldCommandCallbacks::StartPlayFromHere));
-
 					const FText PlayFromHereLabel = GEditor->OnlyLoadEditorVisibleLevelsInPIE() ? LOCTEXT("PlayFromHereVisible", "Play From Here (visible levels)") : LOCTEXT("PlayFromHere", "Play From Here");
-					Section.AddMenuEntry(NAME_None, PlayFromHereLabel, LOCTEXT("PlayFromHere_ToolTip", "Starts a game preview from the clicked location"), FSlateIcon(FAppStyle::GetAppStyleSetName(), "Icons.Play"), PlayFromHereAction);
+					Section.AddMenuEntry(FLevelEditorCommands::Get().PlayFromHere, PlayFromHereLabel);
 				}
 			}
 
@@ -767,14 +767,9 @@ void FLevelEditorContextMenu::RegisterElementContextMenu()
 			// TODO: May someday support custom play-from-here behavior the way that some actors do, and
 			//  may someday support simulation changes
 			FToolMenuSection& Section = InMenu->AddSection("ElementPreview", LOCTEXT("PreviewHeading", "Preview"));
-
-			// Note: not using a command for play from here since it requires a mouse click
-			FUIAction PlayFromHereAction(
-				FExecuteAction::CreateStatic(&FPlayWorldCommandCallbacks::StartPlayFromHere));
-
 			const FText PlayFromHereLabel = GEditor->OnlyLoadEditorVisibleLevelsInPIE() ? LOCTEXT("PlayFromHereVisible", "Play From Here (visible levels)") : LOCTEXT("PlayFromHere", "Play From Here");
-			Section.AddMenuEntry(NAME_None, PlayFromHereLabel, LOCTEXT("PlayFromHere_ToolTip", "Starts a game preview from the clicked location"), FSlateIcon(FAppStyle::GetAppStyleSetName(), "Icons.Play"), PlayFromHereAction);
-}
+			Section.AddMenuEntry(FLevelEditorCommands::Get().PlayFromHere, PlayFromHereLabel);
+		}
 	}));
 }
 
diff --git a/Engine/Source/Editor/LevelEditor/Public/LevelEditorActions.h b/Engine/Source/Editor/LevelEditor/Public/LevelEditorActions.h
--- a/Engine/Source/Editor/LevelEditor/Public/LevelEditorActions.h
+++ b/Engine/Source/Editor/LevelEditor/Public/LevelEditorActions.h
@@ -138,6 +138,9 @@ public:
 
 	/** Moves the camera to the current mouse position */
 	TSharedPtr< FUICommandInfo > GoHere;
+	
+	/** Plays from the current position (selected actor or clicked) */
+	TSharedPtr< FUICommandInfo > PlayFromHere;
 
 	/** Snaps the camera to the selected object. */
 	TSharedPtr< FUICommandInfo > SnapCameraToObject;
@@ -1059,6 +1062,9 @@ public:
 	/** Called when selected actor can be used to start a play session */
 	static UE_API void PlayFromHere_Clicked(bool bFloatingWindow);
 	static UE_API bool PlayFromHere_IsVisible();
+	
+	/** If an actor is selected, plays from its position (AActor::bCanPlayFromHere must be true) otherwise it will play from last clicked position */
+	static UE_API void PlayFromHere();
 
 	/** Called when 'Go to Code for Actor' is clicked */
 	static UE_API void GoToCodeForActor_Clicked();
diff --git a/Engine/Source/Editor/UnrealEd/Private/Kismet2/DebuggerCommands.cpp b/Engine/Source/Editor/UnrealEd/Private/Kismet2/DebuggerCommands.cpp
--- a/Engine/Source/Editor/UnrealEd/Private/Kismet2/DebuggerCommands.cpp
+++ b/Engine/Source/Editor/UnrealEd/Private/Kismet2/DebuggerCommands.cpp
@@ -904,52 +904,96 @@ TSharedRef< SWidget > FPlayWorldCommands::GeneratePlayMenuContent(TSharedRef<FUI
 
 void FPlayWorldCommandCallbacks::StartPlayFromHere()
 {
-	// Figure out the start location of the player
-	UClass* const PlayerStartClass = GUnrealEd->PlayFromHerePlayerStartClass ? (UClass*)GUnrealEd->PlayFromHerePlayerStartClass : APlayerStart::StaticClass();
-	UCapsuleComponent* DefaultCollisionComponent = CastChecked<UCapsuleComponent>(PlayerStartClass->GetDefaultObject<AActor>()->GetRootComponent());
-	FVector	CollisionExtent = FVector(DefaultCollisionComponent->GetScaledCapsuleRadius(), DefaultCollisionComponent->GetScaledCapsuleRadius(), DefaultCollisionComponent->GetScaledCapsuleHalfHeight());
-	FVector StartLocation = GEditor->UnsnappedClickLocation + GEditor->ClickPlane * (FVector::BoxPushOut(GEditor->ClickPlane, CollisionExtent) + 0.1f);
-	TOptional<FRotator> StartRotation;
 	FLevelEditorModule& LevelEditorModule = FModuleManager::GetModuleChecked<FLevelEditorModule>("LevelEditor");
-
-	TSharedPtr<IAssetViewport> ActiveLevelViewport = LevelEditorModule.GetFirstActiveViewport();
-
-	if (ActiveLevelViewport.IsValid() && ActiveLevelViewport->GetAssetViewportClient().IsPerspective())
-	{
-		StartRotation = ActiveLevelViewport->GetAssetViewportClient().GetViewRotation();
-	}
-
-	StartPlayFromHere(StartLocation, StartRotation, ActiveLevelViewport);
+	FStartPlayFromHereOptions Options;
+	Options.Mode = EPlayFromHereMode::Click;
+	Options.DestinationSlateViewport = LevelEditorModule.GetFirstActiveViewport();
+	StartPlayFromHere(Options);
 }
 
 void FPlayWorldCommandCallbacks::StartPlayFromHere(const TOptional<FVector>& Location, const TOptional<FRotator>& Rotation, const TSharedPtr<IAssetViewport>& ActiveLevelViewport)
+{
+	FStartPlayFromHereOptions Options;
+	Options.Mode = EPlayFromHereMode::Custom;
+	Options.Location = Location;
+	Options.Rotation = Rotation;
+	Options.DestinationSlateViewport = ActiveLevelViewport;
+	StartPlayFromHere(Options);
+}
+
+bool FPlayWorldCommandCallbacks::StartPlayFromHere(const FStartPlayFromHereOptions& InOptions)
 {
 	if (FInternalPlayWorldCommandCallbacks::IsStoppedAtBreakpoint_InEngineMode() || DebuggerCommands::IsVREditorActive())
 	{
-		return;
+		return false;
 	}
 
-	// Is a PIE session already running?  If so we close it first
-	if (GUnrealEd->PlayWorld != NULL)
+	TOptional<FVector> StartLocation = InOptions.Location;
+	TOptional<FRotator> StartRotation = InOptions.Rotation;
+
+	bool bCanProceed = InOptions.Mode == EPlayFromHereMode::Custom;
+	if (InOptions.Mode == EPlayFromHereMode::Click)
 	{
-		GUnrealEd->EndPlayMap();
+		// Figure out the start location of the player
+		const UClass* const PlayerStartClass = GUnrealEd->PlayFromHerePlayerStartClass ? (UClass*)GUnrealEd->PlayFromHerePlayerStartClass : APlayerStart::StaticClass();
+		const UCapsuleComponent* DefaultCollisionComponent = CastChecked<UCapsuleComponent>(PlayerStartClass->GetDefaultObject<AActor>()->GetRootComponent());
+		const FVector CollisionExtent = FVector(DefaultCollisionComponent->GetScaledCapsuleRadius(), DefaultCollisionComponent->GetScaledCapsuleRadius(), DefaultCollisionComponent->GetScaledCapsuleHalfHeight());
+		StartLocation = GEditor->UnsnappedClickLocation + GEditor->ClickPlane * (FVector::BoxPushOut(GEditor->ClickPlane, CollisionExtent) + 0.1f);
+		FLevelEditorModule& LevelEditorModule = FModuleManager::GetModuleChecked<FLevelEditorModule>("LevelEditor");
+
+		const TSharedPtr<IAssetViewport> ActiveLevelViewport = LevelEditorModule.GetFirstActiveViewport();
+		if (ActiveLevelViewport.IsValid() && ActiveLevelViewport->GetAssetViewportClient().IsPerspective())
+		{
+			StartRotation = ActiveLevelViewport->GetAssetViewportClient().GetViewRotation();
+		}
+
+		bCanProceed = true;
+	}
+	else if (InOptions.Mode == EPlayFromHereMode::SelectedActor)
+	{
+		if (GEditor->GetSelectedActorCount() == 1)
+		{
+			if (AActor* Actor = Cast<AActor>(*GEditor->GetSelectedActorIterator()))
+			{
+				if (const UWorld* ActorWorld = Actor->GetWorld())
+				{
+					if (ActorWorld->PersistentLevel)
+					{
+						ActorWorld->PersistentLevel->PlayFromHereActor = Actor;
+					}
+				}
+				StartLocation = Actor->GetActorLocation();
+				StartRotation = Actor->GetActorRotation();
+				bCanProceed = true;
+			}
+		}
+	}
+	
+	if (!bCanProceed)
+	{
+		return false;
 	}
 
 	FRequestPlaySessionParams SessionParams;
-	SessionParams.StartLocation = Location;
-	SessionParams.StartRotation = Rotation;
+	SessionParams.StartLocation = StartLocation;
+	SessionParams.StartRotation = StartRotation;
 
-	FLevelEditorModule& LevelEditorModule = FModuleManager::GetModuleChecked<FLevelEditorModule>("LevelEditor");
-
-	if (ActiveLevelViewport.IsValid() && ActiveLevelViewport->GetAssetViewportClient().IsPerspective())
+	if (InOptions.DestinationSlateViewport.IsValid() && InOptions.DestinationSlateViewport->GetAssetViewportClient().IsPerspective())
 	{
 		// If there is no level viewport, a new window will be spawned to play in.
-		SessionParams.DestinationSlateViewport = ActiveLevelViewport;
+		SessionParams.DestinationSlateViewport = InOptions.DestinationSlateViewport;
+	}
+
+	// Is a PIE session already running?  If so we close it first
+	if (GUnrealEd->PlayWorld != nullptr)
+	{
+		GUnrealEd->EndPlayMap();
 	}
 
 	GUnrealEd->RequestPlaySession(SessionParams);
-}
 
+	return true;
+}
 
 void FPlayWorldCommandCallbacks::ResumePlaySession_Clicked()
 {
diff --git a/Engine/Source/Editor/UnrealEd/Public/Kismet2/DebuggerCommands.h b/Engine/Source/Editor/UnrealEd/Public/Kismet2/DebuggerCommands.h
--- a/Engine/Source/Editor/UnrealEd/Public/Kismet2/DebuggerCommands.h
+++ b/Engine/Source/Editor/UnrealEd/Public/Kismet2/DebuggerCommands.h
@@ -139,6 +139,26 @@ public:
 	static UNREALED_API void StartPlayFromHere();
 	static UNREALED_API void StartPlayFromHere(const TOptional<FVector>& Location, const TOptional<FRotator>& Rotation, const TSharedPtr<IAssetViewport>& ActiveLevelViewport);
 
+	enum EPlayFromHereMode
+	{
+		/** Provide value for Location, Rotation */
+		Custom,
+		/** Determines value for Location, Rotation based on click position in viewport */
+		Click,
+		/** Determines values for Location, Rotation based on the selected actor in level */
+		SelectedActor
+	};
+
+	struct FStartPlayFromHereOptions
+	{
+		EPlayFromHereMode Mode = EPlayFromHereMode::Custom;
+		TOptional<FVector> Location;
+		TOptional<FRotator> Rotation;
+		TSharedPtr<IAssetViewport> DestinationSlateViewport;
+	};
+
+	static UNREALED_API bool StartPlayFromHere(const FStartPlayFromHereOptions& InOptions);
+
 	static UNREALED_API void ResumePlaySession_Clicked();
 	static UNREALED_API void PausePlaySession_Clicked();
 	static UNREALED_API void SingleFrameAdvance_Clicked();
