Precalculate in parallel schema hashes and custom versions for class digests.

The BuildConstantLookupTable guard seems redundant, as the function only ever modifies a struct owned by the caller. The parallel gathering of GetCustomVersions leads to a temp construction of FFloatDistribution (through UParticleModuleTypeDataGpu) which ends up calling this function.

#rnx
#Synced-CL 51534191

diff --git a/Engine/Source/Editor/UnrealEd/Private/EditorDomain/EditorDomainUtils.cpp b/Engine/Source/Editor/UnrealEd/Private/EditorDomain/EditorDomainUtils.cpp
--- a/Engine/Source/Editor/UnrealEd/Private/EditorDomain/EditorDomainUtils.cpp
+++ b/Engine/Source/Editor/UnrealEd/Private/EditorDomain/EditorDomainUtils.cpp
@@ -9,6 +9,7 @@
 #include "AssetRegistry/AssetData.h"
 #include "AssetRegistry/IAssetRegistry.h"
 #include "AssetRegistry/PackageReader.h"
+#include "Async/ParallelFor.h"
 #include "Containers/DirectoryTree.h"
 #include "Cooker/CookDependencyContext.h"
 #include "Cooker/CookPackageArtifacts.h"
@@ -159,6 +160,9 @@ void Register()
 
 namespace UE::EditorDomain
 {
+	
+// Enables/disables the parallel precalculation of class data digests (ease of use for debugging purposes).
+constexpr bool bPrecalculateDataForClassDigests = true;
 
 void PrecacheClassDigest(const FTopLevelAssetPath& ClassPath);
 
@@ -214,7 +218,7 @@ ESaveStorageResult SaveStorageResultFromString(FUtf8StringView Text)
 	return ESaveStorageResult::InvalidCode;
 }
 
-TArray<FGuid> GetCustomVersions(UClass& Class);
+TArray<FGuid, TInlineAllocator<1>> GetCustomVersions(UClass& Class, FCustomVersionsStripedMap* InOutCache = nullptr);
 TMap<FGuid, UObject*> FindCustomVersionCulprits(TConstArrayView<FGuid> UnknownGuids, UPackage* Package);
 void InitializeGlobalConstructClasses();
 
@@ -638,10 +642,11 @@ TOptional<FBlake3Hash> TryGetClassDigest(FTopLevelAssetPath ClassPath)
 class FPrecacheClassDigest
 {
 public:
-	FPrecacheClassDigest()
+	FPrecacheClassDigest(FCustomVersionsStripedMap* CustomVersionsCache = nullptr)
 		: ClassDigestsMap(GetClassDigests())
 		, ClassDigests(ClassDigestsMap.Map)
 		, AssetRegistry(IAssetRegistry::GetChecked())
+		, CustomVersionsCache(CustomVersionsCache)
 	{
 		ClassDigestsMap.Lock.WriteLock();
 	}
@@ -671,6 +676,7 @@ private:
 	FClassDigestMap& ClassDigestsMap;
 	TMap<FTopLevelAssetPath, FClassDigestData>& ClassDigests;
 	IAssetRegistry& AssetRegistry;
+	FCustomVersionsStripedMap* CustomVersionsCache;
 
 	// Scratch variables usable during GetRecursive; they are invalidated when a recursive call is made
 	FString NameStringBuffer;
@@ -821,7 +827,7 @@ FClassDigestData* FPrecacheClassDigest::GetRecursive(FTopLevelAssetPath ClassPat
 		// GetCustomVersions and CallDeclareConstructClasses can create the ClassDefaultObject, which can trigger
 		// LoadPackage, which can reenter this function recursively. We have to drop the lock to prevent a deadlock.
 		FUnlockScope UnlockScope(ClassDigestsMap.Lock);
-		FKnownCustomVersions::FindOrAddHandles(DigestData->CustomVersionHandles, GetCustomVersions(*StructAsClass));
+		FKnownCustomVersions::FindOrAddHandles(DigestData->CustomVersionHandles, GetCustomVersions(*StructAsClass, CustomVersionsCache));
 		StructAsClass->CallDeclareConstructClasses(ConstructClasses);
 	}
 	else
@@ -956,10 +962,10 @@ FClassDigestData* FPrecacheClassDigest::GetRecursive(FTopLevelAssetPath ClassPat
 }
 
 /** Try to add the FClassDigestData for each given class into the GetClassDigests map */
-void PrecacheClassDigests(TConstArrayView<FTopLevelAssetPath> ClassPaths)
+void PrecacheClassDigests(TConstArrayView<FTopLevelAssetPath> ClassPaths, FCustomVersionsStripedMap* CustomVersionsCache)
 {
 	SCOPED_EDITORDOMAINTIMER(EditorDomain::PrecacheClassDigests);
-	FPrecacheClassDigest Digester;
+	FPrecacheClassDigest Digester(CustomVersionsCache);
 	for (const FTopLevelAssetPath& ClassPath : ClassPaths)
 	{
 		if (!ClassPath.IsNull())
@@ -1085,9 +1091,14 @@ private:
 };
 
 /** Collect the CustomVersions that can be used by the given Class when it is saved */
-TArray<FGuid> GetCustomVersions(UClass& Class)
+TArray<FGuid, TInlineAllocator<1>> GetCustomVersions(UClass& Class, FCustomVersionsStripedMap* InOutCache)
 {
 	SCOPED_EDITORDOMAINTIMER(EditorDomain::GetCustomVersions);
+	TArray<FGuid, TInlineAllocator<1>> Guids;
+	if (InOutCache && InOutCache->FindAndApply(&Class, [&Guids](const TArray<FGuid, TInlineAllocator<1>>& Cached) { Guids = Cached; }))
+	{
+		return Guids;
+	}
 	FCustomVersionCollectorArchive Ar;
 	Class.CallDeclareCustomVersions(Ar);
 	// Default objects of blueprint classes are serialized during SavePackage with a special call to
@@ -1102,7 +1113,6 @@ TArray<FGuid> GetCustomVersions(UClass& Class)
 	{
 		Ar.UsingCustomVersion(FUE5MainStreamObjectVersion::GUID);
 	}
-	TArray<FGuid> Guids;
 	const FCustomVersionContainer& CustomVersions = Ar.GetCustomVersions();
 	Guids.Reserve(CustomVersions.GetAllVersions().Num());
 	for (const FCustomVersion& CustomVersion : CustomVersions.GetAllVersions())
@@ -1111,6 +1121,10 @@ TArray<FGuid> GetCustomVersions(UClass& Class)
 	}
 	Algo::Sort(Guids);
 	Guids.SetNum(Algo::Unique(Guids));
+	if (InOutCache)
+	{
+		InOutCache->Add(&Class, Guids);
+	}
 	return Guids;
 }
 
@@ -1434,7 +1448,50 @@ void ConstructTargetIterativeClassAllowList()
 	}
 
 	TArray<FTopLevelAssetPath> EnabledClassPathsArray = EnabledClassPaths.Array();
-	PrecacheClassDigests(EnabledClassPathsArray);
+
+	FCustomVersionsStripedMap CustomVersionsCache;
+	if (bPrecalculateDataForClassDigests)
+	{
+		TArray<UStruct*> NativeStructs;
+		{
+			SCOPED_EDITORDOMAINTIMER(CreateNativeCDOs);
+			// Gather the native structs on which PrecacheClassDigests operates on through GetRecursive.
+			for (const FTopLevelAssetPath& ClassPath : EnabledClassPathsArray)
+			{
+				if (FPackageName::IsScriptPackage(WriteToString<256>(ClassPath.GetPackageName())))
+				{
+					if (UStruct* Struct = FindObject<UStruct>(ClassPath))
+					{
+						NativeStructs.Add(Struct);
+						if (UClass* AsClass = Cast<UClass>(Struct))
+						{
+							// Create uninitialized CDOs serially that could be triggered by GetCustomVersions.
+							AsClass->GetDefaultObject();
+						}
+					}
+				}
+			}
+		}
+		{
+			SCOPED_EDITORDOMAINTIMER(CacheSchemasAndCustomVersions);
+			// Precalculate schema hashes and custom versions for use during PrecacheClassDigests.
+			// Both GetSchemaHash and GetCustomVersions are memoized.
+			ParallelFor(NativeStructs.Num(), [&NativeStructs, &CustomVersionsCache](int32 Index)
+			{
+				NativeStructs[Index]->GetSchemaHash(false);
+				if (UClass* AsClass = Cast<UClass>(NativeStructs[Index]))
+				{
+					GetCustomVersions(*AsClass, &CustomVersionsCache);
+				}
+			});
+		}
+		PrecacheClassDigests(EnabledClassPathsArray, &CustomVersionsCache);
+	}
+	else
+	{
+		PrecacheClassDigests(EnabledClassPathsArray);
+	}
+	
 	FClassDigestMap& ClassDigests = GetClassDigests();
 	{
 		FWriteScopeLock ClassDigestsScopeLock(ClassDigests.Lock);
diff --git a/Engine/Source/Editor/UnrealEd/Private/EditorDomain/EditorDomainUtils.h b/Engine/Source/Editor/UnrealEd/Private/EditorDomain/EditorDomainUtils.h
--- a/Engine/Source/Editor/UnrealEd/Private/EditorDomain/EditorDomainUtils.h
+++ b/Engine/Source/Editor/UnrealEd/Private/EditorDomain/EditorDomainUtils.h
@@ -7,6 +7,7 @@
 #include "Containers/Map.h"
 #include "Containers/Set.h"
 #include "Containers/StringView.h"
+#include "Containers/StripedMap.h"
 #include "Containers/UnrealString.h"
 #include "DerivedDataCache.h"
 #include "EditorDomain/EditorDomain.h"
@@ -37,6 +38,8 @@ namespace UE::DerivedData { struct FCacheKey; }
 namespace UE::EditorDomain
 {
 
+using FCustomVersionsStripedMap = TStripedMap<8, const UClass*, TArray<FGuid, TInlineAllocator<1>>>;
+	
 /** A UClass's data that is used in the EditorDomain Digest, and holds other information about Classes the EditorDomain needs. */
 struct FClassDigestData
 {
@@ -101,8 +104,8 @@ bool TryAppendClassDigests(const TArrayView<UE::Cook::FCookDependency>& Dependen
 /** FindOrAdd class digests for the given class dependencies, and return it. Returns empty optional if not found. */
 TOptional<FBlake3Hash> TryGetClassDigest(FTopLevelAssetPath ClassPath);
 
-/** For any ClassNames not already in ClassDigests, look up their UStruct and add them. */
-void PrecacheClassDigests(TConstArrayView<FTopLevelAssetPath> ClassNames);
+/** For any ClassNames not already in ClassDigests, look up their UStruct and add them. An optional cache for precalculated custom versions can be used. */
+void PrecacheClassDigests(TConstArrayView<FTopLevelAssetPath> ClassNames, FCustomVersionsStripedMap* CustomVersionsCache = nullptr);
 
 /** Get the CacheRequest for the given package from the EditorDomain cache bucket. */
 void RequestEditorDomainPackage(const FPackagePath& PackagePath,
diff --git a/Engine/Source/Runtime/Engine/Private/Distributions.cpp b/Engine/Source/Runtime/Engine/Private/Distributions.cpp
--- a/Engine/Source/Runtime/Engine/Private/Distributions.cpp
+++ b/Engine/Source/Runtime/Engine/Private/Distributions.cpp
@@ -69,7 +69,6 @@ bool UDistribution::IsPostLoadThreadSafe() const
  */
 static void BuildConstantLookupTable( FDistributionLookupTable* OutTable, int32 ValuesPerEntry, const float* Values )
 {
-	check(IsInGameThread() || IsInAsyncLoadingThread());
 	check( OutTable != NULL );
 	check( Values != NULL );
 
