Added ConfigureCookSession event, to allow registering extra artifacts and collectors to be used by the cooker

#Synced-CL 51535417

diff --git a/Engine/Source/Editor/UnrealEd/Private/CookOnTheFlyServer.cpp b/Engine/Source/Editor/UnrealEd/Private/CookOnTheFlyServer.cpp
--- a/Engine/Source/Editor/UnrealEd/Private/CookOnTheFlyServer.cpp
+++ b/Engine/Source/Editor/UnrealEd/Private/CookOnTheFlyServer.cpp
@@ -822,7 +822,8 @@ void UCookOnTheFlyServer::AddCookOnTheFlyPlatformFromGameThread(ITargetPlatform*
 
 void UCookOnTheFlyServer::StartCookOnTheFlySessionFromGameThread(ITargetPlatform* TargetPlatform)
 {
-	if (PlatformManager->GetNumSessionPlatforms() == 0)
+	const bool bIsFirstSession = PlatformManager->GetNumSessionPlatforms() == 0;
+	if (bIsFirstSession)
 	{
 		InitializeSession();
 	}
@@ -838,6 +839,11 @@ void UCookOnTheFlyServer::StartCookOnTheFlySessionFromGameThread(ITargetPlatform
 	// before all plugins are loaded.
 	BlockOnAssetRegistry(TConstArrayView<FString>());
 
+	if (bIsFirstSession)
+	{
+		UE::Cook::FDelegates::ConfigureCookSession.Broadcast(*this);
+	}
+
 	if (CookOnTheFlyRequestManager)
 	{
 		FName PlatformName(*TargetPlatform->PlatformName());
@@ -11069,6 +11075,7 @@ void UCookOnTheFlyServer::StartCookByTheBook( const FCookByTheBookStartupOptions
 	CreateSandboxFile(BeginContext);
 	LoadBeginCookConfigSettings(BeginContext);
 	SelectSessionPlatforms(BeginContext);
+	UE::Cook::FDelegates::ConfigureCookSession.Broadcast(*this);
 	CalculateGlobalDependenciesHashes();
 	RegisterShaderLibraryIncrementalCookArtifact(BeginContext);
 	LoadBeginCookIncrementalFlags(BeginContext);
@@ -11345,6 +11352,7 @@ void UCookOnTheFlyServer::StartCookAsCookWorker()
 	CreateSandboxFile(BeginContext);
 	LoadBeginCookConfigSettings(BeginContext);
 	SelectSessionPlatforms(BeginContext);
+	UE::Cook::FDelegates::ConfigureCookSession.Broadcast(*this);
 	CalculateGlobalDependenciesHashes();
 	LoadBeginCookIncrementalFlags(BeginContext);
 	if (IsDirectorCookOnTheFly())
diff --git a/Engine/Source/Runtime/CoreUObject/Private/UObject/ICookInfo.cpp b/Engine/Source/Runtime/CoreUObject/Private/UObject/ICookInfo.cpp
--- a/Engine/Source/Runtime/CoreUObject/Private/UObject/ICookInfo.cpp
+++ b/Engine/Source/Runtime/CoreUObject/Private/UObject/ICookInfo.cpp
@@ -84,6 +84,7 @@ FString FInstigator::ToString() const
 	return FString(Result);
 }
 
+FCookInfoEvent FDelegates::ConfigureCookSession;
 FCookInfoEvent FDelegates::CookStarted;
 FCookInfoEvent FDelegates::CookFinished;
 FValidateSourcePackage FDelegates::ValidateSourcePackage;
diff --git a/Engine/Source/Runtime/CoreUObject/Public/UObject/ICookInfo.h b/Engine/Source/Runtime/CoreUObject/Public/UObject/ICookInfo.h
--- a/Engine/Source/Runtime/CoreUObject/Public/UObject/ICookInfo.h
+++ b/Engine/Source/Runtime/CoreUObject/Public/UObject/ICookInfo.h
@@ -319,6 +319,11 @@ DECLARE_MULTICAST_DELEGATE_TwoParams(FCookLoadIdleEvent, ICookInfo& /* CookInfo
 struct FDelegates
 {
 public:
+	/**
+	 * Called as early as possible after parameters for the current cook session are loaded to allow modifying cook
+	 * startup behavior, such as registering extra artifacts or collectors.
+	 */
+	static COREUOBJECT_API FCookInfoEvent ConfigureCookSession;
 	/** Called after a cook session has been initialized and is about to start ticking and loading/saving packages. */
 	static COREUOBJECT_API FCookInfoEvent CookStarted;
 	/** Called at the end of a cook session, after writing all cook artifacts. */
