EditorViewportClient - call ModeTools->DrawHUD from Draw vs. DrawCanvas...

#Synced-CL 51530715

diff --git a/Engine/Source/Editor/UnrealEd/Private/EditorViewportClient.cpp b/Engine/Source/Editor/UnrealEd/Private/EditorViewportClient.cpp
--- a/Engine/Source/Editor/UnrealEd/Private/EditorViewportClient.cpp
+++ b/Engine/Source/Editor/UnrealEd/Private/EditorViewportClient.cpp
@@ -4668,9 +4668,6 @@ void FEditorViewportClient::DrawCanvas(FViewport& InViewport, FSceneView& View,
 
 	// Render the marquee drag tool
 	RenderDragTool(&View, &Canvas);
-
-	// Draw any HUD from modes
-	ModeTools->DrawHUD(this, &InViewport, &View, &Canvas);
 }
 
 void FEditorViewportClient::SetupViewForRendering(FSceneViewFamily& ViewFamily, FSceneView& View)
@@ -4984,6 +4981,9 @@ void FEditorViewportClient::Draw(FViewport* InViewport, FCanvas* Canvas)
 	{
 		DrawCanvas( *Viewport, *View, *Canvas );
 
+		// Draw any HUD from modes
+		ModeTools->DrawHUD(this, Viewport, View, Canvas);
+
 		DrawSafeFrames(*Viewport, *View, *Canvas);
 	}
 
