[Cook/ALT]
- Fix missing package referencing scope when using async loading for preloading

#Synced-CL 51263423

diff --git a/Engine/Source/Editor/UnrealEd/Private/Cooker/CookPackagePreloader.cpp b/Engine/Source/Editor/UnrealEd/Private/Cooker/CookPackagePreloader.cpp
--- a/Engine/Source/Editor/UnrealEd/Private/Cooker/CookPackagePreloader.cpp
+++ b/Engine/Source/Editor/UnrealEd/Private/Cooker/CookPackagePreloader.cpp
@@ -10,6 +10,8 @@
 #include "HAL/PlatformMath.h"
 #include "Logging/LogMacros.h"
 #include "Misc/AssertionMacros.h"
+#include "Misc/PackageAccessTracking.h"
+#include "Misc/PackageAccessTrackingOps.h"
 #include "Misc/PackagePath.h"
 #include "Misc/PreloadableFile.h"
 #include "Misc/ScopeExit.h"
@@ -239,6 +241,12 @@ bool FPackagePreloader::TryPreload()
 	{
 		if (!AsyncRequest.IsValid())
 		{
+			// Declare the package as being referenced by the ScriptPackageNameEngine; we look for this in
+			// ProcessUnsolicitedPackages so we can avoid adding an error for the package being loaded without a
+			// referencer.
+			static FName ScriptPackageNameEngine(TEXT("/Script/Engine"));
+			UE_TRACK_REFERENCING_PACKAGE_SCOPED(ScriptPackageNameEngine, PackageAccessTrackingOps::NAME_Load);
+
 			PackageData.GetPackageDatas().GetMonitor().OnPreloadAllocatedChanged(PackageData, true);
 			AsyncRequest = MakeShared<FAsyncRequest>();
 			AsyncRequest->RequestID = LoadPackageAsync(
