Cooker: AssetRegistry: Bugfix for CookOnTheFly failing after the change to not copy cooked AssetRegistry from editor AssetRegistry. CookOnTheFly still needs that copy behavior since it needs to pass a bootstrap AssetRegistry to the game before cooking any packages.
#rnx
#Synced-CL 51558650

diff --git a/Engine/Source/Editor/UnrealEd/Private/Commandlets/AssetRegistryGenerator.cpp b/Engine/Source/Editor/UnrealEd/Private/Commandlets/AssetRegistryGenerator.cpp
--- a/Engine/Source/Editor/UnrealEd/Private/Commandlets/AssetRegistryGenerator.cpp
+++ b/Engine/Source/Editor/UnrealEd/Private/Commandlets/AssetRegistryGenerator.cpp
@@ -1303,6 +1303,14 @@ void FAssetRegistryGenerator::Initialize()
 	FGameDelegates::Get().GetAssignLayerChunkDelegate() = FAssignLayerChunkDelegate::CreateStatic(AssignLayerChunkDelegate);
 }
 
+void FAssetRegistryGenerator::InitializeAssetRegistryForCookOnTheFly(const FAssetRegistrySerializationOptions& DevelopmentSaveOptions)
+{
+	// If the asset registry is still doing its background scan, we need to wait for it to finish and tick it so that the results are flushed out
+	AssetRegistry.WaitForCompletion();
+	ensureMsgf(!AssetRegistry.IsLoadingAssets(), TEXT("Cannot initialize asset registry generator while asset registry is still scanning source assets "));
+	AssetRegistry.InitializeTemporaryAssetRegistryState(State, DevelopmentSaveOptions);
+}
+
 static UE::Cook::ECookResult DiskSizeToCookResult(int64 DiskSize)
 {
 	if (DiskSize >= 0)
diff --git a/Engine/Source/Editor/UnrealEd/Private/CookOnTheFlyServer.cpp b/Engine/Source/Editor/UnrealEd/Private/CookOnTheFlyServer.cpp
--- a/Engine/Source/Editor/UnrealEd/Private/CookOnTheFlyServer.cpp
+++ b/Engine/Source/Editor/UnrealEd/Private/CookOnTheFlyServer.cpp
@@ -805,7 +805,12 @@ void UCookOnTheFlyServer::AddCookOnTheFlyPlatformFromGameThread(ITargetPlatform*
 	// Initialize systems that need to write files to the sandbox directory, for consumption later in AddCookOnTheFlyPlatformFromGameThread
 	// Functions in this section are not dependent upon each other and can be ordered arbitrarily or for async performance
 	InitializeShadersForCookOnTheFly(BeginContext.TargetPlatforms);
-	// SaveAssetRegistry is done in CookByTheBookFinished for CBTB, but we need at the start of CookOnTheFly to send as startup information to connecting clients
+	// SaveAssetRegistry is done in CookByTheBookFinished for CBTB, but we need at the start of CookOnTheFly to send
+	// as startup information to connecting clients
+	// We copy the data from the global AssetRegistry since we don't have the up to date data from AssetRegistry
+	// saves yet. 
+	PlatformData->RegistryGenerator->InitializeAssetRegistryForCookOnTheFly(
+		PlatformData->AssetRegistryDevelopmentSaveOptions);
 	uint64 DevelopmentAssetRegistryHash = 0;
 	PlatformData->RegistryGenerator->SaveAssetRegistry(GetSandboxAssetRegistryFilename(), true, false,
 		DevelopmentAssetRegistryHash, TSet<FName>(), PlatformData->AssetRegistryDevelopmentSaveOptions);
diff --git a/Engine/Source/Editor/UnrealEd/Public/Commandlets/AssetRegistryGenerator.h b/Engine/Source/Editor/UnrealEd/Public/Commandlets/AssetRegistryGenerator.h
--- a/Engine/Source/Editor/UnrealEd/Public/Commandlets/AssetRegistryGenerator.h
+++ b/Engine/Source/Editor/UnrealEd/Public/Commandlets/AssetRegistryGenerator.h
@@ -50,6 +50,11 @@ public:
 	~FAssetRegistryGenerator();
 
 	void Initialize();
+	/**
+	 * Called at the beginning of CookOnTheFly to bootstrap the targetplatform's AssetRegistry by copying it from the
+	 * editor's AssetRegistry.
+	 */
+	void InitializeAssetRegistryForCookOnTheFly(const FAssetRegistrySerializationOptions& DevelopmentSaveOptions);
 
 	const ITargetPlatform* GetTargetPlatform() const { return TargetPlatform; }
 
