ExternalActorsCommandlet: fix disabling OFPA on a partitioned world with straming enabled when using the -Force argument.

#Synced-CL 51536473

diff --git a/Engine/Source/Editor/UnrealEd/Private/Commandlets/ExternalActorsCommandlet.cpp b/Engine/Source/Editor/UnrealEd/Private/Commandlets/ExternalActorsCommandlet.cpp
--- a/Engine/Source/Editor/UnrealEd/Private/Commandlets/ExternalActorsCommandlet.cpp
+++ b/Engine/Source/Editor/UnrealEd/Private/Commandlets/ExternalActorsCommandlet.cpp
@@ -210,15 +210,19 @@ int32 UExternalActorsCommandlet::Main(const FString& Params)
 			{
 				int32 NumErrors = 0;
 
+				WorldPartition->Initialize(MainWorld, FTransform::Identity);
+
 				if (WorldPartition->IsStreamingEnabled())
 				{
-					UE_LOG(LogExternalActorsCommandlet, Error, TEXT("Cannot disable external actors for partitioned level '%s' with streaming enabled"), *FullMapName);
-					NumErrors++;
-
 					if (bForce)
 					{
 						WorldPartition->SetEnableStreaming(false);
 					}
+					else
+					{
+						UE_LOG(LogExternalActorsCommandlet, Error, TEXT("Cannot disable external actors for partitioned level '%s' with streaming enabled"), *FullMapName);
+						NumErrors++;
+					}
 				}
 
 				if (UDataLayerManager* DataLayerManager = WorldPartition->GetDataLayerManager())
@@ -246,8 +250,6 @@ int32 UExternalActorsCommandlet::Main(const FString& Params)
 
 				if (!bReport)
 				{
-					WorldPartition->Initialize(MainWorld, FTransform::Identity);
-
 					TArray<FString> LocalPackagesToDelete;
 					WorldPartition->ForEachActorDescContainerInstance([&LocalPackagesToDelete](UActorDescContainerInstance* ActorDescContainerInstance)
 					{
