Animation and Level Editor Viewports - Remove audio listener position log spam
- Also updated LevelEditorViewport logs to use UE_LOGFMT

#rnx

#Synced-CL 51460112

diff --git a/Engine/Source/Editor/Persona/Private/AnimationEditorViewportClient.cpp b/Engine/Source/Editor/Persona/Private/AnimationEditorViewportClient.cpp
--- a/Engine/Source/Editor/Persona/Private/AnimationEditorViewportClient.cpp
+++ b/Engine/Source/Editor/Persona/Private/AnimationEditorViewportClient.cpp
@@ -2825,7 +2825,6 @@ void FAnimationViewportClient::UpdateAudioListener() const
 			ListenerTransform.SetLocation(ViewLocation);
 
 			AudioDevice->SetListener(ViewportWorld, 0, ListenerTransform, 0.0f);
-			UE_LOGFMT(LogAnimation, Log, "Animation Viewport {Index}, {Transform}", ViewIndex, ListenerTransform.ToHumanReadableString());
 		}
 	}
 }
diff --git a/Engine/Source/Editor/UnrealEd/Private/LevelEditorViewport.cpp b/Engine/Source/Editor/UnrealEd/Private/LevelEditorViewport.cpp
--- a/Engine/Source/Editor/UnrealEd/Private/LevelEditorViewport.cpp
+++ b/Engine/Source/Editor/UnrealEd/Private/LevelEditorViewport.cpp
@@ -702,7 +702,7 @@ UObject* FLevelEditorViewportClient::GetOrCreateMaterialFromTexture(UTexture* Un
 		{
 			return FoundAsset;
 		}
-		UE_LOG(LogEditorViewport, Warning, TEXT("Failed to create material %s from texture because a non-material asset already exists with that name"), *NewPackageName);
+		UE_LOGFMT(LogEditorViewport, Warning, "Failed to create material {NewPackageName} from texture because a non-material asset already exists with that name", NewPackageName);
 		return nullptr;
 	}
 
@@ -762,7 +762,7 @@ UObject* FLevelEditorViewportClient::GetOrCreateMaterialFromTexture(UTexture* Un
 			}
 			else
 			{
-				UE_LOG(LogEditorViewport, Warning, TEXT("Dropped texture not assigned to material instance %s because no material parameter name was defined for Unknown or Base in LevelEditorViewportSettings"), *MaterialFullName);
+				UE_LOGFMT(LogEditorViewport, Warning, "Dropped texture not assigned to material instance {MaterialFullName} because no material parameter name was defined for Unknown or Base in LevelEditorViewportSettings", MaterialFullName);
 			}
 		}
 		else
@@ -790,7 +790,7 @@ UObject* FLevelEditorViewportClient::GetOrCreateMaterialFromTexture(UTexture* Un
 			bAssignedInputTexture |= AssignParam(EMaterialKind::Emissive, EmissiveSuffixes) && MaterialKind == EMaterialKind::Emissive;
 			if (!bAssignedInputTexture)
 			{
-				UE_LOG(LogEditorViewport, Warning, TEXT("Dropped texture not assigned to material instance %s because no material parameter name was defined in LevelEditorViewportSettings"), *MaterialFullName);
+				UE_LOGFMT(LogEditorViewport, Warning, "Dropped texture not assigned to material instance {MaterialFullName} because no material parameter name was defined in LevelEditorViewportSettings", MaterialFullName);
 			}
 		}
 	}
@@ -5224,7 +5224,6 @@ void FLevelEditorViewportClient::UpdateAudioListener()
 			ListenerTransform.SetLocation(ViewLocation);
 
 			AudioDevice->SetListener(ViewportWorld, 0, ListenerTransform, 0.f);
-			UE_LOGFMT(LogEditorViewport, Log, "Level Editor Viewport {Index}, {Transform}", ViewIndex, ListenerTransform.ToHumanReadableString());
 		}
 	}
 }
