TEDS: Register ULayers as rows with TEDS compat
Add a column + widget on actor rows for the layers the actor belongs to

#Synced-CL 51628790

diff --git a/Engine/Plugins/Experimental/EditorDataStorageFeatures/Source/TedsOutliner/Private/Compatibility/SceneOutlinerTedsBridge.cpp b/Engine/Plugins/Experimental/EditorDataStorageFeatures/Source/TedsOutliner/Private/Compatibility/SceneOutlinerTedsBridge.cpp
--- a/Engine/Plugins/Experimental/EditorDataStorageFeatures/Source/TedsOutliner/Private/Compatibility/SceneOutlinerTedsBridge.cpp
+++ b/Engine/Plugins/Experimental/EditorDataStorageFeatures/Source/TedsOutliner/Private/Compatibility/SceneOutlinerTedsBridge.cpp
@@ -2,6 +2,7 @@
 
 #include "Compatibility/SceneOutlinerTedsBridge.h"
 
+#include "Columns/TedsActorMobilityColumns.h"
 #include "Columns/TedsOutlinerColumns.h"
 #include "Elements/Columns/TypedElementCompatibilityColumns.h"
 #include "Elements/Columns/TypedElementLabelColumns.h"
@@ -9,13 +10,13 @@
 #include "Elements/Columns/TypedElementPackageColumns.h"
 #include "Elements/Columns/TypedElementTypeInfoColumns.h"
 #include "Elements/Columns/TypedElementVisibilityColumns.h"
-#include "Columns/TedsActorMobilityColumns.h"
 #include "Columns/TedsActorSocketColumns.h"
 #include "Elements/Interfaces/TypedElementDataStorageInterface.h"
 #include "Elements/Interfaces/TypedElementDataStorageUiInterface.h"
+#include "Layers/Columns/LayersColumns.h"
+#include "Modules/ModuleManager.h"
 #include "SceneOutlinerPublicTypes.h"
 #include "TedsOutlinerHelpers.h"
-#include "Modules/ModuleManager.h"
 #include "Widgets/Docking/SDockTab.h"
 
 #include UE_INLINE_GENERATED_CPP_BY_NAME(SceneOutlinerTedsBridge)
@@ -39,6 +40,7 @@ void USceneOutlinerTedsBridgeFactory::PreRegister(UE::Editor::DataStorage::ICore
 	TEDSToOutlinerDefaultColumnMapping.Add(FLevelColumn::StaticStruct(), FSceneOutlinerBuiltInColumnTypes::Level());
 	TEDSToOutlinerDefaultColumnMapping.Add(FTedsPrimaryPackageObjectTag::StaticStruct(), FSceneOutlinerBuiltInColumnTypes::Unsaved());
 	TEDSToOutlinerDefaultColumnMapping.Add(FTedsActorSocketColumn::StaticStruct(), FSceneOutlinerBuiltInColumnTypes::Socket());
+	TEDSToOutlinerDefaultColumnMapping.Add(FEditorDataStorageActorLayersColumn::StaticStruct(), FSceneOutlinerBuiltInColumnTypes::Layer());
 }
 
 void USceneOutlinerTedsBridgeFactory::RegisterTables(UE::Editor::DataStorage::ICoreProvider& DataStorage)
diff --git a/Engine/Plugins/Experimental/EditorDataStorageFeatures/Source/TedsOutliner/Private/TedsOutlinerModule.cpp b/Engine/Plugins/Experimental/EditorDataStorageFeatures/Source/TedsOutliner/Private/TedsOutlinerModule.cpp
--- a/Engine/Plugins/Experimental/EditorDataStorageFeatures/Source/TedsOutliner/Private/TedsOutlinerModule.cpp
+++ b/Engine/Plugins/Experimental/EditorDataStorageFeatures/Source/TedsOutliner/Private/TedsOutlinerModule.cpp
@@ -28,6 +28,7 @@
 #include "Modules/ModuleManager.h"
 #include "Compatibility/SceneOutlinerTedsBridge.h"
 #include "Compatibility/SceneOutlinerRowHandleColumn.h"
+#include "Layers/Columns/LayersColumns.h"
 #include "Widgets/Docking/SDockTab.h"
 
 #define LOCTEXT_NAMESPACE "TedsOutlinerModule"
@@ -154,6 +155,7 @@ FTedsOutlinerColumnDescription FTedsOutlinerModule::GetLevelEditorTedsOutlinerCo
 		FTypedElementClassTypeInfoColumn::StaticStruct(),
 		FLevelColumn::StaticStruct(), // PackageShortName
 		FTedsActorMobilityColumn::StaticStruct(),
+		FEditorDataStorageActorLayersColumn::StaticStruct(),
 		FTedsActorSocketColumn::StaticStruct(),
 		FUObjectIdNameColumn::StaticStruct(),
 		FAlertColumn::StaticStruct(),
@@ -194,7 +196,7 @@ TSharedRef<SWidget> FTedsOutlinerModule::CreateLevelEditorTedsOutliner()
 	const FQueryDescription OutlinerQueryDescription =
 		Select()
 		.Where()
-			.All<FTypedElementUObjectColumn, FTypedElementWorldColumn>()
+			.All<FTypedElementUObjectColumn, FTypedElementLabelColumn, FTypedElementWorldColumn>()
 			.None<FHideRowFromUITag>()
 		.Compile();
 
diff --git a/Engine/Source/Editor/UnrealEd/Private/Layers/Factories/LayersFactory.cpp b/Engine/Source/Editor/UnrealEd/Private/Layers/Factories/LayersFactory.cpp
new file mode 100644
--- /dev/null
+++ b/Engine/Source/Editor/UnrealEd/Private/Layers/Factories/LayersFactory.cpp
@@ -0,0 +1,84 @@
+// Copyright Epic Games, Inc. All Rights Reserved.
+
+#include "LayersFactory.h"
+
+#include "DataStorage/Features.h"
+#include "Editor/UnrealEd/Private/Layers/Widgets/ActorLayersWidget.h"
+#include "Elements/Columns/TypedElementCompatibilityColumns.h"
+#include "Elements/Columns/TypedElementMiscColumns.h"
+#include "Elements/Framework/TypedElementQueryBuilder.h"
+#include "Elements/Interfaces/TypedElementDataStorageCompatibilityInterface.h"
+#include "Layers/Layer.h"
+#include "Layers/Columns/LayersColumns.h"
+
+namespace UE::Editor::Layers::Private
+{
+	static FName TableName(TEXT("LayersTable"));
+}
+
+void UEditorDataStorageLayersFactory::RegisterTables(UE::Editor::DataStorage::ICoreProvider& DataStorage)
+{
+	using namespace UE::Editor::Layers;
+	using namespace UE::Editor::DataStorage;
+	
+	LayersTableHandle = DataStorage.RegisterTable<FLayerTag, FLayerNameColumn>(UE::Editor::Layers::Private::TableName);
+	
+	if (ICompatibilityProvider* TedsCompat = GetMutableDataStorageFeature<ICompatibilityProvider>(CompatibilityFeatureName))
+	{
+		TedsCompat->RegisterTypeTableAssociation(ULayer::StaticClass(), LayersTableHandle);
+	}
+}
+
+void UEditorDataStorageLayersFactory::RegisterQueries(UE::Editor::DataStorage::ICoreProvider& DataStorage)
+{
+	using namespace UE::Editor::DataStorage;
+	using namespace UE::Editor::DataStorage::Queries;
+	
+	DataStorage.RegisterQuery(
+		Select(
+			TEXT("Sync column to layer name"),
+			FProcessor(EQueryTickPhase::PrePhysics, DataStorage.GetQueryTickGroupName(EQueryTickGroups::SyncExternalToDataStorage))
+				.SetExecutionMode(EExecutionMode::GameThread),
+			[](const FTypedElementUObjectColumn& Layer, FEditorDataStorageLayerNameColumn& NameColumn)
+			{
+				if (const ULayer* LayerInstance = Cast<ULayer>(Layer.Object))
+				{
+					NameColumn.LayerName = LayerInstance->GetLayerName();
+				}
+			}
+		)
+		.Where()
+			.All<FEditorDataStorageLayerTag, FTypedElementSyncFromWorldTag>()
+		.Compile());
+	
+	DataStorage.RegisterQuery(
+		Select(
+			TEXT("Sync layer name to column"),
+			FProcessor(EQueryTickPhase::PostPhysics, DataStorage.GetQueryTickGroupName(EQueryTickGroups::SyncDataStorageToExternal))
+				.SetExecutionMode(EExecutionMode::GameThread),
+			[](FTypedElementUObjectColumn& Layer, const FEditorDataStorageLayerNameColumn& NameColumn)
+			{
+				if (ULayer* LayerInstance = Cast<ULayer>(Layer.Object))
+				{
+					LayerInstance->SetLayerName(NameColumn.LayerName);
+				}
+			}
+		)
+		.Where()
+			.All<FEditorDataStorageLayerTag, FTypedElementSyncBackToWorldTag>()
+		.Compile());
+}
+
+void UEditorDataStorageLayersFactory::RegisterWidgetConstructors(UE::Editor::DataStorage::ICoreProvider& DataStorage,
+	UE::Editor::DataStorage::IUiProvider& DataStorageUi) const
+{
+	using namespace UE::Editor::DataStorage::Queries;
+	using namespace UE::Editor::Layers;
+
+	DataStorageUi.RegisterWidgetFactory<FActorLayersWidgetConstructor>(DataStorageUi.FindPurpose(DataStorageUi.GetGeneralWidgetPurposeID()), TColumn<FActorLayersColumn>());
+}
+
+UE::Editor::DataStorage::TableHandle UEditorDataStorageLayersFactory::GetLayersTableHandle() const
+{
+	return LayersTableHandle;
+}
\ No newline at end of file
diff --git a/Engine/Source/Editor/UnrealEd/Private/Layers/Factories/LayersFactory.h b/Engine/Source/Editor/UnrealEd/Private/Layers/Factories/LayersFactory.h
new file mode 100644
--- /dev/null
+++ b/Engine/Source/Editor/UnrealEd/Private/Layers/Factories/LayersFactory.h
@@ -0,0 +1,29 @@
+// Copyright Epic Games, Inc. All Rights Reserved.
+
+#pragma once
+
+#include "Elements/Framework/TypedElementHandle.h"
+#include "Elements/Interfaces/TypedElementDataStorageInterface.h"
+#include "Elements/Interfaces/TypedElementDataStorageFactory.h"
+
+#include "LayersFactory.generated.h"
+
+/**
+ */
+UCLASS(Transient)
+class UEditorDataStorageLayersFactory : public UEditorDataStorageFactory
+{
+	GENERATED_BODY()
+public:
+
+	virtual void RegisterTables(UE::Editor::DataStorage::ICoreProvider& DataStorage) override;
+	virtual void RegisterQueries(UE::Editor::DataStorage::ICoreProvider& DataStorage) override;
+	virtual void RegisterWidgetConstructors(UE::Editor::DataStorage::ICoreProvider& DataStorage, UE::Editor::DataStorage::IUiProvider& DataStorageUi) const override;
+	
+	UE::Editor::DataStorage::TableHandle GetLayersTableHandle() const;
+	
+private:
+	UE::Editor::DataStorage::TableHandle LayersTableHandle;
+};
+
+
diff --git a/Engine/Source/Editor/UnrealEd/Private/Layers/LayersSubsystem.cpp b/Engine/Source/Editor/UnrealEd/Private/Layers/LayersSubsystem.cpp
--- a/Engine/Source/Editor/UnrealEd/Private/Layers/LayersSubsystem.cpp
+++ b/Engine/Source/Editor/UnrealEd/Private/Layers/LayersSubsystem.cpp
@@ -12,6 +12,8 @@
 #include "Editor.h"
 
 #include "ActorEditorUtils.h"
+#include "DataStorage/Features.h"
+#include "Layers/Columns/LayersColumns.h"
 #include "Model.h"
 
 #include UE_INLINE_GENERATED_CPP_BY_NAME(LayersSubsystem)
@@ -53,6 +55,19 @@ private:
 	bool bIsInitialized;
 };
 
+namespace UE::Editor::Layers::Private
+{
+	static bool bRegisterLayersInTeds = true;
+	
+	FAutoConsoleVariableRef CVarRegisterLayersInTeds(
+		TEXT("TEDS.Feature.Layers"),
+		bRegisterLayersInTeds,
+		TEXT("Registers Layers in TEDS"),
+		ECVF_ReadOnly);
+}
+
+
+
 FLayersBroadcast::FLayersBroadcast(ULayersSubsystem* InLayersSubsystem)
 	: LayersSubsystem(InLayersSubsystem)
 	, bIsInitialized(false)
@@ -110,12 +125,37 @@ void ULayersSubsystem::Initialize(FSubsystemCollectionBase& Collection)
 
 	// Subsystems are loaded, so we can set up the broadcast functions for GetEditorSubsystem<ULayersSubsystem>()
 	LayersBroadcast = MakeShareable(new FLayersBroadcast(this));
+	
+	if (UE::Editor::Layers::Private::bRegisterLayersInTeds)
+	{
+		using namespace UE::Editor::DataStorage;
+
+		auto OnDataStorage = [this]
+			{
+				DataStorage = GetMutableDataStorageFeature<ICoreProvider>(StorageFeatureName);
+				DataStorageCompat = GetMutableDataStorageFeature<ICompatibilityProvider>(CompatibilityFeatureName);
+			};
+
+		if (AreEditorDataStorageFeaturesEnabled())
+		{
+			OnDataStorage();
+		}
+		else
+		{
+			OnDataStorageInitialized = OnEditorDataStorageFeaturesEnabled().AddLambda(OnDataStorage);
+		}
+	}
 }
 
 void ULayersSubsystem::Deinitialize()
 {
 	Super::Deinitialize();
-
+	
+	if (UE::Editor::Layers::Private::bRegisterLayersInTeds)
+	{
+		UE::Editor::DataStorage::OnEditorDataStorageFeaturesEnabled().Remove(OnDataStorageInitialized);
+	}
+	
 	LayersBroadcast->Deinitialize();
 }
 
@@ -139,7 +179,7 @@ void ULayersSubsystem::EditorRefreshLayerBrowser()
 void ULayersSubsystem::PostUndoRedo()
 {
 	LayersChanged.Broadcast(ELayersAction::Reset, NULL, NAME_None);
-	UpdateAllActorsVisibility(true, true);
+	UpdateAllActorsVisibility_Internal(true, true, true);
 }
 
 ///////////////////////////////////////////////////////////////////////////////////////////////////////////////////////
@@ -196,6 +236,7 @@ bool ULayersSubsystem::InitializeNewActorLayers(AActor* Actor)
 
 		Layer->Modify();
 		AddActorToStats( Layer, Actor);
+		AddActorToLayerInTeds(Layer, Actor);
 	}
 
 	// update per-view visibility info
@@ -297,6 +338,7 @@ bool ULayersSubsystem::AddActorsToLayers( const TArray< AActor* >& Actors, const
 
 					Layer->Modify();
 					AddActorToStats( Layer, Actor);
+					AddActorToLayerInTeds(Layer, Actor);
 
 					ActorsLayersChanged.Broadcast( Actor );
 				}
@@ -353,6 +395,7 @@ bool ULayersSubsystem::DisassociateActorsFromLayers(const TArray<AActor*>& Actor
 
 			Layer->Modify();
 			RemoveActorFromStats(Layer, Actor);
+			RemoveActorFromLayerInTeds(Layer, Actor);
 			bChangesOccurred = true;
 		}
 	}
@@ -424,12 +467,17 @@ bool ULayersSubsystem::RemoveActorsFromLayers( const TArray< AActor* >& Actors,
 				Actor->Layers.Remove( LayerName );
 
 				ULayer* Layer;
-				if( bUpdateStats && TryGetLayer( LayerName, Layer ))
+				if (TryGetLayer( LayerName, Layer ))
 				{
-					Layer->Modify();
-					RemoveActorFromStats( Layer, Actor);
+					if( bUpdateStats)
+					{
+						Layer->Modify();
+						RemoveActorFromStats( Layer, Actor);
+					}
+					
+					RemoveActorFromLayerInTeds(Layer, Actor);
 				}
-
+				
 				ActorsLayersChanged.Broadcast( Actor );
 			}
 		} //END Iteration over Layers
@@ -919,32 +967,7 @@ bool ULayersSubsystem::UpdateActorVisibility(AActor* Actor, bool& bOutSelectionC
 
 bool ULayersSubsystem::UpdateAllActorsVisibility(const bool bNotifySelectionChange, const bool bRedrawViewports)
 {
-	bool bSelectionChanged = false;
-	bool bChangesOccurred = false;
-	for( FActorIterator It(GetWorld()) ; It ; ++It )
-	{
-		AActor* Actor = *It;
-
-		bool bActorModified = false;
-		bool bActorSelectionChanged = false;
-		const bool bActorNotifySelectionChange = false;
-		const bool bActorRedrawViewports = false;
-
-		bChangesOccurred |= UpdateActorVisibility( Actor, bActorSelectionChanged /*OUT*/, bActorModified /*OUT*/, bActorNotifySelectionChange, bActorRedrawViewports );
-		bSelectionChanged |= bActorSelectionChanged;
-	}
-
-	if ( bNotifySelectionChange && bSelectionChanged )
-	{
-		GEditor->NoteSelectionChange();
-	}
-
-	if( bRedrawViewports )
-	{
-		GEditor->RedrawLevelEditingViewports();
-	}
-
-	return bChangesOccurred;
+	return UpdateAllActorsVisibility_Internal(bNotifySelectionChange, bRedrawViewports);
 }
 
 
@@ -1203,7 +1226,15 @@ ULayer* ULayersSubsystem::CreateLayer(const FName& LayerName)
 
 	NewLayer->SetLayerName(LayerName);
 	NewLayer->SetVisible(true);
-
+	
+	if (UE::Editor::Layers::Private::bRegisterLayersInTeds)
+	{
+		if (DataStorage && DataStorageCompat)
+		{
+			UE::Editor::DataStorage::RowHandle LayerRow = DataStorageCompat->AddCompatibleObject(NewLayer);
+		}
+	}
+	
 	LayersChanged.Broadcast(ELayersAction::Add, NewLayer, NAME_None);
 
 	return NewLayer;
@@ -1237,6 +1268,14 @@ void ULayersSubsystem::DeleteLayers( const TArray< FName >& LayersToDelete )
 	{
 		if( LayersToDelete.Contains( GetWorld()->Layers[ LayerIndex]->GetLayerName() ) )
 		{
+			if (UE::Editor::Layers::Private::bRegisterLayersInTeds)
+			{
+				if (DataStorage && DataStorageCompat)
+				{
+					DataStorageCompat->RemoveCompatibleObject(GetWorld()->Layers[ LayerIndex]);
+				}
+			}
+			
 			GetWorld()->Modify(false);
 			GetWorld()->Layers.RemoveAt( LayerIndex );
 			bValidLayerExisted = true;
@@ -1269,6 +1308,14 @@ void ULayersSubsystem::DeleteLayer( const FName& LayerToDelete )
 	{
 		if( LayerToDelete == GetWorld()->Layers[ LayerIndex]->GetLayerName() )
 		{
+			if (UE::Editor::Layers::Private::bRegisterLayersInTeds)
+			{
+				if (DataStorage && DataStorageCompat)
+				{
+					DataStorageCompat->RemoveCompatibleObject(GetWorld()->Layers[ LayerIndex]);
+				}
+			}
+			
 			GetWorld()->Modify(false);
 			GetWorld()->Layers.RemoveAt( LayerIndex );
 			bValidLayerExisted = true;
@@ -1361,6 +1408,145 @@ void ULayersSubsystem::RemoveActorFromStats(ULayer* Layer, AActor* Actor)
 	}
 }
 
+void ULayersSubsystem::AddActorToLayerInTeds(ULayer* Layer, AActor* Actor)
+{
+	using namespace UE::Editor::DataStorage;
+	using namespace UE::Editor::Layers;
+	
+	if (UE::Editor::Layers::Private::bRegisterLayersInTeds)
+	{
+		if (DataStorage && DataStorageCompat)
+		{
+			RowHandle LayerRow = DataStorageCompat->FindRowWithCompatibleObject(Layer);
+			RowHandle ActorRow = DataStorageCompat->FindRowWithCompatibleObject(Actor);
+			
+			if (DataStorage->IsRowAvailable(LayerRow))
+			{
+				if (FActorLayersColumn* LayersColumn = DataStorage->GetColumn<FActorLayersColumn>(ActorRow))
+				{
+					LayersColumn->Layers.Add(LayerRow);
+				
+					// This also sorts the array
+					LayersColumn->Layers.MakeUnique();
+				}
+				else
+				{
+					FActorLayersColumn NewLayersColumn;
+					NewLayersColumn.Layers.Add(LayerRow, FRowHandleArrayView::EFlags::IsSorted);
+					DataStorage->AddColumn(ActorRow, MoveTemp(NewLayersColumn));
+				}
+			}
+		}
+	}
+}
+
+void ULayersSubsystem::RemoveActorFromLayerInTeds(ULayer* Layer, AActor* Actor)
+{
+	using namespace UE::Editor::DataStorage;
+	using namespace UE::Editor::Layers;
+	
+	if (UE::Editor::Layers::Private::bRegisterLayersInTeds)
+	{
+		if (DataStorage && DataStorageCompat)
+		{
+			RowHandle LayerRow = DataStorageCompat->FindRowWithCompatibleObject(Layer);
+			RowHandle ActorRow = DataStorageCompat->FindRowWithCompatibleObject(Actor);
+			
+			if (FActorLayersColumn* LayersColumn = DataStorage->GetColumn<FActorLayersColumn>(ActorRow))
+			{
+				LayersColumn->Layers.Remove(LayerRow);
+				
+				if (LayersColumn->Layers.IsEmpty())
+				{
+					DataStorage->RemoveColumn<FActorLayersColumn>(ActorRow);
+				}
+			}
+		}
+	}
+}
+
+bool ULayersSubsystem::UpdateAllActorsVisibility_Internal(const bool bNotifySelectionChange, const bool bRedrawViewports, bool bUpdateTedsState)
+{
+	bool bSelectionChanged = false;
+	bool bChangesOccurred = false;
+	for( FActorIterator It(GetWorld()) ; It ; ++It )
+	{
+		AActor* Actor = *It;
+
+		bool bActorModified = false;
+		bool bActorSelectionChanged = false;
+		const bool bActorNotifySelectionChange = false;
+		const bool bActorRedrawViewports = false;
+
+		bChangesOccurred |= UpdateActorVisibility( Actor, bActorSelectionChanged /*OUT*/, bActorModified /*OUT*/, bActorNotifySelectionChange, bActorRedrawViewports );
+		bSelectionChanged |= bActorSelectionChanged;
+		
+		// The state in TEDS for the layers an actor belongs to might not be in sync after an undo/redo, since we are already iterating over all actors here
+		// we can use this opportunity to also make sure the layer column is updated
+		if (bUpdateTedsState)
+		{
+			SyncActorLayersToTeds(Actor);
+		}
+	}
+
+	if ( bNotifySelectionChange && bSelectionChanged )
+	{
+		GEditor->NoteSelectionChange();
+	}
+
+	if( bRedrawViewports )
+	{
+		GEditor->RedrawLevelEditingViewports();
+	}
+
+	return bChangesOccurred;
+}
+
+void ULayersSubsystem::SyncActorLayersToTeds(AActor* Actor)
+{
+	using namespace UE::Editor::DataStorage;
+	using namespace UE::Editor::Layers;
+	
+	if (UE::Editor::Layers::Private::bRegisterLayersInTeds)
+	{
+		if (DataStorage && DataStorageCompat)
+		{
+			RowHandle ActorRow = DataStorageCompat->FindRowWithCompatibleObject(Actor);
+			
+			if (DataStorage->IsRowAvailable(ActorRow))
+			{
+				// Get the state of the layers on the Actor UObject
+				FRowHandleArray NewLayers;
+		
+				for (const FName& LayerName : Actor->Layers)
+				{
+					ULayer* Layer = EnsureLayerExists(LayerName);
+					RowHandle LayerRow = DataStorageCompat->FindRowWithCompatibleObject(Layer);
+			
+					if (DataStorage->IsRowAvailable(LayerRow))
+					{
+						NewLayers.Add(LayerRow);
+					}
+				}
+		
+				NewLayers.Sort();
+		
+				// Overwrite the state in TEDS
+				if (NewLayers.Num())
+				{
+					FActorLayersColumn LayersColumn;
+					LayersColumn.Layers = MoveTemp(NewLayers);
+					DataStorage->AddColumn<FActorLayersColumn>(ActorRow, MoveTemp(LayersColumn));
+				}
+				else
+				{
+					DataStorage->RemoveColumn<FActorLayersColumn>(ActorRow);
+				}
+			}
+		}
+	}
+}
+
 ULayer* ULayersSubsystem::EnsureLayerExists(const FName& LayerName)
 {
 	ULayer* Layer;
diff --git a/Engine/Source/Editor/UnrealEd/Private/Layers/Widgets/ActorLayersWidget.cpp b/Engine/Source/Editor/UnrealEd/Private/Layers/Widgets/ActorLayersWidget.cpp
new file mode 100644
--- /dev/null
+++ b/Engine/Source/Editor/UnrealEd/Private/Layers/Widgets/ActorLayersWidget.cpp
@@ -0,0 +1,51 @@
+// Copyright Epic Games, Inc. All Rights Reserved.
+
+#include "ActorLayersWidget.h"
+
+#include "Elements/Framework/TypedElementAttributeBinding.h"
+#include "Layers/Columns/LayersColumns.h"
+#include "Widgets/SBoxPanel.h"
+#include "Widgets/Text/STextBlock.h"
+
+FActorLayersWidgetConstructor::FActorLayersWidgetConstructor()
+	: FSimpleWidgetConstructor(StaticStruct())
+{
+}
+
+TSharedPtr<SWidget> FActorLayersWidgetConstructor::CreateWidget(UE::Editor::DataStorage::ICoreProvider* DataStorage,
+                                                                UE::Editor::DataStorage::IUiProvider* DataStorageUi, UE::Editor::DataStorage::RowHandle TargetRow, UE::Editor::DataStorage::RowHandle WidgetRow,
+                                                                const UE::Editor::DataStorage::FMetaDataView& Arguments)
+{
+	using namespace UE::Editor::DataStorage;
+	using namespace UE::Editor::Layers;
+	
+	FAttributeBinder Binder(TargetRow, DataStorage);
+	
+	return SNew(SHorizontalBox)
+		+ SHorizontalBox::Slot()
+		.AutoWidth()
+		.VAlign(VAlign_Center)
+		.Padding(8, 0, 0, 0)
+		[
+			SNew(STextBlock)
+				.ColorAndOpacity(FSlateColor::UseSubduedForeground())
+				.Text(Binder.BindData<FText, FRowHandleArray>(&FActorLayersColumn::Layers, [DataStorage](const FRowHandleArray& Layers) -> FText
+				{
+					FString Result;
+						
+					for (RowHandle LayerRow : Layers.GetRows())
+					{
+						if (FLayerNameColumn* LayerNameColumn = DataStorage->GetColumn<FLayerNameColumn>(LayerRow))
+						{
+							if (Result.Len())
+							{
+								Result += TEXT(", ");
+							}
+
+							Result += LayerNameColumn->LayerName.ToString();
+						}
+					}
+					return FText::FromString(Result);
+				}))
+		];
+}
\ No newline at end of file
diff --git a/Engine/Source/Editor/UnrealEd/Private/Layers/Widgets/ActorLayersWidget.h b/Engine/Source/Editor/UnrealEd/Private/Layers/Widgets/ActorLayersWidget.h
new file mode 100644
--- /dev/null
+++ b/Engine/Source/Editor/UnrealEd/Private/Layers/Widgets/ActorLayersWidget.h
@@ -0,0 +1,27 @@
+// Copyright Epic Games, Inc. All Rights Reserved.
+
+#pragma once
+
+#include "Elements/Interfaces/TypedElementDataStorageFactory.h"
+#include "Elements/Interfaces/TypedElementDataStorageUiInterface.h"
+
+#include "ActorLayersWidget.generated.h"
+
+
+// Widget to show the layers an actor belongs to
+USTRUCT()
+struct FActorLayersWidgetConstructor : public FSimpleWidgetConstructor
+{
+	GENERATED_BODY()
+
+public:
+	FActorLayersWidgetConstructor();
+	~FActorLayersWidgetConstructor() override = default;
+	
+	virtual TSharedPtr<SWidget> CreateWidget(
+		UE::Editor::DataStorage::ICoreProvider* DataStorage,
+		UE::Editor::DataStorage::IUiProvider* DataStorageUi,
+		UE::Editor::DataStorage::RowHandle TargetRow,
+		UE::Editor::DataStorage::RowHandle WidgetRow,
+		const UE::Editor::DataStorage::FMetaDataView& Arguments) override;
+};
\ No newline at end of file
diff --git a/Engine/Source/Editor/UnrealEd/Public/Layers/Columns/LayersColumns.h b/Engine/Source/Editor/UnrealEd/Public/Layers/Columns/LayersColumns.h
new file mode 100644
--- /dev/null
+++ b/Engine/Source/Editor/UnrealEd/Public/Layers/Columns/LayersColumns.h
@@ -0,0 +1,46 @@
+// Copyright Epic Games, Inc. All Rights Reserved.
+
+#pragma once
+
+#include "Elements/Interfaces/TypedElementDataStorageInterface.h"
+
+#include "LayersColumns.generated.h"
+
+/**
+ * Tag that indicates this row in TEDS represents a ULayer
+ */
+USTRUCT(meta = (DisplayName = "Layer Tag"))
+struct FEditorDataStorageLayerTag final : public FEditorDataStorageTag
+{
+	GENERATED_BODY()
+};
+
+/**
+ * Column containing the name of a ULayer
+ */
+USTRUCT(meta = (DisplayName = "Layer Name Column"))
+struct FEditorDataStorageLayerNameColumn final : public FEditorDataStorageColumn
+{
+	GENERATED_BODY()
+	
+	UPROPERTY(meta = (Searchable, Sortable))
+	FName LayerName;
+};
+
+/**
+ * Column on actor rows that contains the layers this actor belongs to
+ */
+USTRUCT(meta = (DisplayName = "Layers"))
+struct FEditorDataStorageActorLayersColumn final : public FEditorDataStorageColumn
+{
+	GENERATED_BODY()
+
+	UE::Editor::DataStorage::FRowHandleArray Layers;
+};
+
+namespace UE::Editor::Layers
+{
+	using FLayerTag = FEditorDataStorageLayerTag;
+	using FLayerNameColumn = FEditorDataStorageLayerNameColumn;
+	using FActorLayersColumn = FEditorDataStorageActorLayersColumn;
+}
\ No newline at end of file
diff --git a/Engine/Source/Editor/UnrealEd/Public/Layers/LayersSubsystem.h b/Engine/Source/Editor/UnrealEd/Public/Layers/LayersSubsystem.h
--- a/Engine/Source/Editor/UnrealEd/Public/Layers/LayersSubsystem.h
+++ b/Engine/Source/Editor/UnrealEd/Public/Layers/LayersSubsystem.h
@@ -4,6 +4,8 @@
 
 #include "CoreMinimal.h"
 #include "EditorSubsystem.h"
+#include "Elements/Interfaces/TypedElementDataStorageCompatibilityInterface.h"
+#include "Elements/Interfaces/TypedElementDataStorageInterface.h"
 
 #include "LayersSubsystem.generated.h"
 
@@ -622,6 +624,19 @@ public:
 private:
 	UNREALED_API void AddActorToStats(ULayer* Layer, AActor* Actor);
 	UNREALED_API void RemoveActorFromStats(ULayer* Layer, AActor* Actor);
+	
+	// Update the layers column on the given actor's row in TEDS
+	void AddActorToLayerInTeds(ULayer* Layer, AActor* Actor);
+	void RemoveActorFromLayerInTeds(ULayer* Layer, AActor* Actor);
+	
+	// Internal function to update all actors visibility and also optionally update the TEDS state. This is only needed in specific 
+	// situations and not every time UpdateAllActorsVisibility is called, but we also don't want to iterate all actors twice so the internal function
+	// can be called manually after things like undo/redo when TEDS needs to be updated
+	bool UpdateAllActorsVisibility_Internal(const bool bNotifySelectionChange, const bool bRedrawViewports, bool bUpdateTedsState = false);
+	
+	// Force sync the layers state for the given actor from the world to TEDS (stomping the old state in TEDS)
+	void SyncActorLayersToTeds(AActor* Actor);
+
 
 	/**
 	 *	Checks to see if the named layer exists, and if it doesn't creates it.
@@ -642,4 +657,11 @@ private:
 	 * Note that UClasses objects do not allow broadcast, that is why we need this auxiliary class.
 	 */
 	TSharedPtr<class FLayersBroadcast> LayersBroadcast;
+	
+	/**	Cached Teds interfaces */
+	UE::Editor::DataStorage::ICoreProvider* DataStorage = nullptr;
+	UE::Editor::DataStorage::ICompatibilityProvider* DataStorageCompat = nullptr;
+	
+	/**	Delegate run when TEDS is first initialized */
+	FDelegateHandle OnDataStorageInitialized;
 };
