[Selection] Fix an issue where asset editors were having gizmos positions updated from the level editor with some default hook ups in the mode manager

#Synced-CL 51488868

diff --git a/Engine/Source/Editor/UnrealEd/Private/EditorModeManager.cpp b/Engine/Source/Editor/UnrealEd/Private/EditorModeManager.cpp
--- a/Engine/Source/Editor/UnrealEd/Private/EditorModeManager.cpp
+++ b/Engine/Source/Editor/UnrealEd/Private/EditorModeManager.cpp
@@ -89,11 +89,6 @@ FEditorModeTools::FEditorModeTools()
 	// Load the last used settings
 	LoadConfig();
 
-	// Register our callback for actor selection changes
-	USelection::SelectNoneEvent.AddRaw(this, &FEditorModeTools::OnEditorSelectNone);
-	USelection::SelectionChangedEvent.AddRaw(this, &FEditorModeTools::OnEditorSelectionChanged);
-	USelection::SelectObjectEvent.AddRaw(this, &FEditorModeTools::OnEditorSelectionChanged);
-
 	if( GEditor )
 	{
 		// Register our callback for undo/redo
@@ -579,65 +574,6 @@ bool FEditorModeTools::IsOnlyActiveMode(FEditorModeID InMode) const
 	return ActiveScriptableModes.Num() == 1 && ActiveScriptableModes[0]->GetID() == InMode;
 }
 
-void FEditorModeTools::OnEditorSelectionChanged(UObject* NewSelection)
-{
-	if (NewSelection == GetSelectedActors())
-	{
-		// when actors are selected check if there is at least one component selected and cache that off
-		// Editor modes use this primarily to determine of transform gizmos should be drawn.  
-		// Performing this check each frame with lots of actors is expensive so only do this when selection changes
-		bSelectionHasSceneComponent = false;
-		for(FSelectionIterator It(*GetSelectedActors()); It; ++It)
-		{
-			AActor* Actor = Cast<AActor>(*It);
-			if(Actor != nullptr && Actor->FindComponentByClass<USceneComponent>() != nullptr)
-			{
-				bSelectionHasSceneComponent = true;
-				break;
-			}
-		}
-
-	}
-	else
-	{
-		// If selecting an actor, move the pivot location.
-		AActor* Actor = Cast<AActor>(NewSelection);
-		if(Actor != nullptr)
-		{
-			if(Actor->IsSelected())
-			{
-				SetPivotLocation(Actor->GetActorLocation(), false);
-
-				// If this actor wasn't part of the original selection set during pie/sie, clear it now
-				if(GEditor->ActorsThatWereSelected.Num() > 0)
-				{
-					AActor* EditorActor = EditorUtilities::GetEditorWorldCounterpartActor(Actor);
-					if(!EditorActor || !GEditor->ActorsThatWereSelected.Contains(EditorActor))
-					{
-						GEditor->ActorsThatWereSelected.Empty();
-					}
-				}
-			}
-			else if(GEditor->ActorsThatWereSelected.Num() > 0)
-			{
-				// Clear the selection set
-				GEditor->ActorsThatWereSelected.Empty();
-			}
-		}
-	}
-
-	for(const auto& Pair : FEditorModeRegistry::Get().GetFactoryMap())
-	{
-		Pair.Value->OnSelectionChanged(*this, NewSelection);
-	}
-}
-
-void FEditorModeTools::OnEditorSelectNone()
-{
-	GEditor->SelectNone( false, true );
-	GEditor->ActorsThatWereSelected.Empty();
-}
-
 void FEditorModeTools::DrawBrackets(FEditorViewportClient* ViewportClient, FViewport* Viewport, const FSceneView* View, FCanvas* Canvas)
 {
 	if (!ViewportClient->IsPerspective() || !GetDefault<ULevelEditorViewportSettings>()->bHighlightWithBrackets)
@@ -827,14 +763,6 @@ void FEditorModeTools::RemoveAllDelegateHandlers()
 	}
 
 	FWorldDelegates::OnWorldCleanup.RemoveAll(this);
-	
-	// For now, check that UObjects are even valid, because the level editor has a global static mode tools
-	if (UObjectInitialized())
-	{
-		USelection::SelectionChangedEvent.RemoveAll(this);
-		USelection::SelectNoneEvent.RemoveAll(this);
-		USelection::SelectObjectEvent.RemoveAll(this);
-	}
 
 	OnEditorModeIDChanged().Clear();
 	OnWidgetModeChanged().Clear();
diff --git a/Engine/Source/Editor/UnrealEd/Private/UnrealEdGlobals.cpp b/Engine/Source/Editor/UnrealEd/Private/UnrealEdGlobals.cpp
--- a/Engine/Source/Editor/UnrealEd/Private/UnrealEdGlobals.cpp
+++ b/Engine/Source/Editor/UnrealEd/Private/UnrealEdGlobals.cpp
@@ -20,6 +20,7 @@
 #include "Misc/OutputDeviceConsole.h"
 #include "UnrealEngine.h"
 #include "UnrealEdMisc.h"
+#include "Selection.h"
 #include "EditorModes.h"
 #include "Framework/Application/SlateApplication.h"
 #include "DesktopPlatformModule.h"
@@ -48,12 +49,84 @@ DEFINE_LOG_CATEGORY_STATIC(LogUnrealEd, Log, All);
 namespace Internal
 {
 	static TSharedPtr<FEditorModeTools> EditorModeToolsSingleton;
-};
+	class FLevelEditorModeTools : public FEditorModeTools
+	{
+	public:
+		FLevelEditorModeTools()
+		{
+			USelection::SelectNoneEvent.AddRaw(this, &FLevelEditorModeTools::OnEditorSelectNone);
+			USelection::SelectionChangedEvent.AddRaw(this, &FLevelEditorModeTools::OnEditorSelectionChanged);
+			USelection::SelectObjectEvent.AddRaw(this, &FLevelEditorModeTools::OnEditorSelectionChanged);
+		}
+		virtual ~FLevelEditorModeTools()
+		{
+			if (UObjectInitialized())
+			{
+				USelection::SelectObjectEvent.RemoveAll(this);
+				USelection::SelectionChangedEvent.RemoveAll(this);
+				USelection::SelectNoneEvent.RemoveAll(this);
+			}
+		}
+	private:
+		void OnEditorSelectionChanged(UObject* NewSelection)
+		{
+			if (NewSelection == GetSelectedActors())
+			{
+				// when actors are selected check if there is at least one component selected and cache that off
+				// Editor modes use this primarily to determine of transform gizmos should be drawn.  
+				// Performing this check each frame with lots of actors is expensive so only do this when selection changes
+				bSelectionHasSceneComponent = false;
+				for (FSelectionIterator It(*GetSelectedActors()); It; ++It)
+				{
+					AActor* Actor = Cast<AActor>(*It);
+					if (Actor != nullptr && Actor->FindComponentByClass<USceneComponent>() != nullptr)
+					{
+						bSelectionHasSceneComponent = true;
+						break;
+					}
+				}
 
-bool GLevelEditorModeToolsIsValid()
-{
-	return Internal::EditorModeToolsSingleton.IsValid();
-}
+			}
+			else
+			{
+				// If selecting an actor, move the pivot location.
+				AActor* Actor = Cast<AActor>(NewSelection);
+				if (Actor != nullptr)
+				{
+					if (Actor->IsSelected())
+					{
+						SetPivotLocation(Actor->GetActorLocation(), false);
+
+						// If this actor wasn't part of the original selection set during pie/sie, clear it now
+						if (GEditor->ActorsThatWereSelected.Num() > 0)
+						{
+							AActor* EditorActor = EditorUtilities::GetEditorWorldCounterpartActor(Actor);
+							if (!EditorActor || !GEditor->ActorsThatWereSelected.Contains(EditorActor))
+							{
+								GEditor->ActorsThatWereSelected.Empty();
+							}
+						}
+					}
+					else if (GEditor->ActorsThatWereSelected.Num() > 0)
+					{
+						// Clear the selection set
+						GEditor->ActorsThatWereSelected.Empty();
+					}
+				}
+			}
+
+			for (const auto& Pair : FEditorModeRegistry::Get().GetFactoryMap())
+			{
+				Pair.Value->OnSelectionChanged(*this, NewSelection);
+			}
+		}
+		void OnEditorSelectNone()
+		{
+			GEditor->SelectNone(false, true);
+			GEditor->ActorsThatWereSelected.Empty();
+		}
+	};
+};
 
 // @todo: Can remove after permanent switch to new style
 namespace UE::Editor::ContentBrowser
@@ -118,7 +191,7 @@ FEditorModeTools& GLevelEditorModeTools()
 	checkf(!IsRunningCommandlet(), TEXT("The global mode manager should not be created or accessed in a commandlet environment. Check that your mode or module is not accessing the global mode tools or that scriptable features of modes have been moved to subsystems."));
 	if (!ensureMsgf(Internal::EditorModeToolsSingleton.IsValid(), TEXT("The level editor is not started up yet. If you need to access the global mode manager early in the startup phase, please use FLevelEditorModule::OnLevelEditorCreated to gate the access.")))
 	{
-		Internal::EditorModeToolsSingleton = MakeShared<FEditorModeTools>();
+		Internal::EditorModeToolsSingleton = MakeShared<Internal::FLevelEditorModeTools>();
 	}
 	return *Internal::EditorModeToolsSingleton.Get();
 }
diff --git a/Engine/Source/Editor/UnrealEd/Public/EditorModeManager.h b/Engine/Source/Editor/UnrealEd/Public/EditorModeManager.h
--- a/Engine/Source/Editor/UnrealEd/Public/EditorModeManager.h
+++ b/Engine/Source/Editor/UnrealEd/Public/EditorModeManager.h
@@ -573,11 +573,6 @@ public:
 	/** End New TRS Gizmo interface */
 	
 protected:
-	/** 
-	 * Delegate handlers
-	 **/
-	UNREALED_API void OnEditorSelectionChanged(UObject* NewSelection);
-	UNREALED_API void OnEditorSelectNone();
 
 
 	/** Handles the notification when a world is going through GC to clean up any modes pending deactivation. */
