Fixing issues with Preview Platform TextureLODSettings override
Implicit cast of DeviceProfile to TextureLODSettings is an ugly source of bugs that lets users avoid calling GetTextureLODSettings
LODSettings are cached on the TargetPlatform and must be kept up to date

#Synced-CL 51590761

diff --git a/Engine/Source/Developer/TargetPlatform/Public/Interfaces/ITargetPlatformSettings.h b/Engine/Source/Developer/TargetPlatform/Public/Interfaces/ITargetPlatformSettings.h
--- a/Engine/Source/Developer/TargetPlatform/Public/Interfaces/ITargetPlatformSettings.h
+++ b/Engine/Source/Developer/TargetPlatform/Public/Interfaces/ITargetPlatformSettings.h
@@ -303,6 +303,11 @@ public:
 	 * Gets the texture LOD settings used by this platform.
 	 *
 	 * @return A texture LOD settings structure.
+	 *
+	 * This should generally be equal to :
+	 *	const UTextureLODSettings* LODSettings = UDeviceProfileManager::Get().FindProfile(TargetPlatform->CookingDeviceProfileName())->GetTextureLODSettings();
+	 * It is cached by RegisterTextureLODSettings()
+	 *
 	 */
 	virtual const class UTextureLODSettings& GetTextureLODSettings() const = 0;
 
diff --git a/Engine/Source/Editor/TextureEditor/Private/TextureEditorToolkit.cpp b/Engine/Source/Editor/TextureEditor/Private/TextureEditorToolkit.cpp
--- a/Engine/Source/Editor/TextureEditor/Private/TextureEditorToolkit.cpp
+++ b/Engine/Source/Editor/TextureEditor/Private/TextureEditorToolkit.cpp
@@ -749,7 +749,8 @@ void FTextureEditorToolkit::PopulateQuickInfo( )
 												TextureDescription.GetLinearMipSizesInBytes(TextureExtendedData.MipSizesInBytes);
 											}
 
-											const UTextureLODSettings* LODSettings = (UTextureLODSettings*)UDeviceProfileManager::Get().FindProfile(TargetPlatform->PlatformName());
+											const UTextureLODSettings* LODSettings = 
+												UDeviceProfileManager::Get().FindProfile(TargetPlatform->CookingDeviceProfileName())->GetTextureLODSettings();
 											int32 FirstMipToSerialize = LODSettings->CalculateLODBias(
 												TextureDescription.TopMipSizeX, TextureDescription.TopMipSizeY, Texture->MaxTextureSize, 
 												Texture->LODGroup, Texture->LODBias, 0, Texture->MipGenSettings, false, Texture->IsStreamable());
diff --git a/Engine/Source/Editor/UnrealEd/Private/EditorEngine.cpp b/Engine/Source/Editor/UnrealEd/Private/EditorEngine.cpp
--- a/Engine/Source/Editor/UnrealEd/Private/EditorEngine.cpp
+++ b/Engine/Source/Editor/UnrealEd/Private/EditorEngine.cpp
@@ -6511,19 +6511,81 @@ static void SaveShaderPlatformAsDisabled(FPreviewPlatformInfo& PreviewPlatform)
 	Settings->PreviewShaderPlatformFriendlyName = PreviewPlatform.GetFriendlyName();
 }
 
+static void PreviewPlatformChangePre()
+{
+	// We're changing global platform config; make sure async ops are not running
+
+	// Wait for pending load requests.
+	FlushAsyncLoading();
+
+	// Need to block on all in-progress asset builds as we may be changing their config	
+	FAssetCompilingManager::Get().FinishAllCompilation();
+	
+	// Stop any texture streaming in progress :
+	UTexture::CancelPendingTextureStreaming();
+	
+	FlushRenderingCommands();
+}
+
+static void PreviewPlatformChangePost()
+{
+	// scope for RegisterTextureLODSettings ; @todo CBTLDP : remove when possible
+	{
+		// Need to update the cached TextureLODSettings in TargetPlatform with RegisterTextureLODSettings
+
+		const UTextureLODSettings * TLS_DP = UDeviceProfileManager::Get().GetActiveProfile()->GetTextureLODSettings();
+
+		ITargetPlatformManagerModule* TPM = GetTargetPlatformManager();
+		ITargetPlatform* TargetPlatform = TPM->GetRunningTargetPlatform();
+		//const UTextureLODSettings & TLS_TP = TargetPlatform->GetTextureLODSettings();
+		
+		// Alternate way to get DP : (should be used when cooking)
+		//  during interactive editor run, this should == TLS_DP
+		UTextureLODSettings * TLS_DP2 = UDeviceProfileManager::Get().FindProfile(TargetPlatform->CookingDeviceProfileName())->GetTextureLODSettings();
+
+		check( TLS_DP == TLS_DP2 );
+
+		TargetPlatform->RegisterTextureLODSettings(TLS_DP);
+
+		// after Register, TP->GetTextureLODSettings is equal to the DP way
+		check( &(TargetPlatform->GetTextureLODSettings()) == TLS_DP );
+	}
+
+	{
+		// TextureLODSettings and platform cvars may have changed
+		//	need to notify all textures
+		//	don't do this until all changes are done
+		// We didn't do a PreEdit on all textures; rather we did a system-wide Wait at the start of this function
+	
+		for( TObjectIterator<UTexture> It; It; ++It )
+		{
+			UTexture* Texture = *It;
+			if ( Texture->GetTextureClass() == ETextureClass::RenderTarget )
+			{
+				// @todo CBTLDP : why?
+				continue;
+			}
+
+			Texture->PostEditChange();
+		}
+	}
+}
+
 void UEditorEngine::SetPreviewPlatform(const FPreviewPlatformInfo& NewPreviewPlatform, bool bSaveSettings)
 {
 	// Get the requested preview platform, make sure it is valid.
 	EShaderPlatform ShaderPlatform = NewPreviewPlatform.ShaderPlatform;
 	if (ShaderPlatform == GMaxRHIShaderPlatform)
 	{
-		UE_LOG(LogEditor, Log, TEXT("Disable Preview Platform"));
+		UE_LOG(LogEditor, Display, TEXT("Disable Preview Platform"));
 	}
 	else
 	{
-		UE_LOG(LogEditor, Log, TEXT("Set Preview Platform: %s"), 
+		UE_LOG(LogEditor, Display, TEXT("Set Preview Platform: %s"), 
 			*NewPreviewPlatform.GetFriendlyName().ToString());
 	}
+	
+	PreviewPlatformChangePre();
 
 	check(FDataDrivenShaderPlatformInfo::IsValid(ShaderPlatform));
 	ERHIFeatureLevel::Type MaxFeatureLevel = NewPreviewPlatform.PreviewShaderFormatName != NAME_None ? (ERHIFeatureLevel::Type)GetMaxSupportedFeatureLevel(ShaderPlatform) : ERHIFeatureLevel::SM5;
@@ -6612,7 +6674,8 @@ void UEditorEngine::SetPreviewPlatform(const FPreviewPlatformInfo& NewPreviewPla
 	}
 
 	Scalability::ChangeScalabilityPreviewPlatform(PreviewPlatform.GetEffectivePreviewPlatformName(), GetActiveShaderPlatform());
-
+	
+	// this sets TextureLODSettings and platform CVars :
 	UDeviceProfileManager::Get().RestorePreviewDeviceProfile();
 
 	UStaticMesh::OnLodStrippingQualityLevelChanged(nullptr);
@@ -6624,6 +6687,7 @@ void UEditorEngine::SetPreviewPlatform(const FPreviewPlatformInfo& NewPreviewPla
 		{
 			if (UDeviceProfile* DP = UDeviceProfileManager::Get().FindProfile(PreviewPlatform.DeviceProfileName.ToString(), false))
 			{
+				// this sets TextureLODSettings and platform CVars :
 				UDeviceProfileManager::Get().SetPreviewDeviceProfile(DP);
 			}
 		}
@@ -6631,12 +6695,16 @@ void UEditorEngine::SetPreviewPlatform(const FPreviewPlatformInfo& NewPreviewPla
 
 	Scalability::ApplyCachedQualityLevels();
 
+	PreviewPlatformChangePost();
+	
 	PreviewPlatformChanged.Broadcast();
 
 	if (bSaveSettings)
 	{
 		SaveEditorShaderPlatformPreview();
 	}
+	
+	UE_LOG(LogEditor, Display, TEXT("Done with Preview Platform"));
 }
 
 void UEditorEngine::ToggleFeatureLevelPreview()
@@ -6645,7 +6713,7 @@ void UEditorEngine::ToggleFeatureLevelPreview()
 		!PreviewPlatform.bPreviewFeatureLevelActive ? TEXT("ON") : TEXT("OFF"),
 		*PreviewPlatform.GetFriendlyName().ToString());
 
-	FlushRenderingCommands();
+	PreviewPlatformChangePre();
 
 	PreviewPlatform.bPreviewFeatureLevelActive ^= 1;
 
@@ -6694,6 +6762,8 @@ void UEditorEngine::ToggleFeatureLevelPreview()
 	Scalability::ApplyCachedQualityLevels();
 	OnEffectivePreviewShaderPlatformChange();
 
+	PreviewPlatformChangePost();
+	
 	PreviewPlatformChanged.Broadcast();
 
 	UStaticMesh::OnLodStrippingQualityLevelChanged(nullptr);
@@ -6773,6 +6843,9 @@ void UEditorEngine::LoadEditorShaderPlatformPreview()
 		{
 			PreviewShaderPlatformName = *(LexToString(GShaderPlatformForFeatureLevel[Settings->PreviewFeatureLevel]));
 		}
+
+		UE_LOG(LogExit, Display, TEXT("LoadEditorShaderPlatformPreview: %s %s"), *PreviewShaderPlatformName.ToString(), *Settings->PreviewPlatformName.ToString() );
+
 		SetPreviewPlatform(FPreviewPlatformInfo(
 			(ERHIFeatureLevel::Type)Settings->PreviewFeatureLevel, 
 			ShaderPlatformToPreview, 
diff --git a/Engine/Source/Runtime/Engine/Classes/DeviceProfiles/DeviceProfile.h b/Engine/Source/Runtime/Engine/Classes/DeviceProfiles/DeviceProfile.h
--- a/Engine/Source/Runtime/Engine/Classes/DeviceProfiles/DeviceProfile.h
+++ b/Engine/Source/Runtime/Engine/Classes/DeviceProfiles/DeviceProfile.h
@@ -20,6 +20,11 @@ DECLARE_DELEGATE(FOnCVarsUpdated);
 UCLASS(config=DeviceProfiles, perObjectConfig, MinimalAPI)
 class UDeviceProfile : public UTextureLODSettings
 {
+	/** do NOT cast to UTextureLODSettings !
+	use GetTextureLODSettings() instead !
+	@todo CBTLDP : remove the public inheritance, hold TextureLODSettings as member instead?
+	**/
+
 	GENERATED_UCLASS_BODY()
 
 	/** The type of this profile, I.e. IOS, Windows, PS4 etc */
diff --git a/Engine/Source/Runtime/Engine/Classes/DeviceProfiles/DeviceProfileManager.h b/Engine/Source/Runtime/Engine/Classes/DeviceProfiles/DeviceProfileManager.h
--- a/Engine/Source/Runtime/Engine/Classes/DeviceProfiles/DeviceProfileManager.h
+++ b/Engine/Source/Runtime/Engine/Classes/DeviceProfiles/DeviceProfileManager.h
@@ -117,6 +117,9 @@ public:
 	 * @param bCreateProfileOnFail - Whether to create the profile from config if the object doesn't exist yet.
 	 * @param OptionalPlatformName - The platform name to use for loading.
 	 * @return The found profile.
+	 *
+	 * NOTE: when finding a profile for a TargetPlatform, use TargetPlatform->CookingDeviceProfileName()
+	 *   do NOT use TargetPlatform->PlatformName() or ->IniPlatformName()
 	 */
 	ENGINE_API UDeviceProfile* FindProfile(const FString& ProfileName, bool bCreateProfileOnFail = true, FName OptionalPlatformName = FName());
 
diff --git a/Engine/Source/Runtime/Engine/Classes/Engine/TextureLODSettings.h b/Engine/Source/Runtime/Engine/Classes/Engine/TextureLODSettings.h
--- a/Engine/Source/Runtime/Engine/Classes/Engine/TextureLODSettings.h
+++ b/Engine/Source/Runtime/Engine/Classes/Engine/TextureLODSettings.h
@@ -271,6 +271,6 @@ public:
 	}
 
 private:
-	UTextureLODSettings* PreviewTextureGroupRedirect;
+	UTextureLODSettings* PreviewTextureGroupRedirect = nullptr;  // @todo CBTLDP : why is this a member here and not a static or global ?
 #endif
 };
diff --git a/Engine/Source/Runtime/Engine/Private/DeviceProfiles/DeviceProfile.cpp b/Engine/Source/Runtime/Engine/Private/DeviceProfiles/DeviceProfile.cpp
--- a/Engine/Source/Runtime/Engine/Private/DeviceProfiles/DeviceProfile.cpp
+++ b/Engine/Source/Runtime/Engine/Private/DeviceProfiles/DeviceProfile.cpp
@@ -69,6 +69,7 @@ void UDeviceProfile::GatherParentCVarInformationRecursively(OUT TMap<FString, FS
 	}
 }
 
+// @todo CBTLDP : fixme : this casts away const
 UTextureLODSettings* UDeviceProfile::GetTextureLODSettings() const
 {
 #if ALLOW_OTHER_PLATFORM_CONFIG
@@ -79,7 +80,7 @@ UTextureLODSettings* UDeviceProfile::GetTextureLODSettings() const
 	}
 #endif
 
-	return (UTextureLODSettings*)this;
+	return const_cast<UDeviceProfile*>(this); // inheritance
 }
 
 UDeviceProfile* UDeviceProfile::GetParentProfile(bool bIncludeDefaultObject) const
diff --git a/Engine/Source/Runtime/Engine/Private/DeviceProfiles/DeviceProfileManager.cpp b/Engine/Source/Runtime/Engine/Private/DeviceProfiles/DeviceProfileManager.cpp
--- a/Engine/Source/Runtime/Engine/Private/DeviceProfiles/DeviceProfileManager.cpp
+++ b/Engine/Source/Runtime/Engine/Private/DeviceProfiles/DeviceProfileManager.cpp
@@ -1363,6 +1363,8 @@ void UDeviceProfileManager::SetPreviewDeviceProfile(UDeviceProfile* DeviceProfil
 		PreviewModeTag = "UnknownPreviewMode";
 	}
 	
+	UE_LOGF(LogDeviceProfileManager, Log, "SetPreviewDeviceProfile Start : %ls %ls", *(DeviceProfile->GetName()), *(PreviewModeTag.ToString()) );
+
 	RestorePreviewDeviceProfile(PreviewModeTag);
 
 	PreviewDeviceProfile = DeviceProfile;
@@ -1390,6 +1392,9 @@ void UDeviceProfileManager::SetPreviewDeviceProfile(UDeviceProfile* DeviceProfil
 	IConsoleManager::Get().CallAllConsoleVariableSinks();
 
 	PreviewDeviceProfileChangedDelegate.Broadcast();
+	
+	UE_LOGF(LogDeviceProfileManager, Log, "SetPreviewDeviceProfile End : %ls %ls", *(DeviceProfile->GetName()), *(PreviewModeTag.ToString()) );
+
 }
 
 
diff --git a/Engine/Source/Runtime/Engine/Private/TextureDerivedData.cpp b/Engine/Source/Runtime/Engine/Private/TextureDerivedData.cpp
--- a/Engine/Source/Runtime/Engine/Private/TextureDerivedData.cpp
+++ b/Engine/Source/Runtime/Engine/Private/TextureDerivedData.cpp
@@ -1471,7 +1471,7 @@ static void GetBuildSettingsForTargetPlatform(
 
 	check(TargetPlatform != NULL);
 
-	const UTextureLODSettings* LODSettings = (UTextureLODSettings*)UDeviceProfileManager::Get().FindProfile(TargetPlatform->PlatformName());
+	const UTextureLODSettings* LODSettings = UDeviceProfileManager::Get().FindProfile(TargetPlatform->CookingDeviceProfileName())->GetTextureLODSettings();
 	FTextureBuildSettings SourceBuildSettings;
 	FTexturePlatformData::FTextureEncodeResultMetadata SourceMetadata;
 	GetTextureBuildSettings(Texture, *LODSettings, *TargetPlatform, InEncodeSpeed, SourceBuildSettings, &SourceMetadata);
@@ -4577,7 +4577,7 @@ void UTexture::GetTargetPlatformBuildSettings(const ITargetPlatform* TargetPlatf
 		return;
 	}
 
-	const UTextureLODSettings* LODSettings = (UTextureLODSettings*)UDeviceProfileManager::Get().FindProfile(TargetPlatform->PlatformName());
+	const UTextureLODSettings* LODSettings = UDeviceProfileManager::Get().FindProfile(TargetPlatform->CookingDeviceProfileName())->GetTextureLODSettings();
 	FTextureBuildSettings SourceBuildSettings;
 	FTexturePlatformData::FTextureEncodeResultMetadata SourceMetadata;
 	GetTextureBuildSettings(*this, *LODSettings, *TargetPlatform, EncodeSpeed, SourceBuildSettings, &SourceMetadata);
diff --git a/Engine/Source/Runtime/Engine/Private/TextureLODSettings.cpp b/Engine/Source/Runtime/Engine/Private/TextureLODSettings.cpp
--- a/Engine/Source/Runtime/Engine/Private/TextureLODSettings.cpp
+++ b/Engine/Source/Runtime/Engine/Private/TextureLODSettings.cpp
@@ -541,19 +541,8 @@ ETextureMipLoadOptions UTextureLODSettings::GetMipLoadOptions(const UTexture* Te
 #if ALLOW_OTHER_PLATFORM_CONFIG
 void UTextureLODSettings::SetPreviewTextureGroupRedirect(UTextureLODSettings* InPreviewTextureGroupRedirect)
 {
-	if (PreviewTextureGroupRedirect != InPreviewTextureGroupRedirect)
-	{
-		PreviewTextureGroupRedirect = InPreviewTextureGroupRedirect;
+	//@todo CBTLDP : why is this here? move to DeviceProfile?
 
-		for (FThreadSafeObjectIterator ObjIterator(UTexture::StaticClass(), /*bOnlyGCedObjects = */false, /*AdditionalExclusionFlags = */RF_NoFlags); ObjIterator; ++ObjIterator)
-		{
-			//Call UpdateResource on Textures when changing the TextureLodGroup Settings to update the CachePlatformData
-			UTexture* Texture = Cast<UTexture>(*ObjIterator);
-			if (!Texture->IsA<UTextureRenderTarget>())
-			{
-				Texture->UpdateResource();
-			}
-		}
-	}
+	PreviewTextureGroupRedirect = InPreviewTextureGroupRedirect;
 }
 #endif
\ No newline at end of file
