Fix the 'resavepackages' commandlet to no longer regenerate new persistent guids of packages that were checked out and re-saved via the '-BatchSourceControl' option.

#Synced-CL 51587407

diff --git a/Engine/Source/Editor/UnrealEd/Private/Commandlets/ContentCommandlets.cpp b/Engine/Source/Editor/UnrealEd/Private/Commandlets/ContentCommandlets.cpp
--- a/Engine/Source/Editor/UnrealEd/Private/Commandlets/ContentCommandlets.cpp
+++ b/Engine/Source/Editor/UnrealEd/Private/Commandlets/ContentCommandlets.cpp
@@ -995,13 +995,13 @@ void UResavePackagesCommandlet::LoadAndSaveOnePackage(const FString& Filename)
 							*Filename,Linker->Summary.GetFileVersionUE().ToValue(), Linker->Summary.GetFileVersionLicenseeUE(), GPackageFileUEVersion.ToValue(), VER_LATEST_ENGINE_LICENSEEUE4 );
 					}
 
-					ESaveFlags SaveFlags = SAVE_None;
+					uint32 SaveFlags = ESaveFlags::SAVE_None;
 					
 					if (bIsReadOnly == false || SourceControlQueue == nullptr)
 					{
 						TRACE_CPUPROFILER_EVENT_SCOPE(UResavePackagesCommandlet::LoadAndSaveOnePackage::SavePackage);
 
-						if (SavePackageHelper(Package, Filename, RF_Standalone, GWarn, SaveFlags))
+						if (SavePackageHelper(Package, Filename, RF_Standalone, GWarn, (ESaveFlags)SaveFlags))
 						{
 							PackagesResaved++;
 							if (Verbosity == VERY_VERBOSE)
@@ -1016,7 +1016,11 @@ void UResavePackagesCommandlet::LoadAndSaveOnePackage(const FString& Filename)
 						// need to save to a temporary file first, and then queue the result.
 						const FString TempFilename = CreateTempFilename();
 
-						if (SavePackageHelper(Package, TempFilename, RF_Standalone, GWarn, SaveFlags))
+						// Prevent the persistent GUID of the package being regenerated as we are saving out to a different package path.
+						// We want the GUID to remain the same as we will eventually replace the source package with this temp one.
+						SaveFlags |= ESaveFlags::SAVE_KeepPersistentGUID;
+
+						if (SavePackageHelper(Package, TempFilename, RF_Standalone, GWarn, (ESaveFlags)SaveFlags))
 						{
 							SourceControlQueue->QueueCheckoutAndReplaceOperation(Filename, TempFilename, Package);
 						}
