EditorUsability : UnrealEd
- Fixed restore open asset tab location on restart (docked or undocked) after a crash

#Synced-CL 51539177

diff --git a/Engine/Source/Editor/UnrealEd/Private/Toolkits/SStandaloneAssetEditorToolkitHost.cpp b/Engine/Source/Editor/UnrealEd/Private/Toolkits/SStandaloneAssetEditorToolkitHost.cpp
--- a/Engine/Source/Editor/UnrealEd/Private/Toolkits/SStandaloneAssetEditorToolkitHost.cpp
+++ b/Engine/Source/Editor/UnrealEd/Private/Toolkits/SStandaloneAssetEditorToolkitHost.cpp
@@ -181,6 +181,7 @@ void SStandaloneAssetEditorToolkitHost::RestoreFromLayout( const TSharedRef<FTab
 
 	const TSharedRef<SDockTab> HostTab = HostTabPtr.Pin().ToSharedRef();
 	HostTab->SetOnTabClosed(SDockTab::FOnTabClosedCallback::CreateSP(this, &SStandaloneAssetEditorToolkitHost::OnTabClosed));
+	HostTab->SetOnTabRelocated(FSimpleDelegate::CreateSP(this, &SStandaloneAssetEditorToolkitHost::OnTabRelocated));
 
 	// Let the manager handle some logic on the Areas before destroying
 	MyTabManager->CloseAllAreas();
@@ -474,7 +475,7 @@ FReply SStandaloneAssetEditorToolkitHost::OnKeyDown( const FGeometry& MyGeometry
 }
 
 
-void SStandaloneAssetEditorToolkitHost::OnTabClosed(TSharedRef<SDockTab> TabClosed) const
+void SStandaloneAssetEditorToolkitHost::OnTabClosed(TSharedRef<SDockTab> TabClosed)
 {
 	check(TabClosed == HostTabPtr.Pin());
 
@@ -482,30 +483,78 @@ void SStandaloneAssetEditorToolkitHost::OnTabClosed(TSharedRef<SDockTab> TabClos
 
 	MyTabManager->SetMenuMultiBox(nullptr, nullptr);
 	
-	if(HostedAssetEditorToolkit.IsValid())
+	SaveTabLocation();
+}
+
+void SStandaloneAssetEditorToolkitHost::OnTabRelocated()
+{
+	SaveTabLocation();
+}
+
+void SStandaloneAssetEditorToolkitHost::SaveTabLocation() const
+{
+	if (!HostedAssetEditorToolkit.IsValid())
+	{
+		return;
+	}
+
+	const TSharedPtr<SDockTab> HostTab = HostTabPtr.Pin();
+	if (!HostTab.IsValid())
+	{
+		return;
+	}
+
+	const TArray<UObject*>* const ObjectsBeingEdited = HostedAssetEditorToolkit->GetObjectsCurrentlyBeingEdited();
+	if (!ObjectsBeingEdited)
+	{
+		return;
+	}
+
+	bool bConfigChanged = false;
+	const bool IsDockedAssetEditor = HostTab->HasSiblingTab(FName("DockedToolkit"), false/*TreatIndexNoneAsWildcard*/);
+	const EAssetEditorToolkitTabLocation AssetEditorToolkitTabLocation = (IsDockedAssetEditor) ? EAssetEditorToolkitTabLocation::Docked : EAssetEditorToolkitTabLocation::Standalone;
+	for (const UObject* ObjectBeingEdited : *ObjectsBeingEdited)
 	{
-		const TArray<UObject*>* const ObjectsBeingEdited = HostedAssetEditorToolkit->GetObjectsCurrentlyBeingEdited();
-		if(ObjectsBeingEdited)
+		// Only record assets that have a valid saved package
+		UPackage* const Package = ObjectBeingEdited->GetOutermost();
+		if (Package && Package->GetFileSize())
 		{
-			const bool IsDockedAssetEditor = TabClosed->HasSiblingTab(FName("DockedToolkit"), false/*TreatIndexNoneAsWildcard*/);
-			const EAssetEditorToolkitTabLocation AssetEditorToolkitTabLocation = (IsDockedAssetEditor) ? EAssetEditorToolkitTabLocation::Docked : EAssetEditorToolkitTabLocation::Standalone;
-			for(const UObject* ObjectBeingEdited : *ObjectsBeingEdited)
+			GConfig->SetInt(
+				TEXT("AssetEditorToolkitTabLocation"),
+				*ObjectBeingEdited->GetPathName(),
+				static_cast<int32>(AssetEditorToolkitTabLocation),
+				GEditorPerProjectIni
+			);
+
+			bConfigChanged = true;
+		}
+	}
+
+	if (!bConfigChanged)
+	{
+		return;
+	}
+
+	// Only write this section to disk to avoid loosing last tab position on crash
+	if (const FConfigFile* ConfigFile = GConfig->FindConfigFile(GEditorPerProjectIni))
+	{
+		const FConfigSection* Section = ConfigFile->FindSection(TEXT("AssetEditorToolkitTabLocation"));
+		if (Section && ConfigFile->Branch)
+		{
+			FString IniPath = ConfigFile->Branch->IniPath;
+			FConfigSection SectionCopy = *Section;
+
+			// Avoid editor freeze since config flush happens on main thread...
+			Async(EAsyncExecution::ThreadPool, [IniPath, SectionCopy = MoveTemp(SectionCopy)]()
 			{
-				// Only record assets that have a valid saved package
-				UPackage* const Package = ObjectBeingEdited->GetOutermost();
-				if(Package && Package->GetFileSize())
-				{
-					GConfig->SetInt(
-						TEXT("AssetEditorToolkitTabLocation"), 
-						*ObjectBeingEdited->GetPathName(), 
-						static_cast<int32>(AssetEditorToolkitTabLocation), 
-						GEditorPerProjectIni
-						);
-				}
-			}
+				FConfigFile DiskFile;
+				DiskFile.Read(IniPath);
+				FConfigSection* MutableSection = const_cast<FConfigSection*>(DiskFile.FindOrAddConfigSection(TEXT("AssetEditorToolkitTabLocation")));
+				*MutableSection = SectionCopy;
+				DiskFile.Write(IniPath);
+			});
 		}
 	}
 }
 
-
 #undef LOCTEXT_NAMESPACE
diff --git a/Engine/Source/Editor/UnrealEd/Private/Toolkits/SStandaloneAssetEditorToolkitHost.h b/Engine/Source/Editor/UnrealEd/Private/Toolkits/SStandaloneAssetEditorToolkitHost.h
--- a/Engine/Source/Editor/UnrealEd/Private/Toolkits/SStandaloneAssetEditorToolkitHost.h
+++ b/Engine/Source/Editor/UnrealEd/Private/Toolkits/SStandaloneAssetEditorToolkitHost.h
@@ -113,7 +113,12 @@ public:
 	void UnbindEditorCloseRequestFromHostTab();
 
 private:
-	void OnTabClosed(TSharedRef<SDockTab> TabClosed) const;
+	void OnTabClosed(TSharedRef<SDockTab> TabClosed);
+	void OnTabRelocated();
+
+	/** Saves this tab location into config to reopen properly */
+	void SaveTabLocation() const;
+
 	void ShutdownToolkitHost();
 
 	/** Create the main menu bar, either generated by the toolkit host or by the MainFrame module. */
