[Gizmos] Fix an issue where gizmos were always getting a postive for should draw queries, even when nothing was selected

Asset Editor default mode for new gizmos uses Mode manager has any scene component selected to determine if gizmos should draw.

#Synced-CL 51646065

diff --git a/Engine/Source/Editor/UnrealEd/Private/Tools/DefaultEdMode.cpp b/Engine/Source/Editor/UnrealEd/Private/Tools/DefaultEdMode.cpp
--- a/Engine/Source/Editor/UnrealEd/Private/Tools/DefaultEdMode.cpp
+++ b/Engine/Source/Editor/UnrealEd/Private/Tools/DefaultEdMode.cpp
@@ -2,6 +2,7 @@
 
 #include "Tools/DefaultEdMode.h"
 
+#include "EditorModeManager.h"
 #include "EditorModes.h"
 #include "Textures/SlateIcon.h"
 #include "LevelEditorViewport.h"
@@ -59,7 +60,12 @@ class FAssetEdModeWidgetHelper : public FLegacyEdModeWidgetHelper
 public:
 	virtual bool ShouldDrawWidget() const override
 	{
-		return true;
+		if (Owner && Owner->SelectionHasSceneComponent())
+		{
+			return true;
+		}
+
+		return false;
 	}
 	
 	virtual bool RequiresLegacyViewportInteractions() const override { return false; }
