[Backout] - CL51680395
#rnx
Original CL Desc
-----------------------------------------------------------------
[Gizmos] Fix an issue where gizmos were always getting a postive for should draw queries, even when nothing was selected

Asset Editor default mode for new gizmos uses Mode manager has any scene component selected to determine if gizmos should draw.

diff --git a/Engine/Source/Editor/UnrealEd/Private/Tools/DefaultEdMode.cpp b/Engine/Source/Editor/UnrealEd/Private/Tools/DefaultEdMode.cpp
--- a/Engine/Source/Editor/UnrealEd/Private/Tools/DefaultEdMode.cpp
+++ b/Engine/Source/Editor/UnrealEd/Private/Tools/DefaultEdMode.cpp
@@ -2,7 +2,6 @@
 
 #include "Tools/DefaultEdMode.h"
 
-#include "EditorModeManager.h"
 #include "EditorModes.h"
 #include "Textures/SlateIcon.h"
 #include "LevelEditorViewport.h"
@@ -60,12 +59,7 @@ class FAssetEdModeWidgetHelper : public FLegacyEdModeWidgetHelper
 public:
 	virtual bool ShouldDrawWidget() const override
 	{
-		if (Owner && Owner->SelectionHasSceneComponent())
-		{
-			return true;
-		}
-
-		return false;
+		return true;
 	}
 	
 	virtual bool RequiresLegacyViewportInteractions() const override { return false; }
