Fix regression of gizmo not showing up when selecting an actor in other level editor modes

#Synced-CL 51679831

diff --git a/Engine/Source/Editor/UnrealEd/Private/UnrealEdGlobals.cpp b/Engine/Source/Editor/UnrealEd/Private/UnrealEdGlobals.cpp
--- a/Engine/Source/Editor/UnrealEd/Private/UnrealEdGlobals.cpp
+++ b/Engine/Source/Editor/UnrealEd/Private/UnrealEdGlobals.cpp
@@ -260,7 +260,7 @@ int32 EditorInit( IEngineLoop& EngineLoop )
 	FActorFolders::Get();
 	if (!Internal::EditorModeToolsSingleton.IsValid())
 	{
-		Internal::EditorModeToolsSingleton = MakeShared<FEditorModeTools>();
+		Internal::EditorModeToolsSingleton = MakeShared<Internal::FLevelEditorModeTools>();
 	}
 
 	// Initialize the misc editor
