[PIEAuthorizer]
Refactored IPIEAuthorizer interface to support richer error reporting and a non-blocking availability query

[RewindDebugger]
- Conditionally hides AutoEject, AutoRecord, and their separator when PIE is unavailable, using the new IsPIEAuthorized() silent check.

#rnx

#Synced-CL 51685953

diff --git a/Engine/Plugins/Animation/GameplayInsights/Source/RewindDebugger/Private/RewindDebugger.cpp b/Engine/Plugins/Animation/GameplayInsights/Source/RewindDebugger/Private/RewindDebugger.cpp
--- a/Engine/Plugins/Animation/GameplayInsights/Source/RewindDebugger/Private/RewindDebugger.cpp
+++ b/Engine/Plugins/Animation/GameplayInsights/Source/RewindDebugger/Private/RewindDebugger.cpp
@@ -8,6 +8,7 @@
 #include "Framework/Application/SlateApplication.h"
 #include "HAL/PlatformFileManager.h"
 #include "IDesktopPlatform.h"
+#include "IPIEAuthorizer.h"
 #include "IRewindDebuggerExtension.h"
 #include "Kismet2/DebuggerCommands.h"
 #include "LevelEditor.h"
@@ -1869,19 +1870,29 @@ void FRewindDebugger::RegisterToolBar()
 
 	Section.AddSeparator(NAME_None);
 
-	Section.AddEntry(FToolMenuEntry::InitToolBarButton(
-		Commands.AutoEject,
-		FText(),
-		TAttribute<FText>(),
-		FSlateIcon("RewindDebuggerStyle", "RewindDebugger.AutoEject")));
+	Section.AddDynamicEntry("PIEControls", FNewToolMenuSectionDelegate::CreateLambda([](FToolMenuSection& InSection)
+		{
+			const TValueOrError<bool, FText> AuthorizationResult = IPIEAuthorizer::IsPIEAuthorized(/*bIsSimulateInEditor*/false);
+			if (AuthorizationResult.HasValue()
+				&& AuthorizationResult.GetValue())
+			{
+				const FRewindDebuggerCommands& Commands = FRewindDebuggerCommands::Get();
 
-	Section.AddEntry(FToolMenuEntry::InitToolBarButton(
-		Commands.AutoRecord,
-		FText(),
-		TAttribute<FText>(),
-		FSlateIcon("RewindDebuggerStyle", "RewindDebugger.AutoRecord")));
+				InSection.AddEntry(FToolMenuEntry::InitToolBarButton(
+					Commands.AutoEject,
+					FText(),
+					TAttribute<FText>(),
+					FSlateIcon("RewindDebuggerStyle", "RewindDebugger.AutoEject")));
+
+				InSection.AddEntry(FToolMenuEntry::InitToolBarButton(
+					Commands.AutoRecord,
+					FText(),
+					TAttribute<FText>(),
+					FSlateIcon("RewindDebuggerStyle", "RewindDebugger.AutoRecord")));
 
-	Section.AddSeparator("NAME_None");
+				InSection.AddSeparator(NAME_None);
+			}
+		}));
 
 	Section.AddEntry(FToolMenuEntry::InitComboButton(
 		"Display World",
diff --git a/Engine/Source/Editor/UnrealEd/Private/IPIEAuthorizer.cpp b/Engine/Source/Editor/UnrealEd/Private/IPIEAuthorizer.cpp
--- a/Engine/Source/Editor/UnrealEd/Private/IPIEAuthorizer.cpp
+++ b/Engine/Source/Editor/UnrealEd/Private/IPIEAuthorizer.cpp
@@ -1,3 +1,51 @@
 // Copyright Epic Games, Inc. All Rights Reserved.
 
-#include "IPIEAuthorizer.h"
\ No newline at end of file
+#include "IPIEAuthorizer.h"
+
+#include "Features/IModularFeatures.h"
+#include "Internationalization/Text.h"
+
+TValueOrError<bool, FText> IPIEAuthorizer::IsPIEAuthorized(const bool bIsSimulateInEditor)
+{
+	TArray<IPIEAuthorizer*> PlayAuthorizers = IModularFeatures::Get().GetModularFeatureImplementations<IPIEAuthorizer>(GetModularFeatureName());
+	for (const IPIEAuthorizer* Authority : PlayAuthorizers)
+	{
+		// break on first denial
+		TValueOrError<bool, FText> Result = Authority->IsPIEAuthorizedInternal(bIsSimulateInEditor);
+		if (Result.HasError()
+			|| !Result.GetValue())
+		{
+			return MoveTemp(Result);
+		}
+	}
+	return MakeValue(true);
+}
+
+TValueOrError<bool, FText> IPIEAuthorizer::RequestPIEPermission(const bool bIsSimulateInEditor) const
+{
+	const TValueOrError<bool, FText> IsAuthorized = IsPIEAuthorizedInternal(bIsSimulateInEditor);
+	if (IsAuthorized.HasValue()
+		&& IsAuthorized.GetValue())
+	{
+		return RequestPIEPermissionInternal(bIsSimulateInEditor);
+	}
+
+	return IsAuthorized;
+}
+
+TValueOrError<bool, FText> IPIEAuthorizer::RequestPIEPermissionInternal(bool bIsSimulateInEditor) const
+{
+	// Allowed by default
+	return MakeValue(true);
+}
+
+// deprecated
+bool IPIEAuthorizer::RequestPIEPermission(const bool bIsSimulateInEditor, FString& OutReason) const
+{
+	TValueOrError<bool, FText> ValueOrError = RequestPIEPermission(bIsSimulateInEditor);
+	if (ValueOrError.HasError())
+	{
+		OutReason = ValueOrError.GetError().ToString();
+	}
+	return ValueOrError.HasValue() && ValueOrError.GetValue();
+}
diff --git a/Engine/Source/Editor/UnrealEd/Private/PlayLevel.cpp b/Engine/Source/Editor/UnrealEd/Private/PlayLevel.cpp
--- a/Engine/Source/Editor/UnrealEd/Private/PlayLevel.cpp
+++ b/Engine/Source/Editor/UnrealEd/Private/PlayLevel.cpp
@@ -2627,11 +2627,14 @@ void UEditorEngine::StartPlayInEditorSession(FRequestPlaySessionParams& InReques
 	TArray<IPIEAuthorizer*> PlayAuthorizers = IModularFeatures::Get().GetModularFeatureImplementations<IPIEAuthorizer>(IPIEAuthorizer::GetModularFeatureName());
 	for (const IPIEAuthorizer* Authority : PlayAuthorizers)
 	{
-		FString DeniedReason;
-		if (!Authority->RequestPIEPermission(InRequestParams.WorldType == EPlaySessionWorldType::SimulateInEditor, DeniedReason))
+		TValueOrError<bool, FText> AuthorizationResult = Authority->RequestPIEPermission(InRequestParams.WorldType == EPlaySessionWorldType::SimulateInEditor);
+		if (AuthorizationResult.HasError() || !AuthorizationResult.GetValue())
 		{
-			// In case the authorizer didn't notify the user as to why this was blocked.
-			UE_LOG(LogPlayLevel, Warning, TEXT("Play-In-Editor canceled by plugin: %s"), *DeniedReason);
+			if (AuthorizationResult.HasError())
+			{
+				UE_LOG(LogPlayLevel, Warning, TEXT("Play-In-Editor canceled by plugin: %s"), *AuthorizationResult.GetError().ToString());
+				FSlateNotificationManager::Get().AddNotification(FNotificationInfo(AuthorizationResult.GetError()));
+			}
 
 			CancelRequestPlaySession();
 			return;
diff --git a/Engine/Source/Editor/UnrealEd/Public/IPIEAuthorizer.h b/Engine/Source/Editor/UnrealEd/Public/IPIEAuthorizer.h
--- a/Engine/Source/Editor/UnrealEd/Public/IPIEAuthorizer.h
+++ b/Engine/Source/Editor/UnrealEd/Public/IPIEAuthorizer.h
@@ -2,12 +2,15 @@
 
 #pragma once
 
+#if UE_ENABLE_INCLUDE_ORDER_DEPRECATED_IN_5_8
 #include "CoreMinimal.h"
+#endif // UE_ENABLE_INCLUDE_ORDER_DEPRECATED_IN_5_8
+
 #include "Features/IModularFeature.h"
-#include "HAL/Platform.h"
+#include "Templates/ValueOrError.h"
 #include "UObject/NameTypes.h"
 
-class FString;
+#define UE_API UNREALED_API
 
 class IPIEAuthorizer : public IModularFeature
 {
@@ -21,11 +24,49 @@ public:
 	}
 
 	/**
-	 * Used to gate play-in-editor functionality, where a PIE session may be 
+	 * Helper function that can be used to iterate through all registered implementations of IPIEAuthorizer
+	 * to gate play-in-editor functionality, where a PIE session may be
+	 * undesirable given some external plugin's state.
+	 *
+	 * @param bIsSimulateInEditor Whether the request is for Simulate In Editor (SIE) or Play In Editor (PIE)
+	 * @return Whether PIE is allowed and optional error description if denied.
+	 * @note This function will not require any interactions with users and can be called for "silent" queries
+	 */
+	static UE_API TValueOrError<bool, FText> IsPIEAuthorized(bool bIsSimulateInEditor);
+
+	/**
+	 * Potentially blocking methods that can be used to gate play-in-editor functionality, where a PIE session may be
+	 * undesirable given some external plugin's state.
+	 *
+	 * @param bIsSimulateInEditor Whether the request is for Simulate In Editor (SIE) or Play In Editor (PIE)
+	 * @return Whether PIE is allowed and optional error description if denied.
+	 * @note This method might require interactions with users.
+	 */
+	TValueOrError<bool, FText> RequestPIEPermission(bool bIsSimulateInEditor) const;
+
+protected:
+
+	/**
+	 * Non-blocking methods that can be overridden to gate play-in-editor functionality.
+	 *
+	 * @param bIsSimulateInEditor Whether the request is for Simulate In Editor (SIE) or Play In Editor (PIE)
+	 * @return Whether PIE is allowed and optional error description if denied.
+	 * @note This method should NEVER require interactions with users.
+	 */
+	virtual TValueOrError<bool, FText> IsPIEAuthorizedInternal(bool bIsSimulateInEditor) const = 0;
+
+	/**
+	 * Used to gate play-in-editor functionality, where a PIE session may be
 	 * undesirable given some external plugin's state.
 	 *
-	 * @param  OutReason	[out] A string describing why play was denied by this plugin.
-	 * @return False to block play-in-editor from continuing, true to allow.
+	 * @param bIsSimulateInEditor Whether the request is for Simulate In Editor (SIE) or Play In Editor (PIE)
+	 * @return Whether PIE is allowed and optional error description if denied.
+	 * @note This method might require interactions with users.
 	 */
-	virtual bool RequestPIEPermission(bool bIsSimulateInEditor, FString& OutReason) const = 0;
+	virtual UE_API TValueOrError<bool, FText> RequestPIEPermissionInternal(bool bIsSimulateInEditor) const;
+
+	UE_DEPRECATED(5.8, "Use the version taking a localized string")
+	virtual UE_API bool RequestPIEPermission(bool bIsSimulateInEditor, class FString& OutReason) const final;
 };
+
+#undef UE_API
