Improved support for incremental validation in incremental cooks

Incremental cooks already mostly handle incremental validation, as they skip loading the package if it hasn't changed (as validation runs via UCookOnTheFlyServer::LoadPackageInQueue), however changes to a world would still validate all their external actors. This can be quite slow for large worlds.

This change improves the incremental validation support so that it only validates the specific actors that have changed in a world being re-cooked. It does this by comparing the PackageSavedHash of each external actor package against the previous version that was validated, and skipping any that haven't changed.

This feature is disabled by default, but can be enabled by setting the CVar 'cook.validation.IncrementalWorldValidationEnabled' to true. Note: It only currently works for single-process by-the-book cooks.

You can also now trigger a full revalidation when running incremental validation against a new build version, by setting the CVar 'cook.validation.IncrementalValidationRequiresSameBuildVersion' to true. This is disabled by default.

#Synced-CL 51656536

diff --git a/Engine/Source/Editor/UnrealEd/Classes/CookOnTheSide/CookOnTheFlyServer.h b/Engine/Source/Editor/UnrealEd/Classes/CookOnTheSide/CookOnTheFlyServer.h
--- a/Engine/Source/Editor/UnrealEd/Classes/CookOnTheSide/CookOnTheFlyServer.h
+++ b/Engine/Source/Editor/UnrealEd/Classes/CookOnTheSide/CookOnTheFlyServer.h
@@ -196,6 +196,7 @@ namespace UE::Cook
 	class FShaderLibraryCookArtifact;
 	class FSoftGCHistory;
 	class FStallDetector;
+	class FValidationCookArtifact;
 	class FWorkerRequestsLocal;
 	class FWorkerRequestsRemote;
 	class ICookOnTheFlyRequestManager;
@@ -1071,6 +1072,8 @@ private:
 	void AddFlexPathToCook(TArray<FName>& InOutFilesToCook, TMap<FName, UE::Cook::FInstigator>& InOutInstigators,
 		const FString& InFlexPath, const UE::Cook::FInstigator& Instigator) const;
 
+	void RegisterValidationCookArtifact();
+
 	/** Return the name to use for the project's global shader library */
 	FString GetProjectShaderLibraryName() const;
 
@@ -1744,6 +1747,7 @@ private:
 
 	TArray<TRefCountPtr<UE::Cook::ICookArtifact>> CookArtifacts;
 	UE::Cook::FGlobalCookArtifact* GlobalArtifact = nullptr;
+	UE::Cook::FValidationCookArtifact* ValidationArtifact = nullptr;
 
 	/**
 	 * Heap of Pollables to tick, ordered by NextTimeSeconds.
diff --git a/Engine/Source/Editor/UnrealEd/Private/CookOnTheFlyServer.cpp b/Engine/Source/Editor/UnrealEd/Private/CookOnTheFlyServer.cpp
--- a/Engine/Source/Editor/UnrealEd/Private/CookOnTheFlyServer.cpp
+++ b/Engine/Source/Editor/UnrealEd/Private/CookOnTheFlyServer.cpp
@@ -52,6 +52,7 @@
 #include "Cooker/PackageTracker.h"
 #include "Cooker/ShaderLibraryCooking.h"
 #include "Cooker/StallDetector.h"
+#include "Cooker/ValidationCookArtifact.h"
 #include "Cooker/WorkerRequestsLocal.h"
 #include "Cooker/WorkerRequestsRemote.h"
 #include "CookerSettings.h"
@@ -3256,6 +3257,11 @@ void UCookOnTheFlyServer::LoadPackageInQueue(UE::Cook::FPackageData& PackageData
 
 	if (ValidateSourcePackage(PackageData, LoadedPackage) == EDataValidationResult::Invalid)
 	{
+		if (ValidationArtifact)
+		{
+			ValidationArtifact->MarkPackageRequiresFullValidation(LoadedPackage->GetFName());
+		}
+
 		if (EnumHasAnyFlags(CookByTheBookOptions->StartupOptions, ECookByTheBookOptions::ValidationErrorsAreFatal))
 		{
 			UE_LOG(LogCook, Error, TEXT("%s failed validation"), *LoadedPackage->GetName());
@@ -3419,6 +3425,11 @@ EDataValidationResult UCookOnTheFlyServer::ValidateSourcePackage(UE::Cook::FPack
 		{
 			const FString ExternalActorsPathForWorld = ULevel::GetExternalActorsPath(Package);
 			AssetRegistry->GetAssetsByPath(*ExternalActorsPathForWorld, ExternalObjects, /*bRecursive*/true, /*bIncludeOnlyOnDiskAssets*/true);
+
+			if (ValidationArtifact)
+			{
+				ValidationArtifact->FilterExternalActorsToValidate(Package->GetFName(), ExternalObjects);
+			}
 		}
 
 		FMessageLogScopedOverride AssetCheckLogOverride(UE::DataValidation::MessageLogName);
@@ -8388,6 +8399,22 @@ void UCookOnTheFlyServer::AddFileToCook( TArray<FName>& InOutFilesToCook,
 	}
 }
 
+void UCookOnTheFlyServer::RegisterValidationCookArtifact()
+{
+	if (ValidationArtifact)
+	{
+		UnregisterArtifact(ValidationArtifact);
+		ValidationArtifact = nullptr;
+	}
+
+	check(!ValidationArtifact);
+	ValidationArtifact = new UE::Cook::FValidationCookArtifact(*this);
+	TRefCountPtr<UE::Cook::ICookArtifact> RefCounter(ValidationArtifact);
+	RegisterArtifact(ValidationArtifact);
+	// We keep a raw pointer and rely on RegisterArtifact holding the reference.
+	check(RefCounter.GetRefCount() > 1);
+}
+
 const TCHAR* GCookRequestUsageMessage = TEXT(
 	"By default, the cooker does not cook any packages. Packages must be requested by one of the following methods.\n"
 	"All transitive dependencies of a requested package are also cooked. Packages can be specified by LocalFilename/Filepath\n"
@@ -11096,6 +11123,7 @@ void UCookOnTheFlyServer::StartCookByTheBook( const FCookByTheBookStartupOptions
 	SelectSessionPlatforms(BeginContext);
 	UE::Cook::FDelegates::ConfigureCookSession.Broadcast(*this);
 	CalculateGlobalDependenciesHashes();
+	RegisterValidationCookArtifact();
 	RegisterShaderLibraryIncrementalCookArtifact(BeginContext);
 	LoadBeginCookIncrementalFlags(BeginContext);
 
diff --git a/Engine/Source/Editor/UnrealEd/Private/Cooker/ValidationCookArtifact.cpp b/Engine/Source/Editor/UnrealEd/Private/Cooker/ValidationCookArtifact.cpp
new file mode 100644
--- /dev/null
+++ b/Engine/Source/Editor/UnrealEd/Private/Cooker/ValidationCookArtifact.cpp
@@ -0,0 +1,323 @@
+// Copyright Epic Games, Inc. All Rights Reserved.
+
+#include "Cooker/ValidationCookArtifact.h"
+
+#include "AssetRegistry/AssetData.h"
+#include "AssetRegistry/IAssetRegistry.h"
+#include "CookOnTheSide/CookOnTheFlyServer.h"
+#include "Engine/Level.h"
+#include "JsonUtils/RapidJsonUtils.h"
+#include "HAL/PlatformFileManager.h"
+#include "Interfaces/ITargetPlatform.h"
+#include "Misc/FileHelper.h"
+#include "Misc/Paths.h"
+
+namespace UE::Cook
+{
+
+namespace Private
+{
+
+static const FString ConfigSectionName = TEXT("ValidationSettings");
+static const TCHAR JsonKey_WorldValidationCache[] = TEXT("worldValidationCache");
+static const TCHAR JsonKey_LastAssetValidationBuildVersion[] = TEXT("lastAssetValidationBuildVersion");
+
+static bool bIncrementalWorldValidationEnabled = false;
+static FAutoConsoleVariableRef CVarIncrementalWorldValidationEnabled(
+	TEXT("cook.validation.IncrementalWorldValidationEnabled"),
+	bIncrementalWorldValidationEnabled,
+	TEXT("True if incremental world validation is enabled, where we only validate changed external actors in partitioned worlds"));
+
+static bool bIncrementalValidationRequiresSameBuildVersion = false;
+static FAutoConsoleVariableRef CVarIncrementalValidationRequiresSameBuildVersion(
+	TEXT("cook.validation.IncrementalValidationRequiresSameBuildVersion"),
+	bIncrementalValidationRequiresSameBuildVersion,
+	TEXT("True if a different build version needs to re-validate everything, or False if incremental asset validation persists over different build versions"));
+
+} // namespace Private
+
+FValidationCookArtifact::FValidationCookArtifact(UCookOnTheFlyServer& InCOTFS)
+	: COTFS(InCOTFS)
+	, CurrentBuildVersion(FApp::GetBuildVersion())
+	, bIsIncrementalWorldValidationEnabled(Private::bIncrementalWorldValidationEnabled)
+{
+	// We currently only support incremental world validation for SP ByTheBook cooks (MP cooks would require a IMPCollector)
+	if (bIsIncrementalWorldValidationEnabled && (COTFS.GetProcessType() != UE::Cook::EProcessType::SingleProcess || COTFS.GetCookType() != UE::Cook::ECookType::ByTheBook))
+	{
+		UE_LOGF(LogCook, Warning, "Incremental world validation was requested, but it is only supported by single-process cook-by-the-book cooks. Incremental world validation will be disabled for this cook.");
+		bIsIncrementalWorldValidationEnabled = false;
+	}
+}
+
+FString FValidationCookArtifact::GetArtifactName() const
+{
+	return TEXT("validation");
+}
+
+FConfigFile FValidationCookArtifact::CalculateCurrentSettings(ICookInfo& CookInfo, const ITargetPlatform* TargetPlatform)
+{
+	// Initialize the artifact data for this platform
+	FArtifactData& PlatformArtifactData = PerPlatformData.Add(TargetPlatform);
+	PlatformArtifactData.Filename = FPaths::Combine(CookInfo.GetCookMetadataOutputFolder(TargetPlatform), TEXTVIEW("ValidationCookArtifact.json"));
+
+	// Validation doesn't have any settings that affect incremental cooking
+
+	return FConfigFile();
+}
+
+void FValidationCookArtifact::CompareSettings(UE::Cook::Artifact::FCompareSettingsContext& Context)
+{
+	const ITargetPlatform* TargetPlatform = Context.GetTargetPlatform();
+	FArtifactData& PlatformArtifactData = PerPlatformData.FindChecked(TargetPlatform);
+
+	// Validation doesn't have any settings that affect incremental cooking
+	
+	if (LoadCookArtifactFile(PlatformArtifactData))
+	{
+		UE_LOGF(LogCook, Display, "CookArtifact '%ls' loaded.", *PlatformArtifactData.Filename);
+	}
+
+	if (IsAssetValidationEnabledForCook())
+	{
+		if (Private::bIncrementalValidationRequiresSameBuildVersion && PlatformArtifactData.LastAssetValidationBuildVersion != CurrentBuildVersion)
+		{
+			UE_LOGF(LogCook, Display, "CookArtifact '%ls' invalidated cook because IncrementalValidationRequiresSameBuildVersion is true and LastAssetValidationBuildVersion != CurrentBuildVersion (old: '%ls', new: '%ls').",
+				*PlatformArtifactData.Filename, *PlatformArtifactData.LastAssetValidationBuildVersion, *CurrentBuildVersion);
+			Context.RequestFullRecook(true);
+			return;
+		}
+
+		if (!bIsIncrementalWorldValidationEnabled)
+		{
+			PlatformArtifactData.WorldValidationCache.Empty();
+		}
+	}
+}
+
+void FValidationCookArtifact::OnInvalidate(const ITargetPlatform* TargetPlatform)
+{
+	InvalidateCache(TargetPlatform);
+}
+
+void FValidationCookArtifact::OnFullRecook(const ITargetPlatform* TargetPlatform)
+{
+	InvalidateCache(TargetPlatform);
+}
+
+void FValidationCookArtifact::UpdateOplogPackages(UE::Cook::Artifact::FUpdateOplogPackagesContext& Context)
+{
+	const ITargetPlatform* TargetPlatform = Context.GetTargetPlatform();
+	FArtifactData& PlatformArtifactData = PerPlatformData.FindChecked(TargetPlatform);
+
+	// Remove any stale data from the persistent artifact
+	{
+		const TMap<FName, UE::Cook::Artifact::FOplogPackageData>& OplogPackages = Context.GetOplogPackages();
+		for (auto It = PlatformArtifactData.WorldValidationCache.CreateIterator(); It; ++It)
+		{
+			if (!OplogPackages.Contains(It->Key))
+			{
+				// World is no longer known to the cook; remove it
+				It.RemoveCurrent();
+				continue;
+			}
+		}
+	}
+
+	// If validation was enabled then we can also update the last build version in the persistent artifact
+	if (IsAssetValidationEnabledForCook())
+	{
+		PlatformArtifactData.LastAssetValidationBuildVersion = CurrentBuildVersion;
+	}
+
+	if (SaveCookArtifactFile(PlatformArtifactData))
+	{
+		UE_LOGF(LogCook, Display, "CookArtifact '%ls' saved.", *PlatformArtifactData.Filename);
+	}
+	else
+	{
+		UE_LOGF(LogCook, Warning, "CookArtifact '%ls' failed to save.", *PlatformArtifactData.Filename);
+	}
+}
+
+void FValidationCookArtifact::FilterExternalActorsToValidate(const FName WorldPackageName, TArray<FAssetData>& InOutExternalActors, const bool bUpdateCache)
+{
+	if (!bIsIncrementalWorldValidationEnabled || PerPlatformData.IsEmpty())
+	{
+		return;
+	}
+
+	// Read the current hashes for every external actor package
+	TMap<FName, FIoHash> NewExternalActorHashes;
+	{
+		IAssetRegistry& AssetRegistry = IAssetRegistry::GetChecked();
+		for (const FAssetData& ExternalActor : InOutExternalActors)
+		{
+			// Note: TryGetAssetPackageData only takes a read-lock, so this could be parallelized if needed
+			FAssetPackageData ExternalActorPackageData;
+			if (AssetRegistry.TryGetAssetPackageData(ExternalActor.PackageName, ExternalActorPackageData) == AssetRegistry::EExists::Exists)
+			{
+				NewExternalActorHashes.Add(ExternalActor.PackageName, ExternalActorPackageData.GetPackageSavedHash());
+			}
+		}
+	}
+
+	// Compare these hashes against any previous state to know what's changed
+	// Note: Validation is platform agnostic, so we just use the data for the first platform
+	{
+		FArtifactData& PlatformArtifactData = PerPlatformData.begin()->Value;
+		if (const TMap<FName, FIoHash>* ExternalActorHashes = PlatformArtifactData.WorldValidationCache.Find(WorldPackageName))
+		{
+			InOutExternalActors.RemoveAll(
+				[&NewExternalActorHashes, ExternalActorHashes](const FAssetData& ExternalActor)
+				{
+					const FIoHash* CurrentHash = NewExternalActorHashes.Find(ExternalActor.PackageName);
+					const FIoHash* PreviousHash = ExternalActorHashes->Find(ExternalActor.PackageName);
+					return CurrentHash && PreviousHash && *CurrentHash == *PreviousHash && !CurrentHash->IsZero();
+				});
+		}
+	}
+
+	// Persist the current hashes if we were asked to
+	if (bUpdateCache)
+	{
+		for (TTuple<const ITargetPlatform*, FArtifactData>& PerPlatformDataPair : PerPlatformData)
+		{
+			PerPlatformDataPair.Value.WorldValidationCache.Add(WorldPackageName, NewExternalActorHashes);
+		}
+	}
+}
+
+void FValidationCookArtifact::MarkPackageRequiresFullValidation(const FName PackageName)
+{
+	for (TTuple<const ITargetPlatform*, FArtifactData>& PerPlatformDataPair : PerPlatformData)
+	{
+		PerPlatformDataPair.Value.WorldValidationCache.Remove(PackageName);
+	}
+}
+
+bool FValidationCookArtifact::IsAssetValidationEnabledForCook() const
+{
+	return EnumHasAnyFlags(COTFS.GetCookValidationOptions(), UE::Cook::ECookValidationOptions::RunAssetValidation);
+}
+
+void FValidationCookArtifact::InvalidateCache(const ITargetPlatform* TargetPlatform)
+{
+	FArtifactData& PlatformArtifactData = PerPlatformData.FindChecked(TargetPlatform);
+	FPlatformFileManager::Get().GetPlatformFile().DeleteFile(*PlatformArtifactData.Filename);
+	PlatformArtifactData.WorldValidationCache.Reset();
+}
+
+bool FValidationCookArtifact::LoadCookArtifactFile(FArtifactData& PlatformArtifactData)
+{
+	FString PlatformArtifactJsonString;
+	if (!FFileHelper::LoadFileToString(PlatformArtifactJsonString, *PlatformArtifactData.Filename))
+	{
+		return false;
+	}
+
+	TValueOrError<UE::Json::FDocument, UE::Json::FParseError> PlatformArtifactJsonDocument = UE::Json::ParseInPlace(PlatformArtifactJsonString);
+	if (PlatformArtifactJsonDocument.HasError())
+	{
+		return false;
+	}
+
+	TOptional<UE::Json::FConstObject> PlatformArtifactJsonRootObject = UE::Json::GetRootObject(PlatformArtifactJsonDocument.GetValue());
+	if (!PlatformArtifactJsonRootObject.IsSet())
+	{
+		return false;
+	}
+
+	// Read WorldValidationCache
+	if (TOptional<UE::Json::FConstObject> WorldValidationCacheJsonObject = UE::Json::GetObjectField(PlatformArtifactJsonRootObject.GetValue(), Private::JsonKey_WorldValidationCache))
+	{
+		for (const UE::Json::FMember& WorldValidationCacheJsonEntry : WorldValidationCacheJsonObject.GetValue())
+		{
+			if (WorldValidationCacheJsonEntry.value.IsObject())
+			{
+				const FString WorldPackageName = FString(UE::Json::ValueAsStringView(WorldValidationCacheJsonEntry.name));
+				const FString ExternalActorsPrefix = ULevel::GetExternalActorsPath(WorldPackageName) / TEXT("");
+
+				TMap<FName, FIoHash>& ExternalActorHashes = PlatformArtifactData.WorldValidationCache.Add(*WorldPackageName);
+				for (const UE::Json::FMember& ExternalActorHashesJsonEntry : WorldValidationCacheJsonEntry.value.GetObject())
+				{
+					if (ExternalActorHashesJsonEntry.value.IsString())
+					{
+						// Re-add the prefix stripped off during save
+						FString ExternalActorPath = FString(UE::Json::ValueAsStringView(ExternalActorHashesJsonEntry.name));
+						if (ExternalActorPath.Len() > 0 && ExternalActorPath[0] != TEXT('/'))
+						{
+							ExternalActorPath.InsertAt(0, ExternalActorsPrefix);
+						}
+
+						ExternalActorHashes.Add(*ExternalActorPath, FIoHash(UE::Json::ValueAsStringView(ExternalActorHashesJsonEntry.value)));
+					}
+				}
+			}
+		}
+	}
+
+	// Read LastAssetValidationBuildVersion
+	if (TOptional<FStringView> LastAssetValidationBuildVersionJsonString = UE::Json::GetStringField(PlatformArtifactJsonRootObject.GetValue(), Private::JsonKey_LastAssetValidationBuildVersion))
+	{
+		PlatformArtifactData.LastAssetValidationBuildVersion = LastAssetValidationBuildVersionJsonString.GetValue();
+	}
+
+	return true;
+}
+
+bool FValidationCookArtifact::SaveCookArtifactFile(FArtifactData& PlatformArtifactData)
+{
+	UE::Json::FDocument PlatformArtifactJsonDocument(rapidjson::kObjectType);
+	{
+		UE::Json::FObject PlatformArtifactJsonRootObject = PlatformArtifactJsonDocument.GetObject();
+		UE::Json::FAllocator& PlatformArtifactJsonAllocator = PlatformArtifactJsonDocument.GetAllocator();
+
+		// Add WorldValidationCache
+		{
+			UE::Json::FValue WorldValidationCacheJsonObject(rapidjson::kObjectType);
+			PlatformArtifactData.WorldValidationCache.KeySort(FNameLexicalLess());
+			for (TTuple<FName, TMap<FName, FIoHash>>& WorldValidationCachePair : PlatformArtifactData.WorldValidationCache)
+			{
+				const FString ExternalActorsPrefix = ULevel::GetExternalActorsPath(WorldValidationCachePair.Key.ToString()) / TEXT("");
+
+				UE::Json::FValue ExternalActorHashesJsonObject(rapidjson::kObjectType);
+				WorldValidationCachePair.Value.KeySort(FNameLexicalLess());
+				for (const TTuple<FName, FIoHash>& ExternalActorHashesPair : WorldValidationCachePair.Value)
+				{
+					// All external actors in this world have a common prefix, so strip that off when saving for compactness
+					FString ExternalActorPath = ExternalActorHashesPair.Key.ToString();
+					ExternalActorPath.RemoveFromStart(ExternalActorsPrefix);
+
+					ExternalActorHashesJsonObject.AddMember(
+						UE::Json::MakeStringValue(ExternalActorPath, PlatformArtifactJsonAllocator),
+						UE::Json::MakeStringValue(LexToString(ExternalActorHashesPair.Value), PlatformArtifactJsonAllocator),
+						PlatformArtifactJsonAllocator);
+				}
+				WorldValidationCacheJsonObject.AddMember(
+					UE::Json::MakeStringValue(FNameBuilder(WorldValidationCachePair.Key).ToView(), PlatformArtifactJsonAllocator),
+					MoveTemp(ExternalActorHashesJsonObject),
+					PlatformArtifactJsonAllocator);
+			}
+			PlatformArtifactJsonRootObject.AddMember(
+				UE::Json::MakeStringRef(Private::JsonKey_WorldValidationCache),
+				MoveTemp(WorldValidationCacheJsonObject),
+				PlatformArtifactJsonAllocator);
+		}
+
+		// Add LastAssetValidationBuildVersion
+		PlatformArtifactJsonRootObject.AddMember(
+			UE::Json::MakeStringRef(Private::JsonKey_LastAssetValidationBuildVersion), 
+			UE::Json::MakeStringValue(PlatformArtifactData.LastAssetValidationBuildVersion, PlatformArtifactJsonAllocator), 
+			PlatformArtifactJsonAllocator);
+	}
+
+	const FString PlatformArtifactJsonString = UE::Json::WritePretty(PlatformArtifactJsonDocument);
+	if (!FFileHelper::SaveStringToFile(PlatformArtifactJsonString, *PlatformArtifactData.Filename, FFileHelper::EEncodingOptions::ForceUTF8WithoutBOM))
+	{
+		return false;
+	}
+
+	return true;
+}
+
+} // namespace UE::Cook
diff --git a/Engine/Source/Editor/UnrealEd/Private/Cooker/ValidationCookArtifact.h b/Engine/Source/Editor/UnrealEd/Private/Cooker/ValidationCookArtifact.h
new file mode 100644
--- /dev/null
+++ b/Engine/Source/Editor/UnrealEd/Private/Cooker/ValidationCookArtifact.h
@@ -0,0 +1,73 @@
+// Copyright Epic Games, Inc. All Rights Reserved.
+
+#pragma once
+
+#include "Cooker/CookArtifact.h"
+#include "IO/IoHash.h"
+
+struct FAssetData;
+class UCookOnTheFlyServer;
+
+namespace UE::Cook
+{
+
+/**
+ * CookArtifact that stores information related to incremental source package validation during incremental cooking.
+ * 
+ * @note Most incremental validation is handled automatically, as incremental cooking only loads/validates changed packages, however things made up of multiple 
+ *       source packages (like worlds) need special handling since validation runs on load, so before things like world cell generation have run during save.
+ */
+class FValidationCookArtifact : public ICookArtifact
+{
+public:
+	explicit FValidationCookArtifact(UCookOnTheFlyServer& InCOTFS);
+
+	virtual FString GetArtifactName() const override;
+	virtual FConfigFile CalculateCurrentSettings(ICookInfo& CookInfo, const ITargetPlatform* TargetPlatform) override;
+	virtual void CompareSettings(UE::Cook::Artifact::FCompareSettingsContext& Context) override;
+	virtual void OnInvalidate(const ITargetPlatform* TargetPlatform) override;
+	virtual void OnFullRecook(const ITargetPlatform* TargetPlatform) override;
+	virtual void UpdateOplogPackages(UE::Cook::Artifact::FUpdateOplogPackagesContext& Context) override;
+
+	/**
+	 * Given the complete set of external actors associated with the given world package, filter the set down to only the external actors that need validation.
+	 * 
+	 * @param WorldPackageName		The name of the world package that contains the actors objects.
+	 * @param InOutExternalActors	The set of external actors to filter.
+	 * @param bUpdateCache			If set, update the current state with the data for InOutExternalActors, so that the next call to FilterExternalActorsToValidate will use that as its base data for comparison.
+	 */
+	void FilterExternalActorsToValidate(const FName WorldPackageName, TArray<FAssetData>& InOutExternalActors, const bool bUpdateCache = true);
+
+	/**
+	 * Clear any incremental validation state for the given package (eg, if it failed validation).
+	 */
+	void MarkPackageRequiresFullValidation(const FName PackageName);
+
+private:
+	bool IsAssetValidationEnabledForCook() const;
+
+	void InvalidateCache(const ITargetPlatform* TargetPlatform);
+
+	struct FArtifactData
+	{
+		/** World Package -> External Actor Package -> Hash */
+		TMap<FName, TMap<FName, FIoHash>> WorldValidationCache;
+
+		/** FApp::GetBuildVersion() that was used when we last ran asset validation (@see IsAssetValidationEnabledForCook) */
+		FString LastAssetValidationBuildVersion;
+
+		/** Not serialized */
+		FString Filename;
+	};
+
+	bool LoadCookArtifactFile(FArtifactData& PlatformArtifactData);
+	bool SaveCookArtifactFile(FArtifactData& PlatformArtifactData);
+
+	UCookOnTheFlyServer& COTFS;
+	FString CurrentBuildVersion;
+	bool bIsIncrementalWorldValidationEnabled = false;
+
+	TSortedMap<const ITargetPlatform*, FArtifactData> PerPlatformData;
+};
+
+} // namespace UE::Cook
