Add behavior for UAF to prompt user for asset compilation on PIE
- Added ModularFeature for pre-pie asset compilation (allowing for different asset types to have similar behaviour to Blueprints)
- Refactored pop-up dialog to be reusable across BP implementation and ModularFeature
- Added ModularFeature implementation for UAF (behavior closely mirrors BP one)
- Fixed SetVMRecompilationRequired (were setting wrong flag)
- Added flag to UAFRigVMAsset_EditorData to ignore asset for future PIE runs (if explicitly skipped by user)

#Synced-CL 51756661

diff --git a/Engine/Plugins/Experimental/UAF/UAF/Source/UAFEditor/Private/AnimNextEditorModule.cpp b/Engine/Plugins/Experimental/UAF/UAF/Source/UAFEditor/Private/AnimNextEditorModule.cpp
--- a/Engine/Plugins/Experimental/UAF/UAF/Source/UAFEditor/Private/AnimNextEditorModule.cpp
+++ b/Engine/Plugins/Experimental/UAF/UAF/Source/UAFEditor/Private/AnimNextEditorModule.cpp
@@ -582,6 +582,8 @@ void FAnimNextEditorModule::StartupModule()
 	
 	// Hide the UAF content directory and its contents
 	ToggleUAFContentHidden.Execute();
+
+	IModularFeatures::Get().RegisterModularFeature(UE::PIE::IOutOfDateAssetCompilation::GetModularFeatureName(), &PIE_UAFCompilationHandler);
 }
 
 void FAnimNextEditorModule::ShutdownModule()
@@ -649,6 +651,8 @@ void FAnimNextEditorModule::ShutdownModule()
 	FEdGraphUtilities::UnregisterVisualPinFactory(GraphPanelPinFactory);
 
 	DisableAssetWizard();
+	
+	IModularFeatures::Get().UnregisterModularFeature(UE::PIE::IOutOfDateAssetCompilation::GetModularFeatureName(), &PIE_UAFCompilationHandler);
 }
 
 void FAnimNextEditorModule::RegisterLocatorFragmentEditorType(FName InLocatorFragmentEditorName)
diff --git a/Engine/Plugins/Experimental/UAF/UAF/Source/UAFEditor/Private/AnimNextEditorModule.h b/Engine/Plugins/Experimental/UAF/UAF/Source/UAFEditor/Private/AnimNextEditorModule.h
--- a/Engine/Plugins/Experimental/UAF/UAF/Source/UAFEditor/Private/AnimNextEditorModule.h
+++ b/Engine/Plugins/Experimental/UAF/UAF/Source/UAFEditor/Private/AnimNextEditorModule.h
@@ -6,6 +6,7 @@
 #include "Module/ModuleEventPropertyCustomization.h"
 #include "UObject/TopLevelAssetPath.h"
 #include "Variables/UAFAssetDataPropertyCustomization.h"
+#include "PIEAssetCompilationHandler.h"
 
 class UAnimNextWorkspaceSchema;
 class FAnimNextGraphPanelNodeFactory;
@@ -78,6 +79,8 @@ private:
 	FTopLevelAssetPath UAFRigVMAssetTopLevelAssetPath;
 	FTopLevelAssetPath SystemTopLevelAssetPath;
 	FTopLevelAssetPath SharedVariablesTopLevelAssetPath;
+	
+	FOutOfDateAssetCompilation PIE_UAFCompilationHandler;
 
 	friend class FLocatorContext;
 };
diff --git a/Engine/Plugins/Experimental/UAF/UAF/Source/UAFEditor/Private/AssetCompilationHandler.cpp b/Engine/Plugins/Experimental/UAF/UAF/Source/UAFEditor/Private/AssetCompilationHandler.cpp
--- a/Engine/Plugins/Experimental/UAF/UAF/Source/UAFEditor/Private/AssetCompilationHandler.cpp
+++ b/Engine/Plugins/Experimental/UAF/UAF/Source/UAFEditor/Private/AssetCompilationHandler.cpp
@@ -3,8 +3,14 @@
 #include "AssetCompilationHandler.h"
 #include "AnimNextRigVMAsset.h"
 #include "AnimNextRigVMAssetEditorData.h"
+#include "Editor.h"
 #include "UAFCompilationScope.h"
 #include "UncookedOnlyUtils.h"
+#include "Settings/LevelEditorPlaySettings.h"
+#include "PIEAssetCompilationHandler.h"
+#include "Logging/MessageLog.h"
+#include "Settings/UAFEditorUserSettings.h"
+#include "UObject/UObjectIterator.h"
 
 namespace UE::UAF::Editor
 {
@@ -57,4 +63,52 @@ ECompileStatus FAssetCompilationHandler::GetCompileStatus(TSharedRef<Workspace::
 	return ECompileStatus::UpToDate;
 }
 
+bool FOutOfDateAssetCompilation::HandleOutOfDateAssets_Internal(const FRequestPlaySessionParams& InPlaySessionParams, TArray<UObject*>& OutAssetsWithErrors) const
+{
+	// Ignore if user opted not to auto compile on PIE
+	if (!GetDefault<UUAFEditorUserSettings>()->bAutoCompileOutOfDateAssets)
+	{
+		return true;
+	}
+
+	OutAssetsWithErrors.Reset();
+
+	for (TObjectIterator<UUAFRigVMAssetEditorData> UAFAssetEditorData; UAFAssetEditorData; ++UAFAssetEditorData)
+	{
+		if (UUAFRigVMAssetEditorData* EditorData = *UAFAssetEditorData)
+		{
+			if (UUAFRigVMAsset* Asset = UncookedOnly::FUtils::GetAsset(EditorData))
+			{
+				EditorData->RecompileVMIfRequired();
+
+				if (EditorData->bErrorsDuringCompilation && EditorData->bDisplayCompilePIEWarning)
+				{
+					OutAssetsWithErrors.Add(Asset);
+				}
+			}
+		}
+	}
+
+	if (OutAssetsWithErrors.Num())
+	{
+		if (!ShowCompilationErrorsDialog(OutAssetsWithErrors, NSLOCTEXT("UAF", "UAFCompilationTypeText", "UAF"), [](const TWeakObjectPtr<UObject>& InAssetToOpen)
+		{
+			check(InAssetToOpen.IsValid());
+			GEditor->EditObject(InAssetToOpen.Get());
+		}))
+		{
+			// TODO we do not have an equivalent of BP errors showing the compilation log
+			return false;
+		}
+		
+		// Mark asset to ignore asset errors until manually compiled
+		for (UObject* Asset : OutAssetsWithErrors)
+		{
+			UncookedOnly::FUtils::GetEditorData(CastChecked<UUAFRigVMAsset>(Asset))->bDisplayCompilePIEWarning = false;
+		}
+	}
+
+	return true;
+}
+
 }
diff --git a/Engine/Plugins/Experimental/UAF/UAF/Source/UAFEditor/Private/PIEAssetCompilationHandler.h b/Engine/Plugins/Experimental/UAF/UAF/Source/UAFEditor/Private/PIEAssetCompilationHandler.h
new file mode 100644
--- /dev/null
+++ b/Engine/Plugins/Experimental/UAF/UAF/Source/UAFEditor/Private/PIEAssetCompilationHandler.h
@@ -0,0 +1,15 @@
+// Copyright Epic Games, Inc. All Rights Reserved.
+
+#pragma once
+
+#include "IOutOfDateAssetCompilation.h"
+
+namespace UE::UAF::Editor
+{
+
+class FOutOfDateAssetCompilation : public PIE::IOutOfDateAssetCompilation
+{
+private:
+	virtual bool HandleOutOfDateAssets_Internal(const FRequestPlaySessionParams& InPlaySessionParams, TArray<UObject*>& OutAssetsWithErrors) const override;
+};
+}
diff --git a/Engine/Plugins/Experimental/UAF/UAF/Source/UAFEditor/Private/Settings/UAFEditorUserSettings.h b/Engine/Plugins/Experimental/UAF/UAF/Source/UAFEditor/Private/Settings/UAFEditorUserSettings.h
--- a/Engine/Plugins/Experimental/UAF/UAF/Source/UAFEditor/Private/Settings/UAFEditorUserSettings.h
+++ b/Engine/Plugins/Experimental/UAF/UAF/Source/UAFEditor/Private/Settings/UAFEditorUserSettings.h
@@ -24,4 +24,8 @@ public:
 	/** If true, functions in UAF workspaces will open in a separate tab per default */
 	UPROPERTY(EditAnywhere, Category="UAF", Config)
 	bool bAutoOpenFunctionsInSeparateTabs = true;
-};
\ No newline at end of file
+	
+	/** If true, when user starts PIE any UAF assets requiring compilation will automatically be compiled */
+	UPROPERTY(EditAnywhere, Category="UAF", Config)
+	bool bAutoCompileOutOfDateAssets = true;
+};
diff --git a/Engine/Plugins/Experimental/UAF/UAF/Source/UAFUncookedOnly/Internal/AnimNextRigVMAssetEditorData.h b/Engine/Plugins/Experimental/UAF/UAF/Source/UAFUncookedOnly/Internal/AnimNextRigVMAssetEditorData.h
--- a/Engine/Plugins/Experimental/UAF/UAF/Source/UAFUncookedOnly/Internal/AnimNextRigVMAssetEditorData.h
+++ b/Engine/Plugins/Experimental/UAF/UAF/Source/UAFUncookedOnly/Internal/AnimNextRigVMAssetEditorData.h
@@ -75,6 +75,7 @@ namespace UE::UAF::Editor
 	struct FFunctionsOutlinerCollapsedGraphItem;
 	class FFunctionsOutlinerMode;
 	class SCommonOutliner;
+	class FOutOfDateAssetCompilation;
 }
 
 namespace UE::UAF::Tests
@@ -277,7 +278,7 @@ protected:
 	friend UE::UAF::Editor::SVariablesView;
 	friend UUAFSharedVariableNode;
 	friend class UE::UAF::UncookedOnly::FModule;
-	
+	friend class UE::UAF::Editor::FOutOfDateAssetCompilation;
 
 	// UObject interface
 	UE_API virtual void Serialize(FArchive& Ar) override;
@@ -763,6 +764,10 @@ protected:
 	UPROPERTY()
 	bool bAutoRecompileVM = true;
 	
+	/** Whether to show a warning when attempting to start in PIE and there is a compiler error on this asset */
+	UPROPERTY(transient)
+	bool bDisplayCompilePIEWarning;
+
 	bool bVMRecompilationRequested = false;
 	mutable bool bErrorsDuringCompilation = false;
 	mutable bool bWarningsDuringCompilation = false;
diff --git a/Engine/Plugins/Experimental/UAF/UAF/Source/UAFUncookedOnly/Private/AnimNextRigVMAssetEditorData.cpp b/Engine/Plugins/Experimental/UAF/UAF/Source/UAFUncookedOnly/Private/AnimNextRigVMAssetEditorData.cpp
--- a/Engine/Plugins/Experimental/UAF/UAF/Source/UAFUncookedOnly/Private/AnimNextRigVMAssetEditorData.cpp
+++ b/Engine/Plugins/Experimental/UAF/UAF/Source/UAFUncookedOnly/Private/AnimNextRigVMAssetEditorData.cpp
@@ -1089,6 +1089,8 @@ void UUAFRigVMAssetEditorData::RecompileVM()
 
 	Asset->VMRuntimeSettings = VMRuntimeSettings;
 	Asset->Components.Empty();
+	
+	bDisplayCompilePIEWarning = true;
 
 	OnPreCompileAsset(Settings);
 
@@ -1442,9 +1444,9 @@ void UUAFRigVMAssetEditorData::RequestAutoVMRecompilation()
 
 void UUAFRigVMAssetEditorData::SetVMRecompilationRequired(bool bInIsRequired)
 {
-	if (bInIsRequired != bVMRecompilationRequested)
+	if (bInIsRequired != bVMRecompilationRequired)
 	{
-		bVMRecompilationRequested = bInIsRequired;
+		bVMRecompilationRequired = bInIsRequired;
 		RecompileRequiredChangedEvent.Broadcast();
 	}
 }
diff --git a/Engine/Source/Editor/UnrealEd/Private/PlayLevel.cpp b/Engine/Source/Editor/UnrealEd/Private/PlayLevel.cpp
--- a/Engine/Source/Editor/UnrealEd/Private/PlayLevel.cpp
+++ b/Engine/Source/Editor/UnrealEd/Private/PlayLevel.cpp
@@ -125,6 +125,7 @@
 #include "Features/IModularFeatures.h"
 #include "Containers/ConsumeAllMpmcQueue.h"
 #include "TickableEditorObject.h"
+#include "IOutOfDateAssetCompilation.h"
 
 DEFINE_LOG_CATEGORY(LogPlayLevel);
 
@@ -1489,60 +1490,27 @@ bool UEditorEngine::SpawnPlayFromHereStart( UWorld* World, AActor*& PlayerStart,
 	return true;
 }
 
-static bool ShowBlueprintErrorDialog( TArray<UBlueprint*> ErroredBlueprints )
+bool UE::PIE::IOutOfDateAssetCompilation::ShowCompilationErrorsDialog(TConstArrayView<UObject*> InAssetsWithErrors, const FText& AssetTypeText, TFunction<void(const TWeakObjectPtr<UObject>&)> InOpenAssetFunc)
 {
-	if (FApp::IsUnattended() || GIsRunningUnattendedScript)
+	if (FApp::IsUnattended() || GIsRunningUnattendedScript || !InAssetsWithErrors.Num())
 	{
 		// App is running in unattended mode, so we should avoid modal dialogs and proceed
 		return true;
 	}
-
-	struct Local
-	{
-		static void OnHyperlinkClicked( TWeakObjectPtr<UBlueprint> InBlueprint, TSharedPtr<SCustomDialog> InDialog )
-		{
-			if (UBlueprint* BlueprintToEdit = InBlueprint.Get())
-			{
-				// Open the blueprint
-				GEditor->EditObject( BlueprintToEdit );
-			}
-
-			if (InDialog.IsValid())
-			{
-				InDialog->RequestDestroyWindow();
-			}
-		}
-
-		static void OnOpenAllLinkClicked(const TArray<UBlueprint*>& BlueprintsToOpen, TSharedPtr<SCustomDialog> InDialog)
-		{
-			for(UBlueprint* BP : BlueprintsToOpen)
-			{
-				if (BP)
-				{
-					GEditor->EditObject(BP);
-				}	
-			}
-
-			if (InDialog.IsValid())
-			{
-				InDialog->RequestDestroyWindow();
-			}
-		}
-	};
-
+	
 	TSharedRef<SVerticalBox> DialogContents = SNew(SVerticalBox)
 		+ SVerticalBox::Slot()
 		.Padding(0.f, 0.f, 0.f, 16.f)
 		[
 			SNew(STextBlock)
-			.Text(NSLOCTEXT("PlayInEditor", "PrePIE_BlueprintErrors", "Are you sure you want to Play in Editor? The following blueprints have unresolved compiler errors."))
+			.Text(FText::FormatOrdered(NSLOCTEXT("PlayInEditor", "PrePIE_AssetErrors", "Are you sure you want to Play in Editor? The following {0} {1}|plural(one=Asset,other=Assets) {1}|plural(one=has,other=have) unresolved compiler errors."), AssetTypeText, InAssetsWithErrors.Num()))
 		];
 
 	TSharedPtr<SCustomDialog> CustomDialog;
 
-	for (UBlueprint* Blueprint : ErroredBlueprints)
+	for (UObject* Asset : InAssetsWithErrors)
 	{
-		TWeakObjectPtr<UBlueprint> BlueprintPtr = Blueprint;
+		TWeakObjectPtr<UObject> WeakAssetPtr = Asset;
 
 		DialogContents->AddSlot()
 			.AutoHeight()
@@ -1550,14 +1518,22 @@ static bool ShowBlueprintErrorDialog( TArray<UBlueprint*> ErroredBlueprints )
 			[
 				SNew(SHyperlink)
 				.Style(FAppStyle::Get(), "Common.GotoBlueprintHyperlink")
-				.OnNavigate(FSimpleDelegate::CreateLambda([BlueprintPtr, &CustomDialog]() { Local::OnHyperlinkClicked(BlueprintPtr, CustomDialog); }))
-				.Text(FText::FromString(Blueprint->GetName()))
-				.ToolTipText(NSLOCTEXT("SourceHyperlink", "EditBlueprint_ToolTip", "Click to edit the blueprint"))
+				.OnNavigate(FSimpleDelegate::CreateLambda([WeakAssetPtr, &CustomDialog, InOpenAssetFunc]() 
+				{
+					InOpenAssetFunc(WeakAssetPtr);
+
+					if (CustomDialog.IsValid())
+					{
+						CustomDialog->RequestDestroyWindow();
+					}
+				 }))
+				.Text(FText::FromString(Asset->GetName()))
+				.ToolTipText(FText::FormatOrdered(NSLOCTEXT("SourceHyperlink", "EditAsset_ToolTip", "Click to edit the {0}}"), Asset->GetClass()->GetDisplayNameText()))
 			];
 	}
 
-	// Add an option to open all errored blueprints
-	if(ErroredBlueprints.Num() > 1)
+	// Add an option to open all errored asset
+	if(InAssetsWithErrors.Num() > 1)
 	{
 		DialogContents->AddSlot()
 			.Padding(0.f, 16.f, 0.f, 0.f)
@@ -1565,14 +1541,24 @@ static bool ShowBlueprintErrorDialog( TArray<UBlueprint*> ErroredBlueprints )
 			[
 				SNew(SHyperlink)
 				.Style(FAppStyle::Get(), "Common.GotoBlueprintHyperlink")
-				.OnNavigate(FSimpleDelegate::CreateLambda([&ErroredBlueprints, &CustomDialog]() { Local::OnOpenAllLinkClicked(ErroredBlueprints, CustomDialog); }))
-				.Text(NSLOCTEXT("SourceHyperlink", "EditAllErroredBlueprints", "Open all errored blueprints"))
-				.ToolTipText(NSLOCTEXT("SourceHyperlink", "EditAllErroredBlueprints_ToolTip", "Opens all the errored blueprint in the editor"))
+				.OnNavigate(FSimpleDelegate::CreateLambda([&InAssetsWithErrors, &CustomDialog, InOpenAssetFunc]()
+				{
+					for(UObject* Asset : InAssetsWithErrors)
+					{
+						InOpenAssetFunc(Asset);
+					}
+
+					if (CustomDialog.IsValid())
+					{
+						CustomDialog->RequestDestroyWindow();
+					}
+				}))
+				.Text(NSLOCTEXT("SourceHyperlink", "EditAllErroredAssets", "Open all assets with compilation errors"))
+				.ToolTipText(FText::FormatOrdered(NSLOCTEXT("SourceHyperlink", "EditAllErroredAssets_ToolTip", "Opens all the {0} assets with compilation errors in their respective asset editor"), AssetTypeText))
 			];
 	}
 
-	static const FText DialogTitle = NSLOCTEXT("PlayInEditor", "PrePIE_BlueprintErrorsTitle", "Blueprint Compilation Errors");
-
+	const FText DialogTitle = FText::FormatOrdered(NSLOCTEXT("PlayInEditor", "PrePIE_AssetErrorsTitle", "{0} Asset Compilation {1}|plural(one=Error,other=Errors)"), AssetTypeText, InAssetsWithErrors.Num());
 	static const FText OKText = NSLOCTEXT("PlayInEditor", "PrePIE_OkText", "Play in Editor");
 	static const FText CancelText = NSLOCTEXT("Dialogs", "EAppReturnTypeCancel", "Cancel");
 
@@ -2670,7 +2656,11 @@ void UEditorEngine::StartPlayInEditorSession(FRequestPlaySessionParams& InReques
 	if (ErroredBlueprints.Num() && !GIsDemoMode)
 	{
 		// There was at least one blueprint with an error, make sure the user is OK with that.
-		bool bContinuePIE = ShowBlueprintErrorDialog(ErroredBlueprints);
+		bool bContinuePIE = UE::PIE::IOutOfDateAssetCompilation::ShowCompilationErrorsDialog(MakeArrayView<UObject*>((UObject**)ErroredBlueprints.GetData(),ErroredBlueprints.Num()), UBlueprint::StaticClass()->GetDisplayNameText(), [](const TWeakObjectPtr<UObject>& InAssetToOpen)
+		{
+			check(InAssetToOpen.IsValid());
+			GEditor->EditObject(InAssetToOpen.Get());
+		});
 
 		if (!bContinuePIE)
 		{
@@ -2691,6 +2681,21 @@ void UEditorEngine::StartPlayInEditorSession(FRequestPlaySessionParams& InReques
 		}
 	}
 
+	// Auto-compile other dirty assets if needed, and let the user cancel PIE if there are any errors
+	TArray<UE::PIE::IOutOfDateAssetCompilation*> AssetCompilers = IModularFeatures::Get().GetModularFeatureImplementations<UE::PIE::IOutOfDateAssetCompilation>(UE::PIE::IOutOfDateAssetCompilation::GetModularFeatureName());
+	for (UE::PIE::IOutOfDateAssetCompilation* AssetCompiler : AssetCompilers)
+	{
+		TArray<UObject*> AssetsWithErrors;
+		const bool bContinuePIE = AssetCompiler->HandleOutOfDateAssets(InRequestParams, AssetsWithErrors);
+		if (!bContinuePIE && !GIsDemoMode)
+		{
+			FEditorDelegates::EndPIE.Broadcast(InRequestParams.WorldType == EPlaySessionWorldType::SimulateInEditor);
+			FNavigationSystem::OnPIEEnd(*InWorld);
+			CancelRequestPlaySession();
+			return;
+		}
+	}
+
 	// Register for log processing so we can promote errors/warnings to the message log
 	if (GetDefault<ULevelEditorPlaySettings>()->bPromoteOutputLogWarningsDuringPIE)
 	{
diff --git a/Engine/Source/Editor/UnrealEd/Private/PlayLevelViaLauncher.cpp b/Engine/Source/Editor/UnrealEd/Private/PlayLevelViaLauncher.cpp
--- a/Engine/Source/Editor/UnrealEd/Private/PlayLevelViaLauncher.cpp
+++ b/Engine/Source/Editor/UnrealEd/Private/PlayLevelViaLauncher.cpp
@@ -27,7 +27,9 @@
 #include "Settings/PlatformsMenuSettings.h"
 #include "TargetReceipt.h"
 #include "UnrealEdMisc.h"
+#include "Features/IModularFeatures.h"
 #include "Widgets/Notifications/SNotificationList.h"
+#include "IOutOfDateAssetCompilation.h"
 
 #define LOCTEXT_NAMESPACE "PlayLevel"
 
@@ -365,6 +367,13 @@ BuildCookRun->SetUseZenStore(bUseZenStore);
 	TArray<UBlueprint*> ErroredBlueprints;
 	FInternalPlayLevelUtils::ResolveDirtyBlueprints(!EditorPlaySettings->bAutoCompileBlueprintsOnLaunch, ErroredBlueprints, false);
 
+	const TArray<UE::PIE::IOutOfDateAssetCompilation*> AssetCompilers = IModularFeatures::Get().GetModularFeatureImplementations<UE::PIE::IOutOfDateAssetCompilation>(UE::PIE::IOutOfDateAssetCompilation::GetModularFeatureName());
+	for (UE::PIE::IOutOfDateAssetCompilation* AssetCompiler : AssetCompilers)
+	{
+		TArray<UObject*> AssetsWithErrors;
+		AssetCompiler->HandleOutOfDateAssets(InRequestParams, AssetsWithErrors);
+	}
+
 	TArray<FString> MapNames;
 	FWorldContext & EditorContext = GetEditorWorldContext();
 
diff --git a/Engine/Source/Editor/UnrealEd/Public/IOutOfDateAssetCompilation.h b/Engine/Source/Editor/UnrealEd/Public/IOutOfDateAssetCompilation.h
new file mode 100644
--- /dev/null
+++ b/Engine/Source/Editor/UnrealEd/Public/IOutOfDateAssetCompilation.h
@@ -0,0 +1,40 @@
+// Copyright Epic Games, Inc. All Rights Reserved.
+
+#pragma once
+
+#include "CoreMinimal.h"
+#include "Features/IModularFeature.h"
+#include "Stats/Stats.h"
+
+struct FRequestPlaySessionParams;
+
+#define UE_API UNREALED_API
+
+namespace UE::PIE
+{
+
+class IOutOfDateAssetCompilation : public IModularFeature
+{
+public:
+	virtual ~IOutOfDateAssetCompilation() = default;
+
+	static FName GetModularFeatureName()
+	{
+		static FName FeatureName = FName("OutOfDateAssetCompilation");
+		return FeatureName;
+	}
+
+	bool HandleOutOfDateAssets(const FRequestPlaySessionParams& InPlaySessionParams, TArray<UObject*>& OutAssetsWithErrors) const
+	{
+		QUICK_SCOPE_CYCLE_COUNTER(IOutOfDateAssetCompilation_HandleOutOfDateAssets);
+		return HandleOutOfDateAssets_Internal(InPlaySessionParams, OutAssetsWithErrors);
+	}
+
+	virtual bool HandleOutOfDateAssets_Internal(const FRequestPlaySessionParams& InPlaySessionParams, TArray<UObject*>& OutAssetsWithErrors) const = 0;
+
+	// Present user with dialog showing all assets with errors, returns whether the user chose to ignore the errors
+	static UE_API bool ShowCompilationErrorsDialog(TConstArrayView<UObject*> InAssetsWithErrors, const FText& AssetTypeText, TFunction<void(const TWeakObjectPtr<UObject>&)> InOpenAssetFunc);
+};
+}
+
+#undef UE_API // UNREALED_API
