[EditorEngine]
- Add async version of Save

#Synced-CL 51651818

diff --git a/Engine/Source/Editor/UnrealEd/Classes/Editor/EditorEngine.h b/Engine/Source/Editor/UnrealEd/Classes/Editor/EditorEngine.h
--- a/Engine/Source/Editor/UnrealEd/Classes/Editor/EditorEngine.h
+++ b/Engine/Source/Editor/UnrealEd/Classes/Editor/EditorEngine.h
@@ -2276,6 +2276,14 @@ public:
 	UNREALED_API FSavePackageResultStruct Save(UPackage* InOuter, UObject* InAsset, const TCHAR* Filename,
 		const FSavePackageArgs& SaveArgs);
 
+	/** An async editor wrapper for UPackage::SaveAsync. Auto-adds files to source control when necessary
+ 	 *
+ 	 *  @note It is very important that any optional arguments passed via SaveArgs (i.e. ArchiveCookData, SavePackageContext) 
+ 	 *        stay valid for the entire duration of the task.
+	 */
+	UNREALED_API UE::Tasks::TTask<FSavePackageResultStruct> SaveAsync(UPackage* InOuter, UObject* InAsset, const TCHAR* Filename,
+		const FSavePackageArgs& SaveArgs);
+
 	UNREALED_API virtual bool InitializePhysicsSceneForSaveIfNecessary(UWorld* World, bool &bOutForceInitialized);
 	UNREALED_API void CleanupPhysicsSceneThatWasInitializedForSave(UWorld* World, bool bForceInitialized);
 
@@ -3347,6 +3355,9 @@ public:
 private:
 	FObjectSubsystemCollection<UEditorSubsystem> EditorSubsystemCollection;
 
+	void BeginEditorSave(struct FEditorSaveContext& Context, struct FSlowTask* SlowTask);
+	void EndEditorSave(struct FEditorSaveContext& Context, struct FSlowTask* SlowTask);
+
 	// DEPRECATED VARIABLES ONLY
 public:
 	/** True if we're Simulating In Editor, as opposed to Playing In Editor.  In this mode, simulation takes place right the level editing environment */
diff --git a/Engine/Source/Editor/UnrealEd/Private/EditorEngine.cpp b/Engine/Source/Editor/UnrealEd/Private/EditorEngine.cpp
--- a/Engine/Source/Editor/UnrealEd/Private/EditorEngine.cpp
+++ b/Engine/Source/Editor/UnrealEd/Private/EditorEngine.cpp
@@ -4454,151 +4454,249 @@ UE_TRACE_EVENT_BEGIN(Cpu, EditorEngine_Save, NoSync)
 	UE_TRACE_EVENT_FIELD(UE::Trace::WideString, PackageName)
 UE_TRACE_EVENT_END()
 
-FSavePackageResultStruct UEditorEngine::Save(UPackage* InOuter, UObject* InAsset, const TCHAR* Filename,
-	const FSavePackageArgs& InSaveArgs)
-{
-	UE_TRACE_LOG_SCOPED_T(Cpu, EditorEngine_Save, CpuChannel)
-		<< EditorEngine_Save.PackageName(InOuter ? *(InOuter->GetName()) : TEXT("None"));
+struct FEditorSaveContext
+{
+	UPackage*                                  Package = nullptr;
+	UObject*                                   Asset = nullptr;
+	UWorld*                                    World = nullptr;
+	UWorld*                                    OriginalOwningWorld = nullptr;
+	FString                                    Filename;
+	FSavePackageArgs                           SaveArgs;
+	FObjectSaveContextData                     ObjectSaveContext;
+	TUniquePtr<FSavePackageContext>            UniqueContext;
+	FSavePackageResultStruct                   SavePackageResult;
+	bool                                       bAutoAddPkgToSCC = false;
+	bool                                       bInitializedPhysicsSceneForSave = false;
+	bool                                       bForceInitializedWorld = false;
+	UE::Tasks::TTask<FSavePackageResultStruct> InnerTask;
+};
 
-	UE_TRACK_REFERENCING_PACKAGE_SCOPED(InOuter, PackageAccessTrackingOps::NAME_Save); // Needs to be here in addition to UPackage::Save so that InitializePhysicsSceneForSaveIfNecessary and OnPreSaveWorld have appropriate referencing package info
+void UEditorEngine::BeginEditorSave(FEditorSaveContext& Context, FSlowTask* SlowTask)
+{
+	bool bIsCooking = Context.SaveArgs.ArchiveCookData != nullptr;
+	Context.SaveArgs.TopLevelFlags = UE::SavePackageUtilities::NormalizeTopLevelFlags(Context.SaveArgs.TopLevelFlags, bIsCooking);
 
-	FSavePackageArgs SaveArgs(InSaveArgs);
-	FScopedSlowTask SlowTask(100, FText(), SaveArgs.bSlowTask);
-	bool bIsCooking = SaveArgs.ArchiveCookData != nullptr;
-	SaveArgs.TopLevelFlags = UE::SavePackageUtilities::NormalizeTopLevelFlags(SaveArgs.TopLevelFlags, bIsCooking);
-	UObject* Asset = InAsset;
-	if (!Asset && InOuter)
+	if (!Context.Asset && Context.Package)
 	{
 		// Check if the package contains a map and set the world object as the asset
-		if (InOuter->HasAnyPackageFlags(PKG_ContainsMap) )
+		if (Context.Package->HasAnyPackageFlags(PKG_ContainsMap))
 		{
-			Asset = UWorld::FindWorldInPackage(InOuter);
+			Context.Asset = UWorld::FindWorldInPackage(Context.Package);
 		}
 		else
 		{
 			// Otherwise find the main asset of the package
-			Asset = InOuter->FindAssetInPackage();
+			Context.Asset = Context.Package->FindAssetInPackage();
 		}
 	}
 
 	// if no save package context was passed in and the default settings were modified, install a context for the save
-	TUniquePtr<FSavePackageContext> UniqueContext;
-	if (SaveArgs.SavePackageContext == nullptr && !FSavePackageSettings::GetDefaultSettings().IsDefault())
+	if (Context.SaveArgs.SavePackageContext == nullptr && !FSavePackageSettings::GetDefaultSettings().IsDefault())
 	{
-		UniqueContext = MakeUnique<FSavePackageContext>(nullptr, nullptr, FSavePackageSettings::GetDefaultSettings());
-		SaveArgs.SavePackageContext = UniqueContext.Get();
+		Context.UniqueContext = MakeUnique<FSavePackageContext>(nullptr, nullptr, FSavePackageSettings::GetDefaultSettings());
+		Context.SaveArgs.SavePackageContext = Context.UniqueContext.Get();
 	}
 
-	SlowTask.EnterProgressFrame(10);
-
-	UWorld* World = Cast<UWorld>(Asset);
-	bool bInitializedPhysicsSceneForSave = false;
-	bool bForceInitializedWorld = false;
-	const bool bSavingConcurrent = !!(SaveArgs.SaveFlags & ESaveFlags::SAVE_Concurrent);
+	if (SlowTask)
+	{
+		SlowTask->EnterProgressFrame(10);
+	}
 
-	FObjectSaveContextData ObjectSaveContext(InOuter, SaveArgs.GetTargetPlatform(), Filename, SaveArgs.SaveFlags);
-	if (InSaveArgs.ArchiveCookData)
+	if (Context.SaveArgs.ArchiveCookData)
 	{
-		ObjectSaveContext.CookType = InSaveArgs.ArchiveCookData->CookContext.GetCookType();
-		ObjectSaveContext.CookingDLC = InSaveArgs.ArchiveCookData->CookContext.GetCookingDLC();
-		ObjectSaveContext.CookInfo = InSaveArgs.ArchiveCookData->CookContext.GetCookInfo();
+		Context.ObjectSaveContext.CookType = Context.SaveArgs.ArchiveCookData->CookContext.GetCookType();
+		Context.ObjectSaveContext.CookingDLC = Context.SaveArgs.ArchiveCookData->CookContext.GetCookingDLC();
+		Context.ObjectSaveContext.CookInfo = Context.SaveArgs.ArchiveCookData->CookContext.GetCookInfo();
 	}
-	UWorld *OriginalOwningWorld = nullptr;
-	if ( World )
+
+	if (UWorld* World = Cast<UWorld>(Context.Asset))
 	{
+		Context.World = World;
+
+		const bool bSavingConcurrent = !!(Context.SaveArgs.SaveFlags & ESaveFlags::SAVE_Concurrent);
 		if (!bSavingConcurrent)
 		{
-			bInitializedPhysicsSceneForSave = InitializePhysicsSceneForSaveIfNecessary(World, bForceInitializedWorld);
+			Context.bInitializedPhysicsSceneForSave = InitializePhysicsSceneForSaveIfNecessary(World, Context.bForceInitializedWorld);
 
 			// bForceInitialized=true marks that We Called InitWorld and need to call CleanupWorld
 			// but if we are saving during Load and the caller of LoadPackage requested that we keep it initialized, we should keep it initialized.
-			if (bForceInitializedWorld)
+			if (Context.bForceInitializedWorld)
 			{
 				// TODO: If we ever add a way to set bInitialized=false on a UWorld, then in future saves after setting it back to false
 				// we will still parse the KeepInitializedDuringLoadTag here and not CleanupWorld. If we set bInitialized=false on
 				// a UWorld, we need to clear KeepInitializedDuringLoadTag from its InstancingContext.
 				FLinkerLoad* LinkerLoad = World->GetPackage()->GetLinker();
-				if (bForceInitializedWorld && LinkerLoad && LinkerLoad->GetInstancingContext().HasTag(UWorld::KeepInitializedDuringLoadTag))
+				if (Context.bForceInitializedWorld && LinkerLoad && LinkerLoad->GetInstancingContext().HasTag(UWorld::KeepInitializedDuringLoadTag))
 				{
-					bForceInitializedWorld = false;
+					Context.bForceInitializedWorld = false;
 				}
 			}
 
-			OnPreSaveWorld(World, FObjectPreSaveContext(ObjectSaveContext));
+			OnPreSaveWorld(World, FObjectPreSaveContext(Context.ObjectSaveContext));
 		}
 
-		OriginalOwningWorld = World->PersistentLevel->OwningWorld;
+		Context.OriginalOwningWorld = World->PersistentLevel->OwningWorld;
 		World->PersistentLevel->OwningWorld = World;
 	}
 
 	// See if the package is a valid candidate for being auto-added to the default changelist.
 	// Only allows the addition of newly created packages while in the editor and then only if the user has the option enabled.
-	bool bAutoAddPkgToSCC = false;
-	if (!ObjectSaveContext.bProceduralSave)
+	Context.bAutoAddPkgToSCC = false;
+	if (!Context.ObjectSaveContext.bProceduralSave)
 	{
-		bAutoAddPkgToSCC = IsPackageValidForAutoAdding( InOuter, Filename );
+		Context.bAutoAddPkgToSCC = IsPackageValidForAutoAdding(Context.Package, Context.Filename);
 	}
 
-	SlowTask.EnterProgressFrame(70);
+	if (SlowTask)
+	{
+		SlowTask->EnterProgressFrame(70);
+	}
 
-	UPackage::PreSavePackageWithContextEvent.Broadcast(InOuter, FObjectPreSaveContext(ObjectSaveContext));
-	FSavePackageResultStruct Result = UPackage::Save(InOuter, Asset, Filename, SaveArgs);
+	UPackage::PreSavePackageWithContextEvent.Broadcast(Context.Package, FObjectPreSaveContext(Context.ObjectSaveContext));
+}
 
-	SlowTask.EnterProgressFrame(10);
-	ObjectSaveContext.bSaveSucceeded = Result == ESavePackageResult::Success;
+void UEditorEngine::EndEditorSave(FEditorSaveContext& Context, FSlowTask* SlowTask)
+{
+	if (SlowTask)
+	{
+		SlowTask->EnterProgressFrame(10);
+	}
+
+	Context.ObjectSaveContext.bSaveSucceeded = Context.SavePackageResult == ESavePackageResult::Success;
 
 	// If the package is a valid candidate for being automatically-added to source control, go ahead and add it
 	// to the default changelist
-	if (Result == ESavePackageResult::Success && bAutoAddPkgToSCC)
+	if (Context.SavePackageResult == ESavePackageResult::Success && Context.bAutoAddPkgToSCC)
 	{
 		const bool bIsUncontrolledChangelistEnabled = FUncontrolledChangelistsModule::Get().IsEnabled();
 
 		// IsPackageValidForAutoAdding should not return true if SCC is disabled
 		check(ISourceControlModule::Get().IsEnabled() || bIsUncontrolledChangelistEnabled);
 
-		if(!ISourceControlModule::Get().GetProvider().IsAvailable() && !bIsUncontrolledChangelistEnabled)
+		if (!ISourceControlModule::Get().GetProvider().IsAvailable() && !bIsUncontrolledChangelistEnabled)
 		{
 			// Show the login window here & store the file we are trying to add.
 			// We defer the add operation until we have a valid source control connection.
 			ISourceControlModule::Get().ShowLoginDialog(FSourceControlLoginClosed::CreateUObject(this, &UEditorEngine::RunDeferredMarkForAddFiles), ELoginWindowMode::Modeless);
 		}
-		DeferredFilesToAddToSourceControl.Add(Filename);
+		DeferredFilesToAddToSourceControl.Add(Context.Filename);
 	}
 
-	SlowTask.EnterProgressFrame(10);
+	if (SlowTask)
+	{
+		SlowTask->EnterProgressFrame(10);
+	}
 
-	if ( World )
+	if (Context.World)
 	{
-		if (OriginalOwningWorld)
+		if (Context.OriginalOwningWorld)
 		{
-			World->PersistentLevel->OwningWorld = OriginalOwningWorld;
+			Context.World->PersistentLevel->OwningWorld = Context.OriginalOwningWorld;
 		}
 
+		const bool bSavingConcurrent = !!(Context.SaveArgs.SaveFlags & ESaveFlags::SAVE_Concurrent);
 		if (!bSavingConcurrent)
 		{
-			OnPostSaveWorld(World, FObjectPostSaveContext(ObjectSaveContext));
+			OnPostSaveWorld(Context.World, FObjectPostSaveContext(Context.ObjectSaveContext));
 
-			if (bInitializedPhysicsSceneForSave)
+			if (Context.bInitializedPhysicsSceneForSave)
 			{
-				CleanupPhysicsSceneThatWasInitializedForSave(World, bForceInitializedWorld);
+				CleanupPhysicsSceneThatWasInitializedForSave(Context.World, Context.bForceInitializedWorld);
 			}
 
-			if (Result == ESavePackageResult::Success) // Package saved successfully?
+			if (Context.SavePackageResult == ESavePackageResult::Success) // Package saved successfully?
 			{
 				// Rerunning construction scripts may have made it dirty again
-				InOuter->SetDirtyFlag(false);
+				Context.Package->SetDirtyFlag(false);
 			}
 		}
 	}
 
-	if (ObjectSaveContext.bUpdatingLoadedPath)
+	if (Context.ObjectSaveContext.bUpdatingLoadedPath)
 	{
 		// Notify the asset registry
 		IAssetRegistry& AssetRegistry = IAssetRegistry::GetChecked();
-		AssetRegistry.AssetsSaved(MoveTemp(Result.SavedAssets));
+		AssetRegistry.AssetsSaved(MoveTemp(Context.SavePackageResult.SavedAssets));
 	}
+}
 
-	return Result;
+FSavePackageResultStruct UEditorEngine::Save(UPackage* InOuter, UObject* InAsset, const TCHAR* Filename,
+	const FSavePackageArgs& InSaveArgs)
+{
+	UE_TRACE_LOG_SCOPED_T(Cpu, EditorEngine_Save, CpuChannel)
+		<< EditorEngine_Save.PackageName(InOuter ? *(InOuter->GetName()) : TEXT("None"));
+
+	// Needs to be here in addition to UPackage::Save so that InitializePhysicsSceneForSaveIfNecessary and OnPreSaveWorld have appropriate referencing package info
+	UE_TRACK_REFERENCING_PACKAGE_SCOPED(InOuter, PackageAccessTrackingOps::NAME_Save); 
+
+	FScopedSlowTask SlowTask(100, FText(), InSaveArgs.bSlowTask);
+
+	FEditorSaveContext Context
+	{
+		.Package = InOuter,
+		.Asset = InAsset,
+		.Filename = Filename,
+		.SaveArgs = InSaveArgs,
+		.ObjectSaveContext = FObjectSaveContextData(InOuter, InSaveArgs.GetTargetPlatform(), Filename, InSaveArgs.SaveFlags),
+	};
+
+	BeginEditorSave(Context, &SlowTask);
+	Context.SavePackageResult = UPackage::Save(Context.Package, Context.Asset, *Context.Filename, Context.SaveArgs);
+	EndEditorSave(Context, &SlowTask);
+
+	return MoveTemp(Context.SavePackageResult);
+}
+
+UE::Tasks::TTask<FSavePackageResultStruct> UEditorEngine::SaveAsync(UPackage* InOuter, UObject* InAsset, const TCHAR* Filename,
+		const FSavePackageArgs& InSaveArgs)
+{
+	FEditorSaveContext LocalContext
+	{
+		.Package = InOuter,
+		.Asset = InAsset,
+		.Filename = Filename,
+		.SaveArgs = InSaveArgs,
+		.ObjectSaveContext = FObjectSaveContextData(InOuter, InSaveArgs.GetTargetPlatform(), Filename, InSaveArgs.SaveFlags)
+	};
+
+	auto Context = MakeShared<FEditorSaveContext>(MoveTemp(LocalContext));
+
+	UE::Tasks::FTaskEvent TaskEvent(TEXT("EditorSaveEvent"));
+	auto BeginEditorSaveTask =
+		UE::Tasks::Launch(
+			TEXT("BeginEditorSave"),
+			[this, Context, TaskEvent]() mutable
+			{
+				UE_TRACK_REFERENCING_PACKAGE_SCOPED(Context->Package, PackageAccessTrackingOps::NAME_Save);
+
+				BeginEditorSave(*Context, nullptr);
+
+				Context->InnerTask = UPackage::SaveAsync(Context->Package, Context->Asset, *Context->Filename, Context->SaveArgs);
+
+				// Once the result of the task is available, this will allow the endeditorsave to run
+				TaskEvent.AddPrerequisites(Context->InnerTask);
+				TaskEvent.Trigger();
+			},
+			LowLevelTasks::ETaskPriority::Normal,
+			UE::Tasks::EExtendedTaskPriority::GameThreadHiPri
+		);
+
+	auto EndEditorSaveTask =
+		UE::Tasks::Launch(
+			TEXT("EndEditorSave"),
+			[this, Context]()
+			{
+				UE_TRACK_REFERENCING_PACKAGE_SCOPED(Context->Package, PackageAccessTrackingOps::NAME_Save);
+				Context->SavePackageResult = MoveTemp(Context->InnerTask.GetResult());
+				EndEditorSave(*Context, nullptr);
+				return MoveTemp(Context->SavePackageResult);
+			},
+			UE::Tasks::Prerequisites(BeginEditorSaveTask, TaskEvent),
+			LowLevelTasks::ETaskPriority::Normal,
+			UE::Tasks::EExtendedTaskPriority::GameThreadHiPri
+		);
+
+	return EndEditorSaveTask;
 }
 
 bool UEditorEngine::SavePackage(UPackage* InOuter, UObject* InAsset, const TCHAR* Filename,
