Fix missing include directive in thumbnail generator

#rnx
#Synced-CL 51735391

diff --git a/Engine/Source/Editor/UnrealEd/Private/ThumbnailRendering/StaticMeshThumbnailRenderer.cpp b/Engine/Source/Editor/UnrealEd/Private/ThumbnailRendering/StaticMeshThumbnailRenderer.cpp
--- a/Engine/Source/Editor/UnrealEd/Private/ThumbnailRendering/StaticMeshThumbnailRenderer.cpp
+++ b/Engine/Source/Editor/UnrealEd/Private/ThumbnailRendering/StaticMeshThumbnailRenderer.cpp
@@ -7,6 +7,7 @@
 #include "SceneView.h"
 #include "ThumbnailHelpers.h"
 #include "Engine/StaticMesh.h"
+#include "RenderingThread.h"
 
 #include UE_INLINE_GENERATED_CPP_BY_NAME(StaticMeshThumbnailRenderer)
 
