Removed useless code from previous CL

#Synced-CL 51738745

diff --git a/Engine/Source/Editor/UnrealEd/Private/UserDefinedStructEditorData.cpp b/Engine/Source/Editor/UnrealEd/Private/UserDefinedStructEditorData.cpp
--- a/Engine/Source/Editor/UnrealEd/Private/UserDefinedStructEditorData.cpp
+++ b/Engine/Source/Editor/UnrealEd/Private/UserDefinedStructEditorData.cpp
@@ -30,8 +30,7 @@ namespace
 			{
 				for (FStructVariableDescription& StructVariableDesc : UDSEditorData->VariablesDescriptions)
 				{
-					static const FName TextCategory = TEXT("text"); // Must match UEdGraphSchema_K2::PC_Text
-					if (StructVariableDesc.Category == UEdGraphSchema_K2::PC_Text/*TextCategory*/)
+					if (StructVariableDesc.Category == UEdGraphSchema_K2::PC_Text)
 					{
 						FText StructVariableValue;
 						if (FTextStringHelper::ReadFromBuffer(*StructVariableDesc.DefaultValue, StructVariableValue) && UpdateTextsPackageNamespace.UpdateSingleText(StructVariableValue))
