[Typed Elements] Fix up usage of interfaces to use TTypedElement<InterfaceType> from using the interface pointer directly

#rnx

diff --git a/Engine/Plugins/Experimental/EditorDataStorageFeatures/Source/TedsOperations/Private/DragAndDrop/Widgets/ViewportDropHandler.cpp b/Engine/Plugins/Experimental/EditorDataStorageFeatures/Source/TedsOperations/Private/DragAndDrop/Widgets/ViewportDropHandler.cpp
--- a/Engine/Plugins/Experimental/EditorDataStorageFeatures/Source/TedsOperations/Private/DragAndDrop/Widgets/ViewportDropHandler.cpp
+++ b/Engine/Plugins/Experimental/EditorDataStorageFeatures/Source/TedsOperations/Private/DragAndDrop/Widgets/ViewportDropHandler.cpp
@@ -51,7 +51,7 @@ namespace ViewportDropHandler_Private
 				{
 					if (TTypedElement<ITypedElementWorldInterface> WorldElement = Registry->GetElement<ITypedElementWorldInterface>(Column->Handle))
 					{
-						WorldElement.AddIgnoredElementToCollisionQueryParams(Column->Handle, OutParams, true);
+						WorldElement.AddIgnoredElementToCollisionQueryParams(OutParams, true);
 						continue;
 					}
 				}
diff --git a/Engine/Source/Editor/UnrealEd/Private/Editor/ActorPositioning.cpp b/Engine/Source/Editor/UnrealEd/Private/Editor/ActorPositioning.cpp
--- a/Engine/Source/Editor/UnrealEd/Private/Editor/ActorPositioning.cpp
+++ b/Engine/Source/Editor/UnrealEd/Private/Editor/ActorPositioning.cpp
@@ -76,7 +76,7 @@ namespace ActorPositioningLocals
 			TTypedElement<ITypedElementWorldInterface> WorldElement = Registry->GetElement<ITypedElementWorldInterface>(Element);
 			if (WorldElement)
 			{
-				WorldElement.AddIgnoredElementToCollisionQueryParams(Element, OutParams, true);
+				WorldElement.AddIgnoredElementToCollisionQueryParams(OutParams, true);
 			}
 			else
 			{
diff --git a/Engine/Source/Editor/UnrealEd/Private/EditorViewportSelectability.cpp b/Engine/Source/Editor/UnrealEd/Private/EditorViewportSelectability.cpp
--- a/Engine/Source/Editor/UnrealEd/Private/EditorViewportSelectability.cpp
+++ b/Engine/Source/Editor/UnrealEd/Private/EditorViewportSelectability.cpp
@@ -796,13 +796,13 @@ bool FEditorViewportSelectability::IsTypedElementSelectable(const FTypedElementH
 		return false;
 	}
 
-	ITypedElementObjectInterface* const ObjectInterface = TypedElementRegistry->GetElementInterface<ITypedElementObjectInterface>(InElementHandle);
+	TTypedElement<ITypedElementObjectInterface> const ObjectInterface = TypedElementRegistry->GetElement<ITypedElementObjectInterface>(InElementHandle);
 	if (!ObjectInterface)
 	{
 		return false;
 	}
 
-	UPrimitiveComponent* const PrimitiveComponent = ObjectInterface->GetObjectAs<UPrimitiveComponent>(InElementHandle);
+	UPrimitiveComponent* const PrimitiveComponent = ObjectInterface.GetObjectAs<UPrimitiveComponent>();
 	if (!IsValid(PrimitiveComponent))
 	{
 		return false;
diff --git a/Engine/Source/Editor/UnrealEd/Private/Kismet2/BlueprintEditorUtils.cpp b/Engine/Source/Editor/UnrealEd/Private/Kismet2/BlueprintEditorUtils.cpp
--- a/Engine/Source/Editor/UnrealEd/Private/Kismet2/BlueprintEditorUtils.cpp
+++ b/Engine/Source/Editor/UnrealEd/Private/Kismet2/BlueprintEditorUtils.cpp
@@ -8355,13 +8355,13 @@ public:
 				for (int i = 0; i < ChildElementHandles.Num(); ++i)
 				{
 					FTypedElementHandle ChildElementHandle = ChildElementHandles[i];
-					if (ITypedElementHierarchyInterface* ChildElementHierarchyInterface = Registry->GetElementInterface<ITypedElementHierarchyInterface>(ChildElementHandle))
+					if (TTypedElement<ITypedElementHierarchyInterface> ChildElementHierarchyInterface = Registry->GetElement<ITypedElementHierarchyInterface>(ChildElementHandle))
 					{
-						ChildElementHierarchyInterface->GetChildElements(ChildElementHandle, ChildElementHandles, true);
+						ChildElementHierarchyInterface.GetChildElements(ChildElementHandles, true);
 					}
-					if (ITypedElementObjectInterface* ChildElementObjectInterface = Registry->GetElementInterface<ITypedElementObjectInterface>(ChildElementHandle))
+					if (TTypedElement<ITypedElementObjectInterface> ChildElementObjectInterface = Registry->GetElement<ITypedElementObjectInterface>(ChildElementHandle))
 					{
-						if (UObject* ChildObject = ChildElementObjectInterface->GetObject(ChildElementHandle))
+						if (UObject* ChildObject = ChildElementObjectInterface.GetObject())
 						{
 							OutPotentialReferencerObjects.Add(ChildObject);
 						}
diff --git a/Engine/Source/Editor/UnrealEd/Private/LevelEditorViewport.cpp b/Engine/Source/Editor/UnrealEd/Private/LevelEditorViewport.cpp
--- a/Engine/Source/Editor/UnrealEd/Private/LevelEditorViewport.cpp
+++ b/Engine/Source/Editor/UnrealEd/Private/LevelEditorViewport.cpp
@@ -147,12 +147,12 @@ namespace LevelEditorViewportLocals
 
 		for (const FTypedElementHandle& Element : ElementArray)
 		{
-			ITypedElementObjectInterface* ObjectInterface = Registry->GetElementInterface<ITypedElementObjectInterface>(Element);
+			TTypedElement<ITypedElementObjectInterface> ObjectInterface = Registry->GetElement<ITypedElementObjectInterface>(Element);
 			if (!ObjectInterface)
 			{
 				continue;
 			}
-			AActor* Actor = ObjectInterface->GetObjectAs<AActor>(Element);
+			AActor* Actor = ObjectInterface.GetObjectAs<AActor>();
 			if (!Actor)
 			{
 				continue;
@@ -1554,7 +1554,7 @@ bool FLevelEditorViewportClient::UpdateDropPreviewElements(int32 MouseX, int32 M
 	DropPreviewElements->ForEachElement<ITypedElementWorldInterface>(
 	[&CollisionQueryParams](const TTypedElement<ITypedElementWorldInterface>& InElement)
 	{
-		InElement.AddIgnoredElementToCollisionQueryParams(InElement, CollisionQueryParams);
+		InElement.AddIgnoredElementToCollisionQueryParams(CollisionQueryParams);
 
 		// true means continue
 		return true;
@@ -1605,11 +1605,11 @@ bool FLevelEditorViewportClient::UpdateDropPreviewElements(int32 MouseX, int32 M
 		}
 
 		// Used for the actual location update
-		ITypedElementWorldInterface* WorldInterface = Registry->GetElementInterface<ITypedElementWorldInterface>(Element);
+		TTypedElement<ITypedElementWorldInterface> WorldInterface = Registry->GetElement<ITypedElementWorldInterface>(Element);
 
 		// Used for some legacy updates if we're dealing with actors
-		ITypedElementObjectInterface* ObjectInterface = Registry->GetElementInterface<ITypedElementObjectInterface>(Element);
-		AActor* Actor = ObjectInterface ? ObjectInterface->GetObjectAs<AActor>(Element) : nullptr;
+		TTypedElement<ITypedElementObjectInterface> ObjectInterface = Registry->GetElement<ITypedElementObjectInterface>(Element);
+		AActor* Actor = ObjectInterface ? ObjectInterface.GetObjectAs<AActor>() : nullptr;
 
 		if (WorldInterface)
 		{
@@ -1628,7 +1628,7 @@ bool FLevelEditorViewportClient::UpdateDropPreviewElements(int32 MouseX, int32 M
 			FTransform DestinationTransform = FActorPositioning::GetSnappedSurfaceAlignedTransform(PositioningData);
 			DestinationTransform.SetScale3D(PreviousTransform.GetScale3D());		// preserve scaling
 
-			WorldInterface->SetWorldTransform(Element, DestinationTransform);
+			WorldInterface.SetWorldTransform(DestinationTransform);
 		}
 
 		if (Actor)
@@ -1686,8 +1686,8 @@ void FLevelEditorViewportClient::DestroyDropPreviewElements()
 		// Don't delete the BSP builder brush, which legacy code uses for previews and expects to always exist.
 		if (DefaultBrush)
 		{
-			ITypedElementObjectInterface* ObjectInterface = UTypedElementRegistry::GetInstance()->GetElementInterface<ITypedElementObjectInterface>(InElement);
-			const UObject* Object = ObjectInterface ? ObjectInterface->GetObject(InElement) : nullptr;
+			TTypedElement<ITypedElementObjectInterface> ObjectInterface = UTypedElementRegistry::GetInstance()->GetElement<ITypedElementObjectInterface>(InElement);
+			const UObject* Object = ObjectInterface ? ObjectInterface.GetObject() : nullptr;
 			if (Object && Object == DefaultBrush)
 			{
 				// true means continue - don't fall through to the deletion.
@@ -1707,8 +1707,8 @@ void FLevelEditorViewportClient::DestroyDropPreviewElements()
 			// return true even if deletion does not go through because it calls UUnrealEdEngine::DeleteActors,
 			// and that function returns true for many failed deletions. It would be nice to make edits to make
 			// the check more reliable.
-			ITypedElementObjectInterface* ObjectInterface = UTypedElementRegistry::GetInstance()->GetElementInterface<ITypedElementObjectInterface>(InElement);
-			AActor* Actor = ObjectInterface ? ObjectInterface->GetObjectAs<AActor>(InElement) : nullptr;
+			TTypedElement<ITypedElementObjectInterface> ObjectInterface = UTypedElementRegistry::GetInstance()->GetElement<ITypedElementObjectInterface>(InElement);
+			AActor* Actor = ObjectInterface ? ObjectInterface.GetObjectAs<AActor>() : nullptr;
 			if (Actor)
 			{
 				GetWorld()->DestroyActor(Actor);
@@ -2013,7 +2013,7 @@ bool FLevelEditorViewportClient::DropObjectsAtCoordinates(int32 MouseX, int32 Mo
 						return;
 					}
 
-					ITypedElementWorldInterface* WorldInterface = Registry->GetElementInterface<ITypedElementWorldInterface>(Element);
+					TTypedElement<ITypedElementWorldInterface> WorldInterface = Registry->GetElement<ITypedElementWorldInterface>(Element);
 
 					// Don't store preview element if it doesn't have a world interface, because we would need it 
 					// to delete the preview element or to move it around
@@ -3160,7 +3160,7 @@ void FLevelEditorViewportClient::ProjectSelectedElementsIntoWorld(FViewport* InV
 	ElementsToManipulate->ForEachElement<ITypedElementWorldInterface>(
 		[&CollisionQueryParams](const TTypedElement<ITypedElementWorldInterface>& InElement)
 	{
-		InElement.AddIgnoredElementToCollisionQueryParams(InElement, CollisionQueryParams, /*bAlsoIgnoreSubElements*/ true);
+		InElement.AddIgnoredElementToCollisionQueryParams(CollisionQueryParams, /*bAlsoIgnoreSubElements*/ true);
 		return true; // continue iterating
 	});
 
diff --git a/Engine/Source/Runtime/Engine/Private/Elements/Actor/ActorElementWorldInterface.cpp b/Engine/Source/Runtime/Engine/Private/Elements/Actor/ActorElementWorldInterface.cpp
--- a/Engine/Source/Runtime/Engine/Private/Elements/Actor/ActorElementWorldInterface.cpp
+++ b/Engine/Source/Runtime/Engine/Private/Elements/Actor/ActorElementWorldInterface.cpp
@@ -258,9 +258,9 @@ bool UActorElementWorldInterface::FindSuitableTransformAlongPath_WorldSweep(cons
 	{
 		for (const FTypedElementHandle& ElementToIgnore : InElementsToIgnore)
 		{
-			if (ITypedElementWorldInterface* WorldInterface = Registry->GetElementInterface<ITypedElementWorldInterface>(ElementToIgnore))
+			if (TTypedElement<ITypedElementWorldInterface> WorldInterface = Registry->GetElement<ITypedElementWorldInterface>(ElementToIgnore))
 			{
-				WorldInterface->AddIgnoredElementToCollisionQueryParams(ElementToIgnore, InOutParams);
+				WorldInterface.AddIgnoredElementToCollisionQueryParams(InOutParams);
 			}
 		}
 	}
@@ -292,12 +292,12 @@ void UActorElementWorldInterface::AddIgnoredCollisionQueryElement(const FTypedEl
 		return;
 	}
 
-	if (ITypedElementWorldInterface* WorldInterface = Registry->GetElementInterface<ITypedElementWorldInterface>(InElementHandle))
+	if (TTypedElement<ITypedElementWorldInterface> WorldInterface = Registry->GetElement<ITypedElementWorldInterface>(InElementHandle))
 	{
 		// Preserving legacy behavior, where we didn't add child elements:
 		const bool bAlsoIgnoreSubElements = false;
 
-		WorldInterface->AddIgnoredElementToCollisionQueryParams(InElementHandle, InOutParams, bAlsoIgnoreSubElements);
+		WorldInterface.AddIgnoredElementToCollisionQueryParams(InOutParams, bAlsoIgnoreSubElements);
 	}
 }
 
diff --git a/Engine/Source/Runtime/Engine/Public/Elements/Interfaces/TypedElementWorldInterface.h b/Engine/Source/Runtime/Engine/Public/Elements/Interfaces/TypedElementWorldInterface.h
--- a/Engine/Source/Runtime/Engine/Public/Elements/Interfaces/TypedElementWorldInterface.h
+++ b/Engine/Source/Runtime/Engine/Public/Elements/Interfaces/TypedElementWorldInterface.h
@@ -486,7 +486,9 @@ struct TTypedElement<ITypedElementWorldInterface> : public TTypedElementBase<ITy
 	void NotifyMovementEnded() const { InterfacePtr->NotifyMovementEnded(*this); }
 	bool FindSuitableTransformAtPoint(const FTransform& InPotentialTransform, FTransform& OutSuitableTransform) const { return InterfacePtr->FindSuitableTransformAtPoint(*this, InPotentialTransform, OutSuitableTransform); }
 	bool FindSuitableTransformAlongPath(const FVector& InPathStart, const FVector& InPathEnd, const FCollisionShape& InTestShape, TArrayView<const FTypedElementHandle> InElementsToIgnore, FTransform& OutSuitableTransform) const { return InterfacePtr->FindSuitableTransformAlongPath(*this, InPathStart, InPathEnd, InTestShape, InElementsToIgnore, OutSuitableTransform); }
-	bool AddIgnoredElementToCollisionQueryParams(const FTypedElementHandle& InElementHandle, FCollisionQueryParams& InOutParams, bool bAlsoIgnoreSubElements = true) const { return InterfacePtr->AddIgnoredElementToCollisionQueryParams(InElementHandle, InOutParams, bAlsoIgnoreSubElements); }
+	UE_DEPRECATED(5.8, "The element handle is already included in TTypedElement access. Use the version of this function that does not take the FTypedElementHandle parameter.")
+	bool AddIgnoredElementToCollisionQueryParams(const FTypedElementHandle& InElementHandle, FCollisionQueryParams& InOutParams, bool bAlsoIgnoreSubElements = true) const { return AddIgnoredElementToCollisionQueryParams(InOutParams, bAlsoIgnoreSubElements); }
+	bool AddIgnoredElementToCollisionQueryParams(FCollisionQueryParams& InOutParams, bool bAlsoIgnoreSubElements = true) const { return InterfacePtr->AddIgnoredElementToCollisionQueryParams(*this, InOutParams, bAlsoIgnoreSubElements); }
 	bool CanDeleteElement() const { return InterfacePtr->CanDeleteElement(*this); }
 	bool DeleteElement(UWorld* InWorld, UTypedElementSelectionSet* InSelectionSet, const FTypedElementDeletionOptions& InDeletionOptions) const { return InterfacePtr->DeleteElement(*this, InWorld, InSelectionSet, InDeletionOptions); }
 	bool CanDuplicateElement() const { return InterfacePtr->CanDuplicateElement(*this); }
