Localization: Add a "Dev Notes" field to localized text so developers can provide additional context for translators

This is exposed as an extra text field in the editor and with new macros in C++ (NOTELOCTEXT, NSNOTELOCTEXT)

PR #11531

#Synced-CL 51746988

diff --git a/Engine/Plugins/Editor/ContentBrowser/ContentBrowserAssetDataSource/Source/ContentBrowserAssetDataSource/Private/AssetFileContextMenu.cpp b/Engine/Plugins/Editor/ContentBrowser/ContentBrowserAssetDataSource/Source/ContentBrowserAssetDataSource/Private/AssetFileContextMenu.cpp
--- a/Engine/Plugins/Editor/ContentBrowser/ContentBrowserAssetDataSource/Source/ContentBrowserAssetDataSource/Private/AssetFileContextMenu.cpp
+++ b/Engine/Plugins/Editor/ContentBrowser/ContentBrowserAssetDataSource/Source/ContentBrowserAssetDataSource/Private/AssetFileContextMenu.cpp
@@ -1976,6 +1976,7 @@ void FAssetFileContextMenu::ExecuteShowLocalizationCache(const FAssetData InAsse
 			// Read package file summary from the file
 			FPackageFileSummary PackageFileSummary;
 			*FileReader << PackageFileSummary;
+			FileReader->SetCustomVersions(PackageFileSummary.GetCustomVersionContainer());
 
 			CachedLocalizationId = PackageFileSummary.LocalizationId;
 
diff --git a/Engine/Source/Developer/Localization/Private/PortableObjectPipeline.cpp b/Engine/Source/Developer/Localization/Private/PortableObjectPipeline.cpp
--- a/Engine/Source/Developer/Localization/Private/PortableObjectPipeline.cpp
+++ b/Engine/Source/Developer/Localization/Private/PortableObjectPipeline.cpp
@@ -181,7 +181,7 @@ namespace
 			// Preserve only non-procedurally generated extracted comments.
 			const TArray<FString> CommentsToPreserve = Entry->ExtractedComments.FilterByPredicate([](const FString& ExtractedComment) -> bool
 			{
-				return !ExtractedComment.StartsWith(TEXT("Key:"), ESearchCase::CaseSensitive) && !ExtractedComment.StartsWith(TEXT("SourceLocation:"), ESearchCase::CaseSensitive) && !ExtractedComment.StartsWith(TEXT("InfoMetaData:"), ESearchCase::CaseSensitive);
+				return !ExtractedComment.StartsWith(TEXT("Key:"), ESearchCase::CaseSensitive) && !ExtractedComment.StartsWith(TEXT("SourceLocation:"), ESearchCase::CaseSensitive) && !ExtractedComment.StartsWith(TEXT("InfoMetaData:"), ESearchCase::CaseSensitive) && !ExtractedComment.StartsWith(TEXT("DevNotes:"), ESearchCase::CaseSensitive);
 			});
 
 			if (CommentsToPreserve.Num() || Entry->TranslatorComments.Num())
@@ -373,6 +373,12 @@ namespace
 
 				PoEntry->AddExtractedComment(GetConditionedKeyForExtractedComment(Context.Key.GetString())); // "Notes from Programmer" in the form of the Key.
 				PoEntry->AddExtractedComment(GetConditionedReferenceForExtractedComment(Context.SourceLocation)); // "Notes from Programmer" in the form of the Source Location, since this comes in handy too and OneSky doesn't properly show references, only comments.
+#if WITH_EDITORONLY_DATA
+				if (!Context.DevNotes.IsEmpty())
+				{
+					PoEntry->AddExtractedComment(GetConditionedDevNotesForExtractedComment(Context.DevNotes)); // "Notes from Programmer" in the form of the notes entered by the developer.
+					}
+#endif
 
 				TArray<FString> InfoMetaDataStrings;
 				if (Context.InfoMetadataObj.IsValid())
@@ -734,6 +740,11 @@ FString PortableObjectPipeline::GetConditionedReferenceForExtractedComment(const
 	return FString::Printf(TEXT("SourceLocation:\t%s"), *PORefString);
 }
 
+FString PortableObjectPipeline::GetConditionedDevNotesForExtractedComment(const FString& DevNotes)
+{
+	return FString::Printf(TEXT("DevNotes:\t%s"), *ConditionArchiveStrForPO(DevNotes));
+}
+
 FString PortableObjectPipeline::GetConditionedInfoMetaDataForExtractedComment(const FString& KeyName, const FString& ValueString)
 {
 	return FString::Printf(TEXT("InfoMetaData:\t\"%s\" : \"%s\""), *KeyName, *ValueString);
diff --git a/Engine/Source/Developer/Localization/Private/Serialization/JsonInternationalizationManifestSerializer.cpp b/Engine/Source/Developer/Localization/Private/Serialization/JsonInternationalizationManifestSerializer.cpp
--- a/Engine/Source/Developer/Localization/Private/Serialization/JsonInternationalizationManifestSerializer.cpp
+++ b/Engine/Source/Developer/Localization/Private/Serialization/JsonInternationalizationManifestSerializer.cpp
@@ -14,6 +14,7 @@ const FString FJsonInternationalizationManifestSerializer::TAG_NAMESPACE = TEXT(
 const FString FJsonInternationalizationManifestSerializer::TAG_CHILDREN = TEXT("Children");
 const FString FJsonInternationalizationManifestSerializer::TAG_SUBNAMESPACES = TEXT("Subnamespaces");
 const FString FJsonInternationalizationManifestSerializer::TAG_PATH = TEXT("Path");
+const FString FJsonInternationalizationManifestSerializer::TAG_DEVNOTES = TEXT("DevNotes");
 const FString FJsonInternationalizationManifestSerializer::TAG_OPTIONAL = TEXT("Optional");
 const FString FJsonInternationalizationManifestSerializer::TAG_KEYCOLLECTION = TEXT("Keys");
 const FString FJsonInternationalizationManifestSerializer::TAG_KEY = TEXT("Key");
@@ -265,6 +266,7 @@ bool FJsonInternationalizationManifestSerializer::JsonObjToManifest( UE::Json::F
 						CommandContext.Key = FString(*Key);
 						CommandContext.SourceLocation = SourceLocation.Get(FString());
 						CommandContext.PlatformName = PlatformName;
+						CommandContext.DevNotes = UE::Json::GetStringField(ContextJSONObject, *TAG_DEVNOTES).Get(TEXTVIEW(""));
 
 						if (TOptional<bool> IsOptionalValue = UE::Json::GetBoolField(ContextJSONObject, *TAG_OPTIONAL))
 						{
@@ -479,6 +481,10 @@ void FJsonInternationalizationManifestSerializer::StructuredDataToJsonObj( TShar
 
 			KeyJSONObject.AddMember(UE::Json::MakeStringRef(TAG_KEY), UE::Json::MakeStringValue(AContext.Key.GetString(), Allocator), Allocator);
 			KeyJSONObject.AddMember(UE::Json::MakeStringRef(TAG_PATH), UE::Json::MakeStringValue(ProcessedText, Allocator), Allocator);
+			if (!AContext.DevNotes.IsEmpty())
+			{
+				KeyJSONObject.AddMember(UE::Json::MakeStringRef(TAG_DEVNOTES), UE::Json::MakeStringValue(AContext.DevNotes, Allocator), Allocator);
+			}
 
 			// We only add the optional field if it is true, it is assumed to be false otherwise.
 			if (AContext.bIsOptional)
diff --git a/Engine/Source/Developer/Localization/Public/PortableObjectPipeline.h b/Engine/Source/Developer/Localization/Public/PortableObjectPipeline.h
--- a/Engine/Source/Developer/Localization/Public/PortableObjectPipeline.h
+++ b/Engine/Source/Developer/Localization/Public/PortableObjectPipeline.h
@@ -93,6 +93,9 @@ namespace PortableObjectPipeline
 	/** Given a source location string, condition it so it can be written as the extracted comment field of a PO file */
 	LOCALIZATION_API FString GetConditionedReferenceForExtractedComment(const FString& PORefString);
 
+	/** Given developer-entered translator context notes, condition it so it can be written as the extracted comment field of a PO file */
+	LOCALIZATION_API FString GetConditionedDevNotesForExtractedComment(const FString& DevNotes);
+
 	/** Given a meta-data value string, condition it so it can be written as the extracted comment field of a PO file */
 	LOCALIZATION_API FString GetConditionedInfoMetaDataForExtractedComment(const FString& KeyName, const FString& ValueString);
 
diff --git a/Engine/Source/Developer/Localization/Public/Serialization/JsonInternationalizationManifestSerializer.h b/Engine/Source/Developer/Localization/Public/Serialization/JsonInternationalizationManifestSerializer.h
--- a/Engine/Source/Developer/Localization/Public/Serialization/JsonInternationalizationManifestSerializer.h
+++ b/Engine/Source/Developer/Localization/Public/Serialization/JsonInternationalizationManifestSerializer.h
@@ -146,6 +146,7 @@ private:
 	static const FString TAG_CHILDREN;
 	static const FString TAG_SUBNAMESPACES;
 	static const FString TAG_PATH;
+	static const FString TAG_DEVNOTES;
 	static const FString TAG_OPTIONAL;
 	static const FString TAG_KEYCOLLECTION;
 	static const FString TAG_KEY;
diff --git a/Engine/Source/Developer/ScriptDisassembler/Private/ScriptDisassembler.cpp b/Engine/Source/Developer/ScriptDisassembler/Private/ScriptDisassembler.cpp
--- a/Engine/Source/Developer/ScriptDisassembler/Private/ScriptDisassembler.cpp
+++ b/Engine/Source/Developer/ScriptDisassembler/Private/ScriptDisassembler.cpp
@@ -7,6 +7,7 @@
 #include "ScriptDisassembler.h"
 #include "UObject/Object.h"
 #include "UObject/Class.h"
+#include "UObject/Package.h"
 #include "UObject/UObjectHash.h"
 #include "UObject/UObjectIterator.h"
 #include "UObject/UnrealType.h"
@@ -26,6 +27,7 @@ FKismetBytecodeDisassembler::FKismetBytecodeDisassembler(FOutputDevice& InAr)
 // Disassemble all of the script code in a single structure.
 void FKismetBytecodeDisassembler::DisassembleStructure(UFunction* Source)
 {
+	bIsCooked = Source->GetPackage()->HasAnyPackageFlags(PKG_Cooked);
 	Script.Empty();
 	Script.Append(Source->Script);
 
@@ -860,7 +862,17 @@ void FKismetBytecodeDisassembler::ProcessCommon(int32& ScriptIndex, EExprToken O
 					const FString SourceString = ReadString(ScriptIndex);
 					const FString KeyString = ReadString(ScriptIndex);
 					const FString Namespace = ReadString(ScriptIndex);
-					Ar.Logf(TEXT("%s $%X: literal text - localized text { namespace: \"%s\", key: \"%s\", source: \"%s\" }"), *Indents, (int32)Opcode, *Namespace, *KeyString, *SourceString);
+#if WITH_EDITORONLY_DATA
+					if (!bIsCooked)
+					{
+						const FString DevNotes = ReadString(ScriptIndex);
+						Ar.Logf(TEXT("%s $%X: literal text - localized text { namespace: \"%s\", key: \"%s\", source: \"%s\", devnotes: \"%s\" }"), *Indents, (int32)Opcode, *Namespace, *KeyString, *SourceString, *DevNotes);
+					}
+					else
+#endif
+					{
+						Ar.Logf(TEXT("%s $%X: literal text - localized text { namespace: \"%s\", key: \"%s\", source: \"%s\" }"), *Indents, (int32)Opcode, *Namespace, *KeyString, *SourceString);
+					}
 				}
 				break;
 
diff --git a/Engine/Source/Developer/ScriptDisassembler/Public/ScriptDisassembler.h b/Engine/Source/Developer/ScriptDisassembler/Public/ScriptDisassembler.h
--- a/Engine/Source/Developer/ScriptDisassembler/Public/ScriptDisassembler.h
+++ b/Engine/Source/Developer/ScriptDisassembler/Public/ScriptDisassembler.h
@@ -28,6 +28,7 @@ private:
 	TArray<uint8> Script;
 	FString Indents;
 	FOutputDevice& Ar;
+	bool bIsCooked = false;
 public:
 	/**
 	 * Construct a disassembler that will output to the specified archive.
diff --git a/Engine/Source/Developer/TranslationEditor/Private/TranslationDataManager.cpp b/Engine/Source/Developer/TranslationEditor/Private/TranslationDataManager.cpp
--- a/Engine/Source/Developer/TranslationEditor/Private/TranslationDataManager.cpp
+++ b/Engine/Source/Developer/TranslationEditor/Private/TranslationDataManager.cpp
@@ -151,6 +151,9 @@ void FTranslationDataManager::Initialize()
 
 					ContextInfo.Context = AContext.SourceLocation;
 					ContextInfo.Key = AContext.Key.GetString();
+#if WITH_EDITORONLY_DATA
+					ContextInfo.DevNotes = AContext.DevNotes;
+#endif
 
 					// Make sure we have a unique translation unit for each unique identity.
 					UTranslationUnit*& TranslationUnit = IdentityToTranslationUnitMap.FindOrAdd(FLocTextIdentity(ManifestEntry->Namespace, AContext.Key));
@@ -644,7 +647,11 @@ void FTranslationDataManager::PopulateSearchResultsUsingFilter(const FString& Se
 			{
 				if (!bAdded && 
 					(CurrentContext.Context.Contains(SearchFilter) ||
-					CurrentContext.Key.Contains(SearchFilter)))
+					CurrentContext.Key.Contains(SearchFilter)
+#if WITH_EDITORONLY_DATA
+					|| CurrentContext.DevNotes.Contains(SearchFilter)
+#endif
+					))
 				{
 					SearchResults.Add(TranslationUnit);
 					break;
@@ -926,7 +933,9 @@ bool FTranslationDataManager::SaveSelectedTranslations(TArray<UTranslationUnit*>
 
 										PoEntry->AddExtractedComment(PortableObjectPipeline::GetConditionedKeyForExtractedComment(ContextInfo.Key)); // "Notes from Programmer" in the form of the Key.
 										PoEntry->AddExtractedComment(PortableObjectPipeline::GetConditionedReferenceForExtractedComment(ContextInfo.Context)); // "Notes from Programmer" in the form of the Source Location, since this comes in handy too and OneSky doesn't properly show references, only comments.
-
+#if WITH_EDITORONLY_DATA
+										PoEntry->AddExtractedComment(PortableObjectPipeline::GetConditionedDevNotesForExtractedComment(ContextInfo.DevNotes));
+#endif
 										PortableObjectDom.AddEntry(PoEntry);
 									}
 									break;
diff --git a/Engine/Source/Developer/TranslationEditor/Private/TranslationUnit.h b/Engine/Source/Developer/TranslationEditor/Private/TranslationUnit.h
--- a/Engine/Source/Developer/TranslationEditor/Private/TranslationUnit.h
+++ b/Engine/Source/Developer/TranslationEditor/Private/TranslationUnit.h
@@ -52,6 +52,11 @@ public:
 	/** What file and line this translation is from */
 	UPROPERTY(Category=Translation, VisibleAnywhere)
 	FString Context;
+	
+#if WITH_EDITORONLY_DATA
+	UPROPERTY(Category=Translation, VisibleAnywhere)
+	FString DevNotes;
+#endif
 
 	/** List of previous versions of the source text for this context */
 	UPROPERTY(Category=Translation, VisibleAnywhere)
diff --git a/Engine/Source/Editor/EditorWidgets/Private/STextPropertyEditableTextBox.cpp b/Engine/Source/Editor/EditorWidgets/Private/STextPropertyEditableTextBox.cpp
--- a/Engine/Source/Editor/EditorWidgets/Private/STextPropertyEditableTextBox.cpp
+++ b/Engine/Source/Editor/EditorWidgets/Private/STextPropertyEditableTextBox.cpp
@@ -589,7 +589,7 @@ FReply STextPropertyEditableStringTableReference::OnUnlinkClicked()
 			NewKey = TextNamespaceUtil::GenerateRandomTextKey();
 #endif	// USE_STABLE_LOCALIZATION_KEYS
 			
-			EditableTextProperty->SetText(TextIndex, FText::ChangeKey(NewNamespace, NewKey, CurrentText));
+			EditableTextProperty->SetText(TextIndex, FText::ChangeKey(NewNamespace, NewKey, CurrentText, CurrentText.GetDevNotes()));
 		}
 	}
 
@@ -884,6 +884,24 @@ void STextPropertyEditableTextBox::Construct(const FArguments& InArgs, const TSh
 #endif	// USE_STABLE_LOCALIZATION_KEYS
 						]
 					]
+					+ SVerticalBox::Slot()
+					.AutoHeight()
+					[
+						SNew(STextPropertyEditableOptionRow, LinkedBoxManager)
+						.Text(LOCTEXT("DevNotesLabel", "Dev Notes"))
+						.IsEnabled(this, &STextPropertyEditableTextBox::IsTextLocalizable)
+						[
+							SNew(SMultiLineEditableTextBox)
+							.Text(this, &STextPropertyEditableTextBox::GetDevNotesValue)
+							.Font(PropertyNormalFont)
+							.ClearKeyboardFocusOnCommit(false)
+							.SelectAllTextOnCommit(true)
+							.ModiferKeyForNewLine(EModifierKey::Shift)
+							.OnTextCommitted(this, &STextPropertyEditableTextBox::OnDevNotesCommitted)
+							.IsReadOnly(this, &STextPropertyEditableTextBox::IsIdentityReadOnly)
+							.AutoWrapText(true)
+						]
+					]
 				]
 			]
 		];
@@ -1207,7 +1225,7 @@ void STextPropertyEditableTextBox::OnTextCommitted(const FText& NewText, ETextCo
 			}
 #endif	// USE_STABLE_LOCALIZATION_KEYS
 
-			EditableTextProperty->SetText(TextIndex, FText::ChangeKey(NewNamespace, NewKey, NewText));
+			EditableTextProperty->SetText(TextIndex, FText::ChangeKey(NewNamespace, NewKey, NewText, PropertyValue.GetDevNotes()));
 		}
 	}
 }
@@ -1323,7 +1341,7 @@ void STextPropertyEditableTextBox::OnNamespaceCommitted(const FText& NewText, ET
 			}
 #endif	// USE_STABLE_LOCALIZATION_KEYS
 
-			EditableTextProperty->SetText(TextIndex, FText::ChangeKey(NewNamespace, NewKey, PropertyValue));
+			EditableTextProperty->SetText(TextIndex, FText::ChangeKey(NewNamespace, NewKey, PropertyValue, PropertyValue.GetDevNotes()));
 		}
 	}
 }
@@ -1446,7 +1464,7 @@ void STextPropertyEditableTextBox::OnKeyCommitted(const FText& NewText, ETextCom
 				NewKey
 				);
 
-			EditableTextProperty->SetText(TextIndex, FText::ChangeKey(NewNamespace, NewKey, PropertyValue));
+			EditableTextProperty->SetText(TextIndex, FText::ChangeKey(NewNamespace, NewKey, PropertyValue, PropertyValue.GetDevNotes()));
 		}
 	}
 }
@@ -1486,6 +1504,41 @@ FText STextPropertyEditableTextBox::GetPackageValue() const
 
 #endif // USE_STABLE_LOCALIZATION_KEYS
 
+FText STextPropertyEditableTextBox::GetDevNotesValue() const
+{
+	FText DevNotes;
+
+	const int32 NumTexts = EditableTextProperty->GetNumTexts();
+	if (NumTexts == 1)
+	{
+		DevNotes = FText::FromString(EditableTextProperty->GetText(0).GetDevNotes());
+	}
+	else if (NumTexts > 1)
+	{
+		DevNotes = MultipleValuesText;
+	}
+
+	return DevNotes;
+}
+
+void STextPropertyEditableTextBox::OnDevNotesCommitted(const FText& NewText, ETextCommit::Type CommitInfo)
+{
+	const int32 NumTexts = EditableTextProperty->GetNumTexts();
+
+	// Don't commit the Multiple Values text if there are multiple properties being set
+	if (NumTexts > 0 && (NumTexts == 1 || NewText.ToString() != MultipleValuesText.ToString()))
+	{
+		for (int32 TextIndex = 0; TextIndex < NumTexts; ++TextIndex)
+		{
+			const FText PropertyValue = EditableTextProperty->GetText(TextIndex);
+			const FString Namespace = FTextInspector::GetNamespace(PropertyValue).Get(FString());
+			const FString Key = FTextInspector::GetKey(PropertyValue).Get(FString());
+
+			EditableTextProperty->SetText(TextIndex, FText::ChangeKey(Namespace, Key, PropertyValue, NewText.ToString()));
+		}
+	}
+}
+
 ECheckBoxState STextPropertyEditableTextBox::GetLocalizableCheckState() const
 {
 	TOptional<ECheckBoxState> Result;
diff --git a/Engine/Source/Editor/EditorWidgets/Public/STextPropertyEditableTextBox.h b/Engine/Source/Editor/EditorWidgets/Public/STextPropertyEditableTextBox.h
--- a/Engine/Source/Editor/EditorWidgets/Public/STextPropertyEditableTextBox.h
+++ b/Engine/Source/Editor/EditorWidgets/Public/STextPropertyEditableTextBox.h
@@ -234,6 +234,9 @@ private:
 	UE_API FText GetPackageValue() const;
 #endif // USE_STABLE_LOCALIZATION_KEYS
 
+	UE_API FText GetDevNotesValue() const;
+	UE_API void OnDevNotesCommitted(const FText& NewText, ETextCommit::Type CommitInfo);
+
 	UE_API ECheckBoxState GetLocalizableCheckState() const;
 
 	UE_API void HandleLocalizableCheckStateChanged(ECheckBoxState InCheckboxState);
diff --git a/Engine/Source/Editor/KismetCompiler/Private/KismetCompilerVMBackend.cpp b/Engine/Source/Editor/KismetCompiler/Private/KismetCompilerVMBackend.cpp
--- a/Engine/Source/Editor/KismetCompiler/Private/KismetCompilerVMBackend.cpp
+++ b/Engine/Source/Editor/KismetCompiler/Private/KismetCompilerVMBackend.cpp
@@ -730,6 +730,12 @@ public:
 						EmitStringLiteral(*SourceString);
 						EmitStringLiteral(Key);
 						EmitStringLiteral(Namespace);
+#if WITH_EDITORONLY_DATA
+						if (!Writer.IsFilterEditorOnly() && !Writer.IsCooking() && !IsRunningCookCommandlet())
+						{
+							EmitStringLiteral(Term->TextLiteral.GetDevNotes());
+						}
+#endif
 					}
 					else
 					{
diff --git a/Engine/Source/Editor/StringTableEditor/Private/StringTableEditor.cpp b/Engine/Source/Editor/StringTableEditor/Private/StringTableEditor.cpp
--- a/Engine/Source/Editor/StringTableEditor/Private/StringTableEditor.cpp
+++ b/Engine/Source/Editor/StringTableEditor/Private/StringTableEditor.cpp
@@ -31,6 +31,7 @@ const FName FStringTableEditor::StringTableTabId("StringTableEditor_StringTable"
 const FName FStringTableEditor::StringTableDummyColumnId("Dummy");
 const FName FStringTableEditor::StringTableKeyColumnId("Key");
 const FName FStringTableEditor::StringTableSourceStringColumnId("SourceString");
+const FName FStringTableEditor::StringTableDevNotesColumnId("DevNotes");
 const FName FStringTableEditor::StringTableDeleteColumnId("Delete");
 
 class SStringTableEntryRow : public SMultiColumnTableRow<TSharedRef<FStringTableEditor::FCachedStringTableEntry>>
@@ -74,6 +75,17 @@ public:
 					.OnTextCommitted(this, &SStringTableEntryRow::OnSourceStringCommitted)
 				];
 		}
+		else if (ColumnName == FStringTableEditor::StringTableDevNotesColumnId)
+		{
+			Return = SNew(SBorder)
+				.BorderImage(&FCoreStyle::Get().GetWidgetStyle<FEditableTextBoxStyle>("NormalEditableTextBox").BackgroundImageReadOnly)
+				[
+					SNew(SMultiLineEditableTextBox)
+					.Text(FText::FromString(CachedStringTableEntry->DevNotes))
+					.ModiferKeyForNewLine(EModifierKey::Shift)
+					.OnTextCommitted(this, &SStringTableEntryRow::OnDevNotesCommitted)
+				];
+		}
 		else if (ColumnName == FStringTableEditor::StringTableDeleteColumnId)
 		{
 			Return = SNew(SButton)
@@ -99,7 +111,17 @@ private:
 		if (!NewSourceString.IsEmpty() && !CachedStringTableEntry->SourceString.Equals(NewSourceString, ESearchCase::CaseSensitive))
 		{
 			CachedStringTableEntry->SourceString = NewSourceString;
-			OwnerEditor->SetEntry(CachedStringTableEntry->Key, CachedStringTableEntry->SourceString);
+			OwnerEditor->SetEntry(CachedStringTableEntry->Key, CachedStringTableEntry->SourceString, CachedStringTableEntry->DevNotes);
+		}
+	}
+
+	void OnDevNotesCommitted(const FText& InText, ETextCommit::Type)
+	{
+		const FString NewDevNotes = InText.ToString();
+		if (!CachedStringTableEntry->DevNotes.Equals(NewDevNotes, ESearchCase::CaseSensitive))
+		{
+			CachedStringTableEntry->DevNotes = NewDevNotes;
+			OwnerEditor->SetEntry(CachedStringTableEntry->Key, CachedStringTableEntry->SourceString, CachedStringTableEntry->DevNotes);
 		}
 	}
 
@@ -252,6 +274,7 @@ TSharedRef<SDockTab> FStringTableEditor::SpawnTab_StringTable(const FSpawnTabArg
 	EntryTextFilter = MakeShareable(new FEntryTextFilter(FEntryTextFilter::FItemToStringArray::CreateLambda([](const TSharedPtr<FCachedStringTableEntry>& InItem, OUT TArray< FString >& StringArray) {
 		StringArray.Add(InItem->Key);
 		StringArray.Add(InItem->SourceString);
+		StringArray.Add(InItem->DevNotes);
 	})));
 
 	TSharedRef<SDockTab> StringTableTab = SNew(SDockTab)
@@ -356,6 +379,11 @@ TSharedRef<SDockTab> FStringTableEditor::SpawnTab_StringTable(const FSpawnTabArg
 						.DefaultLabel(LOCTEXT("SourceStringColumnLabel", "Source String"))
 						.FillWidth(1.0f)
 
+						+SHeaderRow::Column(StringTableDevNotesColumnId)
+						.DefaultLabel(LOCTEXT("DevNotesColumnLabel", "Developer Notes"))
+						.ToolTipText(LOCTEXT("DenNotesColumnLabelToolTip", "Notes to give translators additional context that will aid them in translating the source string (details on format parameters, homograph disambiguation, tone, etc)"))
+						.FillWidth(0.5f)
+
 						+SHeaderRow::Column(StringTableDeleteColumnId)
 						.DefaultLabel(FText::GetEmpty())
 						.FixedWidth(28)
@@ -409,6 +437,24 @@ TSharedRef<SDockTab> FStringTableEditor::SpawnTab_StringTable(const FSpawnTabArg
 						.OnTextCommitted(this, &FStringTableEditor::OnSourceStringCommitted)
 					]
 
+					+SHorizontalBox::Slot()
+					.AutoWidth()
+					.Padding(2)
+					.VAlign(VAlign_Center)
+					[
+						SNew(STextBlock)
+						.Text(LOCTEXT("DevNotesLabel", "Dev Notes:"))
+					]
+
+					+SHorizontalBox::Slot()
+					.Padding(2)
+					.VAlign(VAlign_Center)
+					[
+						SAssignNew(DevNotesEditableTextBox, SMultiLineEditableTextBox)
+						.ModiferKeyForNewLine(EModifierKey::Shift)
+						.OnTextCommitted(this, &FStringTableEditor::OnSourceStringCommitted)
+					]
+
 					+SHorizontalBox::Slot()
 					.AutoWidth()
 					.Padding(2)
@@ -458,9 +504,9 @@ void FStringTableEditor::RefreshCachedStringTable(const FString& InCachedSelecti
 	UStringTable* StringTable = Cast<UStringTable>(GetEditingObject());
 	if (StringTable)
 	{
-		StringTable->GetStringTable()->EnumerateSourceStrings([&](const FString& InKey, const FString& InSourceString)
+		StringTable->GetStringTable()->EnumerateSourceStrings([&](const FString& InKey, const FString& InSourceString, const FString& InDevNotes)
 		{
-			TSharedRef<FCachedStringTableEntry> NewStringTableEntry = MakeShared<FCachedStringTableEntry>(InKey, InSourceString);
+			TSharedRef<FCachedStringTableEntry> NewStringTableEntry = MakeShared<FCachedStringTableEntry>(InKey, InSourceString, InDevNotes);
 
 			if (EntryTextFilter->PassesFilter(NewStringTableEntry))
 			{
@@ -537,14 +583,14 @@ void FStringTableEditor::OnNamespaceCommitted(const FText& InText, ETextCommit::
 	}
 }
 
-void FStringTableEditor::SetEntry(const FString& InKey, const FString& InSourceString)
+void FStringTableEditor::SetEntry(const FString& InKey, const FString& InSourceString, const FString& InDevNotes)
 {
 	UStringTable* StringTable = Cast<UStringTable>(GetEditingObject());
 	if (StringTable)
 	{
 		const FScopedTransaction Transaction(LOCTEXT("SetStringTableEntry", "Set String Table Entry"));
 		StringTable->Modify();
-		StringTable->GetMutableStringTable()->SetSourceString(InKey, InSourceString);
+		StringTable->GetMutableStringTable()->SetSourceString(InKey, InSourceString, InDevNotes);
 		HandlePostChange(InKey);
 	}
 }
@@ -632,10 +678,11 @@ FReply FStringTableEditor::OnAddClicked()
 
 	if (!NewKey.IsEmptyOrWhitespace() && IsValidIdentity(NewKey) && !NewSourceString.IsEmpty())
 	{
-		SetEntry(NewKey.ToString(), NewSourceString.ToString());
+		SetEntry(NewKey.ToString(), NewSourceString.ToString(), DevNotesEditableTextBox->GetText().ToString());
 
 		KeyEditableTextBox->SetText(FText::GetEmpty());
 		SourceStringEditableTextBox->SetText(FText::GetEmpty());
+		DevNotesEditableTextBox->SetText(FText::GetEmpty());
 
 		KeyEditableTextBox->SetError(FText::GetEmpty());
 	}
diff --git a/Engine/Source/Editor/StringTableEditor/Private/StringTableEditor.h b/Engine/Source/Editor/StringTableEditor/Private/StringTableEditor.h
--- a/Engine/Source/Editor/StringTableEditor/Private/StringTableEditor.h
+++ b/Engine/Source/Editor/StringTableEditor/Private/StringTableEditor.h
@@ -60,14 +60,16 @@ private:
 	/** Cached string table entry */
 	struct FCachedStringTableEntry
 	{
-		FCachedStringTableEntry(FString InKey, FString InSourceString)
+		FCachedStringTableEntry(FString InKey, FString InSourceString, FString InDevNotes)
 			: Key(MoveTemp(InKey))
 			, SourceString(MoveTemp(InSourceString))
+			, DevNotes(MoveTemp(InDevNotes))
 		{
 		}
 
 		FString Key;
 		FString SourceString;
+		FString DevNotes;
 	};
 
 	/**	Spawns the tab with the string table inside */
@@ -89,7 +91,7 @@ private:
 	void OnNamespaceCommitted(const FText& InText, ETextCommit::Type);
 
 	/** Set the value of an entry in the string table */
-	void SetEntry(const FString& InKey, const FString& InSourceString);
+	void SetEntry(const FString& InKey, const FString& InSourceString, const FString& InDevNotes);
 
 	/** Delete the given entry from the string table */
 	void DeleteEntry(const FString& InKey);
@@ -139,6 +141,9 @@ private:
 
 	/** Editable text for the source string */
 	TSharedPtr<SMultiLineEditableTextBox> SourceStringEditableTextBox;
+	
+	/** Editable text for the developer translation context notes */
+	TSharedPtr<SMultiLineEditableTextBox> DevNotesEditableTextBox;
 
 	/** Array of cached string table entries */
 	TArray<TSharedPtr<FCachedStringTableEntry>> CachedStringTableEntries;
@@ -161,6 +166,9 @@ private:
 	/**	The column ID for the source string column */
 	static const FName StringTableSourceStringColumnId;
 
+	/** The column ID for the developer translation context notes column */
+	static const FName StringTableDevNotesColumnId;
+
 	/**	The column ID for the delete column */
 	static const FName StringTableDeleteColumnId;
 };
diff --git a/Engine/Source/Editor/UnrealEd/Classes/Commandlets/GatherTextFromSourceCommandlet.h b/Engine/Source/Editor/UnrealEd/Classes/Commandlets/GatherTextFromSourceCommandlet.h
--- a/Engine/Source/Editor/UnrealEd/Classes/Commandlets/GatherTextFromSourceCommandlet.h
+++ b/Engine/Source/Editor/UnrealEd/Classes/Commandlets/GatherTextFromSourceCommandlet.h
@@ -71,6 +71,7 @@ private:
 		FString SourceString;
 		FSourceLocation SourceLocation;
 		FName PlatformName;
+		FString DevNotes;
 		bool bIsEditorOnly;
 	};
 
@@ -212,7 +213,7 @@ private:
 
 		void AddStringTableFromFile( const FName InTableId, const FString& InTableNamespace, const FString& InTableFilename, const FString& InRootPath );
 
-		void AddStringTableEntry( const FName InTableId, const FString& InKey, const FString& InSourceString );
+		void AddStringTableEntry( const FName InTableId, const FString& InKey, const FString& InSourceString, const FString& InDevNotes );
 
 		void AddStringTableEntryMetaData( const FName InTableId, const FString& InKey, const FName InMetaDataId, const FString& InMetaData );
 
@@ -271,7 +272,7 @@ private:
 
 	private:
 		bool AddStringTableImpl( const FName InTableId, const FString& InTableNamespace );
-		bool AddStringTableEntryImpl( const FName InTableId, const FString& InKey, const FString& InSourceString, const FSourceLocation& InSourceLocation, const FName InPlatformName );
+		bool AddStringTableEntryImpl( const FName InTableId, const FString& InKey, const FString& InSourceString, const FSourceLocation& InSourceLocation, const FName InPlatformName, const FString& InDevNotes );
 		bool AddStringTableEntryMetaDataImpl( const FName InTableId, const FString& InKey, const FName InMetaDataId, const FString& InMetaData, const FSourceLocation& InSourceLocation );
 
 		//Working data
@@ -486,6 +487,7 @@ private:
 			MAS_Namespace,
 			MAS_Identifier,
 			MAS_SourceText,
+			MAS_DevNotes,
 		};
 
 		struct FMacroArg
@@ -496,6 +498,15 @@ private:
 			FMacroArg(EMacroArgSemantic InSema, bool InIsAutoText) : Semantic(InSema), IsAutoText(InIsAutoText) {}
 		};
 
+		FStringMacroDescriptor(FString InName, FMacroArg Arg0, FMacroArg Arg1, FMacroArg Arg2, FMacroArg Arg3) : FMacroDescriptor(InName, 4)
+		{
+			ApplicableFileTypes = EGatherTextSourceFileTypes::Cpp | EGatherTextSourceFileTypes::Ini;
+			Arguments.Add(Arg0);
+			Arguments.Add(Arg1);
+			Arguments.Add(Arg2);
+			Arguments.Add(Arg3);
+		}
+
 		FStringMacroDescriptor(FString InName, FMacroArg Arg0, FMacroArg Arg1, FMacroArg Arg2) : FMacroDescriptor(InName, 3)
 		{
 			ApplicableFileTypes = EGatherTextSourceFileTypes::Cpp | EGatherTextSourceFileTypes::Ini;
@@ -557,6 +568,17 @@ private:
 		virtual void TryParse(const FString& Text, FSourceFileParseContext& Context) const override;
 	};
 
+	class FStringTableNotedEntryMacroDescriptor : public FMacroDescriptor
+	{
+	public:
+		FStringTableNotedEntryMacroDescriptor()
+			: FMacroDescriptor(TEXT("LOCTABLE_SETNOTESTRING"), 4)
+		{
+		}
+
+		virtual void TryParse(const FString& Text, FSourceFileParseContext& Context) const override;
+	};
+
 	class FStringTableEntryMetaDataMacroDescriptor : public FMacroDescriptor
 	{
 	public:
@@ -605,6 +627,8 @@ private:
 	static const FString DefineString;
 	static const FString MacroString_LOCTEXT;
 	static const FString MacroString_NSLOCTEXT;
+	static const FString MacroString_NOTELOCTEXT;
+	static const FString MacroString_NSNOTELOCTEXT;
 	static const FString MacroString_UI_COMMAND;
 	static const FString MacroString_UI_COMMAND_EXT;
 
diff --git a/Engine/Source/Editor/UnrealEd/Private/Commandlets/GatherTextFromAssetsCommandlet.cpp b/Engine/Source/Editor/UnrealEd/Private/Commandlets/GatherTextFromAssetsCommandlet.cpp
--- a/Engine/Source/Editor/UnrealEd/Private/Commandlets/GatherTextFromAssetsCommandlet.cpp
+++ b/Engine/Source/Editor/UnrealEd/Private/Commandlets/GatherTextFromAssetsCommandlet.cpp
@@ -620,6 +620,7 @@ void UGatherTextFromAssetsCommandlet::ProcessGatherableTextDataArray(const TArra
 				Context.KeyMetadataObj = !(FLocMetadataObject::IsMetadataExactMatch(&TextSourceSiteContext.KeyMetaData, &DefaultMetadataObject)) ? MakeShareable(new FLocMetadataObject(TextSourceSiteContext.KeyMetaData)) : nullptr;
 				Context.InfoMetadataObj = !(FLocMetadataObject::IsMetadataExactMatch(&TextSourceSiteContext.InfoMetaData, &DefaultMetadataObject)) ? MakeShareable(new FLocMetadataObject(TextSourceSiteContext.InfoMetaData)) : nullptr;
 				Context.bIsOptional = TextSourceSiteContext.IsOptional;
+				Context.DevNotes = TextSourceSiteContext.DevNotes;
 				Context.SourceLocation = TextSourceSiteContext.SiteDescription;
 				Context.PlatformName = GetSplitPlatformNameFromPath(TextSourceSiteContext.SiteDescription);
 
@@ -1382,6 +1383,7 @@ void UGatherTextFromAssetsCommandlet::ProcessAndRemoveCachedPackages(TMap<FName,
 			// Read package file summary from the file.
 			FPackageFileSummary PackageFileSummary;
 			*FileReader << PackageFileSummary;
+			FileReader->SetCustomVersions(PackageFileSummary.GetCustomVersionContainer());
 
 			PackagePendingGather.PackageLocalizationId = PackageFileSummary.LocalizationId;
 			PackagePendingGather.PackageLocCacheState = CalculatePackageLocCacheState(PackageFileSummary, PackagePendingGather.PackageName, bIsExternalActorPackage);
diff --git a/Engine/Source/Editor/UnrealEd/Private/Commandlets/GatherTextFromSourceCommandlet.cpp b/Engine/Source/Editor/UnrealEd/Private/Commandlets/GatherTextFromSourceCommandlet.cpp
--- a/Engine/Source/Editor/UnrealEd/Private/Commandlets/GatherTextFromSourceCommandlet.cpp
+++ b/Engine/Source/Editor/UnrealEd/Private/Commandlets/GatherTextFromSourceCommandlet.cpp
@@ -47,6 +47,8 @@ const FString UGatherTextFromSourceCommandlet::FPreProcessorDescriptor::IniNames
 const FString UGatherTextFromSourceCommandlet::FMacroDescriptor::TextMacroString(TEXT("TEXT"));
 const FString UGatherTextFromSourceCommandlet::MacroString_LOCTEXT(TEXT("LOCTEXT"));
 const FString UGatherTextFromSourceCommandlet::MacroString_NSLOCTEXT(TEXT("NSLOCTEXT"));
+const FString UGatherTextFromSourceCommandlet::MacroString_NOTELOCTEXT(TEXT("NOTELOCTEXT"));
+const FString UGatherTextFromSourceCommandlet::MacroString_NSNOTELOCTEXT(TEXT("NSNOTELOCTEXT"));
 const FString UGatherTextFromSourceCommandlet::MacroString_UI_COMMAND(TEXT("UI_COMMAND"));
 const FString UGatherTextFromSourceCommandlet::MacroString_UI_COMMAND_EXT(TEXT("UI_COMMAND_EXT"));
 
@@ -63,6 +65,8 @@ struct FParsedNestedMacroStats
 
 	std::atomic<int32> Nested_LOCTEXT = 0;
 	std::atomic<int32> Nested_NSLOCTEXT = 0;
+	std::atomic<int32> Nested_NOTELOCTEXT = 0;
+	std::atomic<int32> Nested_NSNOTELOCTEXT = 0;
 	std::atomic<int32> Nested_UI_COMMAND = 0;
 	std::atomic<int32> Nested_UI_COMMAND_EXT = 0;
 
@@ -479,7 +483,18 @@ void UGatherTextFromSourceCommandlet::GetParsables(TArray<FParsableDescriptor*>&
 		Parsables.Add(new FUICommandMacroDescriptor());
 
 		Parsables.Add(new FUICommandExtMacroDescriptor());
+		
+		Parsables.Add(new FStringMacroDescriptor( FString(MacroString_NSNOTELOCTEXT),
+			FStringMacroDescriptor::FMacroArg(FStringMacroDescriptor::MAS_Namespace, true),
+			FStringMacroDescriptor::FMacroArg(FStringMacroDescriptor::MAS_Identifier, true),
+			FStringMacroDescriptor::FMacroArg(FStringMacroDescriptor::MAS_SourceText, true),
+			FStringMacroDescriptor::FMacroArg(FStringMacroDescriptor::MAS_DevNotes, true)));
 	
+		Parsables.Add(new FStringMacroDescriptor( FString(MacroString_NOTELOCTEXT),
+			FStringMacroDescriptor::FMacroArg(FStringMacroDescriptor::MAS_Identifier, true),
+			FStringMacroDescriptor::FMacroArg(FStringMacroDescriptor::MAS_SourceText, true),
+			FStringMacroDescriptor::FMacroArg(FStringMacroDescriptor::MAS_DevNotes, true)));
+
 		// New Localization System with Namespace as literal argument.
 		Parsables.Add(new FStringMacroDescriptor(FString(MacroString_NSLOCTEXT),
 			FStringMacroDescriptor::FMacroArg(FStringMacroDescriptor::MAS_Namespace, true),
@@ -499,6 +514,8 @@ void UGatherTextFromSourceCommandlet::GetParsables(TArray<FParsableDescriptor*>&
 
 		Parsables.Add(new FStringTableEntryMacroDescriptor());
 
+		Parsables.Add(new FStringTableNotedEntryMacroDescriptor());
+
 		Parsables.Add(new FStringTableEntryMetaDataMacroDescriptor());
 
 		Parsables.Add(new FStructuredLogMacroDescriptor(TEXT("UE_LOGFMT_LOC"), FStructuredLogMacroDescriptor::EFlags::None));
@@ -662,6 +679,7 @@ void UGatherTextFromSourceCommandlet::RunPass(EGatherSourcePasses Pass, bool Sho
 							SourceContext.Key = ParsedStringTableEntryPair.Key;
 							SourceContext.SourceLocation = ParsedStringTableEntryPair.Value.SourceLocation.ToString();
 							SourceContext.PlatformName = ParsedStringTableEntryPair.Value.PlatformName;
+							SourceContext.DevNotes = ParsedStringTableEntryPair.Value.DevNotes;
 
 							const FParsedStringTableEntryMetaDataMap* ParsedMetaDataMap = ParsedStringTablePair.Value.MetaDataEntries.Find(ParsedStringTableEntryPair.Key);
 							if (ParsedMetaDataMap && ParsedMetaDataMap->Num() > 0)
@@ -1585,6 +1603,8 @@ void UGatherTextFromSourceCommandlet::LogStats()
 	UE_LOG(LogGatherTextFromSourceCommandlet, Display, TEXT(""));
 	UE_LOG(LogGatherTextFromSourceCommandlet, Display, TEXT("Nested_LOCTEXT =        %14d"), NestedMacroStats.Nested_LOCTEXT.load());
 	UE_LOG(LogGatherTextFromSourceCommandlet, Display, TEXT("Nested_NSLOCTEXT =      %14d"), NestedMacroStats.Nested_NSLOCTEXT.load());
+	UE_LOG(LogGatherTextFromSourceCommandlet, Display, TEXT("Nested_NOTELOCTEXT =    %14d"), NestedMacroStats.Nested_NOTELOCTEXT.load());
+	UE_LOG(LogGatherTextFromSourceCommandlet, Display, TEXT("Nested_NSNOTELOCTEXT =  %14d"), NestedMacroStats.Nested_NSNOTELOCTEXT.load());
 	UE_LOG(LogGatherTextFromSourceCommandlet, Display, TEXT("Nested_UI_COMMAND =     %14d"), NestedMacroStats.Nested_UI_COMMAND.load());
 	UE_LOG(LogGatherTextFromSourceCommandlet, Display, TEXT("Nested_UI_COMMAND_EXT = %14d"), NestedMacroStats.Nested_UI_COMMAND_EXT.load());
 
@@ -1845,7 +1865,7 @@ bool UGatherTextFromSourceCommandlet::FSourceFileParseContext::AddStringTableImp
 	return true;
 }
 
-bool UGatherTextFromSourceCommandlet::FSourceFileParseContext::AddStringTableEntryImpl(const FName InTableId, const FString& InKey, const FString& InSourceString, const FSourceLocation& InSourceLocation, const FName InPlatformName)
+bool UGatherTextFromSourceCommandlet::FSourceFileParseContext::AddStringTableEntryImpl(const FName InTableId, const FString& InKey, const FString& InSourceString, const FSourceLocation& InSourceLocation, const FName InPlatformName, const FString& InDevNotes)
 {
 	const bool bIsEditorOnly = EvaluateEditorOnlyDefineState() == EEditorOnlyDefineState::Defined;
 
@@ -1879,6 +1899,7 @@ bool UGatherTextFromSourceCommandlet::FSourceFileParseContext::AddStringTableEnt
 		ParsedStringTableEntry.SourceString = InSourceString;
 		ParsedStringTableEntry.SourceLocation = InSourceLocation;
 		ParsedStringTableEntry.PlatformName = InPlatformName;
+		ParsedStringTableEntry.DevNotes = InDevNotes;
 		ParsedStringTableEntry.bIsEditorOnly = bIsEditorOnly;
 		return true;
 	}
@@ -1940,9 +1961,9 @@ void UGatherTextFromSourceCommandlet::FSourceFileParseContext::AddStringTableFro
 			const FSourceLocation SourceLocation = FSourceLocation(InTableFilename, INDEX_NONE);
 			const FName TablePlatformName = GetSplitPlatformNameFromPath_Static(InTableFilename, SplitPlatforms);
 
-			TmpStringTable->EnumerateSourceStrings([&](const FString& InKey, const FString& InSourceString)
+			TmpStringTable->EnumerateSourceStrings([&](const FString& InKey, const FString& InSourceString, const FString& InDevNote)
 			{
-				AddStringTableEntryImpl(InTableId, InKey, InSourceString, SourceLocation, TablePlatformName);
+				AddStringTableEntryImpl(InTableId, InKey, InSourceString, SourceLocation, TablePlatformName, InDevNote);
 
 				TmpStringTable->EnumerateMetaData(InKey, [&](const FName InMetaDataId, const FString& InMetaData)
 				{
@@ -1966,9 +1987,9 @@ void UGatherTextFromSourceCommandlet::FSourceFileParseContext::AddStringTableFro
 	}
 }
 
-void UGatherTextFromSourceCommandlet::FSourceFileParseContext::AddStringTableEntry(const FName InTableId, const FString& InKey, const FString& InSourceString)
+void UGatherTextFromSourceCommandlet::FSourceFileParseContext::AddStringTableEntry(const FName InTableId, const FString& InKey, const FString& InSourceString, const FString& InDevNotes)
 {
-	AddStringTableEntryImpl(InTableId, InKey, InSourceString, FSourceLocation(Filename, LineNumber), FilePlatformName);
+	AddStringTableEntryImpl(InTableId, InKey, InSourceString, FSourceLocation(Filename, LineNumber), FilePlatformName, InDevNotes);
 }
 
 void UGatherTextFromSourceCommandlet::FSourceFileParseContext::AddStringTableEntryMetaData(const FName InTableId, const FString& InKey, const FName InMetaDataId, const FString& InMetaData)
@@ -2603,6 +2624,18 @@ void UGatherTextFromSourceCommandlet::FNestedMacroPrepassDescriptor::TryParse(co
 	int32 LineStart = Context.LineNumber;
 	int32 LineCount = TextLines.Num();
 	bool bAdvance = false;
+	if (Content.Contains(MacroString_NOTELOCTEXT, ESearchCase::CaseSensitive))
+	{
+		PrepassResults.Emplace(MacroName, MacroString_NOTELOCTEXT, Context.Filename, Content, LineStart, LineCount);
+		++NestedMacroStats.Nested_NOTELOCTEXT;
+		bAdvance = true;
+	}
+	if (Content.Contains(MacroString_NSNOTELOCTEXT, ESearchCase::CaseSensitive))
+	{
+		PrepassResults.Emplace(MacroName, MacroString_NSNOTELOCTEXT, Context.Filename, Content, LineStart, LineCount);
+		++NestedMacroStats.Nested_NSNOTELOCTEXT;
+		bAdvance = true;
+	}
 	if (Content.Contains(MacroString_LOCTEXT, ESearchCase::CaseSensitive) &&
 		!Content.Contains(TEXT("LOCTEXT_NAMESPACE"), ESearchCase::CaseSensitive))
 	{
@@ -2820,6 +2853,21 @@ void UGatherTextFromSourceCommandlet::FNestedMacroDescriptor::TryParse(const FSt
 			FStringMacroDescriptor::FMacroArg(FStringMacroDescriptor::MAS_Identifier, true),
 			FStringMacroDescriptor::FMacroArg(FStringMacroDescriptor::MAS_SourceText, true));
 	}
+	else if (MacroNameNested == MacroString_NOTELOCTEXT)
+	{
+		InnerDescriptor = new FStringMacroDescriptor(FString(MacroString_NOTELOCTEXT),
+			FStringMacroDescriptor::FMacroArg(FStringMacroDescriptor::MAS_Identifier, true),
+			FStringMacroDescriptor::FMacroArg(FStringMacroDescriptor::MAS_SourceText, true),
+			FStringMacroDescriptor::FMacroArg(FStringMacroDescriptor::MAS_DevNotes, true));
+	}
+	else if (MacroNameNested == MacroString_NSNOTELOCTEXT)
+	{
+		InnerDescriptor = new FStringMacroDescriptor(FString(MacroString_NSNOTELOCTEXT),
+			FStringMacroDescriptor::FMacroArg(FStringMacroDescriptor::MAS_Namespace, true),
+			FStringMacroDescriptor::FMacroArg(FStringMacroDescriptor::MAS_Identifier, true),
+			FStringMacroDescriptor::FMacroArg(FStringMacroDescriptor::MAS_SourceText, true),
+			FStringMacroDescriptor::FMacroArg(FStringMacroDescriptor::MAS_DevNotes, true));
+	}
 	else if (MacroNameNested == MacroString_UI_COMMAND)
 	{
 		InnerDescriptor = new FUICommandMacroDescriptor();
@@ -3013,6 +3061,7 @@ void UGatherTextFromSourceCommandlet::FStringMacroDescriptor::TryParse(const FSt
 				FString Identifier;
 				FString SourceLocation = FSourceLocation(Context.Filename, Context.LineNumber).ToString();
 				FString SourceText;
+				FString DevNotes;
 
 				TOptional<FString> Namespace;
 				if (Context.WithinNamespaceDefineLineNumber != INDEX_NONE || !Context.Namespace.IsEmpty())
@@ -3052,6 +3101,11 @@ void UGatherTextFromSourceCommandlet::FStringMacroDescriptor::TryParse(const FSt
 							SourceText = ArgText;
 						}
 						break;
+					case MAS_DevNotes:
+						{
+							DevNotes = ArgText;
+						}
+						break;
 					}
 				}
 
@@ -3093,6 +3147,7 @@ void UGatherTextFromSourceCommandlet::FStringMacroDescriptor::TryParse(const FSt
 					MacroContext.Key = Identifier;
 					MacroContext.SourceLocation = SourceLocation;
 					MacroContext.PlatformName = Context.FilePlatformName;
+					MacroContext.DevNotes = DevNotes;
 
 					if (EnumHasAnyFlags(Context.FileTypes, EGatherTextSourceFileTypes::Ini))
 					{
@@ -3273,7 +3328,56 @@ void UGatherTextFromSourceCommandlet::FStringTableEntryMacroDescriptor::TryParse
 					}
 					else if (!SourceString.IsEmpty())
 					{
-						Context.AddStringTableEntry(TableIdName, Key, SourceString);
+						Context.AddStringTableEntry(TableIdName, Key, SourceString, TEXT(""));
+					}
+				}
+			}
+		}
+	}
+}
+
+void UGatherTextFromSourceCommandlet::FStringTableNotedEntryMacroDescriptor::TryParse(const FString& Text, FSourceFileParseContext& Context) const
+{
+	// Attempt to parse something of the format
+	// LOCTABLE_SETNOTESTRING(Id, Key, SourceString, DevNotes)
+
+	if (!Context.ExcludedRegion && !Context.WithinBlockComment && !Context.WithinLineComment && !Context.WithinStringLiteral)
+	{
+		TArray<FString> Arguments;
+		if (ParseArgsFromMacro(StripCommentsFromToken(Text, Context), Arguments, Context))
+		{
+			if (Arguments.Num() != GetMinNumberOfArgument())
+			{
+				UE_LOG(LogGatherTextFromSourceCommandlet, Warning, TEXT("%s(%d): Expected %d arguments for %s macro, but got %d. %s"), *Context.Filename, Context.LineNumber, GetMinNumberOfArgument(), *GetToken(), Arguments.Num(), *FLocTextHelper::SanitizeLogOutput(Context.LineText.TrimStartAndEnd()));
+			}
+			else
+			{
+				Arguments[0].TrimStartInline();
+				FString TableId = Arguments[0];
+				Arguments[1].TrimStartInline();
+				FString Key = Arguments[1];
+				Arguments[2].TrimStartInline();
+				FString SourceString = Arguments[2];
+				Arguments[3].TrimStartInline();
+				FString DevNotes = Arguments[3];
+
+				const FString TableIdMacroDesc = FString::Printf(TEXT("%s(%d): \"Id\" argument in %s macro"), *Context.Filename, Context.LineNumber, *GetToken());
+				const FString KeyMacroDesc = FString::Printf(TEXT("%s(%d): \"Key\" argument in %s macro"), *Context.Filename, Context.LineNumber, *GetToken());
+				const FString SourceStringMacroDesc = FString::Printf(TEXT("%s(%d): \"SourceString\" argument in %s macro"), *Context.Filename, Context.LineNumber, *GetToken());
+				const FString DevNoteMacroDesc = FString::Printf(TEXT("%s(%d): \"DevNote\" argument in %s macro"), *Context.Filename, Context.LineNumber, *GetToken());
+
+				bool HasQuotes;
+				if (PrepareArgument(TableId, true, TableIdMacroDesc, HasQuotes) && PrepareArgument(Key, true, KeyMacroDesc, HasQuotes) && PrepareArgument(SourceString, true, SourceStringMacroDesc, HasQuotes) && PrepareArgument(DevNotes, true, DevNoteMacroDesc, HasQuotes))
+				{
+					const FName TableIdName = *TableId;
+
+					if (TableIdName.IsNone() || Key.IsEmpty())
+					{
+						UE_LOG(LogGatherTextFromSourceCommandlet, Warning, TEXT("%s(%d): %s macro has an empty identifier and cannot be gathered."), *Context.Filename, Context.LineNumber, *GetToken());
+					}
+					else if (!SourceString.IsEmpty())
+					{
+						Context.AddStringTableEntry(TableIdName, Key, SourceString, DevNotes);
 					}
 				}
 			}
diff --git a/Engine/Source/Programs/Shared/EpicGames.UHT/Types/Properties/UhtTextProperty.cs b/Engine/Source/Programs/Shared/EpicGames.UHT/Types/Properties/UhtTextProperty.cs
--- a/Engine/Source/Programs/Shared/EpicGames.UHT/Types/Properties/UhtTextProperty.cs
+++ b/Engine/Source/Programs/Shared/EpicGames.UHT/Types/Properties/UhtTextProperty.cs
@@ -114,6 +114,12 @@ namespace EpicGames.UHT.Types
 				this.LogError($"LOCTEXT default parameter values are not supported; use NSLOCTEXT instead: {SourceName}");
 				return false;
 			}
+			else if (token.IsIdentifier("NOTELOCTEXT"))
+			{
+				//ETSTODO - Add a test case for this
+				this.LogError($"NOTELOCTEXT default parameter values are not supported; use NSNOTELOCTEXT instead: {SourceName}");
+				return false;
+			}
 			else
 			{
 				return SanitizeDefaultValue(this, defaultValueReader, ref token, innerDefaultValue, false);
@@ -164,6 +170,15 @@ namespace EpicGames.UHT.Types
 			return false;
 		}
 
+		[UhtLocTextDefaultValue(Name = "NOTELOCTEXT")]
+		[SuppressMessage("CodeQuality", "IDE0051:Remove unused private members", Justification = "Attribute accessed method")]
+		[SuppressMessage("Style", "IDE0060:Remove unused parameter", Justification = "Attribute accessed method")]
+		private static bool NoteLocTextDefaultValue(UhtTextProperty property, IUhtTokenReader defaultValueReader, ref UhtToken macroToken, StringBuilder innerDefaultValue)
+		{
+			property.LogError($"NOTELOCTEXT default parameter values are not supported; use NSNOTELOCTEXT instead: {property.SourceName}");
+			return false;
+		}
+
 		[UhtLocTextDefaultValue(Name = "NSLOCTEXT")]
 		[SuppressMessage("CodeQuality", "IDE0051:Remove unused private members", Justification = "Attribute accessed method")]
 		[SuppressMessage("Style", "IDE0060:Remove unused parameter", Justification = "Attribute accessed method")]
@@ -199,6 +214,43 @@ namespace EpicGames.UHT.Types
 			return true;
 		}
 
+		[UhtLocTextDefaultValue(Name = "NSNOTELOCTEXT")]
+		[SuppressMessage("CodeQuality", "IDE0051:Remove unused private members", Justification = "Attribute accessed method")]
+		[SuppressMessage("Style", "IDE0060:Remove unused parameter", Justification = "Attribute accessed method")]
+		private static bool NsNoteLocTextDefaultValue(UhtTextProperty property, IUhtTokenReader defaultValueReader, ref UhtToken macroToken, StringBuilder innerDefaultValue)
+		{
+			defaultValueReader.Require('(');
+			StringView namespaceString = defaultValueReader.GetConstQuotedString();
+			defaultValueReader.Require(',');
+			StringView keyString = defaultValueReader.GetConstQuotedString();
+			defaultValueReader.Require(',');
+			StringView sourceString = defaultValueReader.GetConstQuotedString();
+			defaultValueReader.Require(',');
+			StringView devNotesString = defaultValueReader.GetConstQuotedString();
+			defaultValueReader.Require(')');
+
+			// Strip out the package name
+			ReadOnlySpan<char> strippedNS = namespaceString.Span;
+			if (strippedNS.Length > 1)
+			{
+				strippedNS = strippedNS[1..^1].Trim();
+				if (strippedNS.Length > 0)
+				{
+					if (strippedNS[^1] == ']')
+					{
+						int index = strippedNS.LastIndexOf('[');
+						if (index != -1)
+						{
+							strippedNS = strippedNS[..index].TrimEnd();
+						}
+					}
+				}
+			}
+
+			innerDefaultValue.Append("NSNOTELOCTEXT(\"").Append(strippedNS).Append("\", ").Append(keyString).Append(", ").Append(sourceString).Append(", ").Append(devNotesString).Append(')');
+			return true;
+		}
+
 		[UhtLocTextDefaultValue(Name = "LOCTABLE")]
 		[SuppressMessage("CodeQuality", "IDE0051:Remove unused private members", Justification = "Attribute accessed method")]
 		[SuppressMessage("Style", "IDE0060:Remove unused parameter", Justification = "Attribute accessed method")]
diff --git a/Engine/Source/Runtime/Core/Private/Internationalization/GatherableTextData.cpp b/Engine/Source/Runtime/Core/Private/Internationalization/GatherableTextData.cpp
--- a/Engine/Source/Runtime/Core/Private/Internationalization/GatherableTextData.cpp
+++ b/Engine/Source/Runtime/Core/Private/Internationalization/GatherableTextData.cpp
@@ -7,6 +7,7 @@
 #include "Serialization/StructuredArchiveAdapters.h"
 #include "Serialization/StructuredArchiveNameHelpers.h"
 #include "Serialization/StructuredArchiveSlots.h"
+#include "UObject/FortniteMainBranchObjectVersion.h"
 
 class FArchive;
 
@@ -22,6 +23,14 @@ void FTextSourceSiteContext::StreamStructuredArchive(FStructuredArchive::FSlot S
 	Record << SA_VALUE(TEXT("SiteDescription"), SiteDescription);
 	Record << SA_VALUE(TEXT("IsEditorOnly"), IsEditorOnly);
 	Record << SA_VALUE(TEXT("IsOptional"), IsOptional);
+#if WITH_EDITORONLY_DATA
+	FArchive& UnderlyingArchive = Slot.GetUnderlyingArchive();
+	UnderlyingArchive.UsingCustomVersion(FFortniteMainBranchObjectVersion::GUID);
+	if (UnderlyingArchive.CustomVer(FFortniteMainBranchObjectVersion::GUID) >= FFortniteMainBranchObjectVersion::AddDevNotesToFText && !UnderlyingArchive.IsCooking() && !UnderlyingArchive.IsFilterEditorOnly())
+	{
+		Record << SA_VALUE(TEXT("DevNotes"), DevNotes);
+	}
+#endif
 	Record << SA_VALUE(TEXT("InfoMetaData"), InfoMetaData);
 	Record << SA_VALUE(TEXT("KeyMetaData"), KeyMetaData);
 }
diff --git a/Engine/Source/Runtime/Core/Private/Internationalization/InternationalizationManifest.cpp b/Engine/Source/Runtime/Core/Private/Internationalization/InternationalizationManifest.cpp
--- a/Engine/Source/Runtime/Core/Private/Internationalization/InternationalizationManifest.cpp
+++ b/Engine/Source/Runtime/Core/Private/Internationalization/InternationalizationManifest.cpp
@@ -13,6 +13,7 @@ FManifestContext::FManifestContext(const FManifestContext& Other)
 	, SourceLocation(Other.SourceLocation)
 	, PlatformName(Other.PlatformName)
 	, bIsOptional(Other.bIsOptional)
+	, DevNotes(Other.DevNotes)
 {
 	if (Other.InfoMetadataObj.IsValid())
 	{
@@ -33,6 +34,7 @@ FManifestContext& FManifestContext::operator=(const FManifestContext& Other)
 		SourceLocation = Other.SourceLocation;
 		PlatformName = Other.PlatformName;
 		bIsOptional = Other.bIsOptional;
+		DevNotes = Other.DevNotes;
 		InfoMetadataObj.Reset();
 		KeyMetadataObj.Reset();
 		if (Other.InfoMetadataObj.IsValid())
diff --git a/Engine/Source/Runtime/Core/Private/Internationalization/StringTableCore.cpp b/Engine/Source/Runtime/Core/Private/Internationalization/StringTableCore.cpp
--- a/Engine/Source/Runtime/Core/Private/Internationalization/StringTableCore.cpp
+++ b/Engine/Source/Runtime/Core/Private/Internationalization/StringTableCore.cpp
@@ -19,6 +19,7 @@
 #include "Serialization/Csv/CsvParser.h"
 #include "Templates/Tuple.h"
 #include "Trace/Detail/Channel.h"
+#include "UObject/FortniteMainBranchObjectVersion.h"
 
 class UStringTable;
 
@@ -40,10 +41,17 @@ FStringTableEntry::FStringTableEntry()
 {
 }
 
+#if WITH_EDITORONLY_DATA
+FStringTableEntry::FStringTableEntry(FStringTableConstRef InOwnerTable, FString InSourceString, FTextId InDisplayStringId, FString InDevNotes)
+#else
 FStringTableEntry::FStringTableEntry(FStringTableConstRef InOwnerTable, FString InSourceString, FTextId InDisplayStringId)
+#endif
 	: OwnerTable(MoveTemp(InOwnerTable))
 	, SourceString(MoveTemp(InSourceString))
 	, DisplayStringId(MoveTemp(InDisplayStringId))
+#if WITH_EDITORONLY_DATA
+	, DevNotes(MoveTemp(InDevNotes))
+#endif
 {
 }
 
@@ -150,7 +158,11 @@ void FStringTable::SetNamespace(const FTextKey& InNamespace)
 			FStringTableEntryPtr OldEntry = KeyToEntryPair.Value;
 			OldEntry->Disown();
 
+#if WITH_EDITORONLY_DATA
+			KeyToEntryPair.Value = FStringTableEntry::NewStringTableEntry(AsShared(), OldEntry->GetSourceString(), FTextId(TableNamespace, KeyToEntryPair.Key), OldEntry->GetDevNotes());
+#else
 			KeyToEntryPair.Value = FStringTableEntry::NewStringTableEntry(AsShared(), OldEntry->GetSourceString(), FTextId(TableNamespace, KeyToEntryPair.Key));
+#endif
 		}
 	}
 }
@@ -168,7 +180,11 @@ bool FStringTable::GetSourceString(const FTextKey& InKey, FString& OutSourceStri
 	return false;
 }
 
+#if WITH_EDITORONLY_DATA
+void FStringTable::SetSourceString(const FTextKey& InKey, const FString& InSourceString, const FString& InDevNotes)
+#else
 void FStringTable::SetSourceString(const FTextKey& InKey, const FString& InSourceString)
+#endif
 {
 	checkf(!InKey.IsEmpty(), TEXT("String table key cannot be empty!"));
 
@@ -180,7 +196,11 @@ void FStringTable::SetSourceString(const FTextKey& InKey, const FString& InSourc
 		TableEntry->Disown();
 	}
 	
+#if WITH_EDITORONLY_DATA
+	TableEntry = FStringTableEntry::NewStringTableEntry(AsShared(), InSourceString, FTextId(TableNamespace, InKey), InDevNotes);
+#else
 	TableEntry = FStringTableEntry::NewStringTableEntry(AsShared(), InSourceString, FTextId(TableNamespace, InKey));
+#endif
 	KeysToEntries.Emplace(InKey, TableEntry);
 }
 
@@ -197,6 +217,29 @@ void FStringTable::RemoveSourceString(const FTextKey& InKey)
 	}
 }
 
+#if WITH_EDITORONLY_DATA
+void FStringTable::EnumerateSourceStrings(const TFunctionRef<bool(const FString&, const FString&, const FString&)>& InEnumerator) const
+{
+	EnumerateKeysAndSourceStrings([&InEnumerator](const FTextKey& InKey, const FString& InSourceString, const FString& InDevNotes) -> bool
+	{
+		return InEnumerator(InKey.ToString(), InSourceString, InDevNotes);
+	});
+}
+
+void FStringTable::EnumerateKeysAndSourceStrings(const TFunctionRef<bool(const FTextKey&, const FString&, const FString&)>& InEnumerator) const
+{
+	UE::TScopeLock KeyMappingLock(KeyMappingCS);
+
+	for (const auto& KeyToEntryPair : KeysToEntries)
+	{
+		if (!InEnumerator(KeyToEntryPair.Key, KeyToEntryPair.Value->GetSourceString(), KeyToEntryPair.Value->GetDevNotes()))
+		{
+			break;
+		}
+	}
+}
+#endif
+
 void FStringTable::EnumerateSourceStrings(const TFunctionRef<bool(const FString&, const FString&)>& InEnumerator) const
 {
 	EnumerateKeysAndSourceStrings([&InEnumerator](const FTextKey& InKey, const FString& InSourceString) -> bool
@@ -337,6 +380,9 @@ void FStringTable::Serialize(FArchive& Ar)
 	UE::TScopeLock MetaDataLock(KeysToMetaDataCS);
 
 	TableNamespace.SerializeAsString(Ar);
+#if WITH_EDITORONLY_DATA
+	Ar.UsingCustomVersion(FFortniteMainBranchObjectVersion::GUID);
+#endif
 
 	if (Ar.IsSaving())
 	{
@@ -352,6 +398,14 @@ void FStringTable::Serialize(FArchive& Ar)
 
 				FString SourceString = KeyToEntryPair.Value->GetSourceString();
 				Ar << SourceString;
+
+#if WITH_EDITORONLY_DATA
+				if (!Ar.IsCooking() && !Ar.IsFilterEditorOnly())
+				{
+					FString DevNotes = KeyToEntryPair.Value->GetDevNotes();
+					Ar << DevNotes;
+				}
+#endif
 			}
 		}
 
@@ -382,8 +436,17 @@ void FStringTable::Serialize(FArchive& Ar)
 
 				FString SourceString;
 				Ar << SourceString;
-
+#if WITH_EDITORONLY_DATA
+				FString DevNotes;
+				if (Ar.CustomVer(FFortniteMainBranchObjectVersion::GUID) >= FFortniteMainBranchObjectVersion::AddDevNotesToFText && !Ar.IsFilterEditorOnly())
+				{
+					Ar << DevNotes;
+				}
+				
+				FStringTableEntryRef TableEntry = FStringTableEntry::NewStringTableEntry(AsShared(), SourceString, FTextId(TableNamespace, Key), DevNotes);
+#else
 				FStringTableEntryRef TableEntry = FStringTableEntry::NewStringTableEntry(AsShared(), SourceString, FTextId(TableNamespace, Key));
+#endif
 				KeysToEntries.Emplace(Key, TableEntry);
 			}
 		}
@@ -433,6 +496,9 @@ FString FStringTable::ExportStringsToCSVString() const
 
 		// Write header
 		ExportedStrings += TEXT("Key,SourceString");
+#if WITH_EDITORONLY_DATA
+		ExportedStrings += TEXT(",DevNotes");
+#endif
 		for (const FName& MetaDataColumnName : MetaDataColumnNames)
 		{
 			ExportedStrings += TEXT(",");
@@ -460,6 +526,17 @@ FString FStringTable::ExportStringsToCSVString() const
 			ExportedStrings += ExportedSourceString;
 			ExportedStrings += TEXT("\"");
 
+#if WITH_EDITORONLY_DATA
+			FString ExportedDevNotes = KeyToEntryPair.Value->GetDevNotes().ReplaceCharWithEscapedChar();
+			ExportedDevNotes.ReplaceInline(TEXT("\""), TEXT("\"\""));
+			
+			ExportedStrings += TEXT(",");
+
+			ExportedStrings += TEXT("\"");
+			ExportedStrings += ExportedDevNotes;
+			ExportedStrings += TEXT("\"");
+#endif
+
 			for (const FName& MetaDataColumnName : MetaDataColumnNames)
 			{
 				FString ExportedMetaData = GetMetaData(KeyToEntryPair.Key, MetaDataColumnName);
@@ -504,6 +581,9 @@ bool FStringTable::ImportStringsFromCSVString(const FString& InCSVString, const
 
 	int32 KeyColumn = INDEX_NONE;
 	int32 SourceStringColumn = INDEX_NONE;
+#if WITH_EDITORONLY_DATA
+	int32 DevNotesColumn = INDEX_NONE;
+#endif
 	TMap<FName, int32> MetaDataColumns;
 
 	// Validate header
@@ -521,6 +601,12 @@ bool FStringTable::ImportStringsFromCSVString(const FString& InCSVString, const
 			{
 				SourceStringColumn = CellIdx;
 			}
+#if WITH_EDITORONLY_DATA
+			else if(FCString::Stricmp(Cell, TEXT("DevNotes")) == 0 && DevNotesColumn == INDEX_NONE)
+			{
+				DevNotesColumn = CellIdx;
+			}
+#endif
 			else
 			{
 				const FName MetaDataName = Cell;
@@ -567,7 +653,14 @@ bool FStringTable::ImportStringsFromCSVString(const FString& InCSVString, const
 				FString SourceString = Cells[SourceStringColumn];
 				SourceString = SourceString.ReplaceEscapedCharWithChar();
 
+#if WITH_EDITORONLY_DATA
+				FString DevNotes = Cells.IsValidIndex(DevNotesColumn) ? Cells[DevNotesColumn] : TEXT("");
+				DevNotes = DevNotes.ReplaceEscapedCharWithChar();
+				
+				FStringTableEntryRef TableEntry = FStringTableEntry::NewStringTableEntry(AsShared(), SourceString, FTextId(TableNamespace, Key), MoveTemp(DevNotes));
+#else
 				FStringTableEntryRef TableEntry = FStringTableEntry::NewStringTableEntry(AsShared(), SourceString, FTextId(TableNamespace, Key));
+#endif
 				KeysToEntries.Emplace(Key, TableEntry);
 
 				for (const auto& MetaDataColumnPair : MetaDataColumns)
diff --git a/Engine/Source/Runtime/Core/Private/Internationalization/StringTableRegistry.cpp b/Engine/Source/Runtime/Core/Private/Internationalization/StringTableRegistry.cpp
--- a/Engine/Source/Runtime/Core/Private/Internationalization/StringTableRegistry.cpp
+++ b/Engine/Source/Runtime/Core/Private/Internationalization/StringTableRegistry.cpp
@@ -211,12 +211,19 @@ void FStringTableRegistry::Internal_LocTableFromFile(const FName InTableId, cons
 #endif // WITH_EDITOR
 }
 
+#if WITH_EDITORONLY_DATA
+void FStringTableRegistry::Internal_SetLocTableEntry(const FName InTableId, const FTextKey& InKey, const FString& InSourceString, const FString& InDevNotes)
+#else
 void FStringTableRegistry::Internal_SetLocTableEntry(const FName InTableId, const FTextKey& InKey, const FString& InSourceString)
+#endif
 {
 	FStringTablePtr StringTable = FindMutableStringTable(InTableId);
 	checkf(StringTable.IsValid(), TEXT("Attempting to add a string table entry to the unknown string table '%s'"), *InTableId.ToString());
-
+#if WITH_EDITORONLY_DATA
+	StringTable->SetSourceString(InKey, InSourceString, InDevNotes);
+#else
 	StringTable->SetSourceString(InKey, InSourceString);
+#endif
 }
 
 void FStringTableRegistry::Internal_SetLocTableEntryMetaData(const FName InTableId, const FTextKey& InKey, const FName InMetaDataId, const FString& InMetaData)
diff --git a/Engine/Source/Runtime/Core/Private/Internationalization/Text.cpp b/Engine/Source/Runtime/Core/Private/Internationalization/Text.cpp
--- a/Engine/Source/Runtime/Core/Private/Internationalization/Text.cpp
+++ b/Engine/Source/Runtime/Core/Private/Internationalization/Text.cpp
@@ -273,6 +273,14 @@ FText::FText( FString&& InSourceString, const FTextKey& InNamespace, const FText
 {
 }
 
+#if WITH_EDITORONLY_DATA
+FText::FText( FString&& InSourceString, const FTextKey& InNamespace, const FTextKey& InKey, FString InDevNotes, uint32 InFlags )
+	: TextData(MakeRefCount<FTextHistory_Base>(FTextId(InNamespace, InKey), MoveTemp(InSourceString), MoveTemp(InDevNotes)))
+	, Flags(InFlags)
+{
+}
+#endif
+
 FText::FText(FText&& Other)
 	: TextData(MoveTemp(Other.TextData))
 	, Flags(Other.Flags)
@@ -1055,6 +1063,16 @@ FText FText::ChangeKey( const FTextKey& Namespace, const FTextKey& Key, const FT
 {
 	return FText(FString(*Text.GetTextData().GetSourceString()), Namespace, Key, 0);
 }
+
+FText FText::ChangeKey( const FTextKey& Namespace, const FTextKey& Key, const FText& Text, const FString& Notes )
+{
+	return FText(FString(*Text.GetTextData().GetSourceString()), Namespace, Key, Notes, 0);
+}
+
+const FString& FText::GetDevNotes() const
+{
+	return GetTextData().GetTextHistory().GetDevNotes();
+}
 #endif
 
 FText FText::FromStringTable(const FName InTableId, const FTextKey& InKey, const EStringTableLoadingPolicy InLoadingPolicy)
@@ -1212,6 +1230,31 @@ FText FText::AsCultureInvariant( FText Text )
 	return NewText;
 }
 
+#if WITH_EDITORONLY_DATA
+FText FText::AsLocalizable_Advanced(const FTextKey& Namespace, const FTextKey& Key, const TCHAR* String, FStringView InDevNotes)
+{
+	return FText::AsLocalizable_Advanced(Namespace, Key, FStringView(String), InDevNotes);
+}
+
+FText FText::AsLocalizable_Advanced(const FTextKey& Namespace, const FTextKey& Key, FStringView String, FStringView InDevNotes)
+{
+	return String.IsEmpty() 
+		? FText::GetEmpty() 
+		: FTextCache::Get().FindOrCache(String, FTextId(Namespace, Key), InDevNotes);
+}
+
+FText FText::AsLocalizable_Advanced(const FTextKey& Namespace, const FTextKey& Key, FString&& String, FStringView InDevNotes)
+{
+	return String.IsEmpty() 
+		? FText::GetEmpty() 
+		: FTextCache::Get().FindOrCache(MoveTemp(String), FTextId(Namespace, Key), InDevNotes);
+}
+
+FText FText::AsLocalizable_Advanced_LocText(const TCHAR* Namespace, const TCHAR* Key, const TCHAR* String, FStringView InDevNotes)
+{
+	return FText::AsLocalizable_Advanced(Namespace, Key, FStringView(String), InDevNotes);
+}
+#else
 FText FText::AsLocalizable_Advanced(const FTextKey& Namespace, const FTextKey& Key, const TCHAR* String)
 {
 	return FText::AsLocalizable_Advanced(Namespace, Key, FStringView(String));
@@ -1235,6 +1278,7 @@ FText FText::AsLocalizable_Advanced_LocText(const TCHAR* Namespace, const TCHAR*
 {
 	return FText::AsLocalizable_Advanced(Namespace, Key, FStringView(String));
 }
+#endif
 
 bool FText::FindTextInLiveTable_Advanced(const FTextKey& Namespace, const FTextKey& Key, FText& OutText, const FString* const SourceString)
 {
diff --git a/Engine/Source/Runtime/Core/Private/Internationalization/TextCache.cpp b/Engine/Source/Runtime/Core/Private/Internationalization/TextCache.cpp
--- a/Engine/Source/Runtime/Core/Private/Internationalization/TextCache.cpp
+++ b/Engine/Source/Runtime/Core/Private/Internationalization/TextCache.cpp
@@ -20,12 +20,23 @@ void FTextCache::TearDown()
 	return TLazySingleton<FTextCache>::TearDown();
 }
 
+#if WITH_EDITORONLY_DATA
+FText FTextCache::FindOrCache(const TCHAR* InTextLiteral, const FTextId& InTextId, FStringView InDevNotes)
+{
+	return FindOrCache(FStringView(InTextLiteral), InTextId, InDevNotes);
+}
+#else
 FText FTextCache::FindOrCache(const TCHAR* InTextLiteral, const FTextId& InTextId)
 {
 	return FindOrCache(FStringView(InTextLiteral), InTextId);
 }
+#endif
 
-FText FTextCache::FindOrCache(FStringView InTextLiteral, const FTextId& InTextId)
+#if WITH_EDITORONLY_DATA
+FText FTextCache::FindOrCache(const FStringView InTextLiteral, const FTextId& InTextId, FStringView InDevNotes)
+#else
+FText FTextCache::FindOrCache(const FStringView InTextLiteral, const FTextId& InTextId)
+#endif
 {
 	LLM_SCOPE(ELLMTag::Localization);
 
@@ -52,7 +63,11 @@ FText FTextCache::FindOrCache(FStringView InTextLiteral, const FTextId& InTextId
 	}
 
 	// Not currently cached, make a new instance...
+#if WITH_EDITORONLY_DATA
+	FText NewText = FText(FString(InTextLiteral), InTextId.GetNamespace(), InTextId.GetKey(), FString(InDevNotes), ETextFlag::Immutable);
+#else
 	FText NewText = FText(FString(InTextLiteral), InTextId.GetNamespace(), InTextId.GetKey(), ETextFlag::Immutable);
+#endif
 
 	// ... and add it to the cache
 	CachedText.Add(InTextId, NewText);
@@ -60,7 +75,11 @@ FText FTextCache::FindOrCache(FStringView InTextLiteral, const FTextId& InTextId
 	return NewText;
 }
 
+#if WITH_EDITORONLY_DATA
+FText FTextCache::FindOrCache(FString&& InTextLiteral, const FTextId& InTextId, FStringView InDevNotes)
+#else
 FText FTextCache::FindOrCache(FString&& InTextLiteral, const FTextId& InTextId)
+#endif
 {
 	LLM_SCOPE(ELLMTag::Localization);
 
@@ -87,7 +106,11 @@ FText FTextCache::FindOrCache(FString&& InTextLiteral, const FTextId& InTextId)
 	}
 
 	// Not currently cached, make a new instance...
-	FText NewText = FText(MoveTemp(InTextLiteral), InTextId.GetNamespace(), InTextId.GetKey(), ETextFlag::Immutable);
+#if WITH_EDITORONLY_DATA
+	FText NewText = FText(FString(InTextLiteral), InTextId.GetNamespace(), InTextId.GetKey(), FString(InDevNotes), ETextFlag::Immutable);
+#else
+	FText NewText = FText(FString(InTextLiteral), InTextId.GetNamespace(), InTextId.GetKey(), ETextFlag::Immutable);
+#endif
 
 	// ... and add it to the cache
 	CachedText.Add(InTextId, NewText);
diff --git a/Engine/Source/Runtime/Core/Private/Internationalization/TextCache.h b/Engine/Source/Runtime/Core/Private/Internationalization/TextCache.h
--- a/Engine/Source/Runtime/Core/Private/Internationalization/TextCache.h
+++ b/Engine/Source/Runtime/Core/Private/Internationalization/TextCache.h
@@ -29,9 +29,16 @@ public:
 	/**
 	 * Try and find an existing cached entry for the given data, or construct and cache a new entry if one cannot be found.
 	 */
+
+#if WITH_EDITORONLY_DATA
+	FText FindOrCache(const TCHAR* InTextLiteral, const FTextId& InTextId, FStringView InDevNotes);
+	FText FindOrCache(FStringView InTextLiteral, const FTextId& InTextId, FStringView InDevNotes);
+	FText FindOrCache(FString&& InTextLiteral, const FTextId& InTextId, FStringView InDevNotes);
+#else
 	FText FindOrCache(const TCHAR* InTextLiteral, const FTextId& InTextId);
 	FText FindOrCache(FStringView InTextLiteral, const FTextId& InTextId);
 	FText FindOrCache(FString&& InTextLiteral, const FTextId& InTextId);
+#endif
 
 	/**
 	 * Remove any cached entries for the given text IDs.
diff --git a/Engine/Source/Runtime/Core/Private/Internationalization/TextHistory.cpp b/Engine/Source/Runtime/Core/Private/Internationalization/TextHistory.cpp
--- a/Engine/Source/Runtime/Core/Private/Internationalization/TextHistory.cpp
+++ b/Engine/Source/Runtime/Core/Private/Internationalization/TextHistory.cpp
@@ -9,6 +9,7 @@
 #include "Misc/Guid.h"
 #include "Misc/Parse.h"
 #include "Misc/ScopeLock.h"
+#include "UObject/FortniteMainBranchObjectVersion.h"
 #include "UObject/PropertyPortFlags.h"
 
 #include "Internationalization/FastDecimalFormat.h"
@@ -681,6 +682,14 @@ void WriteTextFormatToBuffer(FString& Buffer, const FString& TokenMarker, const
 ///////////////////////////////////////
 // FTextHistory
 
+#if WITH_EDITORONLY_DATA
+const FString& FTextHistory::GetDevNotes() const
+{
+	static const FString EmptyNotes = TEXT("");
+	return EmptyNotes;
+}
+#endif
+
 uint16 FTextHistory::GetGlobalHistoryRevision() const
 {
 	UE::TScopeLock _(Mutex);
@@ -758,6 +767,15 @@ FTextHistory_Base::FTextHistory_Base(const FTextId& InTextId, FString&& InSource
 	MarkDisplayStringUpToDate();
 }
 
+#if WITH_EDITORONLY_DATA
+FTextHistory_Base::FTextHistory_Base(const FTextId& InTextId, FString&& InSourceString, FString&& InDevNotes)
+	: TextId(InTextId)
+	, SourceString(MoveTemp(InSourceString))
+	, DevNotes(MoveTemp(InDevNotes))
+{
+}
+#endif
+
 FTextId FTextHistory_Base::GetTextId() const
 {
 	return TextId;
@@ -893,6 +911,14 @@ void FTextHistory_Base::Serialize(FStructuredArchive::FRecord Record)
 		// Serialize the SourceString
 		Record << SA_VALUE(TEXT("SourceString"), SourceString);
 	}
+
+#if WITH_EDITORONLY_DATA
+	BaseArchive.UsingCustomVersion(FFortniteMainBranchObjectVersion::GUID);
+	if (BaseArchive.CustomVer(FFortniteMainBranchObjectVersion::GUID) >= FFortniteMainBranchObjectVersion::AddDevNotesToFText && !BaseArchive.IsCooking() && !BaseArchive.IsFilterEditorOnly())
+	{
+		Record << SA_VALUE(TEXT("DevNotes"), DevNotes);
+	}
+#endif
 }
 
 bool FTextHistory_Base::CanUpdateDisplayString()
@@ -914,19 +940,36 @@ void FTextHistory_Base::UpdateDisplayString()
 	LocalizedString = NewLocalizedString;
 }
 
+#if WITH_EDITORONLY_DATA
+const FString& FTextHistory_Base::GetDevNotes() const
+{
+	return DevNotes;
+}
+#endif
+
 bool FTextHistory_Base::StaticShouldReadFromBuffer(const TCHAR* Buffer)
 {
 	return TEXT_STRINGIFICATION_PEEK_MARKER(TextStringificationUtil::NsLocTextMarker)
-		|| TEXT_STRINGIFICATION_PEEK_MARKER(TextStringificationUtil::LocTextMarker);
+		|| TEXT_STRINGIFICATION_PEEK_MARKER(TextStringificationUtil::LocTextMarker)
+		|| TEXT_STRINGIFICATION_PEEK_MARKER(TextStringificationUtil::NsNoteLocTextMarker)
+		|| TEXT_STRINGIFICATION_PEEK_MARKER(TextStringificationUtil::NoteLocTextMarker);
 }
 
 const TCHAR* FTextHistory_Base::ReadFromBuffer(const TCHAR* Buffer, const TCHAR* TextNamespace, const TCHAR* PackageNamespace)
 {
 #define LOC_DEFINE_REGION
-	if (TEXT_STRINGIFICATION_PEEK_MARKER(TextStringificationUtil::NsLocTextMarker))
+	const bool bIsNsNoteLocTextMarker = TEXT_STRINGIFICATION_PEEK_MARKER(TextStringificationUtil::NsNoteLocTextMarker);
+	if (TEXT_STRINGIFICATION_PEEK_MARKER(TextStringificationUtil::NsLocTextMarker) || bIsNsNoteLocTextMarker)
 	{
-		// Parsing something of the form: NSLOCTEXT("...", "...", "...")
-		TEXT_STRINGIFICATION_SKIP_MARKER_LEN(TextStringificationUtil::NsLocTextMarker);
+		// Parsing something of the form: NSLOCTEXT("...", "...", "...") or NSNOTELOCTEXT("...", "...", "...", "...")
+		if (bIsNsNoteLocTextMarker)
+		{
+			TEXT_STRINGIFICATION_SKIP_MARKER_LEN(TextStringificationUtil::NsNoteLocTextMarker);
+		}
+		else
+		{
+			TEXT_STRINGIFICATION_SKIP_MARKER_LEN(TextStringificationUtil::NsLocTextMarker);
+		}
 
 		// Skip whitespace before the opening bracket, and then step over it
 		TEXT_STRINGIFICATION_SKIP_WHITESPACE_AND_CHAR('(');
@@ -951,6 +994,19 @@ const TCHAR* FTextHistory_Base::ReadFromBuffer(const TCHAR* Buffer, const TCHAR*
 		TEXT_STRINGIFICATION_SKIP_WHITESPACE();
 		TEXT_STRINGIFICATION_READ_QUOTED_STRING(SourceString);
 
+		// Parse the DevNote if it exists
+		if (bIsNsNoteLocTextMarker)
+		{
+			TEXT_STRINGIFICATION_SKIP_WHITESPACE_AND_CHAR(',');
+			
+			FString ReadDevNotes;
+			TEXT_STRINGIFICATION_SKIP_WHITESPACE();
+			TEXT_STRINGIFICATION_READ_QUOTED_STRING(ReadDevNotes);
+#if WITH_EDITORONLY_DATA
+			DevNotes = MoveTemp(ReadDevNotes);
+#endif
+		}
+
 		// Skip whitespace before the closing bracket, and then step over it
 		TEXT_STRINGIFICATION_SKIP_WHITESPACE_AND_CHAR(')');
 
@@ -984,11 +1040,19 @@ const TCHAR* FTextHistory_Base::ReadFromBuffer(const TCHAR* Buffer, const TCHAR*
 		return Buffer;
 	}
 	
-	if (TEXT_STRINGIFICATION_PEEK_MARKER(TextStringificationUtil::LocTextMarker))
+	const bool bIsNoteLocTextMarker = TEXT_STRINGIFICATION_PEEK_MARKER(TextStringificationUtil::NoteLocTextMarker);
+	if (TEXT_STRINGIFICATION_PEEK_MARKER(TextStringificationUtil::LocTextMarker) || bIsNoteLocTextMarker)
 	{
-		// Parsing something of the form: LOCTEXT("...", "...")
+		// Parsing something of the form: LOCTEXT("...", "...") or NOTELOCTEXT("...", "...", "...")
 		// This only exists as people sometimes do this in config files. We assume an empty namespace should be used
-		TEXT_STRINGIFICATION_SKIP_MARKER_LEN(TextStringificationUtil::LocTextMarker);
+		if (bIsNoteLocTextMarker)
+		{
+			TEXT_STRINGIFICATION_SKIP_MARKER_LEN(TextStringificationUtil::NoteLocTextMarker);
+		}
+		else
+		{
+			TEXT_STRINGIFICATION_SKIP_MARKER_LEN(TextStringificationUtil::LocTextMarker);
+		}
 
 		// Skip whitespace before the opening bracket, and then step over it
 		TEXT_STRINGIFICATION_SKIP_WHITESPACE_AND_CHAR('(');
@@ -1005,6 +1069,19 @@ const TCHAR* FTextHistory_Base::ReadFromBuffer(const TCHAR* Buffer, const TCHAR*
 		TEXT_STRINGIFICATION_SKIP_WHITESPACE();
 		TEXT_STRINGIFICATION_READ_QUOTED_STRING(SourceString);
 
+		// Parse the DevNote if it exists
+		if (bIsNoteLocTextMarker)
+		{
+			TEXT_STRINGIFICATION_SKIP_WHITESPACE_AND_CHAR(',');
+			
+			FString ReadDevNotes;
+			TEXT_STRINGIFICATION_SKIP_WHITESPACE();
+			TEXT_STRINGIFICATION_READ_QUOTED_STRING(ReadDevNotes);
+#if WITH_EDITORONLY_DATA
+			DevNotes = MoveTemp(ReadDevNotes);
+#endif
+		}
+
 		// Skip whitespace before the closing bracket, and then step over it
 		TEXT_STRINGIFICATION_SKIP_WHITESPACE_AND_CHAR(')');
 
@@ -1055,6 +1132,24 @@ bool FTextHistory_Base::WriteToBuffer(FString& Buffer, const bool bStripPackageN
 		}
 
 #define LOC_DEFINE_REGION
+#if WITH_EDITORONLY_DATA
+		const FString DevNotesToWrite = GetDevNotes().TrimStartAndEnd();
+		if (!DevNotesToWrite.IsEmpty())
+		{
+			// Produces NSNOTELOCTEXT("...", "...", "...", "...")
+			Buffer += TEXT("NSNOTELOCTEXT(\"");
+			Buffer += Namespace.ReplaceCharWithEscapedChar();
+			Buffer += TEXT("\", \"");
+			Buffer += Key.ReplaceCharWithEscapedChar();
+			Buffer += TEXT("\", \"");
+			Buffer += SourceString.ReplaceCharWithEscapedChar();
+			Buffer += TEXT("\", \"");
+			Buffer += DevNotesToWrite.ReplaceCharWithEscapedChar();
+			Buffer += TEXT("\")");
+			
+			return true;
+		}
+#endif
 		// Produces NSLOCTEXT("...", "...", "...")
 		Buffer += TEXT("NSLOCTEXT(\"");
 		Buffer += Namespace.ReplaceCharWithEscapedChar();
@@ -2358,6 +2453,13 @@ bool FTextHistory_Transform::GetHistoricNumericData(const FText& InText, FHistor
 	return FTextInspector::GetHistoricNumericData(SourceText, OutHistoricNumericData);
 }
 
+#if WITH_EDITORONLY_DATA
+const FString& FTextHistory_Transform::GetDevNotes() const
+{
+	return SourceText.GetDevNotes();
+}
+#endif
+
 ///////////////////////////////////////
 // FTextHistory_StringTableEntry
 
@@ -2467,6 +2569,21 @@ void FTextHistory_StringTableEntry::UpdateDisplayString()
 	}
 }
 
+#if WITH_EDITORONLY_DATA
+const FString& FTextHistory_StringTableEntry::GetDevNotes() const
+{
+	if (StringTableReferenceData)
+	{
+		if (const FStringTableEntryConstPtr StringTableEntryPin = StringTableReferenceData->ResolveStringTableEntry())
+		{
+			return StringTableEntryPin->GetDevNotes();
+		}
+	}
+	
+	return FTextHistory::GetDevNotes();
+}
+#endif
+
 bool FTextHistory_StringTableEntry::StaticShouldReadFromBuffer(const TCHAR* Buffer)
 {
 	return TEXT_STRINGIFICATION_PEEK_MARKER(TextStringificationUtil::LocTableMarker);
diff --git a/Engine/Source/Runtime/Core/Private/Internationalization/TextHistory.h b/Engine/Source/Runtime/Core/Private/Internationalization/TextHistory.h
--- a/Engine/Source/Runtime/Core/Private/Internationalization/TextHistory.h
+++ b/Engine/Source/Runtime/Core/Private/Internationalization/TextHistory.h
@@ -56,6 +56,8 @@ inline const auto& TextMarker = TEXT("TEXT");
 inline const auto& InvTextMarker = TEXT("INVTEXT");
 inline const auto& NsLocTextMarker = TEXT("NSLOCTEXT");
 inline const auto& LocTextMarker = TEXT("LOCTEXT");
+inline const auto& NsNoteLocTextMarker = TEXT("NSNOTELOCTEXT");
+inline const auto& NoteLocTextMarker = TEXT("NOTELOCTEXT");
 inline const auto& LocTableMarker = TEXT("LOCTABLE");
 inline const auto& LocGenNumberMarker = TEXT("LOCGEN_NUMBER");
 inline const auto& LocGenPercentMarker = TEXT("LOCGEN_PERCENT");
@@ -218,6 +220,10 @@ public:
 
 	/** Get any historic numeric format data from this history */
 	virtual bool GetHistoricNumericData(const FText& InText, FHistoricTextNumericData& OutHistoricNumericData) const { return false; }
+	
+#if WITH_EDITORONLY_DATA
+	virtual const FString& GetDevNotes() const;
+#endif
 
 	/** Update the display string if the history is out-of-date */
 	void UpdateDisplayStringIfOutOfDate();
@@ -251,6 +257,9 @@ public:
 	FTextHistory_Base() = default;
 	FTextHistory_Base(const FTextId& InTextId, FString&& InSourceString);
 	FTextHistory_Base(const FTextId& InTextId, FString&& InSourceString, FTextConstDisplayStringPtr&& InLocalizedString);
+#if WITH_EDITORONLY_DATA
+	FTextHistory_Base(const FTextId& InTextId, FString&& InSourceString, FString&& InDevNotes);
+#endif
 
 	//~ Begin FTextHistory Interface
 	OVERRIDE_TEXT_HISTORY_STRINGIFICATION;
@@ -264,6 +273,9 @@ public:
 	virtual void Serialize(FStructuredArchive::FRecord Record) override;
 	virtual bool CanUpdateDisplayString() override;
 	virtual void UpdateDisplayString() override;
+#if WITH_EDITORONLY_DATA
+	virtual const FString& GetDevNotes() const override;
+#endif
 	//~ End FTextHistory Interface
 
 private:
@@ -273,6 +285,11 @@ private:
 	FString SourceString;
 	/** The localized string (from FTextLocalizationManager, if any) */
 	FTextConstDisplayStringPtr LocalizedString;
+	
+#if WITH_EDITORONLY_DATA
+	/** The developer notes to give translators context on this string */
+	FString DevNotes;
+#endif
 };
 
 /** Base class for text histories that hold a generated display string. */
@@ -559,6 +576,9 @@ public:
 	virtual void Serialize(FStructuredArchive::FRecord Record) override;
 	virtual void GetHistoricFormatData(const FText& InText, TArray<FHistoricTextFormatData>& OutHistoricFormatData) const override;
 	virtual bool GetHistoricNumericData(const FText& InText, FHistoricTextNumericData& OutHistoricNumericData) const override;
+#if WITH_EDITORONLY_DATA
+	virtual const FString& GetDevNotes() const override;
+#endif
 	//~ End FTextHistory Interfaces
 
 private:
@@ -586,6 +606,9 @@ public:
 	virtual bool IdenticalTo(const FTextHistory& Other, const ETextIdenticalModeFlags CompareModeFlags) const override;
 	virtual void Serialize(FStructuredArchive::FRecord Record) override;
 	virtual void UpdateDisplayString() override;
+#if WITH_EDITORONLY_DATA
+	virtual const FString& GetDevNotes() const override;
+#endif
 	//~ End FTextHistory Interface
 
 	void GetTableIdAndKey(FName& OutTableId, FTextKey& OutKey) const;
diff --git a/Engine/Source/Runtime/Core/Private/UObject/FortniteMainBranchObjectVersion.cpp b/Engine/Source/Runtime/Core/Private/UObject/FortniteMainBranchObjectVersion.cpp
--- a/Engine/Source/Runtime/Core/Private/UObject/FortniteMainBranchObjectVersion.cpp
+++ b/Engine/Source/Runtime/Core/Private/UObject/FortniteMainBranchObjectVersion.cpp
@@ -12,7 +12,7 @@ TMap<FGuid, FGuid> FFortniteMainBranchObjectVersion::GetSystemGuids()
 	SystemGuids.Add(DevGuids.LANDSCAPE_MOBILE_COOK_VERSION, FGuid("32D02EF867C74B71A0D4E0FA41392732"));
 	SystemGuids.Add(DevGuids.MATERIALSHADERMAP_DERIVEDDATA_VER, FGuid("FB0EB87E970C4FB9B217E0E00942BFCB"));
 	SystemGuids.Add(DevGuids.NIAGARASHADERMAP_DERIVEDDATA_VER, FGuid("87F6AE6C269A4592B281752A582E7F13"));
-	SystemGuids.Add(DevGuids.Niagara_LatestScriptCompileVersion, FGuid("132E723405764AFAB45E9B353147C87D"));
+	SystemGuids.Add(DevGuids.Niagara_LatestScriptCompileVersion, FGuid("266FBE22ADCD4773B98124B68DBE3A4C"));
 	SystemGuids.Add(DevGuids.MaterialTranslationDDCVersion, FGuid("B79683EE9E684C13BAC2FDE8B180477F"));
 
 	// These GUIDs are in AppendRestrictedFortniteMainStreamObjectVersionGuids as they should rarely be changed.
diff --git a/Engine/Source/Runtime/Core/Public/Internationalization/GatherableTextData.h b/Engine/Source/Runtime/Core/Public/Internationalization/GatherableTextData.h
--- a/Engine/Source/Runtime/Core/Public/Internationalization/GatherableTextData.h
+++ b/Engine/Source/Runtime/Core/Public/Internationalization/GatherableTextData.h
@@ -21,6 +21,9 @@ struct FTextSourceSiteContext
 	FString SiteDescription;
 	bool IsEditorOnly;
 	bool IsOptional;
+#if WITH_EDITORONLY_DATA
+	FString DevNotes;
+#endif
 	FLocMetadataObject InfoMetaData;
 	FLocMetadataObject KeyMetaData;
 
diff --git a/Engine/Source/Runtime/Core/Public/Internationalization/Internationalization.h b/Engine/Source/Runtime/Core/Public/Internationalization/Internationalization.h
--- a/Engine/Source/Runtime/Core/Public/Internationalization/Internationalization.h
+++ b/Engine/Source/Runtime/Core/Public/Internationalization/Internationalization.h
@@ -293,6 +293,23 @@ namespace UE::Private::LocGen
  */
 #define NSLOCTEXT(InNamespace, InKey, InTextLiteral) FText::AsLocalizable_Advanced_LocText(TEXT(InNamespace), TEXT(InKey), TEXT(InTextLiteral))
 
+#if WITH_EDITORONLY_DATA
+/**
+ * Creates an FText with developer notes to provide context to translators. All parameters must be string literals. All literals will be passed through the localization system.
+ * The global LOCTEXT_NAMESPACE macro must be first set to a string literal to specify this localization key's namespace.
+ * Developer notes are stripped out of non-editor builds
+ */
+#define NOTELOCTEXT(InKey, InTextLiteral, InDevNotes) FText::AsLocalizable_Advanced_LocText(TEXT(LOCTEXT_NAMESPACE), TEXT(InKey), TEXT(InTextLiteral), TEXTVIEW(InDevNotes))
+
+/**
+ * Creates an FText with developer notes to provide context to translators. All parameters must be string literals. All literals will be passed through the localization system.
+ * Developer notes are stripped out of non-editor builds
+ */
+#define NSNOTELOCTEXT(InNamespace, InKey, InTextLiteral, InDevNotes) FText::AsLocalizable_Advanced_LocText(TEXT(InNamespace), TEXT(InKey), TEXT(InTextLiteral), TEXTVIEW(InDevNotes))
+#else
+#define NOTELOCTEXT(InKey, InTextLiteral, InDevNotes) FText::AsLocalizable_Advanced_LocText(TEXT(LOCTEXT_NAMESPACE), TEXT(InKey), TEXT(InTextLiteral))
+#define NSNOTELOCTEXT(InNamespace, InKey, InTextLiteral, InDevNotes) FText::AsLocalizable_Advanced_LocText(TEXT(InNamespace), TEXT(InKey), TEXT(InTextLiteral))
+#endif
 /**
  * Creates a culture invariant FText from the given string literal.
  */
diff --git a/Engine/Source/Runtime/Core/Public/Internationalization/InternationalizationManifest.h b/Engine/Source/Runtime/Core/Public/Internationalization/InternationalizationManifest.h
--- a/Engine/Source/Runtime/Core/Public/Internationalization/InternationalizationManifest.h
+++ b/Engine/Source/Runtime/Core/Public/Internationalization/InternationalizationManifest.h
@@ -47,6 +47,7 @@ public:
 	FString SourceLocation;
 	FName PlatformName;
 	bool bIsOptional;
+	FString DevNotes;
 
 	TSharedPtr<FLocMetadataObject> InfoMetadataObj;
 	TSharedPtr<FLocMetadataObject> KeyMetadataObj;
diff --git a/Engine/Source/Runtime/Core/Public/Internationalization/StringTableCore.h b/Engine/Source/Runtime/Core/Public/Internationalization/StringTableCore.h
--- a/Engine/Source/Runtime/Core/Public/Internationalization/StringTableCore.h
+++ b/Engine/Source/Runtime/Core/Public/Internationalization/StringTableCore.h
@@ -33,16 +33,27 @@ class FStringTableEntry
 {
 public:
 	/** Create a new string table entry using the given data */
+#if WITH_EDITORONLY_DATA
+	static FStringTableEntryRef NewStringTableEntry(FStringTableConstRef InOwnerTable, FString InSourceString, FTextId InDisplayStringId, FString InDevNotes)
+	{
+		return MakeShared<FStringTableEntry, ESPMode::ThreadSafe>(MoveTemp(InOwnerTable), MoveTemp(InSourceString), MoveTemp(InDisplayStringId), MoveTemp(InDevNotes));
+	}
+#else
 	static FStringTableEntryRef NewStringTableEntry(FStringTableConstRef InOwnerTable, FString InSourceString, FTextId InDisplayStringId)
 	{
 		return MakeShared<FStringTableEntry, ESPMode::ThreadSafe>(MoveTemp(InOwnerTable), MoveTemp(InSourceString), MoveTemp(InDisplayStringId));
 	}
+#endif
 
 	/** Default constructor */
 	CORE_API FStringTableEntry();
 
 	/** Create a new string table entry using the given data */
+#if WITH_EDITORONLY_DATA
+	CORE_API FStringTableEntry(FStringTableConstRef InOwnerTable, FString InSourceString, FTextId InDisplayStringId, FString InDevNotes);
+#else
 	CORE_API FStringTableEntry(FStringTableConstRef InOwnerTable, FString InSourceString, FTextId InDisplayStringId);
+#endif
 
 	/** @return true if this entry is currently owned by a string table, false if it's been disowned (and should be re-cached) */
 	CORE_API bool IsOwned() const;
@@ -61,6 +72,11 @@ public:
 
 	/** Get the display string ID of this string table entry */
 	CORE_API FTextId GetDisplayStringId() const;
+	
+#if WITH_EDITORONLY_DATA
+	/** Get the translation context notes */
+	const FString& GetDevNotes() const { return DevNotes; }
+#endif
 
 	/** Get the placeholder source string to use for string table entries that are missing */
 	static CORE_API const FString& GetPlaceholderSourceString();
@@ -74,6 +90,11 @@ private:
 
 	/** The display string ID of this entry */
 	FTextId DisplayStringId;
+	
+#if WITH_EDITORONLY_DATA
+	/** Notes from the developer to provide additional context to translators */
+	FString DevNotes;
+#endif
 };
 
 /** String table implementation. Holds Key->SourceString pairs of text. */
@@ -118,12 +139,23 @@ public:
 
 	/** Get the source string used by the given entry (if any) */
 	CORE_API bool GetSourceString(const FTextKey& InKey, FString& OutSourceString) const;
-
+	
+#if WITH_EDITORONLY_DATA
+	/** Set the source string and dev notes used by the given entry (will replace any existing data for that entry) */
+	CORE_API void SetSourceString(const FTextKey& InKey, const FString& InSourceString, const FString& InDevNotes);
+#else
 	/** Set the source string used by the given entry (will replace any existing data for that entry) */
 	CORE_API void SetSourceString(const FTextKey& InKey, const FString& InSourceString);
+#endif
 
 	/** Remove the given entry (including its meta-data) */
 	CORE_API void RemoveSourceString(const FTextKey& InKey);
+	
+#if WITH_EDITORONLY_DATA
+	/** Enumerate all source strings in the table with DevNotes passed in as the third param. Return true from the enumerator to continue, or false to stop */
+	CORE_API void EnumerateSourceStrings(const TFunctionRef<bool(const FString&, const FString&, const FString&)>& InEnumerator) const;
+	CORE_API void EnumerateKeysAndSourceStrings(const TFunctionRef<bool(const FTextKey&, const FString&, const FString&)>& InEnumerator) const;
+#endif
 
 	/** Enumerate all source strings in the table. Return true from the enumerator to continue, or false to stop */
 	CORE_API void EnumerateSourceStrings(const TFunctionRef<bool(const FString&, const FString&)>& InEnumerator) const;
diff --git a/Engine/Source/Runtime/Core/Public/Internationalization/StringTableRegistry.h b/Engine/Source/Runtime/Core/Public/Internationalization/StringTableRegistry.h
--- a/Engine/Source/Runtime/Core/Public/Internationalization/StringTableRegistry.h
+++ b/Engine/Source/Runtime/Core/Public/Internationalization/StringTableRegistry.h
@@ -54,8 +54,12 @@ public:
 	/** Internal function called by LOCTABLE_FROMFILE_X to create and register a FStringTable instance that has been populated from a file */
 	CORE_API void Internal_LocTableFromFile(const FName InTableId, const FTextKey& InNamespace, const FString& InFilePath, const FString& InRootPath);
 
-	/** Internal function called by LOCTABLE_SETSTRING to set the entry denoted by the given key to the given source string, within the given string table (table must have been registered already) */
+	/** Internal function called by LOCTABLE_SETSTRING and LOCTABLE_SETNOTESTRING to set the entry denoted by the given key to the given source string, within the given string table (table must have been registered already) */
+#if WITH_EDITORONLY_DATA
+	CORE_API void Internal_SetLocTableEntry(const FName InTableId, const FTextKey& InKey, const FString& InSourceString, const FString& InDevNotes = TEXT(""));
+#else
 	CORE_API void Internal_SetLocTableEntry(const FName InTableId, const FTextKey& InKey, const FString& InSourceString);
+#endif
 
 	/** Internal function called by LOCTABLE_SETMETA to set meta-data for the entry denoted by the given key, within the given string table (table must have been registered already) */
 	CORE_API void Internal_SetLocTableEntryMetaData(const FName InTableId, const FTextKey& InKey, const FName InMetaDataId, const FString& InMetaData);
@@ -117,6 +121,14 @@ private:
 #define LOCTABLE_SETSTRING(ID, KEY, SRC) \
 	FStringTableRegistry::Get().Internal_SetLocTableEntry(TEXT(ID), TEXT(KEY), TEXT(SRC))
 
+#if WITH_EDITORONLY_DATA
+/** Add a string table entry with the given key, source string, and developer translation context note. */
+#define LOCTABLE_SETNOTESTRING(ID, KEY, SRC, NOTE) \
+	FStringTableRegistry::Get().Internal_SetLocTableEntry(TEXT(ID), TEXT(KEY), TEXT(SRC), TEXT(NOTE))
+#else
+#define LOCTABLE_SETNOTESTRING(ID, KEY, SRC, NOTE) LOCTABLE_SETSTRING(ID, KEY, SRC)
+#endif
+
 /** Add a meta-data for the entry with the given key and source string. */
 #define LOCTABLE_SETMETA(ID, KEY, METAID, META) \
 	FStringTableRegistry::Get().Internal_SetLocTableEntryMetaData(TEXT(ID), TEXT(KEY), TEXT(METAID), TEXT(META))
diff --git a/Engine/Source/Runtime/Core/Public/Internationalization/Text.h b/Engine/Source/Runtime/Core/Public/Internationalization/Text.h
--- a/Engine/Source/Runtime/Core/Public/Internationalization/Text.h
+++ b/Engine/Source/Runtime/Core/Public/Internationalization/Text.h
@@ -532,10 +532,17 @@ public:
 	 * When using it to dynamically reference an EXISTING FText (such as known LOCTEXT in code) it can be thought of as an alternative to FText::FindTextInLiveTable_Advanced that also handles untranslated text.
 	 * Note: Direct dynamic references to FText are EXTREMELY FRAGILE, and you may want to use a string table instead!
 	 */
+#if WITH_EDITORONLY_DATA
+	[[nodiscard]] static CORE_API FText AsLocalizable_Advanced(const FTextKey& Namespace, const FTextKey& Key, const TCHAR* String, FStringView InDevNotes = TEXTVIEW(""));
+	[[nodiscard]] static CORE_API FText AsLocalizable_Advanced(const FTextKey& Namespace, const FTextKey& Key, FStringView String, FStringView InDevNotes = TEXTVIEW(""));
+	[[nodiscard]] static CORE_API FText AsLocalizable_Advanced(const FTextKey& Namespace, const FTextKey& Key, FString&& String, FStringView InDevNotes = TEXTVIEW(""));
+	[[nodiscard]] static CORE_API FText AsLocalizable_Advanced_LocText(const TCHAR* Namespace, const TCHAR* Key, const TCHAR* String, FStringView InDevNotes = TEXTVIEW("")); // For use by LOCTEXT to avoid constructing temp FTextKeys in every caller
+#else
 	[[nodiscard]] static CORE_API FText AsLocalizable_Advanced(const FTextKey& Namespace, const FTextKey& Key, const TCHAR* String);
 	[[nodiscard]] static CORE_API FText AsLocalizable_Advanced(const FTextKey& Namespace, const FTextKey& Key, FStringView String);
 	[[nodiscard]] static CORE_API FText AsLocalizable_Advanced(const FTextKey& Namespace, const FTextKey& Key, FString&& String);
 	[[nodiscard]] static CORE_API FText AsLocalizable_Advanced_LocText(const TCHAR* Namespace, const TCHAR* Key, const TCHAR* String); // For use by LOCTEXT to avoid constructing temp FTextKeys in every caller
+#endif
 
 	/**
 	 * === !! This is an ADVANCED function. USE WITH CAUTION !! ===
@@ -824,6 +831,13 @@ public:
 	 * Constructs a new FText with the SourceString of the specified text but with the specified namespace and key
 	 */
 	[[nodiscard]] static CORE_API FText ChangeKey( const FTextKey& Namespace, const FTextKey& Key, const FText& Text );
+
+	/**
+	 * Constructs a new FText with the SourceString of the specified text but with the specified namespace, key, and developer translation context notes
+	 */
+	[[nodiscard]] static CORE_API FText ChangeKey( const FTextKey& Namespace, const FTextKey& Key, const FText& Text, const FString& DevNotes );
+
+	CORE_API const FString& GetDevNotes() const;
 #endif
 
 	/////////////////////////////////////////////////////
@@ -860,6 +874,10 @@ private:
 	CORE_API FText( FName InTableId, const FTextKey& InKey, const EStringTableLoadingPolicy InLoadingPolicy );
 
 	CORE_API FText( FString&& InSourceString, const FTextKey& InNamespace, const FTextKey& InKey, uint32 InFlags=0 );
+	
+#if WITH_EDITORONLY_DATA
+	CORE_API FText( FString&& InSourceString, const FTextKey& InNamespace, const FTextKey& InKey, FString InDevNotes, uint32 InFlags=0 );
+#endif
 
 	static CORE_API void SerializeText( FArchive& Ar, FText& Value );
 	static CORE_API void SerializeText(FStructuredArchive::FSlot Slot, FText& Value);
diff --git a/Engine/Source/Runtime/Core/Public/UObject/FortniteMainBranchObjectVersions.inl b/Engine/Source/Runtime/Core/Public/UObject/FortniteMainBranchObjectVersions.inl
--- a/Engine/Source/Runtime/Core/Public/UObject/FortniteMainBranchObjectVersions.inl
+++ b/Engine/Source/Runtime/Core/Public/UObject/FortniteMainBranchObjectVersions.inl
@@ -790,5 +790,8 @@ FORTNITE_MAIN_VERSION(FixPublicEntityComponentsAndSubobjects, 258)
 // Added override flag for Physical Material
 FORTNITE_MAIN_VERSION(AddPhysMatOverrideFlag, 259)
 
+// Added context DevNotes for translators to FText
+FORTNITE_MAIN_VERSION(AddDevNotesToFText, 260)
+
 // -----<new versions can be added above this line>-------------------------------------------------
 #undef FORTNITE_MAIN_VERSION
diff --git a/Engine/Source/Runtime/Core/Tests/Internationalization/TextTest.cpp b/Engine/Source/Runtime/Core/Tests/Internationalization/TextTest.cpp
--- a/Engine/Source/Runtime/Core/Tests/Internationalization/TextTest.cpp
+++ b/Engine/Source/Runtime/Core/Tests/Internationalization/TextTest.cpp
@@ -1224,11 +1224,41 @@ TEST_CASE_NAMED(FTextStringificationTest, "System::Core::Misc::TextStringificati
 	{
 		LOCTABLE_NEW("Core.Tests.TextFormatTest", "Core.Tests.TextFormatTest");
 		LOCTABLE_SETSTRING("Core.Tests.TextFormatTest", "TextStringificationTest_Lorem", "Lorem");
+		LOCTABLE_SETNOTESTRING("Core.Tests.TextFormatTest", "TextStringificationTest_Ipsum", "Ipsum", "A note about Ipsum");
 	}
 
 	TEST(NSLOCTEXT("Core.Tests.TextFormatTest", "TextStringificationTest_Lorem", "Lorem"), "NSLOCTEXT(\"Core.Tests.TextFormatTest\", \"TextStringificationTest_Lorem\", \"Lorem\")");
 	TEST(LOCTEXT("TextStringificationTest_Lorem", "Lorem"), "NSLOCTEXT(\"Core.Tests.TextFormatTest\", \"TextStringificationTest_Lorem\", \"Lorem\")");
 	TEST(LOCTABLE("Core.Tests.TextFormatTest", "TextStringificationTest_Lorem"), "LOCTABLE(\"Core.Tests.TextFormatTest\", \"TextStringificationTest_Lorem\")");
+	
+#if WITH_EDITORONLY_DATA
+	TEST(NSNOTELOCTEXT("Core.Tests.TextFormatTest", "TextStringificationTest_Ipsum", "Ipsum", "A note about Ipsum"), "NSNOTELOCTEXT(\"Core.Tests.TextFormatTest\", \"TextStringificationTest_Ipsum\", \"Ipsum\", \"A note about Ipsum\")");
+	TEST(NOTELOCTEXT("TextStringificationTest_Ipsum", "Ipsum", "A note about Ipsum"), "NSNOTELOCTEXT(\"Core.Tests.TextFormatTest\", \"TextStringificationTest_Ipsum\", \"Ipsum\", \"A note about Ipsum\")");
+#else
+	// Notes are stripped outside of editor
+	TEST(NSNOTELOCTEXT("Core.Tests.TextFormatTest", "TextStringificationTest_Ipsum", "Ipsum", "A note about Ipsum"), "NSLOCTEXT(\"Core.Tests.TextFormatTest\", \"TextStringificationTest_Ipsum\", \"Ipsum\")");
+	TEST(NOTELOCTEXT("TextStringificationTest_Ipsum", "Ipsum", "A note about Ipsum"), "NSLOCTEXT(\"Core.Tests.TextFormatTest\", \"TextStringificationTest_Ipsum\", \"Ipsum\")");
+#endif
+
+#if WITH_EDITORONLY_DATA
+	// Test that DevNotes can be retrieved
+	{
+		const FText IpsumText = NOTELOCTEXT("TextStringificationTest_Ipsum", "Ipsum", "A note about Ipsum");
+		const FString IpsumTextDevNotes = IpsumText.GetDevNotes();
+		if (!IpsumTextDevNotes.Equals(TEXT("A note about Ipsum"), ESearchCase::CaseSensitive))
+		{
+			FAIL_CHECK(FString::Printf(TEXT("Text DevNotes import failure (from C++): Expected DevNotes 'A note about Ipsum' but got '%s'."), *IpsumTextDevNotes));
+		}
+
+		const FText NSIpsumText = NSNOTELOCTEXT("Core.Tests.TextFormatTest", "TextStringificationTest_Ipsum", "Ipsum", "A note about Ipsum");
+		const FString NSIpsumTextDevNotes = NSIpsumText.GetDevNotes();
+		if (!NSIpsumTextDevNotes.Equals(TEXT("A note about Ipsum"), ESearchCase::CaseSensitive))
+		{
+			FAIL_CHECK(FString::Printf(TEXT("Text DevNotes import failure (from C++): Expected DevNotes 'A note about Ipsum' but got '%s'."), *NSIpsumTextDevNotes));
+		}
+	}
+#endif
+
 	TEST(INVTEXT("DummyText"), "INVTEXT(\"DummyText\")");
 	if (GIsEditor)
 	{
diff --git a/Engine/Source/Runtime/CoreUObject/Private/Serialization/PropertyLocalizationDataGathering.cpp b/Engine/Source/Runtime/CoreUObject/Private/Serialization/PropertyLocalizationDataGathering.cpp
--- a/Engine/Source/Runtime/CoreUObject/Private/Serialization/PropertyLocalizationDataGathering.cpp
+++ b/Engine/Source/Runtime/CoreUObject/Private/Serialization/PropertyLocalizationDataGathering.cpp
@@ -558,7 +558,7 @@ void FPropertyLocalizationDataGatherer::GatherLocalizationDataFromChildTextPrope
 
 void FPropertyLocalizationDataGatherer::GatherTextInstance(const FText& Text, const FString& Description, const bool bIsEditorOnly)
 {
-	auto AddGatheredText = [this, &Description](const FString& InNamespace, const FString& InKey, const FTextSourceData& InSourceData, const bool InIsEditorOnly)
+	auto AddGatheredText = [this, &Description](const FString& InNamespace, const FString& InKey, const FTextSourceData& InSourceData, const bool InIsEditorOnly, const FString& InDevNotes)
 	{
 		FGatherableTextData* GatherableTextData = GatherableTextDataArray.FindByPredicate([&](const FGatherableTextData& Candidate)
 		{
@@ -593,6 +593,9 @@ void FPropertyLocalizationDataGatherer::GatherTextInstance(const FText& Text, co
 				SourceSiteContext.SiteDescription = Description;
 				SourceSiteContext.IsEditorOnly = InIsEditorOnly;
 				SourceSiteContext.IsOptional = false;
+#if WITH_EDITORONLY_DATA
+				SourceSiteContext.DevNotes = InDevNotes;
+#endif
 			}
 		}
 	};
@@ -621,7 +624,12 @@ void FPropertyLocalizationDataGatherer::GatherTextInstance(const FText& Text, co
 
 	// Always include the text without its package localization ID
 	const FString CleanNamespace = TextNamespaceUtil::StripPackageNamespace(Namespace);
-	AddGatheredText(CleanNamespace, Key, SourceData, bIsEditorOnly);
+#if WITH_EDITORONLY_DATA
+	AddGatheredText(CleanNamespace, Key, SourceData, bIsEditorOnly, Text.GetDevNotes());
+#else
+	static const FString EmptyString = TEXT("");
+	AddGatheredText(CleanNamespace, Key, SourceData, bIsEditorOnly, EmptyString);
+#endif
 
 #if USE_STABLE_LOCALIZATION_KEYS
 	// Sanity check that the text we gathered has the expected package localization ID
@@ -747,9 +755,23 @@ private:
 				SerializeExpr(iCode, Ar);
 				const FString TextNamespace = MoveTemp(LastParsedString);
 
+#if WITH_EDITORONLY_DATA
+				FString DevNotes;
+				const bool bIsCooked = PropertyLocalizationDataGatherer.GetPackage()->HasAnyPackageFlags(PKG_Cooked);
+				if (!bIsCooked)
+				{
+					SerializeExpr(iCode, Ar);
+					DevNotes = MoveTemp(LastParsedString);
+				}
+#endif
+
 				bIsParsingText = false;
 
+#if WITH_EDITORONLY_DATA
+				const FText TextInstance = FText::AsLocalizable_Advanced(TextNamespace, TextKey, MoveTemp(SourceString), DevNotes);
+#else
 				const FText TextInstance = FText::AsLocalizable_Advanced(TextNamespace, TextKey, MoveTemp(SourceString));
+#endif
 				if (!PropertyLocalizationDataGatherer.IsDefaultTextInstance(TextInstance))
 				{
 					PropertyLocalizationDataGatherer.GatherTextInstance(TextInstance, FString::Printf(TEXT("%s [Script Bytecode]"), SourceDescription), bTreatAsEditorOnlyData);
diff --git a/Engine/Source/Runtime/CoreUObject/Private/UObject/ScriptCore.cpp b/Engine/Source/Runtime/CoreUObject/Private/UObject/ScriptCore.cpp
--- a/Engine/Source/Runtime/CoreUObject/Private/UObject/ScriptCore.cpp
+++ b/Engine/Source/Runtime/CoreUObject/Private/UObject/ScriptCore.cpp
@@ -3481,8 +3481,19 @@ DEFINE_FUNCTION(UObject::execTextConst)
 			
 			FString Namespace;
 			Stack.Step(Stack.Object, &Namespace);
+#if WITH_EDITORONLY_DATA
+			if (Stack.Object && Stack.Object->GetPackage() && !Stack.Object->GetPackage()->HasAnyPackageFlags(PKG_Cooked))
+			{
+				FString DevNotes;
+				Stack.Step(Stack.Object, &DevNotes);
 
-			*(FText*)RESULT_PARAM = FText::AsLocalizable_Advanced(Namespace, KeyString, MoveTemp(SourceString));
+				*(FText*)RESULT_PARAM = FText::AsLocalizable_Advanced(Namespace, KeyString, MoveTemp(SourceString), DevNotes);
+			}
+			else
+#endif
+			{
+				*(FText*)RESULT_PARAM = FText::AsLocalizable_Advanced(Namespace, KeyString, MoveTemp(SourceString));
+			}
 		}
 		break;
 
diff --git a/Engine/Source/Runtime/CoreUObject/Public/Serialization/PropertyLocalizationDataGathering.h b/Engine/Source/Runtime/CoreUObject/Public/Serialization/PropertyLocalizationDataGathering.h
--- a/Engine/Source/Runtime/CoreUObject/Public/Serialization/PropertyLocalizationDataGathering.h
+++ b/Engine/Source/Runtime/CoreUObject/Public/Serialization/PropertyLocalizationDataGathering.h
@@ -174,6 +174,8 @@ public:
 		return AllObjectsInPackage.Contains(Object);
 	}
 
+	const UPackage* GetPackage() const { return Package; }
+
 private:
 	COREUOBJECT_API const FGatherableFieldsForType& CacheGatherableFieldsForType(const UStruct* InType);
 	COREUOBJECT_API bool CanGatherFromInnerProperty(const FProperty* InInnerProperty);
diff --git a/Engine/Source/Runtime/Engine/Private/Internationalization/StringTable.cpp b/Engine/Source/Runtime/Engine/Private/Internationalization/StringTable.cpp
--- a/Engine/Source/Runtime/Engine/Private/Internationalization/StringTable.cpp
+++ b/Engine/Source/Runtime/Engine/Private/Internationalization/StringTable.cpp
@@ -43,7 +43,7 @@ namespace UE::StringTable
 			return *GatherableTextData;
 		};
 
-		StringTable->EnumerateSourceStrings([&SourceLocation, &StringTable, &FindOrAddTextData](const FString& InKey, const FString& InSourceString) -> bool
+		StringTable->EnumerateSourceStrings([&SourceLocation, &StringTable, &FindOrAddTextData](const FString& InKey, const FString& InSourceString, const FString& InDevNotes) -> bool
 		{
 			if (!InSourceString.IsEmpty())
 			{
@@ -54,6 +54,7 @@ namespace UE::StringTable
 				SourceSiteContext.SiteDescription = SourceLocation;
 				SourceSiteContext.IsEditorOnly = false;
 				SourceSiteContext.IsOptional = false;
+				SourceSiteContext.DevNotes = InDevNotes;
 
 				StringTable->EnumerateMetaData(InKey, [&](const FName InMetaDataId, const FString& InMetaData)
 				{
