Respect bSkipSourceCheck for chunked LocRes generated during cook

#Synced-CL 51753475

diff --git a/Engine/Source/Editor/UnrealEd/Private/LocalizationChunkDataGenerator.cpp b/Engine/Source/Editor/UnrealEd/Private/LocalizationChunkDataGenerator.cpp
--- a/Engine/Source/Editor/UnrealEd/Private/LocalizationChunkDataGenerator.cpp
+++ b/Engine/Source/Editor/UnrealEd/Private/LocalizationChunkDataGenerator.cpp
@@ -83,7 +83,7 @@ void FLocalizationChunkDataGenerator::GenerateChunkDataFiles(const int32 InChunk
 	ConditionalCacheLocalizationTargetData();
 
 	// We can skip this if we're not actually chunking any localization data
-	if (CachedLocalizationTargetHelpers.Num() == 0)
+	if (CachedLocalizationTargetData.Num() == 0)
 	{
 		return;
 	}
@@ -92,8 +92,10 @@ void FLocalizationChunkDataGenerator::GenerateChunkDataFiles(const int32 InChunk
 	const FString LocalizationContentRoot = (InSandboxFile->GetSandboxDirectory() / InSandboxFile->GetGameSandboxDirectoryName() / TEXT("Content") / TEXT("Localization")).Replace(TEXT("[Platform]"), *InPlatformName);
 	const FString LocalizationMetadataRoot = (InSandboxFile->GetSandboxDirectory() / InSandboxFile->GetGameSandboxDirectoryName() / TEXT("Metadata") / TEXT("Localization")).Replace(TEXT("[Platform]"), *InPlatformName);
 
-	for (const TSharedPtr<FLocTextHelper>& SourceLocTextHelper : CachedLocalizationTargetHelpers)
+	for (const FCachedLocalizationTargetData& CachedLocalizationTargetDataEntry : CachedLocalizationTargetData)
 	{
+		const TSharedPtr<FLocTextHelper>& SourceLocTextHelper = CachedLocalizationTargetDataEntry.LocTextHelper;
+
 		// If this target has no helper, then it was either an invalid target or failed to load when caching - skip it here
 		if (!SourceLocTextHelper)
 		{
@@ -217,7 +219,7 @@ void FLocalizationChunkDataGenerator::GenerateChunkDataFiles(const int32 InChunk
 
 			FTextLocalizationResource ChunkLocRes;
 			TMap<FName, TSharedRef<FTextLocalizationResource>> UnusedPerPlatformLocRes;
-			if (FTextLocalizationResourceGenerator::GenerateLocRes(ChunkLocTextHelper, CultureToCook, EGenerateLocResFlags::None, ChunkLocResFilename, ChunkLocRes, UnusedPerPlatformLocRes) && ChunkLocRes.SaveToFile(ChunkLocResFilename))
+			if (FTextLocalizationResourceGenerator::GenerateLocRes(ChunkLocTextHelper, CultureToCook, CachedLocalizationTargetDataEntry.GenerateLocResFlags, ChunkLocResFilename, ChunkLocRes, UnusedPerPlatformLocRes) && ChunkLocRes.SaveToFile(ChunkLocResFilename))
 			{
 				OutChunkFilenames.Add(ChunkLocResFilename);
 			}
@@ -238,15 +240,61 @@ void FLocalizationChunkDataGenerator::ConditionalCacheLocalizationTargetData()
 	}
 
 	// Already cached?
-	if (LocalizationTargetsToChunk.Num() == CachedLocalizationTargetHelpers.Num())
+	if (LocalizationTargetsToChunk.Num() == CachedLocalizationTargetData.Num())
 	{
 		return;
 	}
 
-	CachedLocalizationTargetHelpers.Reset(LocalizationTargetsToChunk.Num());
+	auto DoesTargetAllowStaleTranslations =
+		[](const FString& LocalizationTarget) -> bool
+		{
+			// Note: We have to try and load this from the gather config data, as this information is not mirrored in the LocMeta file
+			const FString PotentialConfigPaths[] = {
+				FPaths::Combine(FPaths::ProjectConfigDir(), TEXTVIEW("Localization"), LocalizationTarget + TEXTVIEW("_Compile.ini")),
+				FPaths::Combine(FPaths::ProjectConfigDir(), TEXTVIEW("Localization"), LocalizationTarget + TEXTVIEW(".ini")),
+			};
+
+			for (const FString& PotentialConfigPath : PotentialConfigPaths)
+			{
+				if (!FPaths::FileExists(PotentialConfigPath))
+				{
+					continue;
+				}
+
+				const FConfigFile* LocalizationConfig = GConfig->Find(PotentialConfigPath); // Note: Find also loads if needed
+				if (!LocalizationConfig)
+				{
+					continue;
+				}
+
+				UE_LOG(LogLocalizationChunkDataGenerator, Log, TEXT("Loaded config for localization target '%s' (from: '%s')"), *LocalizationTarget, *PotentialConfigPath);
+
+				// Look for bSkipSourceCheck in any section of this config file
+				bool bSkipSourceCheck = false;
+				for (const TTuple<FString, FConfigSection>& SectionPair : *LocalizationConfig)
+				{
+					if (SectionPair.Value.GetBool(TEXT("bSkipSourceCheck"), bSkipSourceCheck))
+					{
+						break;
+					}
+				}
+
+				// Clean-up the config we loaded above
+				GConfig->UnloadFile(PotentialConfigPath);
+
+				return bSkipSourceCheck;
+			}
+
+			UE_LOG(LogLocalizationChunkDataGenerator, Log, TEXT("Could not find any config for localization target '%s'"), *LocalizationTarget);
+
+			return false;
+		};
+
+	CachedLocalizationTargetData.Reset(LocalizationTargetsToChunk.Num());
 	for (const FString& LocalizationTarget : LocalizationTargetsToChunk)
 	{
-		TSharedPtr<FLocTextHelper>& SourceLocTextHelper = CachedLocalizationTargetHelpers.AddDefaulted_GetRef();
+		FCachedLocalizationTargetData& CachedLocalizationTargetDataEntry = CachedLocalizationTargetData.AddDefaulted_GetRef();
+		CachedLocalizationTargetDataEntry.GenerateLocResFlags = DoesTargetAllowStaleTranslations(LocalizationTarget) ? EGenerateLocResFlags::AllowStaleTranslations : EGenerateLocResFlags::None;
 
 		// Does this target exist?
 		// Note: We only allow game localization targets to be chunked, and the layout is assumed to follow our standard pattern (as used by the localization dashboard and FLocTextHelper)
@@ -289,12 +337,16 @@ void FLocalizationChunkDataGenerator::ConditionalCacheLocalizationTargetData()
 		}
 
 		// Load the data for this target
-		SourceLocTextHelper = MakeShared<FLocTextHelper>(SourceRootPath, FString::Printf(TEXT("%s.manifest"), *LocalizationTarget), FString::Printf(TEXT("%s.archive"), *LocalizationTarget), SourceNativeCulture, SourceForeignCulturesToCook, nullptr);
+		CachedLocalizationTargetDataEntry.LocTextHelper = MakeShared<FLocTextHelper>(SourceRootPath, FString::Printf(TEXT("%s.manifest"), *LocalizationTarget), FString::Printf(TEXT("%s.archive"), *LocalizationTarget), SourceNativeCulture, SourceForeignCulturesToCook, nullptr);
 		{
 			FText LoadError;
-			if (!SourceLocTextHelper->LoadAll(ELocTextHelperLoadFlags::Load, &LoadError))
+			if (CachedLocalizationTargetDataEntry.LocTextHelper->LoadAll(ELocTextHelperLoadFlags::Load, &LoadError))
+			{
+				UE_LOG(LogLocalizationChunkDataGenerator, Log, TEXT("Loaded data for localization target '%s'"), *LocalizationTarget);
+			}
+			else
 			{
-				SourceLocTextHelper.Reset();
+				CachedLocalizationTargetDataEntry.LocTextHelper.Reset();
 				UE_LOG(LogLocalizationChunkDataGenerator, Error, TEXT("Failed to load data for localization target '%s' when chunking localization data: %s"), *LocalizationTarget, *LoadError.ToString());
 				continue;
 			}
diff --git a/Engine/Source/Editor/UnrealEd/Private/LocalizationChunkDataGenerator.h b/Engine/Source/Editor/UnrealEd/Private/LocalizationChunkDataGenerator.h
--- a/Engine/Source/Editor/UnrealEd/Private/LocalizationChunkDataGenerator.h
+++ b/Engine/Source/Editor/UnrealEd/Private/LocalizationChunkDataGenerator.h
@@ -9,6 +9,8 @@
 #include "HAL/Platform.h"
 #include "Templates/SharedPointer.h"
 
+enum class EGenerateLocResFlags : uint8;
+
 class FLocTextHelper;
 class FName;
 class FSandboxPlatformFile;
@@ -39,8 +41,13 @@ private:
 	/** Complete list of cultures to cook data for, including inferred parent cultures */
 	TArray<FString> AllCulturesToCook;
 
-	/** Cached localization target helpers, to avoid redundant work for each chunk */
-	TArray<TSharedPtr<FLocTextHelper>> CachedLocalizationTargetHelpers;
+	struct FCachedLocalizationTargetData
+	{
+		TSharedPtr<FLocTextHelper> LocTextHelper;
+		EGenerateLocResFlags GenerateLocResFlags = {};
+	};
+	/** Cached localization target data, to avoid redundant work for each chunk */
+	TArray<FCachedLocalizationTargetData> CachedLocalizationTargetData;
 
 	/** Array of potential content roots, including plugins that aren't currently loaded */
 	TArray<FString> AllPotentialContentRoots;
