PackFactory: Fix unsafe path traversal outside the extract dir during the upack extraction.

#rnx

#Synced-CL 51816408

diff --git a/Engine/Source/Editor/UnrealEd/Private/Factories/PackFactory.cpp b/Engine/Source/Editor/UnrealEd/Private/Factories/PackFactory.cpp
--- a/Engine/Source/Editor/UnrealEd/Private/Factories/PackFactory.cpp
+++ b/Engine/Source/Editor/UnrealEd/Private/Factories/PackFactory.cpp
@@ -33,6 +33,7 @@
 #include "SourceCodeNavigation.h"
 #include "Misc/HotReloadInterface.h"
 #include "Misc/AES.h"
+#include "Hash/xxhash.h"
 #include "GameProjectGenerationModule.h"
 #include "Dialogs/SOutputLogDialog.h"
 #include "Templates/UniquePtr.h"
@@ -226,6 +227,29 @@ namespace PackFactoryHelper
 		MemReader << FileContents;
 	}
 
+	// Utility function to return a safe destination path for extracting a pak entry.
+	// By default it just returns the normalized combined path of NormalizedDestRoot and RelativeDestFilename.
+	// If the normalized path would escape NormalizedDestRoot (e.g. via path traversal), the file is redirected to a quarantine subfolder
+	// generated from a hash of the original EntryFileName. 
+	FString MakeSafeDestPath(const FString& NormalizedDestRoot, const FString& RelativeDestFilename, const FString& EntryFilename)
+	{
+		FString DestFilename = FPaths::ConvertRelativePathToFull(NormalizedDestRoot / RelativeDestFilename);
+
+		if (!FPaths::IsUnderDirectory(DestFilename, NormalizedDestRoot))
+		{
+			FString EntryPath = FPaths::GetPath(EntryFilename);
+			uint64 EntryHash = FXxHash64::HashBuffer(*EntryPath, EntryPath.Len() * sizeof(TCHAR)).Hash;
+			FString HashDir = FString::Printf(TEXT("%016llx"), EntryHash);
+			FString SafeDestFilename = FPaths::ConvertRelativePathToFull(NormalizedDestRoot / TEXT("Quarantined") / HashDir / FPaths::GetCleanFilename(EntryFilename));
+			FMessageLog("AssetTools").Error(FText::Format(
+				NSLOCTEXT("PackFactory", "PathTraversalDetected", "Path traversal detected for entry \"{0}\": \"{1}\" is outside destination root. Quarantined to \"{2}\"."),
+				FText::FromString(EntryFilename), FText::FromString(DestFilename), FText::FromString(SafeDestFilename)));
+			DestFilename = SafeDestFilename;
+		}
+
+		return DestFilename;
+	}
+
 	struct FPackConfigParameters
 	{
 		FPackConfigParameters()
@@ -408,6 +432,7 @@ UObject* UPackFactory::FactoryCreateBinary
 		{
 			ContentDestinationRoot /= PakFile.GetMountPoint().RightChop(ChopIndex + ContentFolder.Len());
 		}
+		ContentDestinationRoot = FPaths::ConvertRelativePathToFull(ContentDestinationRoot);
 
 		TArray<uint8> CopyBuffer;
 		TArray<uint8> PersistentCompressionBuffer;
@@ -520,6 +545,10 @@ UObject* UPackFactory::FactoryCreateBinary
 			}
 		}
 
+		const FString NormalizedSourceRoot = SourceModuleInfo.ModuleSourcePath.IsEmpty()
+			? FString()
+			: FPaths::ConvertRelativePathToFull(SourceModuleInfo.ModuleSourcePath);
+
 		// Process everything else and copy out to disk
 		for (FPakFile::FPakEntryIterator It(PakFile); It; ++It, ++FileCount)
 		{
@@ -546,6 +575,13 @@ UObject* UPackFactory::FactoryCreateBinary
 			{
 				if (EntryFilename->StartsWith(TEXT("Source/")) || EntryFilename->Contains(TEXT("/Source/")))
 				{
+					if (NormalizedSourceRoot.IsEmpty())
+					{
+						UE_LOG(LogPackFactory, Error, TEXT("Skipping source entry \"%s\": no Runtime module source path available."), **EntryFilename);
+						++ErrorCount;
+						continue;
+					}
+
 					FString DestFilename = *EntryFilename;
 					if (DestFilename.StartsWith(TEXT("Source/")))
 					{
@@ -560,7 +596,7 @@ UObject* UPackFactory::FactoryCreateBinary
 						}
 					}
 
-					DestFilename = SourceModuleInfo.ModuleSourcePath / DestFilename;
+					DestFilename = PackFactoryHelper::MakeSafeDestPath(NormalizedSourceRoot, DestFilename, *EntryFilename);
 					UE_LOG(LogPackFactory, Log, TEXT("%s (%ld) -> %s"), **EntryFilename, Entry.Size, *DestFilename);
 
 					FString SourceContents;
@@ -601,7 +637,7 @@ UObject* UPackFactory::FactoryCreateBinary
 							DestFilename.RightChopInline(ContentIndex + 9, EAllowShrinking::No);
 						}
 					}
-					DestFilename = ContentDestinationRoot / DestFilename;
+					DestFilename = PackFactoryHelper::MakeSafeDestPath(ContentDestinationRoot, DestFilename, *EntryFilename);
 					UE_LOG(LogPackFactory, Log, TEXT("%s (%ld) -> %s"), **EntryFilename, Entry.Size, *DestFilename);
 
 					TUniquePtr<FArchive> FileHandle(IFileManager::Get().CreateFileWriter(*DestFilename));
@@ -635,6 +671,13 @@ UObject* UPackFactory::FactoryCreateBinary
 			{
 				if (FileToCopy.StartsWith(TEXT("Source/")) || FileToCopy.Contains(TEXT("/Source/")))
 				{
+					if (NormalizedSourceRoot.IsEmpty())
+					{
+						UE_LOG(LogPackFactory, Error, TEXT("Skipping source file \"%s\": no Runtime module source path available."), *FileToCopy);
+						++ErrorCount;
+						continue;
+					}
+
 					FString DestFilename = FileToCopy;
 					if (DestFilename.StartsWith(TEXT("Source/")))
 					{
@@ -648,7 +691,7 @@ UObject* UPackFactory::FactoryCreateBinary
 							DestFilename.RightChopInline(SourceIndex + 8, EAllowShrinking::No);
 						}
 					}
-					DestFilename = SourceModuleInfo.ModuleSourcePath / DestFilename;
+					DestFilename = PackFactoryHelper::MakeSafeDestPath(NormalizedSourceRoot, DestFilename, FileToCopy);
 
 					FString DestDirectory = FPaths::GetPath(DestFilename);
 
@@ -700,7 +743,7 @@ UObject* UPackFactory::FactoryCreateBinary
 							DestFilename.RightChopInline(ContentIndex + 9, EAllowShrinking::No);
 						}
 					}
-					DestFilename = ContentDestinationRoot / DestFilename;
+					DestFilename = PackFactoryHelper::MakeSafeDestPath(ContentDestinationRoot, DestFilename, FileToCopy);
 
 					FString DestDirectory = FPaths::GetPath(DestFilename);
 
