Fix Slate crash when toggling preview
#Synced-CL 51834001

diff --git a/Engine/Source/Editor/UnrealEd/Private/EditorEngine.cpp b/Engine/Source/Editor/UnrealEd/Private/EditorEngine.cpp
--- a/Engine/Source/Editor/UnrealEd/Private/EditorEngine.cpp
+++ b/Engine/Source/Editor/UnrealEd/Private/EditorEngine.cpp
@@ -6621,7 +6621,13 @@ static void PreviewPlatformChangePre()
 	
 	// Stop any texture streaming in progress :
 	UTexture::CancelPendingTextureStreaming();
-	
+
+	// Clear the Slate renderer's cached scene pointers before we recreate scenes.
+	if (FSlateApplication::IsInitialized())
+	{
+		FSlateApplication::Get().GetRenderer()->ReleaseAccessedResources(false /* bImmediatelyFlush */);
+	}
+
 	FlushRenderingCommands();
 }
 
