FAssetEditorToolkit : backout changelist 49994995

AddReferencedObjects is called on a non-game thread during gc

#Synced-CL 51843813

diff --git a/Engine/Source/Editor/UnrealEd/Private/Toolkits/AssetEditorToolkit.cpp b/Engine/Source/Editor/UnrealEd/Private/Toolkits/AssetEditorToolkit.cpp
--- a/Engine/Source/Editor/UnrealEd/Private/Toolkits/AssetEditorToolkit.cpp
+++ b/Engine/Source/Editor/UnrealEd/Private/Toolkits/AssetEditorToolkit.cpp
@@ -1735,12 +1735,7 @@ void FAssetEditorToolkit::FGCEditingObjects::AddReferencedObjects(FReferenceColl
 	Collector.AddReferencedObjects(OwnerToolkit.EditingObjects);
 	
 	// Remove null objects as a safe guard against assets being forcibly GC'd
-	TArray<TObjectPtr<UObject>> ObjectsToRemove = OwnerToolkit.EditingObjects.FilterByPredicate([](UObject* Obj) { return Obj == nullptr; });
-	for (TObjectPtr<UObject> ObjectToRemove : ObjectsToRemove)
-	{
-		// Call RemoveEditingAsset so that the editor tab gets closed if we remove the last EditingObject
-		OwnerToolkit.RemoveEditingAsset(ObjectToRemove);
-	}
+	OwnerToolkit.EditingObjects.RemoveAllSwap([](UObject* Obj) { return Obj == nullptr; });
 }
 
 FString FAssetEditorToolkit::FGCEditingObjects::GetReferencerName() const
