[Cooker/MTSAVE]
- Replace GOutputCookingWarnings by accessors and make it thread-safe
- Use the new InheritedGuard so that the value follows the whole task chain

#Synced-CL 51863520

diff --git a/Engine/Source/Editor/UnrealEd/Private/CookOnTheFlyServer.cpp b/Engine/Source/Editor/UnrealEd/Private/CookOnTheFlyServer.cpp
--- a/Engine/Source/Editor/UnrealEd/Private/CookOnTheFlyServer.cpp
+++ b/Engine/Source/Editor/UnrealEd/Private/CookOnTheFlyServer.cpp
@@ -11,6 +11,7 @@
 #include "Algo/Find.h"
 #include "Algo/RandomShuffle.h"
 #include "Algo/Unique.h"
+#include "Async/InheritedGuardValue.h"
 #include "AssetCompilingManager.h"
 #include "AssetRegistry/AssetRegistryModule.h"
 #include "AssetRegistry/AssetRegistryState.h"
@@ -1393,7 +1394,8 @@ static FDelayedAutoRegisterHelper GPreloadARHelper(EDelayedRegisterRunPhase::Ear
 );
 
 
-COREUOBJECT_API extern bool GOutputCookingWarnings;
+COREUOBJECT_API extern void SetOutputCookingWarnings(bool bValue);
+COREUOBJECT_API extern bool GetOutputCookingWarnings();
 
 void UCookOnTheFlyServer::WaitForRequests(int TimeoutMs)
 {
@@ -4643,10 +4645,11 @@ bool UCookOnTheFlyServer::LoadPackageForCooking(UE::Cook::FPackageData& PackageD
 	UE_LOG(LogCook, Display, TEXT("Processing request %s"), *ReportingFileName);
 #endif
 	static TSet<FString> CookWarningsList;
+	TOptional<TInheritedGuardValueAccessors<GetOutputCookingWarnings, SetOutputCookingWarnings>> OutputCookingWarningsGuard;
 	if (CookWarningsList.Contains(FileName) == false)
 	{
 		CookWarningsList.Add(FileName);
-		GOutputCookingWarnings = IsCookFlagSet(ECookInitializationFlags::OutputVerboseCookerWarnings);
+		OutputCookingWarningsGuard.Emplace(IsCookFlagSet(ECookInitializationFlags::OutputVerboseCookerWarnings));
 	}
 
 	bool bSuccess = true;
@@ -4701,7 +4704,6 @@ bool UCookOnTheFlyServer::LoadPackageForCooking(UE::Cook::FPackageData& PackageD
 			LogCookerMessage(FString::Printf(TEXT("Error loading %s!"), *ReportingFileName), EMessageSeverity::Error);
 		}
 	}
-	GOutputCookingWarnings = false;
 	return bSuccess;
 }
 
diff --git a/Engine/Source/Editor/UnrealEd/Private/Cooker/CookSavePackage.cpp b/Engine/Source/Editor/UnrealEd/Private/Cooker/CookSavePackage.cpp
--- a/Engine/Source/Editor/UnrealEd/Private/Cooker/CookSavePackage.cpp
+++ b/Engine/Source/Editor/UnrealEd/Private/Cooker/CookSavePackage.cpp
@@ -4,6 +4,7 @@
 
 #include "AssetRegistry/AssetData.h"
 #include "AssetRegistry/IAssetRegistry.h"
+#include "Async/InheritedGuardValue.h"
 #include "Commandlets/AssetRegistryGenerator.h"
 #include "Cooker/CookDependency.h"
 #include "Cooker/CookDeterminismManager.h"
@@ -50,7 +51,8 @@
 #include "UObject/SoftObjectPath.h"
 #include "UObject/UObjectGlobals.h"
 
-COREUOBJECT_API extern bool GOutputCookingWarnings;
+COREUOBJECT_API extern bool GetOutputCookingWarnings();
+COREUOBJECT_API extern void SetOutputCookingWarnings(bool bValue);
 
 #if OUTPUT_COOKTIMING
 
@@ -78,8 +80,7 @@ void UCookOnTheFlyServer::SaveCookedPackage(UE::Cook::FSaveCookedPackageContext&
 
 	Context.SetupPackage();
 
-	TGuardValue<bool> ScopedOutputCookerWarnings(GOutputCookingWarnings, IsCookFlagSet(ECookInitializationFlags::OutputVerboseCookerWarnings));
-	// SavePackage can CollectGarbage, so we need to store the currently-unqueued PackageData in a separate variable that we register for garbage collection
+	TInheritedGuardValueAccessors<GetOutputCookingWarnings, SetOutputCookingWarnings> ScopedOutputCookerWarnings(IsCookFlagSet(ECookInitializationFlags::OutputVerboseCookerWarnings));
 	TGuardValue<UE::Cook::FPackageData*> ScopedSavingPackageData(SavingPackageData, &Context.PackageData);
 	TGuardValue<bool> ScopedIsSavingPackage(bIsSavingPackage, true);
 	// For legacy reasons we set GIsCookerLoadingPackage == true during save. Some classes use it to conditionally execute cook operations in both save and load
diff --git a/Engine/Source/Runtime/CoreUObject/Private/UObject/SavePackage2.cpp b/Engine/Source/Runtime/CoreUObject/Private/UObject/SavePackage2.cpp
--- a/Engine/Source/Runtime/CoreUObject/Private/UObject/SavePackage2.cpp
+++ b/Engine/Source/Runtime/CoreUObject/Private/UObject/SavePackage2.cpp
@@ -84,7 +84,7 @@ FSavePackageContext::~FSavePackageContext()
 
 // defined in UObjectGlobals.cpp
 extern const FName NAME_UniqueObjectNameForCooking;
-COREUOBJECT_API extern bool GOutputCookingWarnings;
+COREUOBJECT_API extern bool GetOutputCookingWarnings();
 
 namespace
 {
@@ -801,7 +801,7 @@ ESavePackageResult ValidateExports(FSaveContext& SaveContext)
 
 	// Cooking checks
 #if WITH_EDITOR && !UE_FNAME_OUTLINE_NUMBER
-	if (GOutputCookingWarnings)
+	if (GetOutputCookingWarnings())
 	{
 		// check the name list for UniqueObjectNameForCooking cooking. Since check is a fast check for any NAME_UniqueObjectNameForCooking
 		// we can only perform this check when UE_FNAME_OUTLINE_NUMBER=0 as otherwise all prefixed names will have a unique ComparisonIndex
diff --git a/Engine/Source/Runtime/CoreUObject/Private/UObject/UObjectGlobals.cpp b/Engine/Source/Runtime/CoreUObject/Private/UObject/UObjectGlobals.cpp
--- a/Engine/Source/Runtime/CoreUObject/Private/UObject/UObjectGlobals.cpp
+++ b/Engine/Source/Runtime/CoreUObject/Private/UObject/UObjectGlobals.cpp
@@ -3261,7 +3261,17 @@ bool StaticAllocateObjectErrorTests( const UClass* InClass, UObject* InOuter, FN
 static thread_local FRestoreForUObjectOverwrite* ObjectRestoreAfterInitProps = nullptr;
 
 extern const FName NAME_UniqueObjectNameForCooking("UniqueObjectNameForCooking");
-COREUOBJECT_API bool GOutputCookingWarnings = false;
+static thread_local bool GOutputCookingWarnings = false;
+
+COREUOBJECT_API bool GetOutputCookingWarnings()
+{
+	return GOutputCookingWarnings;
+}
+
+COREUOBJECT_API void SetOutputCookingWarnings(bool bValue)
+{
+	GOutputCookingWarnings = bValue;
+}
 
 // Used by StaticAllocateObject when constructing an object on top of existing object to simulate GC behavior
 class FGCReconstructionGuard
@@ -3417,7 +3427,7 @@ UObject* StaticAllocateObject
 		UE_AUTORTFM_OPEN
 		{
 #if WITH_EDITOR
-			if ( GOutputCookingWarnings && GetTransientPackage() != InOuter->GetOutermost() )
+			if ( GetOutputCookingWarnings() && GetTransientPackage() != InOuter->GetOutermost() )
 			{
 				InName = MakeUniqueObjectName(InOuter, InClass, NAME_UniqueObjectNameForCooking);
 			}
diff --git a/Engine/Source/Runtime/Engine/Private/StaticMesh.cpp b/Engine/Source/Runtime/Engine/Private/StaticMesh.cpp
--- a/Engine/Source/Runtime/Engine/Private/StaticMesh.cpp
+++ b/Engine/Source/Runtime/Engine/Private/StaticMesh.cpp
@@ -7322,11 +7322,6 @@ FName UStaticMesh::GetLODPathName(const UStaticMesh* Mesh, int32 LODIndex)
 FUObjectAnnotationSparseBool GStaticMeshesThatNeedMaterialFixup;
 #endif // #if WITH_EDITORONLY_DATA
 
-#if WITH_EDITOR
-COREUOBJECT_API extern bool GOutputCookingWarnings;
-#endif
-
-
 LLM_DEFINE_TAG(StaticMesh_Serialize); // This is an important test case for LLM_DEFINE_TAG
 
 /**
