[Viewport ITF] Remove remaining non-context usages of UE::Editor::ViewportInteractions::UseEditorViewportInteractions(), replacing them with a new UsesViewportInteractions() function that takes in an FEditorModeTools and provides correct per-mode-tools-state. Fixes issues that could arise where viewport ITF interactions were enabled at the global level but disabled within an FEditorModeTools context.

#rnx
#Synced-CL 51825263

diff --git a/Engine/Source/Editor/Experimental/EditorInteractiveToolsFramework/Private/ViewportInteractions/ViewportInteractionsUtilities.cpp b/Engine/Source/Editor/Experimental/EditorInteractiveToolsFramework/Private/ViewportInteractions/ViewportInteractionsUtilities.cpp
--- a/Engine/Source/Editor/Experimental/EditorInteractiveToolsFramework/Private/ViewportInteractions/ViewportInteractionsUtilities.cpp
+++ b/Engine/Source/Editor/Experimental/EditorInteractiveToolsFramework/Private/ViewportInteractions/ViewportInteractionsUtilities.cpp
@@ -2,6 +2,7 @@
 
 #include "ViewportInteractions/ViewportInteractionsUtilities.h"
 
+#include "EditorModeManager.h"
 #include "EditorDragTools/BoxSelectInteraction.h"
 #include "EditorDragTools/FrustumSelectInteraction.h"
 #include "EditorDragTools/MeasureToolInteraction.h"
@@ -19,7 +20,10 @@
 #include "ViewportInteractions/ViewportViewAngleInteraction.h"
 #include "ViewportInteractions/ViewportZoomInteraction.h"
 
-void UE::Editor::ViewportInteractions::AddDefaultCameraMovementInteractions(
+namespace UE::Editor::ViewportInteractions
+{
+
+void AddDefaultCameraMovementInteractions(
 	UViewportInteractionsBehaviorSource* InInteractionsBehaviorSource
 )
 {
@@ -43,7 +47,7 @@ void UE::Editor::ViewportInteractions::AddDefaultCameraMovementInteractions(
 	}
 }
 
-void UE::Editor::ViewportInteractions::AddDefaultDragToolsInteractions(UViewportInteractionsBehaviorSource* InInteractionsBehaviorSource)
+void AddDefaultDragToolsInteractions(UViewportInteractionsBehaviorSource* InInteractionsBehaviorSource)
 {
 	if (InInteractionsBehaviorSource)
 	{
@@ -53,3 +57,18 @@ void UE::Editor::ViewportInteractions::AddDefaultDragToolsInteractions(UViewport
 		InInteractionsBehaviorSource->AddInteraction<UViewportChangeInteraction>();
 	}
 }
+
+bool UsesViewportInteractions(const FEditorModeTools* ModeTools)
+{
+	if (ModeTools && ModeTools->UsesViewportITF())
+	{
+		if (const UModeManagerInteractiveToolsContext* InputToolsBehaviorSource = ModeTools->GetInteractiveToolsContext())
+		{
+			return InputToolsBehaviorSource->UsesViewportInteractions();
+		}
+	}
+	return false;
+}
+
+
+}
diff --git a/Engine/Source/Editor/Experimental/EditorInteractiveToolsFramework/Public/ViewportInteractions/ViewportInteractionsUtilities.h b/Engine/Source/Editor/Experimental/EditorInteractiveToolsFramework/Public/ViewportInteractions/ViewportInteractionsUtilities.h
--- a/Engine/Source/Editor/Experimental/EditorInteractiveToolsFramework/Public/ViewportInteractions/ViewportInteractionsUtilities.h
+++ b/Engine/Source/Editor/Experimental/EditorInteractiveToolsFramework/Public/ViewportInteractions/ViewportInteractionsUtilities.h
@@ -5,6 +5,7 @@
 #include "HAL/Platform.h"
 
 class UViewportInteractionsBehaviorSource;
+class FEditorModeTools;
 
 namespace UE::Editor::ViewportInteractions
 {
@@ -15,4 +16,10 @@ namespace UE::Editor::ViewportInteractions
 EDITORINTERACTIVETOOLSFRAMEWORK_API void AddDefaultCameraMovementInteractions(UViewportInteractionsBehaviorSource* InInteractionsBehaviorSource);
 
 EDITORINTERACTIVETOOLSFRAMEWORK_API void AddDefaultDragToolsInteractions(UViewportInteractionsBehaviorSource* InInteractionsBehaviorSource);
+
+/**
+ * Returns true if the provided mode tools is valid and uses viewport interactions.
+ */
+EDITORINTERACTIVETOOLSFRAMEWORK_API bool UsesViewportInteractions(const FEditorModeTools* ModeTools);
+
 }
diff --git a/Engine/Source/Editor/UnrealEd/Private/EditorViewportClient.cpp b/Engine/Source/Editor/UnrealEd/Private/EditorViewportClient.cpp
--- a/Engine/Source/Editor/UnrealEd/Private/EditorViewportClient.cpp
+++ b/Engine/Source/Editor/UnrealEd/Private/EditorViewportClient.cpp
@@ -89,6 +89,7 @@
 #include "Engine/LevelBounds.h"
 #include "ViewportInteractions/ViewportInteractionsBehaviorSource.h"
 #include "ViewportInteractions/ViewportInteraction.h"
+#include "ViewportInteractions/ViewportInteractionsUtilities.h"
 #include "WorldPartition/WorldPartition.h"
 #include "RenderViewportFeedback.h"
 #include "EditorGizmos/EditorTransformGizmoUtil.h"
@@ -143,18 +144,6 @@ static bool GetDefaultLowDPIPreviewValue()
 	return CVarEditorViewportHighDPIPtr->GetInt() == 0;
 }
 
-bool UsesViewportInteractions(const FEditorModeTools* InModeTools)
-{
-	if (InModeTools && UE::Editor::ViewportInteractions::UseEditorViewportInteractions())
-	{
-		if (UModeManagerInteractiveToolsContext* InputToolsBehaviorSource = InModeTools->GetInteractiveToolsContext())
-		{
-			return InputToolsBehaviorSource->UsesViewportInteractions();
-		}
-	}
-	return false;
-}
-
 /**
  * Returns true if a viewport interaction with the specified name is enabled and able to act 
  */
@@ -1742,7 +1731,7 @@ void FEditorViewportClient::Tick(float DeltaTime)
 		// Update any real-time camera movement
 		UpdateCameraMovement( DeltaTime );
 
-		if (UsesViewportInteractions(ModeTools.Get()))
+		if (UE::Editor::ViewportInteractions::UsesViewportInteractions(ModeTools.Get()))
 		{
 			ConsumeExternalCameraMovement();
 		}
@@ -2114,7 +2103,7 @@ bool FEditorViewportClient::IsActiveViewportType(ELevelViewportType InViewportTy
 // Updates real-time camera movement.  Should be called every viewport tick!
 void FEditorViewportClient::UpdateCameraMovement( float DeltaTime )
 {
-	const bool bIsITFInteractions = UsesViewportInteractions(ModeTools.Get());
+	const bool bIsITFInteractions = UE::Editor::ViewportInteractions::UsesViewportInteractions(ModeTools.Get());
 	
 	// We only want to move perspective cameras around like this
 	if( Viewport && IsPerspective() && (bIsITFInteractions ? !IsViewportInteractionActive(ModeTools.Get(), UE::Editor::ViewportInteractions::Orbit) : !ShouldOrbitCamera()))
@@ -3645,7 +3634,7 @@ bool FEditorViewportClient::Internal_InputKey(const FInputKeyEventArgs& EventArg
 	const EInputEvent& Event = EventArgs.Event;
 	const FViewport* InViewport = EventArgs.Viewport;
 
-	bool bRouteToToolsContext = !bIsTracking || UE::Editor::ViewportInteractions::UseEditorViewportInteractions();
+	bool bRouteToToolsContext = !bIsTracking || UE::Editor::ViewportInteractions::UsesViewportInteractions(ModeTools.Get());
 	// Let the current mode have a look at the input before reacting to it. 
 	// Note that bIsTracking tells us whether the viewport client is capturing mouse behavior to fly around,
 	// move objects, etc. In this case we don't want to pass the input the input router, because it is already
@@ -5757,7 +5746,7 @@ bool FEditorViewportClient::IsFlightCameraInputModeActive() const
 	{
 		if( CameraController != NULL )
 		{
-			if (UsesViewportInteractions(ModeTools.Get()))
+			if (UE::Editor::ViewportInteractions::UsesViewportInteractions(ModeTools.Get()))
 			{
 				return ViewportNavigationHelper->bIsMouseLooking;
 			}
diff --git a/Engine/Source/Editor/UnrealEd/Private/LevelEditorViewport.cpp b/Engine/Source/Editor/UnrealEd/Private/LevelEditorViewport.cpp
--- a/Engine/Source/Editor/UnrealEd/Private/LevelEditorViewport.cpp
+++ b/Engine/Source/Editor/UnrealEd/Private/LevelEditorViewport.cpp
@@ -113,6 +113,7 @@
 #include "Subsystems/PlacementSubsystem.h"
 #include "Subsystems/EditorElementSubsystem.h"
 #include "ViewportInteractions/ViewportInteractionsBehaviorSource.h"
+#include "ViewportInteractions/ViewportInteractionsUtilities.h"
 
 DEFINE_LOG_CATEGORY(LogEditorViewport);
 
@@ -3303,11 +3304,8 @@ bool FLevelEditorViewportClient::InputWidgetDelta(FViewport* InViewport, EAxisLi
 				const bool RightMouseButtonDown = InViewport->KeyState(EKeys::RightMouseButton);
 				const bool MiddleMouseButtonDown = InViewport->KeyState(EKeys::MiddleMouseButton);
 
-				// Some functionality can be skipped if new gizmos and ITF drag tools are active
-				const bool bUseLegacyGizmoDragTools = !(ModeTools->UsesViewportITF() && UE::Editor::ViewportInteractions::UseEditorViewportInteractions());
-
-				// If duplicate dragging . . .
-				if (IsAltPressed() && (LeftMouseButtonDown || RightMouseButtonDown) && bUseLegacyGizmoDragTools)
+				// If duplicate dragging and using legacy viewport interaction behavior 
+				if (IsAltPressed() && (LeftMouseButtonDown || RightMouseButtonDown) && !UE::Editor::ViewportInteractions::UsesViewportInteractions(ModeTools.Get()))
 				{
 					// The widget has been offset, so check if we should duplicate the selection.
 					if ( bDuplicateOnNextDrag )
diff --git a/Engine/Source/Editor/UnrealEd/Private/MouseDeltaTracker.cpp b/Engine/Source/Editor/UnrealEd/Private/MouseDeltaTracker.cpp
--- a/Engine/Source/Editor/UnrealEd/Private/MouseDeltaTracker.cpp
+++ b/Engine/Source/Editor/UnrealEd/Private/MouseDeltaTracker.cpp
@@ -14,6 +14,7 @@
 #include "Misc/AxisDisplayInfo.h"
 #include "EditorModeManager.h"
 #include "ViewportInteractions/ViewportInteractionsBehaviorSource.h"
+#include "ViewportInteractions/ViewportInteractionsUtilities.h"
 
 #define LOCTEXT_NAMESPACE "MouseDeltaTracker"
 
@@ -345,18 +346,9 @@ bool FMouseDeltaTracker::EndTracking(FEditorViewportClient* InViewportClient)
 
 void FMouseDeltaTracker::ConditionalBeginUsingDragTool( FEditorViewportClient* InViewportClient )
 {
-	if (FEditorModeTools* ModeTools = InViewportClient->GetModeTools())
+	if (UE::Editor::ViewportInteractions::UsesViewportInteractions(InViewportClient->GetModeTools()))
 	{
-		if (UE::Editor::ViewportInteractions::UseEditorViewportInteractions())
-		{
-			if (UModeManagerInteractiveToolsContext* InputToolsBehaviorSource = ModeTools->GetInteractiveToolsContext())
-			{
-				if (InputToolsBehaviorSource->UsesViewportInteractions())
-				{
-					return;
-				}
-			}
-		}
+		return;
 	}
 
 	const bool LeftMouseButtonDown = InViewportClient->Viewport->KeyState(EKeys::LeftMouseButton);
