Include external objects related via their OptionalOuterPathName when running cook-time asset validation

#rnx
#Synced-CL 51753475

diff --git a/Engine/Source/Editor/UnrealEd/Private/CookOnTheFlyServer.cpp b/Engine/Source/Editor/UnrealEd/Private/CookOnTheFlyServer.cpp
--- a/Engine/Source/Editor/UnrealEd/Private/CookOnTheFlyServer.cpp
+++ b/Engine/Source/Editor/UnrealEd/Private/CookOnTheFlyServer.cpp
@@ -13,6 +13,7 @@
 #include "Algo/Unique.h"
 #include "Async/InheritedGuardValue.h"
 #include "AssetCompilingManager.h"
+#include "AssetRegistry/AssetRegistryHelpers.h"
 #include "AssetRegistry/AssetRegistryModule.h"
 #include "AssetRegistry/AssetRegistryState.h"
 #include "Async/Async.h"
@@ -3439,6 +3440,11 @@ EDataValidationResult UCookOnTheFlyServer::ValidateSourcePackage(UE::Cook::FPack
 				ValidationArtifact->FilterExternalActorsToValidate(Package->GetFName(), ExternalObjects);
 			}
 		}
+		if (TArray<FAssetData> AssetsWithPackageAsOuter = UAssetRegistryHelpers::GetAssetsWithOuterForPaths({ Package->GetFName() }, Package->GetFName(), EAssetsWithOuterForPathsFlags::RecursivePaths | EAssetsWithOuterForPathsFlags::IncludeOnlyOnDiskAsset);
+			!AssetsWithPackageAsOuter.IsEmpty())
+		{
+			ExternalObjects.Append(MoveTemp(AssetsWithPackageAsOuter));
+		}
 
 		FMessageLogScopedOverride AssetCheckLogOverride(UE::DataValidation::MessageLogName);
 		if (bLogErrorsAsWarnings)
