[Cooker/MTSAVE]
- Replace FScopedActivePackage with one that is both thread-safe and properly inherited throught the whole task chain

#Synced-CL 51868743

diff --git a/Engine/Source/Editor/UnrealEd/Classes/CookOnTheSide/CookOnTheFlyServer.h b/Engine/Source/Editor/UnrealEd/Classes/CookOnTheSide/CookOnTheFlyServer.h
--- a/Engine/Source/Editor/UnrealEd/Classes/CookOnTheSide/CookOnTheFlyServer.h
+++ b/Engine/Source/Editor/UnrealEd/Classes/CookOnTheSide/CookOnTheFlyServer.h
@@ -1526,29 +1526,6 @@ private:
 	 */
 	void DeleteOutputForPackage(FName PackageName, const ITargetPlatform* TargetPlatform);
 
-	/**
-	 * Set the given package as the active package for diagnostics (e.g. hidden dependencies).
- 	 * Must be balanced with call to ClearActivePackage. Consider using FScopedActivePackage instead.
-	 * @param PackageTrackingOpsName Optional op for TObjectPtr tracking. Set to NAME_None to not set TObjectPtr context.
-	 */
-	void SetActivePackage(FName PackageName, FName PackageTrackingOpsName);
-	/** Clear the package for diagnostics. Must be balanced with call to SetActivePackage */
-	void ClearActivePackage();
-	struct FActivePackageData
-	{
-		FName PackageName;
-		bool bActive = false;
-		UE_TRACK_REFERENCING_PACKAGE_DECLARE_SCOPE_VARIABLE(ReferenceTrackingScope);
-	};
-	/** Scoped type to call Set/ClearActivePackage. */
-	struct FScopedActivePackage
-	{
-		FScopedActivePackage(UCookOnTheFlyServer& InCOTFS, FName InPackageName, FName InPackageTrackingOpsName);
-		~FScopedActivePackage();
-
-		UCookOnTheFlyServer& COTFS;
-	};
-
 	/** Callback for FGenericCrashContext; provides the current ActivePackage as context. */
 	void DumpCrashContext(FCrashContextExtendedWriter& Writer);
 	/** Callback for analytics when a new UPackage is loaded. */
@@ -1587,9 +1564,6 @@ private:
 #if UE_WITH_OBJECT_HANDLE_TRACKING
 	UE::CoreUObject::FObjectHandleTrackingCallbackId ObjectHandleReadHandle;
 #endif // UE_WITH_OBJECT_HANDLE_TRACKING
-	/** Data tracking the package currently having calls made (Load/Save/Other) from COTFS. Used for diagnostics. */
-	FActivePackageData ActivePackageData;
-
 	/** Classes (and all subclasses) that were listed as the only classes that should be cooked in the filter settings. */
 	TSet<FName> CookFilterIncludedClasses;
 	/** Classes (and all subclasses) that were listed as the only asset classes that should be cooked in the filter settings. */
diff --git a/Engine/Source/Editor/UnrealEd/Private/CookOnTheFlyServer.cpp b/Engine/Source/Editor/UnrealEd/Private/CookOnTheFlyServer.cpp
--- a/Engine/Source/Editor/UnrealEd/Private/CookOnTheFlyServer.cpp
+++ b/Engine/Source/Editor/UnrealEd/Private/CookOnTheFlyServer.cpp
@@ -21,6 +21,7 @@
 #include "Commandlets/AssetRegistryGenerator.h"
 #include "Containers/DirectoryTree.h"
 #include "Containers/RingBuffer.h"
+#include "Cooker/CookActivePackage.h"
 #include "Cooker/AsyncIODelete.h"
 #include "Cooker/CookAssetRegistryAccessTracker.h"
 #include "Cooker/CookConfigAccessTracker.h"
@@ -3918,7 +3919,7 @@ UE::Cook::EPollStatus UCookOnTheFlyServer::PreSaveGeneratorPackage(UE::Cook::FPa
 	// function so that other GenerationHelper code can rely on the support tasks for the PreSave being done.
 	if (GenerationHelper.IsSaveGenerator())
 	{
-		FScopedActivePackage ScopedActivePackage(*this, GenerationHelper.GetOwner().GetPackageName(),
+		UE::Cook::FScopedCookActivePackage ScopedActivePackage(GenerationHelper.GetOwner().GetPackageName(),
 #if UE_WITH_OBJECT_HANDLE_TRACKING
 			PackageAccessTrackingOps::NAME_CookerBuildObject
 #else
@@ -4087,7 +4088,7 @@ UE::Cook::EPollStatus UCookOnTheFlyServer::TryPreSaveGeneratedPackage(UE::Cook::
 	GeneratedPackagePopulateData.bCreatedAsMap = GeneratedInfo.IsCreateAsMap();
 	UE::Cook::CookPackageSplitter::FPopulateContextData PopulateData;
 	{
-		FScopedActivePackage ScopedActivePackage(*this, GenerationHelper.GetOwner().GetPackageName(),
+		UE::Cook::FScopedCookActivePackage ScopedActivePackage(GenerationHelper.GetOwner().GetPackageName(),
 #if UE_WITH_OBJECT_HANDLE_TRACKING
 			PackageAccessTrackingOps::NAME_CookerBuildObject
 #else
@@ -4640,7 +4641,7 @@ bool UCookOnTheFlyServer::LoadPackageForCooking(UE::Cook::FPackageData& PackageD
 	UE::Cook::FPackageData* ReportingPackageData)
 {
 	UE_SCOPED_HIERARCHICAL_COOKTIMER_AND_DURATION(LoadPackageForCooking, DetailedCookStats::TickCookOnTheSideLoadPackagesTimeSec);
-	FScopedActivePackage ScopedActivePackage(*this, PackageData.GetPackageName(), NAME_None);
+	UE::Cook::FScopedCookActivePackage ScopedActivePackage(PackageData.GetPackageName(), NAME_None);
 
 	FString PackageName = PackageData.GetPackageName().ToString();
 	OutPackage = FindObject<UPackage>(nullptr, *PackageName);
@@ -4713,42 +4714,13 @@ bool UCookOnTheFlyServer::LoadPackageForCooking(UE::Cook::FPackageData& PackageD
 	return bSuccess;
 }
 
-void UCookOnTheFlyServer::SetActivePackage(FName PackageName, FName PackageTrackingOpsName)
-{
-	check(!ActivePackageData.bActive);
-	ActivePackageData.bActive = true;
-	ActivePackageData.PackageName = PackageName;
-	if (!PackageTrackingOpsName.IsNone())
-	{
-		UE_TRACK_REFERENCING_PACKAGE_ACTIVATE_SCOPE_VARIABLE(ActivePackageData.ReferenceTrackingScope,
-			PackageName, PackageTrackingOpsName);
-	}
-}
-
-void UCookOnTheFlyServer::ClearActivePackage()
-{
-	check(ActivePackageData.bActive);
-	UE_TRACK_REFERENCING_PACKAGE_DEACTIVATE_SCOPE_VARIABLE(ActivePackageData.ReferenceTrackingScope);
-	ActivePackageData.PackageName = NAME_None;
-	ActivePackageData.bActive = false;
-}
-
-UCookOnTheFlyServer::FScopedActivePackage::FScopedActivePackage(UCookOnTheFlyServer& InCOTFS, FName PackageName, FName PackageTrackingOpsName)
-	: COTFS(InCOTFS)
-{
-	COTFS.SetActivePackage(PackageName, PackageTrackingOpsName);
-}
-UCookOnTheFlyServer::FScopedActivePackage::~FScopedActivePackage()
-{
-	COTFS.ClearActivePackage();
-}
-
 void UCookOnTheFlyServer::DumpCrashContext(FCrashContextExtendedWriter& Writer)
 {
 #if WITH_ADDITIONAL_CRASH_CONTEXTS
-	if (ActivePackageData.bActive)
+	FName PackageName = UE::Cook::GetCookActivePackageName();
+	if (!PackageName.IsNone())
 	{
-		Writer.AddString(TEXT("ActivePackage"), *WriteToString<256>(ActivePackageData.PackageName));
+		Writer.AddString(TEXT("ActivePackage"), *WriteToString<256>(PackageName));
 	}
 #endif
 }
@@ -12912,19 +12884,17 @@ void UCookOnTheFlyServer::RouteBeginCacheForCookedPlatformData(UE::Cook::FPackag
 		return;
 	}
 
-	// We need to set our scopes for e.g. TObjectPtr reads around the call to BeginCacheForCookedPlatformData,
-	// but in some cases we have already set the scope (e.g. when calling BeginCache from inside SavePackage)
-	TOptional<FScopedActivePackage> ScopedActivePackage;
-	if (!ActivePackageData.bActive)
-	{
-		ScopedActivePackage.Emplace(*this, PackageName,
+	// Set scopes for e.g. TObjectPtr reads around the call to BeginCacheForCookedPlatformData.
+	// When called from inside an existing tracking scope (e.g. SavePackage), preserve the outer
+	// scope's tracking op rather than overriding it with NAME_CookerBuildObject.
+	UE::Cook::FScopedCookActivePackage ScopedActivePackage(PackageName,
 #if UE_WITH_OBJECT_HANDLE_TRACKING
-			PackageAccessTrackingOps::NAME_CookerBuildObject
+		PackageAccessTracking_Private::FPackageAccessRefScope::GetCurrentThreadScope()
+			? NAME_None : PackageAccessTrackingOps::NAME_CookerBuildObject
 #else
-			FName()
+		FName()
 #endif
-		);
-	}
+	);
 	Obj->BeginCacheForCookedPlatformData(TargetPlatform);
 	*ExistingEvent = ECachedCookedPlatformDataEvent::BeginCacheForCookedPlatformDataCalled;
 }
@@ -12946,19 +12916,17 @@ bool UCookOnTheFlyServer::RouteIsCachedCookedPlatformDataLoaded(UE::Cook::FPacka
 		ExistingEvent = &CCPDState.PlatformStates.FindOrAdd(TargetPlatform, ECachedCookedPlatformDataEvent::None);
 	}
 
-	// We need to set our scopes for e.g. TObjectPtr reads around the call to BeginCacheForCookedPlatformData,
-	// but in some cases we have already set the scope (e.g. when calling IsCached from inside SavePackage)
-	TOptional<FScopedActivePackage> ScopedActivePackage;
-	if (!ActivePackageData.bActive)
-	{
-		ScopedActivePackage.Emplace(*this, PackageName,
+	// Set scopes for e.g. TObjectPtr reads around the call to IsCachedCookedPlatformDataLoaded.
+	// When called from inside an existing tracking scope (e.g. SavePackage), preserve the outer
+	// scope's tracking op rather than overriding it with NAME_CookerBuildObject.
+	UE::Cook::FScopedCookActivePackage ScopedActivePackage(PackageName,
 #if UE_WITH_OBJECT_HANDLE_TRACKING
-			PackageAccessTrackingOps::NAME_CookerBuildObject
+		PackageAccessTracking_Private::FPackageAccessRefScope::GetCurrentThreadScope()
+			? NAME_None : PackageAccessTrackingOps::NAME_CookerBuildObject
 #else
-			FName()
+		FName()
 #endif
-			);
-	}
+	);
 
 	if (*ExistingEvent != ECachedCookedPlatformDataEvent::BeginCacheForCookedPlatformDataCalled
 		&& *ExistingEvent != ECachedCookedPlatformDataEvent::IsCachedCookedPlatformDataLoadedCalled)
diff --git a/Engine/Source/Editor/UnrealEd/Private/Cooker/CookActivePackage.cpp b/Engine/Source/Editor/UnrealEd/Private/Cooker/CookActivePackage.cpp
new file mode 100644
--- /dev/null
+++ b/Engine/Source/Editor/UnrealEd/Private/Cooker/CookActivePackage.cpp
@@ -0,0 +1,35 @@
+// Copyright Epic Games, Inc. All Rights Reserved.
+
+#include "Cooker/CookActivePackage.h"
+
+// Thread_local active cook package name for crash context diagnostics.
+// Propagated to child tasks via FInheritedContextExtension.
+static thread_local FName GCookActivePackageName;
+
+static UE::FInheritedContextExtension& GetCookActivePackageExtension()
+{
+	static UE::FInheritedContextExtension Extension = UE::MakeInheritedContextExtension(
+		[]() -> FName& { return GCookActivePackageName; });
+	return Extension;
+}
+
+FName UE::Cook::GetCookActivePackageName()
+{
+	return GCookActivePackageName;
+}
+
+UE::Cook::FScopedCookActivePackage::FScopedCookActivePackage(FName PackageName, FName PackageTrackingOpsName)
+	: PreviousPackageName(GCookActivePackageName)
+	, ExtensionScope(GetCookActivePackageExtension())
+{
+	GCookActivePackageName = PackageName;
+	if (!PackageTrackingOpsName.IsNone())
+	{
+		UE_TRACK_REFERENCING_PACKAGE_ACTIVATE_SCOPE_VARIABLE(TrackingScope, PackageName, PackageTrackingOpsName);
+	}
+}
+
+UE::Cook::FScopedCookActivePackage::~FScopedCookActivePackage()
+{
+	GCookActivePackageName = PreviousPackageName;
+}
diff --git a/Engine/Source/Editor/UnrealEd/Private/Cooker/CookActivePackage.h b/Engine/Source/Editor/UnrealEd/Private/Cooker/CookActivePackage.h
new file mode 100644
--- /dev/null
+++ b/Engine/Source/Editor/UnrealEd/Private/Cooker/CookActivePackage.h
@@ -0,0 +1,34 @@
+// Copyright Epic Games, Inc. All Rights Reserved.
+
+#pragma once
+
+#include "Async/InheritedContext.h"
+#include "Misc/PackageAccessTracking.h"
+#include "UObject/NameTypes.h"
+
+namespace UE::Cook
+{
+
+/** Get the active cook package name for the current thread. Thread-safe (thread_local). */
+FName GetCookActivePackageName();
+
+/**
+ * RAII scope that sets the active cook package for crash-context diagnostics and package access tracking.
+ * Thread-safe and task-friendly:
+ * - Crash context uses a thread_local FName propagated to child tasks via FInheritedContextExtension.
+ * - Package access tracking uses FPackageAccessRefScope which is already thread_local and InheritedContext-aware.
+ */
+struct FScopedCookActivePackage
+{
+	FScopedCookActivePackage(FName PackageName, FName PackageTrackingOpsName);
+	~FScopedCookActivePackage();
+
+	UE_NONCOPYABLE(FScopedCookActivePackage);
+
+private:
+	FName PreviousPackageName;
+	UE::FInheritedContextExtensionScope ExtensionScope;
+	UE_TRACK_REFERENCING_PACKAGE_DECLARE_SCOPE_VARIABLE(TrackingScope);
+};
+
+} // namespace UE::Cook
diff --git a/Engine/Source/Editor/UnrealEd/Private/Cooker/CookGenerationHelper.cpp b/Engine/Source/Editor/UnrealEd/Private/Cooker/CookGenerationHelper.cpp
--- a/Engine/Source/Editor/UnrealEd/Private/Cooker/CookGenerationHelper.cpp
+++ b/Engine/Source/Editor/UnrealEd/Private/Cooker/CookGenerationHelper.cpp
@@ -3,6 +3,7 @@
 #include "Cooker/CookGenerationHelper.h"
 
 #include "Algo/Unique.h"
+#include "Cooker/CookActivePackage.h"
 #include "AssetRegistry/AssetData.h"
 #include "Cooker/CookDirector.h"
 #include "Cooker/CookGarbageCollect.h"
@@ -986,7 +987,7 @@ bool FGenerationHelper::TryReportGenerationManifest()
 
 	ICookPackageSplitter::FGenerationManifest Manifest;
 	{
-		UCookOnTheFlyServer::FScopedActivePackage ScopedActivePackage(COTFS, OwnerPackageName,
+		UE::Cook::FScopedCookActivePackage ScopedActivePackage(OwnerPackageName,
 #if UE_WITH_OBJECT_HANDLE_TRACKING
 			PackageAccessTrackingOps::NAME_CookerBuildObject
 #else
@@ -1203,7 +1204,7 @@ bool FGenerationHelper::TryCallPopulateGeneratorPackage(
 	UE::Cook::CookPackageSplitter::FPopulateContextData PopulateData;
 	if (bPopulateCallAllowed)
 	{
-		UCookOnTheFlyServer::FScopedActivePackage ScopedActivePackage(COTFS, OwnerPackageName,
+		UE::Cook::FScopedCookActivePackage ScopedActivePackage(OwnerPackageName,
 #if UE_WITH_OBJECT_HANDLE_TRACKING
 			PackageAccessTrackingOps::NAME_CookerBuildObject
 #else
@@ -1265,7 +1266,6 @@ bool FGenerationHelper::TryCallPopulateGeneratedPackage(UE::Cook::FCookGeneratio
 	}
 	FPackageData& OwnerPackageData = GetOwner();
 	FName OwnerPackageName = OwnerPackageData.GetPackageName();
-	UCookOnTheFlyServer& COTFS = OwnerPackageData.GetPackageDatas().GetCookOnTheFlyServer();
 
 	if (!bReportedManifest)
 	{
@@ -1291,7 +1291,7 @@ bool FGenerationHelper::TryCallPopulateGeneratedPackage(UE::Cook::FCookGeneratio
 	ICookPackageSplitter::FGeneratedPackageForPopulate SplitterInfo{ Info.RelativePath,
 		Info.GeneratedRootPath, Package, Info.IsCreateAsMap() };
 
-	UCookOnTheFlyServer::FScopedActivePackage ScopedActivePackage(COTFS, OwnerPackageName,
+	UE::Cook::FScopedCookActivePackage ScopedActivePackage(OwnerPackageName,
 #if UE_WITH_OBJECT_HANDLE_TRACKING
 		PackageAccessTrackingOps::NAME_CookerBuildObject
 #else
@@ -1911,8 +1911,7 @@ void FGenerationHelper::ResetSaveState(FCookGenerationInfo& Info, UPackage* Pack
 		}
 		else
 		{
-			UCookOnTheFlyServer& COTFS = GetOwner().GetPackageDatas().GetCookOnTheFlyServer();
-			UCookOnTheFlyServer::FScopedActivePackage ScopedActivePackage(COTFS, GetOwner().GetPackageName(),
+			UE::Cook::FScopedCookActivePackage ScopedActivePackage(GetOwner().GetPackageName(),
 #if UE_WITH_OBJECT_HANDLE_TRACKING
 				PackageAccessTrackingOps::NAME_CookerBuildObject
 #else
diff --git a/Engine/Source/Editor/UnrealEd/Private/Cooker/CookPackageData.cpp b/Engine/Source/Editor/UnrealEd/Private/Cooker/CookPackageData.cpp
--- a/Engine/Source/Editor/UnrealEd/Private/Cooker/CookPackageData.cpp
+++ b/Engine/Source/Editor/UnrealEd/Private/Cooker/CookPackageData.cpp
@@ -11,6 +11,7 @@
 #include "AssetRegistry/IAssetRegistry.h"
 #include "Async/ParallelFor.h"
 #include "CompactBinaryTCP.h"
+#include "Cooker/CookActivePackage.h"
 #include "Cooker/BuildResultDependenciesMap.h"
 #include "Cooker/CookDirector.h"
 #include "Cooker/CookGenerationHelper.h"
@@ -45,6 +46,8 @@
 #include "UObject/UObjectGlobals.h"
 #include "UObject/UObjectHash.h"
 
+#include <atomic>
+
 namespace UE::Cook
 {
 
@@ -2994,26 +2997,38 @@ void FPackageDatas::PollPendingCookedPlatformDatas(bool bForce, double& LastCook
 		GShaderCompilingManager->ProcessAsyncResults(TimeSlice, false /* bBlockOnGlobalShaderCompletion */);
 	}
 
+	// Future proof as each callback thread gets its own thread_local scope so that the thread_local state
+	// (GCookActivePackageName, FInheritedContextExtensionScope) is always manipulated on the correct thread.
+	// The atomic counter tracks the total number of open scopes across all threads so we can assert they are
+	// all closed by the time ProcessAsyncTasks returns.
+	static thread_local TOptional<UE::Cook::FScopedCookActivePackage> ActivePackageScope;
+	std::atomic<int32> OpenScopeCount{ 0 };
 	FDelegateHandle EventHandle = FAssetCompilingManager::Get().OnPackageScopeEvent().AddLambda(
-		[this](UPackage* Package, bool bEntering)
+		[&OpenScopeCount](UPackage* Package, bool bEntering)
 		{
 			if (bEntering)
 			{
-				CookOnTheFlyServer.SetActivePackage(Package->GetFName(),
+				// We need the scope to always be properly balanced
+				check(!ActivePackageScope.IsSet());
+				ActivePackageScope.Emplace(Package->GetFName(),
 #if UE_WITH_OBJECT_HANDLE_TRACKING
 					PackageAccessTrackingOps::NAME_CookerBuildObject
 #else
 					FName()
 #endif
 				);
+				OpenScopeCount.fetch_add(1, std::memory_order_relaxed);
 			}
 			else
 			{
-				CookOnTheFlyServer.ClearActivePackage();
+				ActivePackageScope.Reset();
+				OpenScopeCount.fetch_sub(1, std::memory_order_relaxed);
 			}
 		});
 	FAssetCompilingManager::Get().ProcessAsyncTasks(true);
 	FAssetCompilingManager::Get().OnPackageScopeEvent().Remove(EventHandle);
+	// Validate that we don't have any open scopes left after exiting
+	check(OpenScopeCount.load(std::memory_order_relaxed) == 0);
 
 	if (LastCookableObjectTickTime + TickCookableObjectsFrameTime <= CurrentTime)
 	{
diff --git a/Engine/Source/Editor/UnrealEd/Private/Cooker/CookSavePackage.cpp b/Engine/Source/Editor/UnrealEd/Private/Cooker/CookSavePackage.cpp
--- a/Engine/Source/Editor/UnrealEd/Private/Cooker/CookSavePackage.cpp
+++ b/Engine/Source/Editor/UnrealEd/Private/Cooker/CookSavePackage.cpp
@@ -6,6 +6,7 @@
 #include "AssetRegistry/IAssetRegistry.h"
 #include "Async/InheritedGuardValue.h"
 #include "Commandlets/AssetRegistryGenerator.h"
+#include "Cooker/CookActivePackage.h"
 #include "Cooker/CookDependency.h"
 #include "Cooker/CookDeterminismManager.h"
 #include "Cooker/CookDiagnostics.h"
@@ -36,6 +37,7 @@
 #include "Interfaces/ITargetPlatform.h"
 #include "Interfaces/ITargetPlatformManagerModule.h"
 #include "Misc/PackageAccessTracking.h"
+#include "Misc/PackageAccessTrackingOps.h"
 #include "Misc/Paths.h"
 #include "Misc/RedirectCollector.h"
 #include "Logging/LogMacros.h"
@@ -117,7 +119,7 @@ void UCookOnTheFlyServer::SaveCookedPackage(UE::Cook::FSaveCookedPackageContext&
 #endif
 				{
 					LLM_SCOPE_BYTAG(Cooker_SavePackage);
-					FScopedActivePackage ScopedActivePackage(*this, Context.PackageData.GetPackageName(), NAME_None);
+					UE::Cook::FScopedCookActivePackage ScopedActivePackage(Context.PackageData.GetPackageName(), NAME_None);
 					if (bSkipSave)
 					{
 						Context.SavePackageResult = ESavePackageResult::Success;
