Fix missing Velocity shader in Material Editor Preview Viewport caused by FVelocity being templated but the name in FPreviewMaterial::ShouldCache still pointing to FVelocity.

[RN] minor | Rendering Architecture
#Synced-CL 51712365

diff --git a/Engine/Source/Editor/UnrealEd/Private/PreviewMaterial.cpp b/Engine/Source/Editor/UnrealEd/Private/PreviewMaterial.cpp
--- a/Engine/Source/Editor/UnrealEd/Private/PreviewMaterial.cpp
+++ b/Engine/Source/Editor/UnrealEd/Private/PreviewMaterial.cpp
@@ -210,7 +210,7 @@ public:
 			{
 				bShaderTypeMatches = true;
 			}
-			else if (FCString::Stristr(ShaderType->GetName(), TEXT("FVelocity")))
+			else if (FCString::Stristr(ShaderType->GetName(), TEXT("TVelocity")))
 			{
 				bShaderTypeMatches = true;
 			}
