SaveCookedPackage: restore accidentally-deleted comment.
#rnx
#Synced-CL 51945718

diff --git a/Engine/Source/Editor/UnrealEd/Private/Cooker/CookSavePackage.cpp b/Engine/Source/Editor/UnrealEd/Private/Cooker/CookSavePackage.cpp
--- a/Engine/Source/Editor/UnrealEd/Private/Cooker/CookSavePackage.cpp
+++ b/Engine/Source/Editor/UnrealEd/Private/Cooker/CookSavePackage.cpp
@@ -82,10 +82,14 @@ void UCookOnTheFlyServer::SaveCookedPackage(UE::Cook::FSaveCookedPackageContext&
 
 	Context.SetupPackage();
 
-	TInheritedGuardValueAccessors<GetOutputCookingWarnings, SetOutputCookingWarnings> ScopedOutputCookerWarnings(IsCookFlagSet(ECookInitializationFlags::OutputVerboseCookerWarnings));
+	TInheritedGuardValueAccessors<GetOutputCookingWarnings, SetOutputCookingWarnings>
+		ScopedOutputCookerWarnings(IsCookFlagSet(ECookInitializationFlags::OutputVerboseCookerWarnings));
+	// SavePackage can CollectGarbage, so we need to store the currently-unqueued PackageData in a separate variable
+	// that we register for garbage collection
 	TGuardValue<UE::Cook::FPackageData*> ScopedSavingPackageData(SavingPackageData, &Context.PackageData);
 	TGuardValue<bool> ScopedIsSavingPackage(bIsSavingPackage, true);
-	// For legacy reasons we set GIsCookerLoadingPackage == true during save. Some classes use it to conditionally execute cook operations in both save and load
+	// For legacy reasons we set GIsCookerLoadingPackage == true during save. Some classes use it to conditionally
+	// execute cook operations in both save and load
 	TGuardValue<bool> ScopedIsCookerLoadingPackage(GIsCookerLoadingPackage, true);
 
 	for (int32 PlatformIndex = 0; PlatformIndex < Context.PlatformsForPackage.Num(); ++PlatformIndex)
