[Viewport ITF] Refine 51898335 by using an existing function.

#rnx
#Synced-CL 51825263

diff --git a/Engine/Source/Editor/Experimental/EditorInteractiveToolsFramework/Private/ViewportInteractions/ViewportInteractionsUtilities.cpp b/Engine/Source/Editor/Experimental/EditorInteractiveToolsFramework/Private/ViewportInteractions/ViewportInteractionsUtilities.cpp
--- a/Engine/Source/Editor/Experimental/EditorInteractiveToolsFramework/Private/ViewportInteractions/ViewportInteractionsUtilities.cpp
+++ b/Engine/Source/Editor/Experimental/EditorInteractiveToolsFramework/Private/ViewportInteractions/ViewportInteractionsUtilities.cpp
@@ -58,17 +58,4 @@ void AddDefaultDragToolsInteractions(UViewportInteractionsBehaviorSource* InInte
 	}
 }
 
-bool UsesViewportInteractions(const FEditorModeTools* ModeTools)
-{
-	if (ModeTools && ModeTools->UsesViewportITF())
-	{
-		if (const UModeManagerInteractiveToolsContext* InputToolsBehaviorSource = ModeTools->GetInteractiveToolsContext())
-		{
-			return InputToolsBehaviorSource->UsesViewportInteractions();
-		}
-	}
-	return false;
-}
-
-
 }
diff --git a/Engine/Source/Editor/Experimental/EditorInteractiveToolsFramework/Public/ViewportInteractions/ViewportInteractionsUtilities.h b/Engine/Source/Editor/Experimental/EditorInteractiveToolsFramework/Public/ViewportInteractions/ViewportInteractionsUtilities.h
--- a/Engine/Source/Editor/Experimental/EditorInteractiveToolsFramework/Public/ViewportInteractions/ViewportInteractionsUtilities.h
+++ b/Engine/Source/Editor/Experimental/EditorInteractiveToolsFramework/Public/ViewportInteractions/ViewportInteractionsUtilities.h
@@ -17,9 +17,4 @@ EDITORINTERACTIVETOOLSFRAMEWORK_API void AddDefaultCameraMovementInteractions(UV
 
 EDITORINTERACTIVETOOLSFRAMEWORK_API void AddDefaultDragToolsInteractions(UViewportInteractionsBehaviorSource* InInteractionsBehaviorSource);
 
-/**
- * Returns true if the provided mode tools is valid and uses viewport interactions.
- */
-EDITORINTERACTIVETOOLSFRAMEWORK_API bool UsesViewportInteractions(const FEditorModeTools* ModeTools);
-
 }
diff --git a/Engine/Source/Editor/UnrealEd/Private/EditorViewportClient.cpp b/Engine/Source/Editor/UnrealEd/Private/EditorViewportClient.cpp
--- a/Engine/Source/Editor/UnrealEd/Private/EditorViewportClient.cpp
+++ b/Engine/Source/Editor/UnrealEd/Private/EditorViewportClient.cpp
@@ -1731,7 +1731,7 @@ void FEditorViewportClient::Tick(float DeltaTime)
 		// Update any real-time camera movement
 		UpdateCameraMovement( DeltaTime );
 
-		if (UE::Editor::ViewportInteractions::UsesViewportInteractions(ModeTools.Get()))
+		if (ModeTools && ModeTools->UsesViewportITF())
 		{
 			ConsumeExternalCameraMovement();
 		}
@@ -2103,7 +2103,7 @@ bool FEditorViewportClient::IsActiveViewportType(ELevelViewportType InViewportTy
 // Updates real-time camera movement.  Should be called every viewport tick!
 void FEditorViewportClient::UpdateCameraMovement( float DeltaTime )
 {
-	const bool bIsITFInteractions = UE::Editor::ViewportInteractions::UsesViewportInteractions(ModeTools.Get());
+	const bool bIsITFInteractions = ModeTools->UsesViewportITF();
 	
 	// We only want to move perspective cameras around like this
 	if( Viewport && IsPerspective() && (bIsITFInteractions ? !IsViewportInteractionActive(ModeTools.Get(), UE::Editor::ViewportInteractions::Orbit) : !ShouldOrbitCamera()))
@@ -3634,7 +3634,7 @@ bool FEditorViewportClient::Internal_InputKey(const FInputKeyEventArgs& EventArg
 	const EInputEvent& Event = EventArgs.Event;
 	const FViewport* InViewport = EventArgs.Viewport;
 
-	bool bRouteToToolsContext = !bIsTracking || UE::Editor::ViewportInteractions::UsesViewportInteractions(ModeTools.Get());
+	bool bRouteToToolsContext = !bIsTracking || ModeTools->UsesViewportITF();
 	// Let the current mode have a look at the input before reacting to it. 
 	// Note that bIsTracking tells us whether the viewport client is capturing mouse behavior to fly around,
 	// move objects, etc. In this case we don't want to pass the input the input router, because it is already
@@ -5746,7 +5746,7 @@ bool FEditorViewportClient::IsFlightCameraInputModeActive() const
 	{
 		if( CameraController != NULL )
 		{
-			if (UE::Editor::ViewportInteractions::UsesViewportInteractions(ModeTools.Get()))
+			if (ModeTools->UsesViewportITF())
 			{
 				return ViewportNavigationHelper->bIsMouseLooking;
 			}
diff --git a/Engine/Source/Editor/UnrealEd/Private/LevelEditorViewport.cpp b/Engine/Source/Editor/UnrealEd/Private/LevelEditorViewport.cpp
--- a/Engine/Source/Editor/UnrealEd/Private/LevelEditorViewport.cpp
+++ b/Engine/Source/Editor/UnrealEd/Private/LevelEditorViewport.cpp
@@ -3304,8 +3304,10 @@ bool FLevelEditorViewportClient::InputWidgetDelta(FViewport* InViewport, EAxisLi
 				const bool RightMouseButtonDown = InViewport->KeyState(EKeys::RightMouseButton);
 				const bool MiddleMouseButtonDown = InViewport->KeyState(EKeys::MiddleMouseButton);
 
+				const bool bUsesViewportITF = ModeTools->UsesViewportITF();
+
 				// If duplicate dragging and using legacy viewport interaction behavior 
-				if (IsAltPressed() && (LeftMouseButtonDown || RightMouseButtonDown) && !UE::Editor::ViewportInteractions::UsesViewportInteractions(ModeTools.Get()))
+				if (IsAltPressed() && (LeftMouseButtonDown || RightMouseButtonDown) && !bUsesViewportITF)
 				{
 					// The widget has been offset, so check if we should duplicate the selection.
 					if ( bDuplicateOnNextDrag )
@@ -3321,8 +3323,6 @@ bool FLevelEditorViewportClient::InputWidgetDelta(FViewport* InViewport, EAxisLi
 					}
 				}
 
-				const bool bUsesViewportITF = ModeTools->UsesViewportITF();
-
 				// We do not want actors updated if we are holding down the middle mouse button.
 				// enable MMB for New TRS Gizmos
 				const bool bModifyTransform = [this, MiddleMouseButtonDown]() -> bool
diff --git a/Engine/Source/Editor/UnrealEd/Private/MouseDeltaTracker.cpp b/Engine/Source/Editor/UnrealEd/Private/MouseDeltaTracker.cpp
--- a/Engine/Source/Editor/UnrealEd/Private/MouseDeltaTracker.cpp
+++ b/Engine/Source/Editor/UnrealEd/Private/MouseDeltaTracker.cpp
@@ -346,7 +346,7 @@ bool FMouseDeltaTracker::EndTracking(FEditorViewportClient* InViewportClient)
 
 void FMouseDeltaTracker::ConditionalBeginUsingDragTool( FEditorViewportClient* InViewportClient )
 {
-	if (UE::Editor::ViewportInteractions::UsesViewportInteractions(InViewportClient->GetModeTools()))
+	if (InViewportClient->GetModeTools()->UsesViewportITF())
 	{
 		return;
 	}
