UE: Bump cook incremental version. Function change not picked up

diff --git a/Engine/Source/Editor/UnrealEd/Private/Cooker/CookIncrementalVersion.cpp b/Engine/Source/Editor/UnrealEd/Private/Cooker/CookIncrementalVersion.cpp
--- a/Engine/Source/Editor/UnrealEd/Private/Cooker/CookIncrementalVersion.cpp
+++ b/Engine/Source/Editor/UnrealEd/Private/Cooker/CookIncrementalVersion.cpp
@@ -8,6 +8,6 @@ namespace UE::Cook
 {
 
 // Change CookIncrementalVersion to a new guid when all packages in an incremental cook need to be invalidated
-FGuid CookIncrementalVersion(0xFFB442DC, 0x88224FE2, 0xBEA2627A, 0x1D71D307);
+FGuid CookIncrementalVersion(0xFFB442DC, 0x58224FE2, 0xBEA2627A, 0x1A71D307);
 
 }
