FGlobalEditorCommonCommands : Call GetActiveTopLevelRegularWindow() in ::OpenPopupMenu (ctrl+tab) when determining the ParentWindow for the popup.

This fixes a crash where GetActiveTopLevelWindow() returns a SVirtualWindow (bIsPopupWindow == true) after starting PIE.
The ParentWindow for the PopupMenu should be a regular window and not a popup window.

#Synced-CL 51954903

diff --git a/Engine/Source/Editor/UnrealEd/Private/Toolkits/GlobalEditorCommonCommands.cpp b/Engine/Source/Editor/UnrealEd/Private/Toolkits/GlobalEditorCommonCommands.cpp
--- a/Engine/Source/Editor/UnrealEd/Private/Toolkits/GlobalEditorCommonCommands.cpp
+++ b/Engine/Source/Editor/UnrealEd/Private/Toolkits/GlobalEditorCommonCommands.cpp
@@ -209,7 +209,7 @@ void FGlobalEditorCommonCommands::OnSummonedAssetPicker()
 TSharedPtr<IMenu> FGlobalEditorCommonCommands::OpenPopupMenu(TSharedRef<SWidget> WindowContents, const FVector2D& PopupDesiredSize)
 {
 	// Determine where the pop-up should open
-	TSharedPtr<SWindow> ParentWindow = FSlateApplication::Get().GetActiveTopLevelWindow();
+	TSharedPtr<SWindow> ParentWindow = FSlateApplication::Get().GetActiveTopLevelRegularWindow();
 	FVector2D WindowPosition = FSlateApplication::Get().GetCursorPos();
 	if (!ParentWindow.IsValid())
 	{
