Updated the DerivedDataBuildRemoteExecutor to use new endpoints and improved scheduling logic

- Added DerivedDataBuildRemoteExecutor.bContentNegotiationEnabled tweakable which controls whether actions go through a two-phase or one-phase submit. Since we are initially dealing with very small payloads which are typically ephemeral, and no cloud DDC or other remote data store, the default is one-phase submit.
- Fixed bug in TryStartNewBuild() which didn't handle the bEnableLimits setting correctly
- Fixed bug in FRemoteBuildExecutionRequest::PostWorkerPackage_Async where it did not update the correct stats (TotalActionPackagesUploaded -> TotalWorkerPackagesUploaded)
- Added support for per-action priority assignment, which is taken into account downstream

#rnx
#Synced-CL 51955189

diff --git a/Engine/Source/Editor/UnrealEd/Private/DerivedDataBuildRemoteExecutor.cpp b/Engine/Source/Editor/UnrealEd/Private/DerivedDataBuildRemoteExecutor.cpp
--- a/Engine/Source/Editor/UnrealEd/Private/DerivedDataBuildRemoteExecutor.cpp
+++ b/Engine/Source/Editor/UnrealEd/Private/DerivedDataBuildRemoteExecutor.cpp
@@ -43,6 +43,21 @@ DEFINE_LOG_CATEGORY_STATIC(LogDerivedDataBuildRemoteExecutor, Log, All);
 
 class FRemoteBuildWorkerExecutor;
 
+/**
+ * This implements a remote execution interface for Derived Data Build
+ * using Zen Structured Compute interfaces to get work done. 
+ 
+ * All build actions are typically submitted to the same Zen server which 
+ * is responsible for build outputs (Project Store) and local DDC. This Zen 
+ * server is responsible for dispatching the work elsewhere but can also 
+ * execute it locally when appropriate.
+ * 
+ * The current implementation uses the "legacy" REST API for simplicity,
+ * but should be upgraded to use the new queue-oriented API as soon
+ * as possible since that is more efficient and more flexible.
+ * 
+ */
+
 class FRemoteBuildExecutionRequest final : public FRequestBase
 {
 public:
@@ -113,6 +128,7 @@ private:
 
 	// General utility methods
 	static FCbObject BuildWorkerDescriptor(const FBuildWorker& Worker, const int TimeoutSeconds);
+	void PrepareActionForPost();
 
 	bool ProcessCancellation();
 	bool IsResultOk(const UE::Zen::FZenHttpRequest::Result& Result, const TCHAR* OperationDesc);
@@ -149,6 +165,7 @@ public:
 
 		bool bConfigEnabled = false;
 		GConfig->GetBool(TEXT("DerivedDataBuildRemoteExecutor"), TEXT("bEnabled"), bConfigEnabled, GEngineIni);
+		GConfig->GetBool(TEXT("DerivedDataBuildRemoteExecutor"), TEXT("bContentNegotiationEnabled"), bContentNegotiationEnabled, GEngineIni);
 		GConfig->GetInt(TEXT("DerivedDataBuildRemoteExecutor"), TEXT("GlobalExecutionTimeoutSeconds"), GlobalExecutionTimeoutSeconds, GEngineIni);
 
 		if (bConfigEnabled || FParse::Param(FCommandLine::Get(), TEXT("DDC2RemoteExecution")))
@@ -203,12 +220,6 @@ public:
 		FOnBuildWorkerActionComplete&& OnComplete) final
 	{
 		{
-			// TODO: This block forces resolution of inputs before we attempt to determine which
-			//		 inputs need to be uploaded.  This is required because we can't refer to inputs
-			//		 in the Merkle tree by their RawHash/RawSize but instead must send their CompressedHash/
-			//		 CompressedSize.  Once the remote execution API allows us to represent inputs with RawHash/
-			//		 RawSize, this block can be removed and we can find missing CAS inputs without having resolved
-			//		 the inputs first.
 			TArray<FUtf8StringView> MissingInputs;
 			uint64 TotalInputSize = 0;
 			uint64 TotalMissingInputSize = 0;
@@ -294,6 +305,11 @@ public:
 		ProcessingThreadEvent->Trigger();
 	}
 
+	bool IsContentNegotiationEnabled() const
+	{
+		return bContentNegotiationEnabled;
+	}
+
 private:
 	struct FStats
 	{
@@ -420,13 +436,18 @@ private:
 
 		bool TryStartNewBuild(FStats& InStats)
 		{
-			if ((InStats.TotalRemoteBuilds.fetch_add(1, std::memory_order_relaxed) >= MaxTotalRemoteBuilds) && bEnableLimits)
+			if (!bEnableLimits)
+			{
+				return true;
+			}
+
+			if ((InStats.TotalRemoteBuilds.fetch_add(1, std::memory_order_relaxed) >= MaxTotalRemoteBuilds))
 			{
 				InStats.TotalRemoteBuilds.fetch_sub(1, std::memory_order_relaxed);
 				return false;
 			}
 
-			if ((InStats.InFlightRemoteBuilds.fetch_add(1, std::memory_order_relaxed) >= MaxInFlightRemoteBuilds) && bEnableLimits)
+			if ((InStats.InFlightRemoteBuilds.fetch_add(1, std::memory_order_relaxed) >= MaxInFlightRemoteBuilds))
 			{
 				InStats.TotalRemoteBuilds.fetch_sub(1, std::memory_order_relaxed);
 				InStats.InFlightRemoteBuilds.fetch_sub(1, std::memory_order_relaxed);
@@ -461,6 +482,7 @@ private:
 	TUniquePtr<UE::Zen::FScopeZenService> ScopeZenService;
 	TUniquePtr<UE::Zen::FZenHttpRequestPool> RequestPool;
 	bool bEnabled;
+	bool bContentNegotiationEnabled = false;
 	std::atomic<bool> bProcessingThreadRunning;
 };
 
@@ -583,7 +605,7 @@ bool FRemoteBuildExecutionRequest::IsResultOk(const UE::Zen::FZenHttpRequest::Re
 {
 	if (Result == Zen::FZenHttpRequest::Result::Failed)
 	{
-		UE_LOG(LogDerivedDataBuildRemoteExecutor, Warning, TEXT("Remote execution system error: operation '%s' produced an failed result!"), OperationDesc);
+		UE_LOG(LogDerivedDataBuildRemoteExecutor, Warning, TEXT("Remote execution system error: operation '%s' produced a failed result!"), OperationDesc);
 
 		State.Owner.End(this, [this]
 		{
@@ -619,7 +641,7 @@ void FRemoteBuildExecutionRequest::DetermineIfWorkerExists_Async()
 			State.WorkerDescriptionId = State.WorkerDescriptor.GetHash();
 
 			TStringBuilder<128> WorkerUri;
-			WorkerUri.Append("/apply/workers/");
+			WorkerUri.Append("/compute/workers/");
 			WorkerUri << State.WorkerDescriptionId;
 
 			UE::Zen::FZenScopedRequestPtr Request(Executor.RequestPool.Get(), false);
@@ -636,7 +658,7 @@ void FRemoteBuildExecutionRequest::PostWorkerObject_Async()
 	UE::Tasks::Launch(TEXT("FRemoteBuildExecutionRequest::PostWorkerObject"), [this]
 		{
 			TStringBuilder<128> WorkerUri;
-			WorkerUri.Append("/apply/workers/");
+			WorkerUri.Append("/compute/workers/");
 			WorkerUri << State.WorkerDescriptionId;
 
 			Executor.Stats.TotalWorkerObjectsUploaded.AddBlob(State.WorkerDescriptor.GetSize());
@@ -668,7 +690,7 @@ void FRemoteBuildExecutionRequest::PostWorkerPackage_Async()
 	UE::Tasks::Launch(TEXT("FRemoteBuildExecutionRequest::PostWorkerPackage"), [this]
 		{
 			TStringBuilder<128> WorkerUri;
-			WorkerUri.Append("/apply/workers/");
+			WorkerUri.Append("/compute/workers/");
 			WorkerUri << State.WorkerDescriptionId;
 
 			uint64_t AttachmentBytes{};
@@ -715,7 +737,7 @@ void FRemoteBuildExecutionRequest::PostWorkerPackage_Async()
 
 			Package.SetObject(State.WorkerDescriptor);
 
-			Executor.Stats.TotalActionPackagesUploaded.AddBlob(AttachmentBytes + State.WorkerDescriptor.GetSize());
+			Executor.Stats.TotalWorkerPackagesUploaded.AddBlob(AttachmentBytes + State.WorkerDescriptor.GetSize());
 
 			UE::Zen::FZenScopedRequestPtr Request(Executor.RequestPool.Get(), false);
 			State.Result = Request->PerformBlockingPostPackage(WorkerUri, Package);
@@ -726,20 +748,27 @@ void FRemoteBuildExecutionRequest::PostWorkerPackage_Async()
 	);
 }
 
+void FRemoteBuildExecutionRequest::PrepareActionForPost()
+{
+	if (!State.Action) 
+	{
+		FCbWriter BuildActionWriter;
+		State.BuildAction.Save(BuildActionWriter);
+		State.Action = BuildActionWriter.Save().AsObject();
+		State.ActionId = State.Action.GetHash();
+	}
+}
+
 void FRemoteBuildExecutionRequest::PostActionObject_Async()
 {
 	UE::Tasks::Launch(TEXT("FRemoteBuildExecutionRequest::PostActionObject"), [this]
 		{
-			{
-				FCbWriter BuildActionWriter;
-				State.BuildAction.Save(BuildActionWriter);
-				State.Action = BuildActionWriter.Save().AsObject();
-				State.ActionId = State.Action.GetHash();
-			}
+			PrepareActionForPost();
 
 			TStringBuilder<128> ActionUri;
-			ActionUri.Append("/apply/jobs/");
+			ActionUri.Append("/compute/jobs/");
 			ActionUri << State.WorkerDescriptionId;
+			ActionUri << "?priority=" << (uint8)State.Owner.GetPriority();
 
 			State.NeedHashes.Empty();
 
@@ -769,19 +798,25 @@ void FRemoteBuildExecutionRequest::PostActionPackage_Async()
 {
 	UE::Tasks::Launch(TEXT("FRemoteBuildExecutionRequest::PostActionPackage"), [this]
 		{
+			PrepareActionForPost();
+
 			uint64_t AttachmentBytes{};
 			FCbPackage ActionPackage;
 
+			// It's valid for NeedHashes to be empty, in the case we go ahead and skip
+			// the initial "needs" probe because the inputs are very small so we may
+			// as well go ahead and send all of them at once without prior negotiation.
+
 			State.BuildInputs.Get().IterateInputs([NeedHashes = State.NeedHashes, &ActionPackage, &AttachmentBytes](FUtf8StringView Key, const FCompressedBuffer& Buffer)
 				{
-					if (NeedHashes.Contains(Buffer.GetRawHash()))
+					if (NeedHashes.IsEmpty() || NeedHashes.Contains(Buffer.GetRawHash()))
 					{
 						ActionPackage.AddAttachment(FCbAttachment{ Buffer });
 						AttachmentBytes += Buffer.GetCompressedSize();
 					}
 				});
 
-			if (State.NeedHashes.Num() != ActionPackage.GetAttachments().Num())
+			if (!State.NeedHashes.IsEmpty() && (State.NeedHashes.Num() != ActionPackage.GetAttachments().Num()))
 			{
 				UE_LOG(LogDerivedDataBuildRemoteExecutor, Warning, TEXT("Remote execution system error: build input attachment missing!"));
 				State.Result = UE::Zen::FZenHttpRequest::Result::Failed;
@@ -793,8 +828,9 @@ void FRemoteBuildExecutionRequest::PostActionPackage_Async()
 			ActionPackage.SetObject(State.Action);
 
 			TStringBuilder<128> ActionUri;
-			ActionUri.Append("/apply/jobs/");
+			ActionUri.Append("/compute/jobs/");
 			ActionUri << State.WorkerDescriptionId;
+			ActionUri << "?priority=" << (uint8)State.Owner.GetPriority();
 
 			Executor.Stats.TotalActionPackagesUploaded.AddBlob(AttachmentBytes + State.Action.GetSize());
 
@@ -812,7 +848,7 @@ void FRemoteBuildExecutionRequest::GetResultPackage_Async()
 	UE::Tasks::Launch(TEXT("FRemoteBuildExecutionRequest::GetResultPackage"), [this]
 		{
 			TStringBuilder<128> JobGetUri;
-			JobGetUri.Append("/apply/jobs/");
+			JobGetUri.Append("/compute/jobs/");
 			JobGetUri << State.WorkerDescriptionId;
 			JobGetUri << "/";
 			JobGetUri << State.ActionId;
@@ -828,6 +864,8 @@ void FRemoteBuildExecutionRequest::GetResultPackage_Async()
 			else if (State.ResponseCode == 404)
 			{
 				Executor.Stats.TotalTimedOutRemoteBuilds.fetch_add(1, std::memory_order_relaxed);
+
+				UE_LOG(LogDerivedDataBuildRemoteExecutor, Warning, TEXT("result package missing (uri: %s)"), *JobGetUri);
 			}
 			
 			OnGetResultPackageComplete(State.Result);
@@ -879,7 +917,17 @@ void FRemoteBuildExecutionRequest::OnWorkerExistsDetermined(const UE::Zen::FZenH
 	}
 
 	// Worker exists, proceed to action
-	PostActionObject_Async();
+	if (Executor.IsContentNegotiationEnabled())
+	{
+		PostActionObject_Async();
+	}
+	else
+	{
+		// Skip the content negotiation and just post the action if
+		// the inputs are small enough or if we just know that it's
+		// very unlikely that the remote end has the content.
+		PostActionPackage_Async();
+	}
 }
 
 void FRemoteBuildExecutionRequest::OnPostWorkerObjectComplete(const UE::Zen::FZenHttpRequest::Result& Result)
@@ -944,7 +992,7 @@ void FRemoteBuildExecutionRequest::OnPostActionObjectComplete(const UE::Zen::FZe
 		return;
 	}
 
-	if (!IsResponseOk(State.ResponseCode, TEXT("PostWorkerObject")))
+	if (!IsResponseOk(State.ResponseCode, TEXT("PostActionObject")))
 	{
 		return;
 	}
@@ -976,17 +1024,23 @@ void FRemoteBuildExecutionRequest::OnGetResultPackageComplete(const UE::Zen::FZe
 		return;
 	}
 
-	if (!IsResultOk(Result, TEXT("GetResultPackage")) || !IsResponseOk(State.ResponseCode, TEXT("GetResultPackage")))
+	if (!IsResultOk(Result, TEXT("GetResultPackage")))
 	{
 		return;
 	}
 
-	if (State.ResponseCode == 202)
+	if (State.ResponseCode == 202 /* Accepted */)
 	{
+		// Result is not yet available but the job is pending/running
 		QueueGetResultPackage();
 		return;
 	}
 
+	if (!IsResponseOk(State.ResponseCode, TEXT("GetResultPackage")))
+	{
+		return;
+	}
+
 	// We're done!
 
 	FOptionalBuildOutput RemoteBuildOutput = FBuildOutput::Load(State.BuildAction.GetName(), State.BuildAction.GetFunction(), State.ResultPackage.GetObject());
