Cooker: Remove duplicate declaration of LogCookerMessage.
#rnx
#Synced-CL 51945718

diff --git a/Engine/Source/Editor/UnrealEd/Private/Cooker/CookTypes.h b/Engine/Source/Editor/UnrealEd/Private/Cooker/CookTypes.h
--- a/Engine/Source/Editor/UnrealEd/Private/Cooker/CookTypes.h
+++ b/Engine/Source/Editor/UnrealEd/Private/Cooker/CookTypes.h
@@ -763,7 +763,6 @@ struct FBeginCookContextForWorker
 FCbWriter& operator<<(FCbWriter& Writer, const FBeginCookContextForWorker& Value);
 bool LoadFromCompactBinary(FCbFieldView Field, FBeginCookContextForWorker& Value);
 
-void LogCookerMessage(const FString& MessageText, EMessageSeverity::Type Severity);
 LLM_DECLARE_TAG(Cooker);
 
 #define REMAPPED_PLUGINS TEXTVIEW("RemappedPlugins")
