Avoid compiling blueprints that were loaded for diff

[RN] Blueprints loaded for diffing will no longer attempt to compile before PIE, even if they have been marked as dirty
#Synced-CL 51986075

diff --git a/Engine/Source/Editor/UnrealEd/Private/PlayLevel.cpp b/Engine/Source/Editor/UnrealEd/Private/PlayLevel.cpp
--- a/Engine/Source/Editor/UnrealEd/Private/PlayLevel.cpp
+++ b/Engine/Source/Editor/UnrealEd/Private/PlayLevel.cpp
@@ -1280,6 +1280,12 @@ int32 FInternalPlayLevelUtils::ResolveDirtyBlueprints(const bool bPromptForCompi
 			continue;
 		}
 
+		UPackage* Package = Blueprint->GetPackage();
+		if(Package && Package->HasAnyPackageFlags(PKG_ForDiffing))
+		{
+			continue;
+		}
+
 		// do not try to recompile BPs that have not changed since they last failed to compile, so don't check Blueprint->IsUpToDate()
 		const bool bIsDirtyAndShouldBeRecompiled = Blueprint->IsPossiblyDirty();
 		if (!FBlueprintEditorUtils::IsDataOnlyBlueprint(Blueprint)
