EditorUsability : UnrealEd
- Fixed General Appearance - Editor Preference : Import values are not saved

#Synced-CL 51539177

diff --git a/Engine/Source/Editor/UnrealEd/Private/Settings/EditorStyleClasses.cpp b/Engine/Source/Editor/UnrealEd/Private/Settings/EditorStyleClasses.cpp
--- a/Engine/Source/Editor/UnrealEd/Private/Settings/EditorStyleClasses.cpp
+++ b/Engine/Source/Editor/UnrealEd/Private/Settings/EditorStyleClasses.cpp
@@ -127,6 +127,9 @@ bool UEditorStyleSettings::OnImportBegin(const FString& ImportFromPath)
 	{
 		FEditorStyleSettingsCustomization::PromptToImportTheme(ImportFromPath); 
 	}
+
+	SaveConfig();
+
 	return true; 
 }
 
