[Gizmos] Fix an issue where the tools context was not updating its state when various gizmo cvars were changing

diff --git a/Engine/Source/Editor/UnrealEd/Private/Tools/EdModeInteractiveToolsContext.cpp b/Engine/Source/Editor/UnrealEd/Private/Tools/EdModeInteractiveToolsContext.cpp
--- a/Engine/Source/Editor/UnrealEd/Private/Tools/EdModeInteractiveToolsContext.cpp
+++ b/Engine/Source/Editor/UnrealEd/Private/Tools/EdModeInteractiveToolsContext.cpp
@@ -1450,8 +1450,9 @@ void UModeManagerInteractiveToolsContext::Initialize(IToolsContextQueriesAPI* Qu
 	GetParentEditorModeManager()->OnEditorModeIDChanged().AddUObject(this, &UModeManagerInteractiveToolsContext::OnEditorModeChanged);
 
 	// Delegates used to handle InputTools.EnableITFTools CVAR changes
-	UE::Editor::ViewportInteractions::OnEditorViewportInteractionsActivated().AddUObject(this, &UModeManagerInteractiveToolsContext::RegisterViewportInteractions);
-	UE::Editor::ViewportInteractions::OnEditorViewportInteractionsDeactivated().AddUObject(this, &UModeManagerInteractiveToolsContext::UnregisterViewportInteractions);
+	UE::Editor::ViewportInteractions::OnEditorViewportInteractionsActivated().AddWeakLambda(this, [this]() { bViewportInteractionsDirty = true; });
+	UE::Editor::ViewportInteractions::OnEditorViewportInteractionsDeactivated().AddWeakLambda(this, [this]() { bViewportInteractionsDirty = true; });
+	UEditorInteractiveGizmoManager::OnUsesNewTRSGizmosChangedDelegate().AddWeakLambda(this, [this](const bool) { bViewportInteractionsDirty = true; });
 	
 	CurrentMouseState.OnGetButtonStateFallback.BindUObject(this, &UModeManagerInteractiveToolsContext::OnGetButtonStateFallback);
 	
@@ -1466,8 +1467,9 @@ void UModeManagerInteractiveToolsContext::Shutdown()
 	UnregisterViewportInteractions();
 
 	GetParentEditorModeManager()->OnEditorModeIDChanged().RemoveAll(this);
-	UE::Editor::DragTools::OnEditorDragToolsActivated().RemoveAll(this);
-	UE::Editor::DragTools::OnEditorDragToolsDeactivated().RemoveAll(this);
+	UEditorInteractiveGizmoManager::OnUsesNewTRSGizmosChangedDelegate().RemoveAll(this);
+	UE::Editor::ViewportInteractions::OnEditorViewportInteractionsActivated().RemoveAll(this);
+	UE::Editor::ViewportInteractions::OnEditorViewportInteractionsDeactivated().RemoveAll(this);
 
 	if (ToolManager)
 	{
@@ -1554,6 +1556,7 @@ void UModeManagerInteractiveToolsContext::SetViewportInteractionsEnabled(bool bI
 {
 	if (bUsesViewportInteractions != bInEnabled)
 	{
+		bViewportInteractionsDirty = true;
 		bUsesViewportInteractions = bInEnabled;
 		if (!bInEnabled)
 		{
