[Viewport Interactions] Pass TEDS selection set name through Element click handler

diff --git a/Engine/Source/Editor/UnrealEd/Private/ViewportSelectionUtilities.cpp b/Engine/Source/Editor/UnrealEd/Private/ViewportSelectionUtilities.cpp
--- a/Engine/Source/Editor/UnrealEd/Private/ViewportSelectionUtilities.cpp
+++ b/Engine/Source/Editor/UnrealEd/Private/ViewportSelectionUtilities.cpp
@@ -1,4 +1,4 @@
-﻿// Copyright Epic Games, Inc. All Rights Reserved.
+// Copyright Epic Games, Inc. All Rights Reserved.
 
 #include "ViewportSelectionUtilities.h"
 #include "ActorEditorUtils.h"
@@ -135,7 +135,8 @@ bool ClickElement(FEditorViewportClient* InViewportClient, const FTypedElementHa
 			bHandledClick = true;
 
 			const FTypedElementSelectionOptions SelectionOptions =
-				FTypedElementSelectionOptions().SetAllowHidden(true).SetWarnIfLocked(true);
+				FTypedElementSelectionOptions().SetAllowHidden(true).SetWarnIfLocked(true)
+					.SetNameForTEDSIntegration(ElementSelectionSet->GetNameForTedsIntegration());
 
 			bool bNeedViewportRefresh = false;
 
@@ -148,7 +149,9 @@ bool ClickElement(FEditorViewportClient* InViewportClient, const FTypedElementHa
 				if (Click.IsControlDown() && bAllowSelectionModifiers)
 				{
 					if (ElementSelectionSet->IsElementSelected(
-							ResolvedElement, FTypedElementIsSelectedOptions().SetAllowIndirect(true)
+							ResolvedElement,
+							FTypedElementIsSelectedOptions().SetAllowIndirect(true)
+								.SetNameForTEDSIntegration(ElementSelectionSet->GetNameForTedsIntegration())
 						))
 					{
 						ElementSelectionSet->DeselectElement(ResolvedElement, SelectionOptions);
@@ -167,7 +170,9 @@ bool ClickElement(FEditorViewportClient* InViewportClient, const FTypedElementHa
 					// Skip the clear if we're doing a RMB select and this actor is already selected, as we want to summon the menu for the current selection
 					if (bIsLeftClickSelection
 						|| !ElementSelectionSet->IsElementSelected(
-							ResolvedElement, FTypedElementIsSelectedOptions().SetAllowIndirect(true)
+							ResolvedElement,
+							FTypedElementIsSelectedOptions().SetAllowIndirect(true)
+								.SetNameForTEDSIntegration(ElementSelectionSet->GetNameForTedsIntegration())
 						))
 					{
 						bNeedViewportRefresh =
