Cooker: Diagnostics: Add immediate mode for IncrementallyModifiedDiagnostics. Change commandline parsing to use new cook.diagnostic.modified naming convention.
#rnx
#Synced-CL 52052889

diff --git a/Engine/Source/Editor/UnrealEd/Classes/CookOnTheSide/CookOnTheFlyServer.h b/Engine/Source/Editor/UnrealEd/Classes/CookOnTheSide/CookOnTheFlyServer.h
--- a/Engine/Source/Editor/UnrealEd/Classes/CookOnTheSide/CookOnTheFlyServer.h
+++ b/Engine/Source/Editor/UnrealEd/Classes/CookOnTheSide/CookOnTheFlyServer.h
@@ -1613,7 +1613,6 @@ private:
 	bool bCookIncrementalAllowAllClasses = false;
 	bool bHiddenDependenciesDebug = false;
 	bool bIncrementallyModifiedDiagnostics = false;
-	bool bWriteIncrementallyModifiedDiagnosticsToLogs = false;
 	bool bOnlyEditorOnlyDebug = false;
 	bool bHiddenDependenciesClassPathFilterListIsAllowList = true;
 	bool bFirstCookInThisProcessInitialized = false;
diff --git a/Engine/Source/Editor/UnrealEd/Private/CookOnTheFlyServer.cpp b/Engine/Source/Editor/UnrealEd/Private/CookOnTheFlyServer.cpp
--- a/Engine/Source/Editor/UnrealEd/Private/CookOnTheFlyServer.cpp
+++ b/Engine/Source/Editor/UnrealEd/Private/CookOnTheFlyServer.cpp
@@ -6518,8 +6518,8 @@ void UCookOnTheFlyServer::SetInitializeConfigSettings(UE::Cook::FInitializeConfi
 	FParse::Value(CommandLine, TEXT("-CookCVarControl="), OverrideCookCVarControl);
 	FParse::Value(CommandLine, TEXT("-DeviceProfile="), OverrideDeviceProfileName);
 
-	bIncrementallyModifiedDiagnostics = FParse::Param(CommandLine, TEXT("CookIncrementallyModifiedDiagnostics"));
-	bWriteIncrementallyModifiedDiagnosticsToLogs = FParse::Param(CommandLine, TEXT("WriteCookIncrementallyModifiedDiagnosticsToLogs"));
+	bIncrementallyModifiedDiagnostics = FParse::Param(CommandLine, TEXT("CookIncrementallyModifiedDiagnostics"))
+		|| FParse::Param(CommandLine, TEXT("Cook.Diagnostic.Modified"));
 
 	bCookingDLCBaseGamePlugin = FParse::Param(CommandLine, TEXT("CookingDLCBaseGamePlugin"));
 }
@@ -9697,30 +9697,19 @@ void UCookOnTheFlyServer::CookByTheBookFinishedInternal()
 		}
 	}
 
-	if (!PlatformManager->GetSessionPlatforms().IsEmpty())
+	if (!PlatformManager->GetSessionPlatforms().IsEmpty() && IncrementallyModifiedDiagnostics)
 	{
-		if (IncrementallyModifiedDiagnostics)
+		const ITargetPlatform* TargetPlatform = PlatformManager->GetSessionPlatforms()[0];
+		FString DiagnosticsDir;
+		if (GIsBuildMachine)
 		{
-			if (bWriteIncrementallyModifiedDiagnosticsToLogs)
-			{
-				IncrementallyModifiedDiagnostics->WriteToLogs();
-			}
-			else
-			{
-				const ITargetPlatform* TargetPlatform = PlatformManager->GetSessionPlatforms()[0];
-				FString DiagnosticsDir;
-				
-				if (GIsBuildMachine)
-				{
-					DiagnosticsDir = FPaths::Combine(*FPaths::EngineDir(), TEXT("Programs"), TEXT("AutomationTool"), TEXT("Saved"), TEXT("Logs"), TargetPlatform->PlatformName());
-				}
-				else
-				{
-					DiagnosticsDir = GetCookMetadataOutputFolder(TargetPlatform);
-				}
-				IncrementallyModifiedDiagnostics->WriteToFile(DiagnosticsDir);
-			}
+			DiagnosticsDir = FPaths::Combine(*FPaths::EngineDir(), TEXT("Programs"), TEXT("AutomationTool"), TEXT("Saved"), TEXT("Logs"), TargetPlatform->PlatformName());
+		}
+		else
+		{
+			DiagnosticsDir = GetCookMetadataOutputFolder(TargetPlatform);
 		}
+		IncrementallyModifiedDiagnostics->WriteFinal(DiagnosticsDir);
 	}
 
 	UE_LOG(LogCookStatus, Display, TEXT("Shutting down shader library cooker & shader compilers"));
diff --git a/Engine/Source/Editor/UnrealEd/Private/Cooker/CookDiagnostics.cpp b/Engine/Source/Editor/UnrealEd/Private/Cooker/CookDiagnostics.cpp
--- a/Engine/Source/Editor/UnrealEd/Private/Cooker/CookDiagnostics.cpp
+++ b/Engine/Source/Editor/UnrealEd/Private/Cooker/CookDiagnostics.cpp
@@ -14,8 +14,11 @@
 #include "CookOnTheSide/CookLog.h"
 #include "CookOnTheSide/CookOnTheFlyServer.h"
 #include "Logging/LogMacros.h"
+#include "Logging/StructuredLog.h"
+#include "Misc/CommandLine.h"
 #include "Misc/FileHelper.h"
 #include "Misc/Optional.h"
+#include "Misc/Parse.h"
 #include "Misc/Paths.h"
 #include "Misc/StringBuilder.h"
 #include "Serialization/ArchiveSerializedPropertyChain.h"
@@ -457,6 +460,53 @@ const TCHAR* LexToString(UE::Cook::EIncrementallyModifiedReason Reason)
 	}
 }
 
+FIncrementallyModifiedDiagnostics::FIncrementallyModifiedDiagnostics()
+{
+	const TCHAR* CommandLine = FCommandLine::Get();
+
+	bool bCommandLineOptionSet = false;
+
+	if (FParse::Param(CommandLine, TEXT("WriteCookIncrementallyModifiedDiagnosticsToLogs"))
+		|| FParse::Param(CommandLine, TEXT("Cook.Diagnostic.Modified.FinalLog")))
+	{
+		bWriteFinalLog = true;
+		bCommandLineOptionSet = true;
+	}
+	bCommandLineOptionSet |= FParse::Bool(CommandLine, TEXT("-Cook.Diagnostic.Modified.FinalLog="), bWriteFinalLog);
+	if (FParse::Param(CommandLine, TEXT("Cook.Diagnostic.Modified.FinalFiles")))
+	{
+		bWriteFinalFiles = true;
+		bCommandLineOptionSet = true;
+	}
+	bCommandLineOptionSet |= FParse::Bool(CommandLine, TEXT("-Cook.Diagnostic.Modified.FinalFiles="), bWriteFinalFiles);
+	if (FParse::Param(CommandLine, TEXT("Cook.Diagnostic.Modified.Immediate")))
+	{
+		bWriteImmediateLog = true;
+		bCommandLineOptionSet = true;
+	}
+	bCommandLineOptionSet |= FParse::Bool(CommandLine, TEXT("-Cook.Diagnostic.Modified.Immediate="), bWriteImmediateLog);
+
+	if (!bCommandLineOptionSet)
+	{
+		bWriteFinalFiles = true;
+	}
+}
+
+bool FIncrementallyModifiedDiagnostics::IsWriteFinalLog() const
+{
+	return bWriteFinalLog;
+}
+
+bool FIncrementallyModifiedDiagnostics::IsWriteFinalFiles() const
+{
+	return bWriteFinalFiles;
+}
+
+bool FIncrementallyModifiedDiagnostics::IsWriteImmediateLog() const
+{
+	return bWriteImmediateLog;
+}
+
 void FIncrementallyModifiedDiagnostics::AddPackage(FIncrementallyModifiedContext& Context)
 {
 	check(Context.PackageData);
@@ -499,9 +549,49 @@ void FIncrementallyModifiedDiagnostics::AddPackage(FIncrementallyModifiedContext
 	default:
 		break;
 	}
+
+	if (IsWriteImmediateLog())
+	{
+		if (PackageData.ReasonLines.Num() == 0)
+		{
+			UE_LOGFMT(LogCook, Display, "Modified: {Package}, {Reason}",
+				PackageData.PackageName, LexToString(PackageData.Reason));
+		}
+		else if (PackageData.ReasonLines.Num() == 1)
+		{
+			UE_LOGFMT(LogCook, Display, "Modified: {Package}, {Reason}: {ReasonDetail}",
+				PackageData.PackageName, LexToString(PackageData.Reason), *PackageData.ReasonLines[0]);
+		}
+		else
+		{
+			UE_LOGFMT(LogCook, Display, "Modified: {Package}, {Reason}:",
+				PackageData.PackageName, LexToString(PackageData.Reason));
+			for (const FString& ReasonLine : PackageData.ReasonLines)
+			{
+				UE_LOGFMT(LogCook, Display, "    {ReasonDetail}", ReasonLine);
+			}
+		}
+
+		if (!IsWriteFinalFiles() && !IsWriteFinalLog())
+		{
+			PackageData.ReasonLines.Empty();
+		}
+	}
+}
+
+void FIncrementallyModifiedDiagnostics::WriteFinal(FStringView OutputDirectory) const
+{
+	if (IsWriteFinalLog())
+	{
+		WriteToLog();
+	}
+	if (IsWriteFinalFiles())
+	{
+		WriteToFiles(OutputDirectory);
+	}
 }
 
-void FIncrementallyModifiedDiagnostics::WriteToFile(FStringView OutputDirectory) const
+void FIncrementallyModifiedDiagnostics::WriteToFiles(FStringView OutputDirectory) const
 {
 	TMap<FTopLevelAssetPath, TArray<const FPackageData*>> PackagesPerClass;
 	SortPackageData(PackagesPerClass);
@@ -523,7 +613,7 @@ void FIncrementallyModifiedDiagnostics::WriteToFile(FStringView OutputDirectory)
 	}
 }
 
-void FIncrementallyModifiedDiagnostics::WriteToLogs() const
+void FIncrementallyModifiedDiagnostics::WriteToLog() const
 {
 	TMap<FTopLevelAssetPath, TArray<const FPackageData*>> PackagesPerClass;
 	SortPackageData(PackagesPerClass);
diff --git a/Engine/Source/Editor/UnrealEd/Private/Cooker/CookDiagnostics.h b/Engine/Source/Editor/UnrealEd/Private/Cooker/CookDiagnostics.h
--- a/Engine/Source/Editor/UnrealEd/Private/Cooker/CookDiagnostics.h
+++ b/Engine/Source/Editor/UnrealEd/Private/Cooker/CookDiagnostics.h
@@ -67,9 +67,13 @@ struct FIncrementallyModifiedContext
 class FIncrementallyModifiedDiagnostics
 {
 public:
+	FIncrementallyModifiedDiagnostics();
+
 	void AddPackage(FIncrementallyModifiedContext& Context);
-	void WriteToFile(FStringView OutputDirectory) const;
-	void WriteToLogs() const;
+	void WriteFinal(FStringView OutputDirectory) const;
+	bool IsWriteFinalLog() const;
+	bool IsWriteFinalFiles() const;
+	bool IsWriteImmediateLog() const;
 
 private:
 	/** The saved diagnostic data about each reported modified Package. */
@@ -82,11 +86,18 @@ private:
 		EIncrementallyModifiedReason Reason = EIncrementallyModifiedReason::None;
 	};
 
-	TArray<FPackageData> PackageDatas;
-
+	void WriteToFiles(FStringView OutputDirectory) const;
+	void WriteToLog() const;
 	void SortPackageData(TMap<FTopLevelAssetPath, TArray<const FPackageData*>>& OutPackagesPerClass) const;
-	void CreateModifiedCookedPackageContent(const TMap<FTopLevelAssetPath, TArray<const FPackageData*>>& PackagesPerClass, TArray<FString>& OutLines) const;
-	void CreateModifiedEditorPackageContent(const TMap<FTopLevelAssetPath, TArray<const FPackageData*>>& PackagesPerClass, TArray<FString>& OutLines) const;
+	void CreateModifiedCookedPackageContent(const TMap<FTopLevelAssetPath,
+		TArray<const FPackageData*>>& PackagesPerClass, TArray<FString>& OutLines) const;
+	void CreateModifiedEditorPackageContent(const TMap<FTopLevelAssetPath,
+		TArray<const FPackageData*>>& PackagesPerClass, TArray<FString>& OutLines) const;
+
+	TArray<FPackageData> PackageDatas;
+	bool bWriteFinalFiles = false;
+	bool bWriteFinalLog = false;
+	bool bWriteImmediateLog = false;
 };
 
 } // namespace UE::Cook
\ No newline at end of file
