ObjectTools.h - Add missing doxygen for bIncludeTarget param in GetDefaultTooltipForFunction and fix copy-paste error in GetDefaultTooltipForFunctionReturn comment

#Synced-CL 51991118

diff --git a/Engine/Source/Editor/UnrealEd/Public/ObjectTools.h b/Engine/Source/Editor/UnrealEd/Public/ObjectTools.h
--- a/Engine/Source/Editor/UnrealEd/Public/ObjectTools.h
+++ b/Engine/Source/Editor/UnrealEd/Public/ObjectTools.h
@@ -717,6 +717,7 @@ namespace ObjectTools
 	/**
 	 * Gets the non-specific tooltip for the given UFunction.
 	 * @param Function the function to find the tooltip for
+	 * @param bIncludeTarget if true, appends a "Target is [ClassName]" line identifying the function's owning class
 	 * @return the tooltip as a string
 	 */
 	UNREALED_API FString GetDefaultTooltipForFunction(const UFunction* Function, bool bIncludeTarget = true);
@@ -731,7 +732,7 @@ namespace ObjectTools
 
 	/**
 	 * Gets the non-specific tooltip for the given UFunction's return value.
-	 * @param Function the function to parse the param tooltip from
+	 * @param Function the function to parse the return tooltip from
 	 * @return the tooltip as a string
 	 */
 	UNREALED_API FString GetDefaultTooltipForFunctionReturn(const UFunction* Function);
