Editor Preview Platform Optimizations
Materials are already updated by SetPreviewPlatform, so disable the automatic per-texture material dependency update
8min ->1min Preview transion

#Synced-CL 52057137

diff --git a/Engine/Source/Editor/UnrealEd/Private/EditorEngine.cpp b/Engine/Source/Editor/UnrealEd/Private/EditorEngine.cpp
--- a/Engine/Source/Editor/UnrealEd/Private/EditorEngine.cpp
+++ b/Engine/Source/Editor/UnrealEd/Private/EditorEngine.cpp
@@ -6631,6 +6631,7 @@ static void PreviewPlatformChangePre()
 	FlushRenderingCommands();
 }
 
+extern ENGINE_API bool GDisableAutomaticTextureMaterialUpdateDependencies;
 static void PreviewPlatformChangePost()
 {
 	// scope for RegisterTextureLODSettings ; @todo CBTLDP : remove when possible
@@ -6664,6 +6665,9 @@ static void PreviewPlatformChangePost()
 	}
 
 	{
+		// Materials are already updated by SetPreviewPlatform, so disable the automatic per-texture material dependency update
+		TGuardValue<bool> Guard(GDisableAutomaticTextureMaterialUpdateDependencies, true);
+
 		// TextureLODSettings and platform cvars may have changed
 		//	need to notify all textures
 		//	don't do this until all changes are done
