Fix certain hair view mode having missing info due to ShaderPrint being disabled by high level logic.

#Synced-CL 52056680

diff --git a/Engine/Source/Editor/UnrealEd/Private/EditorViewportClient.cpp b/Engine/Source/Editor/UnrealEd/Private/EditorViewportClient.cpp
--- a/Engine/Source/Editor/UnrealEd/Private/EditorViewportClient.cpp
+++ b/Engine/Source/Editor/UnrealEd/Private/EditorViewportClient.cpp
@@ -3423,6 +3423,7 @@ void FEditorViewportClient::ChangeGroomVisualizationMode(FName InName)
 {
 	SetViewMode(VMI_VisualizeGroom);
 	CurrentGroomVisualizationMode = InName;
+	EngineShowFlags.SetShaderPrint(true);
 }
 
 bool FEditorViewportClient::IsGroomVisualizationModeSelected(FName InName) const
