CookPackageSplitter: Add support for a simpler path for populating Generated packages when !RequiresEmptyPackageBeforePopulate: PopulateGeneratorPackage is now allowed to move the objects into the Generated packages; this prevents the need for implementing PopulateGeneratedPackage/PreSaveGeneratedPackage and juggling moved objects. This is worse for memory - packages have to be kept in memory longer - but the simplification is worth it for generators that are relatively quick to finish.

This also fixes the error "PackageSplitter found a package it generated that was not removed from memory during garbage collection. This will cause errors later during population." That error was caused by UObjects with a time-consuming BeginCacheForCookedPlatformData being moved during PopulateGenerator while still having a FPendingCookedPlatformData; after being moved into the generated package  they kept the generated package referenced when we were not expecting it.
#rnx
#Synced-CL 52053053

diff --git a/Engine/Source/Editor/UnrealEd/Private/CookOnTheFlyServer.cpp b/Engine/Source/Editor/UnrealEd/Private/CookOnTheFlyServer.cpp
--- a/Engine/Source/Editor/UnrealEd/Private/CookOnTheFlyServer.cpp
+++ b/Engine/Source/Editor/UnrealEd/Private/CookOnTheFlyServer.cpp
@@ -3684,6 +3684,14 @@ UE::Cook::EPollStatus UCookOnTheFlyServer::PrepareSaveGenerationPackage(UE::Cook
 		}
 	}
 
+	// If we moved a FPendingCookedPlatformData object between generator and a generated package, wait for it
+	// to complete before we allow the generator or moved generated package to proceed. For simplicity, we currently
+	// block all generated packages from proceeding rather than keeping track of which ones have the pending objects.
+	if (GenerationHelper.HasBlockingPendingCookedPlatformDataObjects())
+	{
+		return EPollStatus::Incomplete;
+	}
+
 	if (PackageData.GetSaveSubState() <= ESaveSubState::Generation_PreMoveCookedPlatformData_WaitingForIsLoaded)
 	{
 		// Both Generator packages and Generated packages should wait for all IsCachedCookedPlatformData
@@ -3744,7 +3752,8 @@ UE::Cook::EPollStatus UCookOnTheFlyServer::PrepareSaveGenerationPackage(UE::Cook
 			PackageData.SetSaveSubStateComplete(ESaveSubState::Generation_BeginCacheObjectsToMove);
 		}
 		check(PackageData.GetSaveSubState() <= ESaveSubState::Generation_FinishCacheObjectsToMove);
-		if (PackageData.GetNumPendingCookedPlatformData() > 0)
+		if (PackageData.GetNumPendingCookedPlatformData() > 0 ||
+			GenerationHelper.HasBlockingPendingCookedPlatformDataObjects())
 		{
 			return EPollStatus::Incomplete;
 		}
@@ -3844,7 +3853,6 @@ UE::Cook::EPollStatus UCookOnTheFlyServer::BeginCacheObjectsToMove(UE::Cook::FGe
 
 	if (PackageData.GetSaveSubState() <= ESaveSubState::Generation_CallObjectsToMove)
 	{
-		bool bPopulateSucceeded = false;
 		bool bGeneratedRequiresGeneratorPopulate = GenerationHelper.DoesGeneratedRequireGenerator()
 			>= ICookPackageSplitter::EGeneratedRequiresGenerator::Populate;
 		if (Info.IsGenerator() || bGeneratedRequiresGeneratorPopulate)
@@ -4498,6 +4506,11 @@ UE::Cook::EPollStatus UCookOnTheFlyServer::CallBeginCacheOnObjects(UE::Cook::FPa
 				bool bNeedsResourceRelease = CurrentAsyncCache != nullptr;
 				PackageDatas->AddPendingCookedPlatformData(FPendingCookedPlatformData(Obj, TargetPlatform,
 					PackageData, bNeedsResourceRelease, *this));
+				TRefCountPtr<FGenerationHelper> GenerationHelper = PackageData.GetGenerationHelper();
+				if (GenerationHelper)
+				{
+					GenerationHelper->OnAddPendingCookedPlatformData(*this, PackageData, Obj);
+				}
 			}
 
 			if (Timer.IsActionTimeUp())
diff --git a/Engine/Source/Editor/UnrealEd/Private/Cooker/CookGenerationHelper.cpp b/Engine/Source/Editor/UnrealEd/Private/Cooker/CookGenerationHelper.cpp
--- a/Engine/Source/Editor/UnrealEd/Private/Cooker/CookGenerationHelper.cpp
+++ b/Engine/Source/Editor/UnrealEd/Private/Cooker/CookGenerationHelper.cpp
@@ -1160,6 +1160,72 @@ bool FGenerationHelper::TryReportGenerationManifest()
 	return true;
 }
 
+UObject* GetObjectFromObjectHolder(UObject*const& ObjectHolder)
+{
+	return ObjectHolder;
+}
+UObject* GetObjectFromObjectHolder(const FCachedObjectInOuter& ObjectHolder)
+{
+	return ObjectHolder.Object.Get();
+}
+template <typename ObjectHolderType>
+TMap<UObject*, UPackage*> GetPackagesForObjects(TConstArrayView<ObjectHolderType> Objects)
+{
+	TMap<UObject*, UPackage*> PackageForObject;
+	PackageForObject.Reserve(Objects.Num());
+	TArray<UObject*, TInlineAllocator<10>> ChildStack;
+	for (const ObjectHolderType& ObjectHolder : Objects)
+	{
+		check(ChildStack.IsEmpty());
+		UObject* Object = GetObjectFromObjectHolder(ObjectHolder);
+		if (!Object)
+		{
+			continue;
+		}
+		UPackage* Package = nullptr;
+		for (;;)
+		{
+			UPackage** PackagePtr = PackageForObject.Find(Object);
+			Package = PackagePtr ? *PackagePtr : nullptr;
+			if (Package)
+			{
+				break;
+			}
+
+			Package = Cast<UPackage>(Object);
+			if (Package)
+			{
+				PackageForObject.Add(Object, Package);
+				break;
+			}
+
+			Package = Object->GetExternalPackage();
+			if (Package)
+			{
+				PackageForObject.Add(Object, Package);
+				break;
+			}
+
+			UObject* Outer = Object->GetOuter();
+			if (!Outer)
+			{
+				Package = GetTransientPackage();
+				PackageForObject.Add(Object, Package);
+				break;
+			}
+			ChildStack.Add(Object);
+			Object = Outer;
+		}
+		check(Package);
+		for (UObject* ChildObject : ChildStack)
+		{
+			PackageForObject.Add(ChildObject, Package);
+		}
+		ChildStack.Reset();
+	}
+	return PackageForObject;
+}
+
 bool FGenerationHelper::TryCallPopulateGeneratorPackage(
 	TArray<ICookPackageSplitter::FGeneratedPackageForPopulate>& InOutGeneratedPackagesForPopulate,
 	bool bPopulateCallAllowed)
@@ -1211,6 +1277,21 @@ bool FGenerationHelper::TryCallPopulateGeneratorPackage(
 			FName()
 #endif
 		);
+
+		// Get the list of Objects in the package. Read them from the PackageData if it already has them,
+		// otherwise calculate them into a temporary buffer.
+		TArray<FCachedObjectInOuter> ObjectsInOuterBuffer;
+		TConstArrayView<FCachedObjectInOuter> ObjectsInOuter;
+		if (OwnerPackageData.GetHasSaveCache())
+		{
+			ObjectsInOuter = OwnerPackageData.GetCachedObjectsInOuter();
+		}
+		else
+		{
+			OwnerPackageData.CreateObjectCache(ObjectsInOuterBuffer, LocalOwnerPackage);
+			ObjectsInOuter = ObjectsInOuterBuffer;
+		}
+
 		ICookPackageSplitter::FPopulateContext PopulateContext(PopulateData);
 		PopulateData.OwnerPackage = LocalOwnerPackage;
 		PopulateData.OwnerObject = LocalSplitDataObject;
@@ -1233,6 +1314,106 @@ bool FGenerationHelper::TryCallPopulateGeneratorPackage(
 		UE_LOG(LogCookGenerationHelper, Verbose, TEXT("%s: PopulateGeneratorPackage."),
 			*WriteToString<256>(GetOwner().GetPackageName()));
 
+		// PopulateGenerator may have moved objects out of the package; detect them and add them to ObjectsToMove.
+		bool bMovedObjects = false;
+		TMap<UObject*, UPackage*> NewPackageForObject = GetPackagesForObjects(ObjectsInOuter);
+		for (const TPair<UObject*, UPackage*>& PackagePair : NewPackageForObject)
+		{
+			UPackage* NewPackage = PackagePair.Value;
+			check(NewPackage);
+			if (NewPackage == LocalOwnerPackage)
+			{
+				continue;
+			}
+			FPackageData* NewPackageData = COTFS.PackageDatas->FindPackageDataByPackageName(NewPackage->GetFName());
+			if (!NewPackageData)
+			{
+				// Object moved into a temporary or other package that the cooker does not track. The cooker will
+				// treat it as still being in the package; SavePackage will make it an import if referenced.
+				continue;
+			}
+			FCookGenerationInfo* Info = FindInfo(*NewPackageData);
+			if (!Info)
+			{
+				// Object moved into another cookable package, but not a generated package of this generator.
+				// During cook of the generator we will treat it as still being in the generator package; SavePackage
+				// will make it an import if referenced.
+				continue;
+			}
+			bMovedObjects = true;
+			PopulateData.ObjectsToMove.Add(PackagePair.Key);
+		}
+		if (bMovedObjects)
+		{
+			// Enforce uniqueness on ObjectsToMove since we added without checking for duplicates
+			PopulateData.ObjectsToMove.Sort();
+			PopulateData.ObjectsToMove.SetNum(Algo::Unique(PopulateData.ObjectsToMove));
+		}
+
+		// If PopulateGenerator already moved any of the UObjects in the generator package out into generated packages,
+		// keep those packages referenced until the generator package is destroyed, since we cannot reload
+		// the generator package to recreate them until that happens. This behavior is only valid in the case where the
+		// splitter has informed us !RequiresEmptyPackageBeforePopulate; otherwise it should not be accessing the
+		// generated packages until their own populate calls.
+		bMovedObjects = false;
+		NewPackageForObject = GetPackagesForObjects<UObject*>(PopulateData.ObjectsToMove);
+		for (const TPair<UObject*, UPackage*>& PackagePair : NewPackageForObject)
+		{
+			UPackage* NewPackage = PackagePair.Value;
+			check(NewPackage);
+			if (NewPackage == LocalOwnerPackage)
+			{
+				continue;
+			}
+			FPackageData* NewPackageData = COTFS.PackageDatas->FindPackageDataByPackageName(NewPackage->GetFName());
+			if (!NewPackageData)
+			{
+				// Object moved into a temporary or other package that the cooker does not track. The cooker will
+				// treat it as still being in the package; SavePackage will make it an import if referenced.
+				continue;
+			}
+			FCookGenerationInfo* Info = FindInfo(*NewPackageData);
+			if (!Info)
+			{
+				// Object moved into another cookable package, but not a generated package of this generator.
+				// During cook of the generator we will treat it as still being in the generator package; SavePackage
+				// will make it an import if referenced.
+				continue;
+			}
+			check(Info != &OwnerInfo); // NewPackage != LocalOwnerPackage, so Info != OwnerInfo
+			if (CookPackageSplitterInstance->RequiresEmptyPackageBeforePopulate(NewPackage))
+			{
+				UE_LOG(LogCook, Error,
+					TEXT("CookPackageSplitter moved an object from the generator package into a generated package during PopulateGeneratorPackage. ")
+					TEXT("This behavior is only allowed if RequiresEmptyPackageBeforePopulate returns false for the generated package. ")
+					TEXT("Edit the splitter code to either not move the object during Populate, or to return false for RequiresEmptyPackageBeforePopulate.")
+					TEXT("\n\tSplitter=%s")
+					TEXT("\n\tMoved Object=%s"),
+					*GetSplitDataObjectName().ToString(), *PackagePair.Key->GetPathName());
+			}
+			else
+			{
+				bMovedObjects = true;
+				PopulateData.KeepReferencedPackages.Add(NewPackage);
+			}
+		}
+		if (bMovedObjects)
+		{
+			// Enforce uniqueness on KeepReferencedPackages since we added without checking for duplicates
+			PopulateData.KeepReferencedPackages.Sort();
+			PopulateData.KeepReferencedPackages.SetNum(Algo::Unique(PopulateData.KeepReferencedPackages));
+		}
+
+		OwnerObjectsToMove.Reserve(PopulateData.ObjectsToMove.Num());
+		for (UObject* Object : PopulateData.ObjectsToMove)
+		{
+			if (Object)
+			{
+				OwnerObjectsToMove.Emplace(Object);
+			}
+		}
+		OwnerInfo.AddKeepReferencedPackages(*this, PopulateData.KeepReferencedPackages);
+
 		// Contract Point 1: We will not call populate again until the splitter has been destroyed
 		// Contract Point 2: We will not call populate again without garbage collecting the generator package
 		OwnerInfo.SetHasCalledPopulate(true);
@@ -1244,15 +1425,6 @@ bool FGenerationHelper::TryCallPopulateGeneratorPackage(
 			*WriteToString<256>(GetOwner().GetPackageName()));
 
 	}
-	OwnerInfo.AddKeepReferencedPackages(*this, PopulateData.KeepReferencedPackages);
-	OwnerObjectsToMove.Reserve(PopulateData.ObjectsToMove.Num());
-	for (UObject* Object : PopulateData.ObjectsToMove)
-	{
-		if (Object)
-		{
-			OwnerObjectsToMove.Emplace(Object);
-		}
-	}
 
 	return true;
 }
@@ -2148,9 +2320,10 @@ void FGenerationHelper::SetPreviousGeneratedPackages(const ITargetPlatform* Targ
 }
 
 template <typename T>
-static void AppendWeakPtrsToObjectPtrArray(TArray<T*>& Out, TArray<TWeakObjectPtr<T>>& In)
+static TArrayView<T*> AppendWeakPtrsToObjectPtrArray(TArray<T*>& Out, const TArray<TWeakObjectPtr<T>>& In)
 {
 	Out.Reserve(Out.Num() + In.Num());
+	int32 StartIndex = Out.Num();
 	for (const TWeakObjectPtr<T>& WeakPtr : In)
 	{
 		T* Object = WeakPtr.Get();
@@ -2159,6 +2332,7 @@ static void AppendWeakPtrsToObjectPtrArray(TArray<T*>& Out, TArray<TWeakObjectPt
 			Out.Add(Object);
 		}
 	}
+	return TArrayView<T*>(Out).Mid(StartIndex);
 }
 
 void FGenerationHelper::PreGarbageCollect(const TRefCountPtr<FGenerationHelper>& RefcountHeldByCaller,
@@ -2198,6 +2372,7 @@ void FGenerationHelper::PreGarbageCollect(const TRefCountPtr<FGenerationHelper>&
 	// We have a contract that KeepReferencedPackages in any Info are kept referenced.
 	bool bKeepingAnyObjects = false;
 	bool bNeedsGeneratorPackage = false;
+	FPackageDatas& PackageDatas = PackageData.GetPackageDatas();
 	if (&Info == &OwnerInfo)
 	{
 		// Handled by bCurrentGCHasKeptGeneratorKeepPackages
@@ -2205,7 +2380,14 @@ void FGenerationHelper::PreGarbageCollect(const TRefCountPtr<FGenerationHelper>&
 	else if (!Info.KeepReferencedPackages.IsEmpty())
 	{
 		bKeepingAnyObjects = true;
-		AppendWeakPtrsToObjectPtrArray(GCKeepPackages, Info.KeepReferencedPackages);
+		for (UPackage* Package : AppendWeakPtrsToObjectPtrArray(GCKeepPackages, Info.KeepReferencedPackages))
+		{
+			FPackageData* PackageDataOfPackage = PackageDatas.FindPackageDataByPackageName(Package->GetFName());
+			if (PackageDataOfPackage)
+			{
+				GCKeepPackageDatas.Add(PackageDataOfPackage);
+			}
+		}
 	}
 	if (!bCurrentGCHasKeptGeneratorKeepPackages)
 	{
@@ -2213,7 +2395,14 @@ void FGenerationHelper::PreGarbageCollect(const TRefCountPtr<FGenerationHelper>&
 		if (!OwnerInfo.KeepReferencedPackages.IsEmpty())
 		{
 			bNeedsGeneratorPackage = true;
-			AppendWeakPtrsToObjectPtrArray(GCKeepPackages, OwnerInfo.KeepReferencedPackages);
+			for (UPackage* Package : AppendWeakPtrsToObjectPtrArray(GCKeepPackages, OwnerInfo.KeepReferencedPackages))
+			{
+				FPackageData* PackageDataOfPackage = PackageDatas.FindPackageDataByPackageName(Package->GetFName());
+				if (PackageDataOfPackage)
+				{
+					GCKeepPackageDatas.Add(PackageDataOfPackage);
+				}
+			}
 		}
 	}
 
@@ -2807,6 +2996,7 @@ void FCookGenerationInfo::TakeOverCachedObjectsAndAddMoved(FGenerationHelper& Ge
 EPollStatus FCookGenerationInfo::RefreshPackageObjects(FGenerationHelper& GenerationHelper, UPackage* Package,
 	bool& bOutFoundNewObjects, ESaveSubState DemotionState)
 {
+	// As in FPackageData::RefreshObjectCache, we add any new objects but do not remove any removed objects.
 	bOutFoundNewObjects = false;
 	TArray<UObject*> CurrentObjectsInOuter;
 	GetObjectsWithOuter(Package, CurrentObjectsInOuter, true /* bIncludeNestedObjects */, RF_NoFlags,
@@ -2849,6 +3039,95 @@ EPollStatus FCookGenerationInfo::RefreshPackageObjects(FGenerationHelper& Genera
 	return EPollStatus::Success;
 }
 
+void FGenerationHelper::OnAddPendingCookedPlatformData(UCookOnTheFlyServer& COTFS, FPackageData& AssignedPackageData,
+	UObject* Obj)
+{
+	if (&AssignedPackageData != OwnerInfo.PackageData)
+	{
+		// We only need to coordinate saves if the assigned package is the generator package.
+		return;
+	}
+	UPackage* AssignedPackage = AssignedPackageData.GetPackage();
+	check(AssignedPackage); // Only be called for assigned packages in the save state.
+	check(Obj); // Only called for non-null objects.
+	UPackage* ParentPackage = Obj->GetPackage();
+	if (ParentPackage == AssignedPackage)
+	{
+		// No action necessary if the object is within the assigned package.
+		return;
+	}
+	FPackageData* ParentPackageData = COTFS.PackageDatas->FindPackageDataByPackageName(ParentPackage->GetFName());
+	if (!ParentPackageData)
+	{
+		// Object is in a temporary or other package that the cooker does not track. The cooker will
+		// treat it as still being in the assigned package; SavePackage will make it an import if referenced.
+		return;
+	}
+	FCookGenerationInfo* ParentInfo = FindInfo(*ParentPackageData);
+	if (!ParentInfo)
+	{
+		// Object moved into another cookable package, but not a generated package of this generator.
+		// During cook of the generator we will treat it as still being in the generator package; SavePackage
+		// will make it an import if referenced.
+		return;
+	}
+	// ParentPackage != AssignedPackage and AssignedPackage == OwnerPackage, so ParentInfo != OwnerInfo
+	check(ParentInfo != &OwnerInfo); 
+
+	// Block any further operations on both the generator and generated package it moved into, until the Object is no
+	// longer pending. We do this because we need to guarantee the cooker contract that we do not call
+	// BeginCachedCookedPlatformData on it again until it has returned true from IsCachedCookedPlatformDataLoaded, and
+	// operations on the generated package might find it in the generated package and call BeginCached.
+	AddBlockingPendingCookedPlatformDataObject(Obj);
+}
+
+void FGenerationHelper::AddBlockingPendingCookedPlatformDataObject(UObject* Object)
+{
+	BlockingPendingCookedPlatformDatas.Add(FWeakObjectPtr(Object));
+}
+
+bool FGenerationHelper::HasBlockingPendingCookedPlatformDataObjects()
+{
+	FPackageDatas& PackageDatas = OwnerInfo.PackageData->GetPackageDatas();
+	for (auto Iter = BlockingPendingCookedPlatformDatas.CreateIterator(); Iter; ++Iter)
+	{
+		UObject* Object = Iter->Get();
+		if (Object && !PackageDatas.GetPendingCookedPlatformDatas(Object).IsEmpty())
+		{
+			return true;
+		}
+		Iter.RemoveCurrentSwap();
+	}
+	return false;
+}
+
+bool FGenerationHelper::ShouldKeepWaitingOnMovedPendingCookedPlatformData(FPendingCookedPlatformData& Pending,
+	FPackageData& AssignedPackage, UObject* Object)
+{
+	check(Object);
+	// Only allow generators to keep waiting on objects in other packages
+	if (&AssignedPackage != OwnerInfo.PackageData)
+	{
+		return false;
+	}
+	UPackage* CurrentPackage = Object->GetPackage();
+	check(CurrentPackage);
+	FPackageData* CurrentPackageData =
+		AssignedPackage.GetPackageDatas().FindPackageDataByPackageName(CurrentPackage->GetFName());
+	if (!CurrentPackageData)
+	{
+		// Only wait on objects in cookable packages
+		return false;
+	}
+	FCookGenerationInfo* CurrentInfo = FindInfo(*CurrentPackageData);
+	if (!CurrentInfo)
+	{
+		// Only wait on objects in generated packages of this generator.
+		return false;
+	}
+	return true;
+}
+
 void FCookGenerationInfo::AddKeepReferencedPackages(FGenerationHelper& GenerationHelper,
 	TArray<UPackage*>& InKeepReferencedPackages)
 {
diff --git a/Engine/Source/Editor/UnrealEd/Private/Cooker/CookGenerationHelper.h b/Engine/Source/Editor/UnrealEd/Private/Cooker/CookGenerationHelper.h
--- a/Engine/Source/Editor/UnrealEd/Private/Cooker/CookGenerationHelper.h
+++ b/Engine/Source/Editor/UnrealEd/Private/Cooker/CookGenerationHelper.h
@@ -522,6 +522,34 @@ public:
 	 */
 	void ForceUninitialize();
 
+	/**
+	 * Called when BeginCacheForCookedPlatformData called on a UObject, and IsCachedCookedPlatformDataLoaded does
+	 * not immediately return true, for the given Object assigned to the given PackageData (but possibly a child of
+	 * another package). The GenerationHelper will need to coordinate save fences if the object is a child
+	 * of a different generated package owned by the generator than the one it is assigned to.
+	 */
+	void OnAddPendingCookedPlatformData(UCookOnTheFlyServer& COTFS, FPackageData& AssignedPackageData, UObject* Obj);
+	/**
+	 * Adds the object to a list that is checked before any save-state action of the generator or its generated
+	 * packages. If the object still has an FPendingCookedPlatformData, the action is blocked and the package
+	 * save is deferred as it would be if there were a FPendingCookedPlatformData assigned to the package. This
+	 * is used for objects that move between generator and generated packages in edge-case flows that we
+	 * don't support efficiently but need to handle for robustness.
+	 */
+	void AddBlockingPendingCookedPlatformDataObject(UObject* Object);
+	/**
+	 * Removes completed objects from the BlockingObjects list and returns whether any remain; it may leave
+	 * some completed objects in the list if others are incomplete.
+	 */
+	bool HasBlockingPendingCookedPlatformDataObjects();
+	/**
+	 * A PendingCookedPlatformData has detected an object moved out of the given package with this generation helper.
+	 * Return whether the pendingCookedPlatformData should keep waiting on it. If not we will stop polling it and will
+	 * no longer let it block the save of its package.
+	 */
+	bool ShouldKeepWaitingOnMovedPendingCookedPlatformData(FPendingCookedPlatformData& Pending,
+		FPackageData& AssignedPackage, UObject* Object);
+
 	/** Helper function for Initialize and for TryCreateValidGenerationHelper. */
 	static void SearchForRegisteredSplitDataObject(UCookOnTheFlyServer& COTFS, FName PackageName, UPackage* Package,
 		TOptional<TConstArrayView<FCachedObjectInOuter>> CachedObjectsInOuter, UObject*& OutSplitDataObject,
@@ -612,6 +640,7 @@ private:
 	TMap<const ITargetPlatform*, FPlatformData> PlatformDatas;
 	TArray<FName> ExternalActorDependencies;
 	TArray<FWeakObjectPtr> OwnerObjectsToMove;
+	TArray<FWeakObjectPtr> BlockingPendingCookedPlatformDatas;
 	TRefCountPtr<FGenerationHelper> ReferenceFromKeepForQueueResults;
 	TRefCountPtr<FGenerationHelper> ReferenceFromKeepForAllSavedOrGC;
 	int32 MPCookNextAssignmentIndex = 0;
diff --git a/Engine/Source/Editor/UnrealEd/Private/Cooker/CookPackageData.cpp b/Engine/Source/Editor/UnrealEd/Private/Cooker/CookPackageData.cpp
--- a/Engine/Source/Editor/UnrealEd/Private/Cooker/CookPackageData.cpp
+++ b/Engine/Source/Editor/UnrealEd/Private/Cooker/CookPackageData.cpp
@@ -1330,20 +1330,7 @@ void FPackageData::CreateObjectCache()
 	UPackage* LocalPackage = GetPackage();
 	if (LocalPackage && LocalPackage->IsFullyLoaded())
 	{
-		PackageName = LocalPackage->GetFName();
-		TArray<UObject*> ObjectsInOuter;
-		// ignore RF_Garbage objects; they will not be serialized out so we don't need to call
-		// BeginCacheForCookedPlatformData on them
-		GetObjectsWithOuter(LocalPackage, ObjectsInOuter, true /* bIncludeNestedObjects */, RF_NoFlags,
-			EInternalObjectFlags::Garbage);
-		CachedObjectsInOuter.Reset(ObjectsInOuter.Num());
-		for (UObject* Object : ObjectsInOuter)
-		{
-			FWeakObjectPtr ObjectWeakPointer(Object);
-			// GetObjectsWithOuter with Garbage filtered out should only return valid-for-weakptr objects
-			check(ObjectWeakPointer.Get());
-			CachedObjectsInOuter.Emplace(ObjectWeakPointer);
-		}
+		CreateObjectCache(CachedObjectsInOuter, LocalPackage);
 
 		for (TPair<const ITargetPlatform*, FPackagePlatformData>& Pair : PlatformDatas)
 		{
@@ -1363,6 +1350,23 @@ void FPackageData::CreateObjectCache()
 	}
 }
 
+void FPackageData::CreateObjectCache(TArray<FCachedObjectInOuter>& OutBuffer, UPackage* InPackage)
+{
+	TArray<UObject*> ObjectsInOuter;
+	// Ignore garbage objects; they will not be serialized out so we don't need to call
+	// BeginCacheForCookedPlatformData on them
+	GetObjectsWithOuter(InPackage, ObjectsInOuter, true /* bIncludeNestedObjects */, RF_NoFlags,
+		EInternalObjectFlags::Garbage);
+	OutBuffer.Reset(ObjectsInOuter.Num());
+	for (UObject* Object : ObjectsInOuter)
+	{
+		FWeakObjectPtr ObjectWeakPointer(Object);
+		// GetObjectsWithOuter with Garbage filtered out should only return valid-for-weakptr objects
+		check(ObjectWeakPointer.Get());
+		OutBuffer.Emplace(ObjectWeakPointer);
+	}
+}
+
 static TArray<UObject*> SetDifference(TArray<UObject*>& A, TArray<UObject*>& B)
 {
 	Algo::Sort(A); // Don't use TArray.Sort, it sorts pointers as references and we want to sort them as pointers
@@ -1421,6 +1425,11 @@ EPollStatus FPackageData::RefreshObjectCache(bool& bOutFoundNewObjects)
 	GetObjectsWithOuter(Package.Get(), CurrentObjects, true /* bIncludeNestedObjects */, RF_NoFlags, 
 		EInternalObjectFlags::Garbage);
 
+	// We add any new objects but do not remove any objects moved out of the package, and do not remove any objects
+	// deleted by garbage collection. FCookGenerationHelper relies on being able to store objects in the generator's
+	// object cache that it needs to call BeginCache/IsCached on but that might be in other packages such as
+	// ExternalActor packages that own the object on disk or generated packages the object has been moved into.
+	// Objects deleted by garbage collection are detected elsewhere and cause control flow decisions.
 	TArray<UObject*> NewObjects = SetDifference(CurrentObjects, OldObjects);
 	bOutFoundNewObjects = NewObjects.Num() > 0;
 	if (bOutFoundNewObjects)
@@ -1434,7 +1443,7 @@ EPollStatus FPackageData::RefreshObjectCache(bool& bOutFoundNewObjects)
 			CachedObjectsInOuter.Emplace(MoveTemp(ObjectWeakPointer));
 		}
 		// GetCookedPlatformDataNextIndex is already where it should be, pointing at the first of the objects we have
-		// added. Caller is respnsible for changing state back to calling BeginCacheForCookedPlatformData.
+		// added. Caller is responsible for changing state back to calling BeginCacheForCookedPlatformData.
 
 		if (++GetNumRetriesBeginCacheOnObjects() > FPackageData::GetMaxNumRetriesBeginCacheOnObjects())
 		{
@@ -1614,7 +1623,7 @@ void FPackageData::UpdateSaveAfterGarbageCollect(bool& bOutDemote)
 			{
 				// Deleting a public object puts the package in an invalid state; demote back to request
 				// and load/save it again
-				bool bPublicDeleted = !!(CachedObjectInOuter.ObjectFlags & RF_Public);;
+				bool bPublicDeleted = !!(CachedObjectInOuter.ObjectFlags & RF_Public);
 				bOutDemote |= bPublicDeleted;
 			}
 		}
@@ -2008,11 +2017,17 @@ bool FPendingCookedPlatformData::PollIsComplete()
 	FName CurrentPackageName = LocalObject->GetPackage()->GetFName();
 	if (CurrentPackageName != PackageData.GetPackageName())
 	{
-		UE_LOG(LogCook, Display,
-			TEXT("We were waiting for IsCachedCookedPlatformData to return true for %s in package %s, but that object has been moved out of the package. We will stop waiting on it."),
-			*LocalObject->GetFullName(), *PackageData.GetPackageName().ToString());
-		Release();
-		return true;
+		// Exception: allow Generator packages to keep waiting on objects that they have moved into generated packages.
+		TRefCountPtr<FGenerationHelper> GenerationHelper = PackageData.GetGenerationHelper();
+		if (!GenerationHelper ||
+			!GenerationHelper->ShouldKeepWaitingOnMovedPendingCookedPlatformData(*this, PackageData, LocalObject))
+		{
+			UE_LOG(LogCook, Display,
+				TEXT("We were waiting for IsCachedCookedPlatformData to return true for %s in package %s, but that object has been moved out of the package. We will stop waiting on it."),
+				*LocalObject->GetFullName(), *PackageData.GetPackageName().ToString());
+			Release();
+			return true;
+		}
 	}
 	if (LocalObject->HasAnyFlags(EObjectFlags::RF_MirroredGarbage))
 	{
@@ -2951,6 +2966,21 @@ void FPackageDatas::AddPendingCookedPlatformData(FPendingCookedPlatformData&& Da
 	++PendingCookedPlatformDataNum;
 }
 
+TArray<FPendingCookedPlatformData*> FPackageDatas::GetPendingCookedPlatformDatas(UObject* Object)
+{
+	check(Object);
+
+	TArray<FPendingCookedPlatformData*> Results;
+	ForEachPendingCookedPlatformData([&Results, Object](FPendingCookedPlatformData& Data)
+		{
+			if (Data.Object.GetEvenIfUnreachable() == Object)
+			{
+				Results.Add(&Data);
+			}
+		});
+	return Results;
+}
+
 void FPackageDatas::PollPendingCookedPlatformDatas(bool bForce, double& LastCookableObjectTickTime,
 	int32& OutNumRetired)
 {
diff --git a/Engine/Source/Editor/UnrealEd/Private/Cooker/CookPackageData.h b/Engine/Source/Editor/UnrealEd/Private/Cooker/CookPackageData.h
--- a/Engine/Source/Editor/UnrealEd/Private/Cooker/CookPackageData.h
+++ b/Engine/Source/Editor/UnrealEd/Private/Cooker/CookPackageData.h
@@ -639,6 +639,8 @@ public:
 	void CheckObjectCacheEmpty() const;
 	/** Populate CachedObjectsInOuter if not already populated. Invalid to call except when in the save state. */
 	void CreateObjectCache();
+	/** Populate OutBuffer with all objects from the given Package; filters out objects not relevant to the cook.  */
+	static void CreateObjectCache(TArray<FCachedObjectInOuter>& OutBuffer, UPackage* Package);
 	/**
 	 * Look for new Objects that were created during BeginCacheForCookedPlatformData calls, and if found add
 	 * them to the ObjectCache and set state so that we call BeginCacheForCookedPlatformData on the new objects.
@@ -1435,6 +1437,12 @@ public:
 	void ForEachPendingCookedPlatformData(const FunctionType& Function);
 	int32 GetPendingCookedPlatformDataNum() const;
 	void AddPendingCookedPlatformData(FPendingCookedPlatformData&& Data);
+	/**
+	 * Return pointers to the PendingCookedPlatformDatas for the given object, if any exist. These
+	 * pointers are to TArray elements and must not be dereferenced after any modifications to
+	 * PendingCookedPlatformData: e.g. AddPendingCookedPlatformData or PollPendingCookedPlatformDatas.
+	 */
+	TArray<FPendingCookedPlatformData*> GetPendingCookedPlatformDatas(UObject* Object);
 
 	/**
 	 * Iterate over all elements in PendingCookedPlatformDatas and check whether they have completed,
diff --git a/Engine/Source/Editor/UnrealEd/Public/CookPackageSplitter.h b/Engine/Source/Editor/UnrealEd/Public/CookPackageSplitter.h
--- a/Engine/Source/Editor/UnrealEd/Public/CookPackageSplitter.h
+++ b/Engine/Source/Editor/UnrealEd/Public/CookPackageSplitter.h
@@ -89,7 +89,7 @@ public:
 	}
 
 	/**
-	 * If true, this splitter forces the Generator package to be empty before calling the populate function.
+	 * If true, this splitter requires that the given generated package be empty before calling the populate function.
 	 * Returning false will allow the splitter to precreate packages instead of relying on the cooker to create them.
 	 */
 	virtual bool RequiresEmptyPackageBeforePopulate(const UPackage* Package)
@@ -327,6 +327,15 @@ public:
 	 * This function will not be called if bCookSaveGeneratorPackage is false in the GenerationManifest, unless it
 	 * is needed for generated packages as well via DoesGeneratedRequireGenerator.
 	 *
+	 * If this splitter returns false for virtual RequiresEmptyPackageBeforePopulate(GeneratedPackage), it can add or
+	 * move objects into that generated package during PopulateGeneratorPackage. Doing so will prevent the targeted
+	 * generated package from being garbage collected until the generator package can be destroyed. This is worse
+	 * performance, but can be useful to reduce Splitter complexity.
+	 * 
+	 * If the splitter returns true for virtual RequiresEmptyPackageBeforePopulate(GeneratedPackage), adding or moving
+	 * objects into that generated package during PopulateGeneratorPackage is an error, and those objects will be
+	 * missing from the generated package when it is saved.
+	 *
 	 * @return True if successfully populated, false on error (this will cause a cook error).
 	 */
 	virtual bool PopulateGeneratorPackage(FPopulateContext& PopulateContext)
@@ -444,7 +453,7 @@ public:
 	UNREALED_API void WarnIfDeprecatedVirtualNotCalled(const TCHAR* FunctionName);
 
 private:
-	UE_DEPRECATED(5.6, "Deprecation support, do not read/write outside of cooker code.")
+	UE_DEPRECATED(5.7, "Deprecation support, do not read/write outside of cooker code.")
 	bool bDeprecatedVirtualCalledAsExpected = false;
 };
 
