UE: Bump CookIncrementalVersion. RVT derived data change was not picked up as a dix

diff --git a/Engine/Source/Editor/UnrealEd/Private/Cooker/CookIncrementalVersion.cpp b/Engine/Source/Editor/UnrealEd/Private/Cooker/CookIncrementalVersion.cpp
--- a/Engine/Source/Editor/UnrealEd/Private/Cooker/CookIncrementalVersion.cpp
+++ b/Engine/Source/Editor/UnrealEd/Private/Cooker/CookIncrementalVersion.cpp
@@ -8,6 +8,6 @@ namespace UE::Cook
 {
 
 // Change CookIncrementalVersion to a new guid when all packages in an incremental cook need to be invalidated
-FGuid CookIncrementalVersion(0xFFB412DC, 0x88224FE2, 0xBEF2627A, 0x1B71D307);
+FGuid CookIncrementalVersion(0xFFB411DC, 0x88224FE2, 0xBEDF2627A, 0x1B71D307);
 
 }
