Merged ToggleFeatureLevelPreview into SetPreviewPlatform using a new bCompileShaders parameter that controls whether or not we compile the shaders.
- Added bRestore parameter to device profile update to skip redundant RestorePreviewDeviceProfile during toggle
- Removed CachedEditorShaderPlatform member, using GMaxRHIShaderPlatform directly.
- Removed toggle button restriction for GMaxRHIFeatureLevel.
- Fix XSX Preview by using PREVIEW_DISABLE_IF_RUNTIME_UNSUPPORTED for SupportsBarycentricsIntrinsics and SupportsBarycentricsSemantic

#Synced-CL 52057137

diff --git a/Engine/Source/Editor/LevelEditor/Private/LevelEditorActions.cpp b/Engine/Source/Editor/LevelEditor/Private/LevelEditorActions.cpp
--- a/Engine/Source/Editor/LevelEditor/Private/LevelEditorActions.cpp
+++ b/Engine/Source/Editor/LevelEditor/Private/LevelEditorActions.cpp
@@ -823,19 +823,11 @@ bool FLevelEditorActionCallbacks::IsFeatureLevelPreviewEnabled()
 	{
 		return false;
 	}
-	if (GEditor->PreviewPlatform.PreviewFeatureLevel == GMaxRHIFeatureLevel)
-	{
-		return false;
-	}
 	return GEditor->IsFeatureLevelPreviewEnabled();
 }
 
 bool FLevelEditorActionCallbacks::IsFeatureLevelPreviewActive()
 {
-	if (GEditor->PreviewPlatform.PreviewFeatureLevel == GMaxRHIFeatureLevel)
-	{
-		return false;
-	}
 	return GEditor->IsFeatureLevelPreviewEnabled() && GEditor->IsFeatureLevelPreviewActive();
 }
 
diff --git a/Engine/Source/Editor/UnrealEd/Classes/Editor/EditorEngine.h b/Engine/Source/Editor/UnrealEd/Classes/Editor/EditorEngine.h
--- a/Engine/Source/Editor/UnrealEd/Classes/Editor/EditorEngine.h
+++ b/Engine/Source/Editor/UnrealEd/Classes/Editor/EditorEngine.h
@@ -646,9 +646,6 @@ public:
 
 	/** The feature level and platform we should use when loading or creating a new world */
 	FPreviewPlatformInfo PreviewPlatform;
-	
-	/** Cached ShaderPlatform so the editor can go back to after previewing */
-	EShaderPlatform CachedEditorShaderPlatform;
 
 	/** A delegate that is called when the preview feature level changes. Primarily used to switch a viewport's feature level. */
 	UE_DEPRECATED(5.8, "FPreviewFeatureLevelChanged is deprecated, use FPreviewShaderPlatformChanged instead")
@@ -3241,7 +3238,7 @@ public:
 	FSceneMaterialsModifiedEvent& OnSceneMaterialsModifiedEvent() { return SceneMaterialsModifiedEvent; }
 
 	/** Call this function to change the feature level and to override the material quality platform of the editor and PIE worlds */
-	UNREALED_API void SetPreviewPlatform(const FPreviewPlatformInfo& NewPreviewPlatform, bool bSaveSettings);
+	UNREALED_API void SetPreviewPlatform(const FPreviewPlatformInfo& NewPreviewPlatform, bool bSaveSettings, bool bCompileShaders = true);
 
 	/** Toggle the feature level preview */
 	UNREALED_API void ToggleFeatureLevelPreview();
diff --git a/Engine/Source/Editor/UnrealEd/Private/EditorEngine.cpp b/Engine/Source/Editor/UnrealEd/Private/EditorEngine.cpp
--- a/Engine/Source/Editor/UnrealEd/Private/EditorEngine.cpp
+++ b/Engine/Source/Editor/UnrealEd/Private/EditorEngine.cpp
@@ -606,7 +606,6 @@ UEditorEngine::UEditorEngine(const FObjectInitializer& ObjectInitializer)
 	bIsEndingPlay = false;
 	DefaultWorldFeatureLevel = GMaxRHIFeatureLevel;
 	PreviewPlatform = FPreviewPlatformInfo(DefaultWorldFeatureLevel, GMaxRHIShaderPlatform);
-	CachedEditorShaderPlatform = GMaxRHIShaderPlatform;
 
 	FCoreDelegates::OnFeatureLevelDisabled.AddLambda([this](int RHIType, const FName& PreviewPlatformName)
 		{
@@ -6579,7 +6578,7 @@ void UEditorEngine::OnEffectivePreviewShaderPlatformChange()
 		if (StereoRenderTargetManager)
 		{
 			StereoRenderTargetManager->ReconfigureForShaderPlatform(
-				PreviewPlatform.bPreviewFeatureLevelActive ? PreviewPlatform.ShaderPlatform : CachedEditorShaderPlatform);
+				PreviewPlatform.bPreviewFeatureLevelActive ? PreviewPlatform.ShaderPlatform : GMaxRHIShaderPlatform);
 		}
 	}
 }
@@ -6687,7 +6686,7 @@ static void PreviewPlatformChangePost()
 	}
 }
 
-void UEditorEngine::SetPreviewPlatform(const FPreviewPlatformInfo& NewPreviewPlatform, bool bSaveSettings)
+void UEditorEngine::SetPreviewPlatform(const FPreviewPlatformInfo& NewPreviewPlatform, bool bSaveSettings, bool bCompileShaders)
 {
 	// Get the requested preview platform, make sure it is valid.
 	EShaderPlatform ShaderPlatform = NewPreviewPlatform.ShaderPlatform;
@@ -6697,212 +6696,180 @@ void UEditorEngine::SetPreviewPlatform(const FPreviewPlatformInfo& NewPreviewPla
 	}
 	else
 	{
-		UE_LOG(LogEditor, Display, TEXT("Set Preview Platform: %s"), 
-			*NewPreviewPlatform.GetFriendlyName().ToString());
+		UE_LOG(LogEditor, Display, TEXT("Set Preview Platform: %s (CompileShaders=%s)"),
+			*NewPreviewPlatform.GetFriendlyName().ToString(),
+			bCompileShaders ? TEXT("true") : TEXT("false"));
 	}
-	
-	PreviewPlatformChangePre();
-
-	check(FDataDrivenShaderPlatformInfo::IsValid(ShaderPlatform));
-	ERHIFeatureLevel::Type MaxFeatureLevel = NewPreviewPlatform.PreviewShaderFormatName != NAME_None ? (ERHIFeatureLevel::Type)GetMaxSupportedFeatureLevel(ShaderPlatform) : ERHIFeatureLevel::SM5;
-	check(NewPreviewPlatform.PreviewShaderFormatName.IsNone() || MaxFeatureLevel == NewPreviewPlatform.PreviewFeatureLevel);
-
-	const bool bChangedPreviewShaderPlatform = NewPreviewPlatform.ShaderPlatform != PreviewPlatform.ShaderPlatform;
-	const bool bChangedEffectiveShaderPlatform = bChangedPreviewShaderPlatform && (PreviewPlatform.bPreviewFeatureLevelActive || NewPreviewPlatform.bPreviewFeatureLevelActive);
-	const ERHIFeatureLevel::Type EffectiveFeatureLevel = NewPreviewPlatform.GetEffectivePreviewFeatureLevel();
 
-	if (NewPreviewPlatform.PreviewShaderFormatName != NAME_None)
 	{
-		// Force generation of the autogen files if they don't already exist
-		FShaderCompileUtilities::GenerateBrdfHeaders(ShaderPlatform);
-	}
-
-	// Record the new preview platform
-	PreviewPlatform = NewPreviewPlatform;
+		FScopedSlowTask SlowTask(100.f, NSLOCTEXT("Engine", "ChangingPreviewPlatform", "Changing Preview Platform"), true);
+		SlowTask.Visibility = ESlowTaskVisibility::ForceVisible;
+		SlowTask.MakeDialog(/*bShowCancelButton=*/false, /*bAllowInPIE=*/true);
 
-	// Initially set preview as disabled in case it fails
-	SaveShaderPlatformAsDisabled(PreviewPlatform);
+		PreviewPlatformChangePre();
 
-	// If we changed the preview platform, we need to update the material quality settings
-	if (bChangedPreviewShaderPlatform)
-	{
-		UMaterialShaderQualitySettings* MaterialShaderQualitySettings = UMaterialShaderQualitySettings::Get();
-		MaterialShaderQualitySettings->SetPreviewPlatform(PreviewPlatform.PreviewShaderPlatformName);
+		check(FDataDrivenShaderPlatformInfo::IsValid(ShaderPlatform));
+		ERHIFeatureLevel::Type MaxFeatureLevel = NewPreviewPlatform.PreviewShaderFormatName != NAME_None ? (ERHIFeatureLevel::Type)GetMaxSupportedFeatureLevel(ShaderPlatform) : ERHIFeatureLevel::SM5;
+		check(NewPreviewPlatform.PreviewShaderFormatName.IsNone() || MaxFeatureLevel == NewPreviewPlatform.PreviewFeatureLevel);
 
-		UStaticMesh::OnLodStrippingQualityLevelChanged(nullptr);
+		const bool bChangedPreviewShaderPlatform = NewPreviewPlatform.ShaderPlatform != PreviewPlatform.ShaderPlatform;
+		const bool bChangedEffectiveShaderPlatform = bChangedPreviewShaderPlatform
+			? (PreviewPlatform.bPreviewFeatureLevelActive || NewPreviewPlatform.bPreviewFeatureLevelActive)
+			: (NewPreviewPlatform.bPreviewFeatureLevelActive != PreviewPlatform.bPreviewFeatureLevelActive);
+		const ERHIFeatureLevel::Type EffectiveFeatureLevel = NewPreviewPlatform.GetEffectivePreviewFeatureLevel();
 
-		if (bChangedEffectiveShaderPlatform)
+		if (NewPreviewPlatform.PreviewShaderFormatName != NAME_None)
 		{
-			OnEffectivePreviewShaderPlatformChange();
+			// Force generation of the autogen files if they don't already exist
+			FShaderCompileUtilities::GenerateBrdfHeaders(ShaderPlatform);
 		}
-	}
 
-	// Update any PerPlatformConfig class defaults or instances
-	for (FThreadSafeObjectIterator ObjIterator(UObject::StaticClass(), RF_NoFlags); ObjIterator; ++ObjIterator)
-	{
-		if ((*ObjIterator) && ObjIterator->GetClass()->HasAnyClassFlags(CLASS_PerPlatformConfig))
+		// Record the new preview platform
+		PreviewPlatform = NewPreviewPlatform;
+
+		// Initially set preview as disabled in case shader compilation crashes
+		if (bCompileShaders)
 		{
-			ObjIterator->LoadConfig();
+			SaveShaderPlatformAsDisabled(PreviewPlatform);
 		}
-	}
 
-	constexpr bool bUpdateProgressDialog = true;
-	constexpr bool bCacheAllRemainingShaders = false;
-
-	{
-		// Set the correct SP preview for the FeatureLevel that is being previewed
-		FScopedSlowTask SlowTask(100.f, NSLOCTEXT("Engine", "ChangingPreviewPlatform", "Changing Preview Platform"), true);
-		SlowTask.Visibility = ESlowTaskVisibility::ForceVisible;
-		SlowTask.MakeDialog();
-
-		//invalidate global bound shader states so they will be created with the new shaders the next time they are set (in SetGlobalBoundShaderState)
-		for (TLinkedList<FGlobalBoundShaderStateResource*>::TIterator It(FGlobalBoundShaderStateResource::GetGlobalBoundShaderStateList()); It; It.Next())
+		// If we changed the preview platform, we need to update the material quality settings
+		if (bChangedPreviewShaderPlatform)
 		{
-			BeginUpdateResourceRHI(*It);
+			UMaterialShaderQualitySettings* MaterialShaderQualitySettings = UMaterialShaderQualitySettings::Get();
+			MaterialShaderQualitySettings->SetPreviewPlatform(PreviewPlatform.PreviewShaderPlatformName);
 		}
 
-		FGlobalComponentReregisterContext RecreateComponents;
-		FlushRenderingCommands();
-
-		// Set only require the preview feature level and the max feature level. The Max feature level is required for the toggle feature.
-		for (uint32 i = (uint32)ERHIFeatureLevel::ES3_1; i < (uint32)ERHIFeatureLevel::Num; i++)
+		// Update any PerPlatformConfig class defaults or instances
+		for (FThreadSafeObjectIterator ObjIterator(UObject::StaticClass(), RF_NoFlags); ObjIterator; ++ObjIterator)
 		{
-			ERHIFeatureLevel::Type FeatureLevel = (ERHIFeatureLevel::Type)i;
-			UMaterialInterface::SetGlobalRequiredFeatureLevel(FeatureLevel, FeatureLevel == PreviewPlatform.PreviewFeatureLevel || FeatureLevel == GMaxRHIFeatureLevel);
+			if ((*ObjIterator) && ObjIterator->GetClass()->HasAnyClassFlags(CLASS_PerPlatformConfig))
+			{
+				ObjIterator->LoadConfig();
+			}
 		}
 
-		GShaderPlatformForFeatureLevel[PreviewPlatform.PreviewFeatureLevel] = ShaderPlatform;
-
-		SlowTask.EnterProgressFrame(35.0f);
-		UMaterial::AllMaterialsCacheResourceShadersForRendering(bUpdateProgressDialog, bCacheAllRemainingShaders);
+		// Restore device profile to defaults, then apply the preview profile
+		// bRestore=true resets CVars first (needed for compile path to get a clean state).
+		// bRestore=false skips the reset (toggle path avoids double CVar sink flush).
+		auto UpdateDeviceProfile = [&](bool bRestore)
+		{
+			// Always restore if no preview profile will be set, or if explicitly requested (compile path)
+			bRestore = bRestore || !PreviewPlatform.bPreviewFeatureLevelActive || PreviewPlatform.DeviceProfileName == NAME_None;
 
-		SlowTask.EnterProgressFrame(35.0f);
-		UMaterialInstance::AllMaterialsCacheResourceShadersForRendering(bUpdateProgressDialog, bCacheAllRemainingShaders);
+			if (bRestore)
+			{
+				UDeviceProfileManager::Get().RestorePreviewDeviceProfile();
+			}
 
-		SlowTask.EnterProgressFrame(15.0f, NSLOCTEXT("Engine", "SlowTaskGlobalShaderMapMessage", "Compiling global shaders"));
-		CompileGlobalShaderMap(PreviewPlatform.PreviewFeatureLevel);
+			if (PreviewPlatform.bPreviewFeatureLevelActive && PreviewPlatform.DeviceProfileName != NAME_None)
+			{
+				if (UDeviceProfile* DP = UDeviceProfileManager::Get().FindProfile(PreviewPlatform.DeviceProfileName.ToString(), false))
+				{
+					UDeviceProfileManager::Get().SetPreviewDeviceProfile(DP);
+				}
+			}
+		};
 
-		SlowTask.EnterProgressFrame(15.0f, NSLOCTEXT("Engine", "SlowTaskFinalizingMessage", "Finalizing"));
-		GShaderCompilingManager->ProcessAsyncResults(false, true);
+		// Full platform change: re-register components, compile shaders
+		if (bCompileShaders)
+		{
+			constexpr bool bUpdateProgressDialog = true;
+			constexpr bool bCacheAllRemainingShaders = false;
 
-		DefaultWorldFeatureLevel = EffectiveFeatureLevel;
-		PreviewShaderPlatformChanged.Broadcast(ShaderPlatform);
-PRAGMA_DISABLE_DEPRECATION_WARNINGS
-		PreviewFeatureLevelChanged.Broadcast(EffectiveFeatureLevel);
-PRAGMA_ENABLE_DEPRECATION_WARNINGS
-	}
+			//invalidate global bound shader states so they will be created with the new shaders the next time they are set (in SetGlobalBoundShaderState)
+			for (TLinkedList<FGlobalBoundShaderStateResource*>::TIterator It(FGlobalBoundShaderStateResource::GetGlobalBoundShaderStateList()); It; It.Next())
+			{
+				BeginUpdateResourceRHI(*It);
+			}
 
-	Scalability::ChangeScalabilityPreviewPlatform(PreviewPlatform.GetEffectivePreviewPlatformName(), GetActiveShaderPlatform());
-	
-	// this sets TextureLODSettings and platform CVars :
-	UDeviceProfileManager::Get().RestorePreviewDeviceProfile();
+			FGlobalComponentReregisterContext RecreateComponents;
+			FlushRenderingCommands();
 
-	UStaticMesh::OnLodStrippingQualityLevelChanged(nullptr);
+			// Set only require the preview feature level and the max feature level. The Max feature level is required for the toggle feature.
+			for (uint32 i = (uint32)ERHIFeatureLevel::ES3_1; i < (uint32)ERHIFeatureLevel::Num; i++)
+			{
+				ERHIFeatureLevel::Type FeatureLevel = (ERHIFeatureLevel::Type)i;
+				UMaterialInterface::SetGlobalRequiredFeatureLevel(FeatureLevel, FeatureLevel == PreviewPlatform.PreviewFeatureLevel || FeatureLevel == GMaxRHIFeatureLevel);
+			}
 
-	if (PreviewPlatform.bPreviewFeatureLevelActive)
-	{
-		//Override the current device profile.
-		if (PreviewPlatform.DeviceProfileName != NAME_None)
-		{
-			if (UDeviceProfile* DP = UDeviceProfileManager::Get().FindProfile(PreviewPlatform.DeviceProfileName.ToString(), false))
+			// Restore editor's shader platform in case a previous SM6 preview left it stale
+			GShaderPlatformForFeatureLevel[GMaxRHIFeatureLevel] = GMaxRHIShaderPlatform;
+			if (PreviewPlatform.bPreviewFeatureLevelActive)
 			{
-				// this sets TextureLODSettings and platform CVars :
-				UDeviceProfileManager::Get().SetPreviewDeviceProfile(DP);
+				GShaderPlatformForFeatureLevel[PreviewPlatform.PreviewFeatureLevel] = ShaderPlatform;
 			}
-		}
-	}
 
-	Scalability::ApplyCachedQualityLevels();
+			SlowTask.EnterProgressFrame(35.0f);
+			UMaterial::AllMaterialsCacheResourceShadersForRendering(bUpdateProgressDialog, bCacheAllRemainingShaders);
 
-	PreviewPlatformChangePost();
+			SlowTask.EnterProgressFrame(35.0f);
+			UMaterialInstance::AllMaterialsCacheResourceShadersForRendering(bUpdateProgressDialog, bCacheAllRemainingShaders);
 
-PRAGMA_DISABLE_DEPRECATION_WARNINGS	
-	PreviewPlatformChanged.Broadcast();
-PRAGMA_ENABLE_DEPRECATION_WARNINGS
+			SlowTask.EnterProgressFrame(15.0f, NSLOCTEXT("Engine", "SlowTaskGlobalShaderMapMessage", "Compiling global shaders"));
+			CompileGlobalShaderMap(PreviewPlatform.PreviewFeatureLevel);
 
-	if (bSaveSettings)
-	{
-		SaveEditorShaderPlatformPreview();
-	}
-	
-	UE_LOG(LogEditor, Display, TEXT("Done with Preview Platform"));
-}
+			SlowTask.EnterProgressFrame(15.0f, NSLOCTEXT("Engine", "SlowTaskFinalizingMessage", "Finalizing"));
+			GShaderCompilingManager->ProcessAsyncResults(false, true);
 
-void UEditorEngine::ToggleFeatureLevelPreview()
-{
-	UE_LOG(LogEditor, Log, TEXT("Toggle %s Preview ShaderPlatform: %s"),
-		!PreviewPlatform.bPreviewFeatureLevelActive ? TEXT("ON") : TEXT("OFF"),
-		*PreviewPlatform.GetFriendlyName().ToString());
+			// Update device profiles after shaders are compiled 
+			UpdateDeviceProfile(/*bRestore=*/true);
+		}
+		else
+		{
+			UpdateDeviceProfile(/*bRestore=*/false);
+		}
 
-	PreviewPlatformChangePre();
+		DefaultWorldFeatureLevel = EffectiveFeatureLevel;
+		PreviewShaderPlatformChanged.Broadcast(GShaderPlatformForFeatureLevel[EffectiveFeatureLevel]);
+	PRAGMA_DISABLE_DEPRECATION_WARNINGS
+		PreviewFeatureLevelChanged.Broadcast(EffectiveFeatureLevel);
+	PRAGMA_ENABLE_DEPRECATION_WARNINGS
 
-	PreviewPlatform.bPreviewFeatureLevelActive ^= 1;
+		Scalability::ChangeScalabilityPreviewPlatform(PreviewPlatform.GetEffectivePreviewPlatformName(), GetActiveShaderPlatform());
+		Scalability::ApplyCachedQualityLevels();
 
-	ERHIFeatureLevel::Type NewPreviewFeatureLevel = PreviewPlatform.GetEffectivePreviewFeatureLevel();
+		if (bChangedEffectiveShaderPlatform)
+		{
+			OnEffectivePreviewShaderPlatformChange();
+		}
 
-	DefaultWorldFeatureLevel = NewPreviewFeatureLevel;
-PRAGMA_DISABLE_DEPRECATION_WARNINGS
-	PreviewFeatureLevelChanged.Broadcast(NewPreviewFeatureLevel);
-PRAGMA_ENABLE_DEPRECATION_WARNINGS
+		UStaticMesh::OnLodStrippingQualityLevelChanged(nullptr);
 
-	Scalability::ChangeScalabilityPreviewPlatform(PreviewPlatform.GetEffectivePreviewPlatformName(), GetActiveShaderPlatform());
+		PreviewPlatformChangePost();
 
-	if (PreviewPlatform.bPreviewFeatureLevelActive)
-	{
-		GShaderPlatformForFeatureLevel[PreviewPlatform.PreviewFeatureLevel] = PreviewPlatform.ShaderPlatform;
+	PRAGMA_DISABLE_DEPRECATION_WARNINGS
+		PreviewPlatformChanged.Broadcast();
+	PRAGMA_ENABLE_DEPRECATION_WARNINGS
 
-		if (PreviewPlatform.DeviceProfileName != NAME_None)
-		{
-			if (UDeviceProfile* DP = UDeviceProfileManager::Get().FindProfile(PreviewPlatform.DeviceProfileName.ToString(), false))
-			{
-				UDeviceProfileManager::Get().SetPreviewDeviceProfile(DP);
-			}
-		}
-		else
-		{
-			UDeviceProfileManager::Get().RestorePreviewDeviceProfile();
-		}
-	}
-	else
-	{
-		// If the Preview FeatureLevel is the same as the Editor, restore the SP
-		if (PreviewPlatform.PreviewFeatureLevel == GMaxRHIFeatureLevel)
+		RedrawAllViewports();
+
+		if (bSaveSettings)
 		{
-			GShaderPlatformForFeatureLevel[PreviewPlatform.PreviewFeatureLevel] = CachedEditorShaderPlatform;
+			SaveEditorShaderPlatformPreview();
 		}
-		UDeviceProfileManager::Get().RestorePreviewDeviceProfile();
 	}
+	UE_LOG(LogEditor, Display, TEXT("Done with Preview Platform"));
+}
 
-	PreviewShaderPlatformChanged.Broadcast(GShaderPlatformForFeatureLevel[NewPreviewFeatureLevel]);
+void UEditorEngine::ToggleFeatureLevelPreview()
+{
+	FPreviewPlatformInfo NewInfo = PreviewPlatform;
+	NewInfo.bPreviewFeatureLevelActive ^= 1;
 
-	// Update any PerPlatformConfig class defaults or instances
-	for (FThreadSafeObjectIterator ObjIterator(UObject::StaticClass(), RF_NoFlags); ObjIterator; ++ObjIterator)
+	// When the preview shares the editor's feature level, toggling changes GShaderPlatformForFeatureLevel which requires the compile path.
+	const bool bNeedsCompile = (NewInfo.PreviewFeatureLevel == GMaxRHIFeatureLevel);
+	if (bNeedsCompile)
 	{
-		if ((*ObjIterator) && ObjIterator->GetClass()->HasAnyClassFlags(CLASS_PerPlatformConfig))
-		{
-			ObjIterator->LoadConfig();
-		}
+		UE_LOG(LogEditor, Display, TEXT("Toggle requires shader compilation: preview feature level matches editor feature level (GMaxRHIFeatureLevel)"));
 	}
 
-	Scalability::ApplyCachedQualityLevels();
-	OnEffectivePreviewShaderPlatformChange();
-
-	PreviewPlatformChangePost();
-	
-PRAGMA_DISABLE_DEPRECATION_WARNINGS	
-	PreviewPlatformChanged.Broadcast();
-PRAGMA_ENABLE_DEPRECATION_WARNINGS
-
-	UStaticMesh::OnLodStrippingQualityLevelChanged(nullptr);
-
-	GEditor->RedrawAllViewports();
-	
-	SaveEditorShaderPlatformPreview();
+	SetPreviewPlatform(NewInfo, /*bSaveSettings=*/true, /*bCompileShaders=*/bNeedsCompile);
 }
 
 bool UEditorEngine::IsFeatureLevelPreviewEnabled() const
 {
-	return PreviewPlatform.PreviewFeatureLevel != GMaxRHIFeatureLevel || PreviewPlatform.PreviewShaderFormatName != NAME_None || PreviewPlatform.PreviewShaderPlatformName != NAME_None;
+	return PreviewPlatform.PreviewShaderFormatName != NAME_None || PreviewPlatform.PreviewShaderPlatformName != NAME_None;
 }
 
 bool UEditorEngine::IsFeatureLevelPreviewActive() const
diff --git a/Engine/Source/Runtime/RHI/Private/DataDrivenShaderPlatformInfo.cpp b/Engine/Source/Runtime/RHI/Private/DataDrivenShaderPlatformInfo.cpp
--- a/Engine/Source/Runtime/RHI/Private/DataDrivenShaderPlatformInfo.cpp
+++ b/Engine/Source/Runtime/RHI/Private/DataDrivenShaderPlatformInfo.cpp
@@ -571,6 +571,8 @@ void FGenericDataDrivenShaderPlatformInfo::UpdatePreviewPlatforms()
 				PREVIEW_DISABLE_IF_RUNTIME_UNSUPPORTED(bStripNaniteVertexShaders);
 				PREVIEW_DISABLE_IF_RUNTIME_UNSUPPORTED(bSupportsMobileMultiView);
 				PREVIEW_DISABLE_IF_RUNTIME_UNSUPPORTED(bSupportsSpecializationConstants);
+				PREVIEW_DISABLE_IF_RUNTIME_UNSUPPORTED(SupportsBarycentricsIntrinsics);
+				PREVIEW_DISABLE_IF_RUNTIME_UNSUPPORTED(SupportsBarycentricsSemantic);
 
 				// Settings that need to match the runtime
 				PREVIEW_USE_RUNTIME_VALUE(bSupportsGPUScene);
