IncrementalCook: Fix IncrementalValidate not exploring the entire known graph when not -incrementalvalidateallowwrite because it failed to queue the dependencies of modified packages.
#rnx
#Synced-CL 52053053

diff --git a/Engine/Source/Editor/UnrealEd/Private/Cooker/CookRequestCluster.cpp b/Engine/Source/Editor/UnrealEd/Private/Cooker/CookRequestCluster.cpp
--- a/Engine/Source/Editor/UnrealEd/Private/Cooker/CookRequestCluster.cpp
+++ b/Engine/Source/Editor/UnrealEd/Private/Cooker/CookRequestCluster.cpp
@@ -2221,32 +2221,6 @@ void FRequestCluster::FGraphSearch::FExploreEdgesContext::ProcessPlatformAttachm
 	{
 		check(PackagePlatformData.IsIncrementallyUnmodifiedSet());
 		bool bIncrementallyUnmodified = PackagePlatformData.IsIncrementallyUnmodified();
-		if (bIncrementallyUnmodified)
-		{
-			// Queue runtime dependencies if we are exploring runtime dependencies, build definitions if
-			// we are queuing those, and always queue build dependencies
-			UE::Cook::FPackageArtifacts& Artifacts = PlatformAttachments.Artifacts;
-			if (bExploreRuntimeDependencies && PackagePlatformData.IsCookable() && Cluster.bAllowSoftDependencies)
-			{
-				TArray<FName, TInlineAllocator<16>> RuntimeContentDependencies;
-				Artifacts.GetRuntimeContentDependencies(RuntimeContentDependencies);
-				AddPlatformDependencyRange(RuntimeContentDependencies, PlatformIndex, EInstigator::SoftDependency);
-			}
-
-			if (Cluster.bPreQueueBuildDefinitions && PackagePlatformData.IsCookable())
-			{
-				bFoundBuildDefinitions = true;
-				Cluster.BuildDefinitions.AddBuildDefinitionList(PackageName, TargetPlatform,
-					PlatformAttachments.BuildDefinitions.Definitions);
-			}
-
-			TArray<FName, TInlineAllocator<10>> TransitiveBuildDependencies;
-			Artifacts.GetTransitiveBuildDependencies(TransitiveBuildDependencies);
-			for (FName TransitivePackageName : TransitiveBuildDependencies)
-			{
-				AddPlatformDependency(TransitivePackageName, PlatformIndex, EInstigator::BuildDependency);
-			}
-		}
 		bool bShouldIncrementallySkip = bIncrementallyUnmodified;
 		ECookResult CookResult =
 			QueryPlatformData.CookAttachments.CommitStatus == IPackageWriter::ECommitStatus::Success
@@ -2301,8 +2275,38 @@ void FRequestCluster::FGraphSearch::FExploreEdgesContext::ProcessPlatformAttachm
 				}
 			}
 		}
+
 		if (bShouldIncrementallySkip)
 		{
+			// If we incrementally skip a package and we have its previously stored dependencies, try to queue those
+			// dependencies.
+			UE::Cook::FPackageArtifacts& Artifacts = PlatformAttachments.Artifacts;
+			if (Artifacts.IsValid())
+			{
+				// Queue runtime dependencies if we are exploring runtime dependencies, build definitions if
+				// we are queuing those, and always queue build dependencies
+				if (bExploreRuntimeDependencies && PackagePlatformData.IsCookable() && Cluster.bAllowSoftDependencies)
+				{
+					TArray<FName, TInlineAllocator<16>> RuntimeContentDependencies;
+					Artifacts.GetRuntimeContentDependencies(RuntimeContentDependencies);
+					AddPlatformDependencyRange(RuntimeContentDependencies, PlatformIndex, EInstigator::SoftDependency);
+				}
+
+				if (Cluster.bPreQueueBuildDefinitions && PackagePlatformData.IsCookable())
+				{
+					bFoundBuildDefinitions = true;
+					Cluster.BuildDefinitions.AddBuildDefinitionList(PackageName, TargetPlatform,
+						PlatformAttachments.BuildDefinitions.Definitions);
+				}
+
+				TArray<FName, TInlineAllocator<10>> TransitiveBuildDependencies;
+				Artifacts.GetTransitiveBuildDependencies(TransitiveBuildDependencies);
+				for (FName TransitivePackageName : TransitiveBuildDependencies)
+				{
+					AddPlatformDependency(TransitivePackageName, PlatformIndex, EInstigator::BuildDependency);
+				}
+			}
+
 			if (Cluster.TraversalMarkCookable())
 			{
 				// Call SetPlatformCooked instead of just PackagePlatformData.SetCookResults because we might also need
