Fix functional AI tests that fail on maps without pre-built navmesh data (e.g. AutoNavLink).

AutomationLoadMap was synchronous - it started PIE immediately after LoadMap(), copying an empty navmesh into the PIE world before the async editor rebuild could finish; since RuntimeGeneration=Static blocks rebuilds in the game world, the navmesh stayed empty for the duration of the test

Restructures AutomationLoadMap to defer PIE startup via two new latent commands in AutomationEditorCommon: FWaitForEditorNavMeshBuildCommand (waits for the editor world's navmesh to finish building before PIE duplicates it) and FStartPIEForAutomationCommand (starts PIE, waits for map load, and fails the test immediately if PIE fails to start rather than hanging indefinitely).

#Synced-CL 52081264

#Synced-CL 52111343

diff --git a/Engine/Source/Editor/UnrealEd/Private/EditorEngine.cpp b/Engine/Source/Editor/UnrealEd/Private/EditorEngine.cpp
--- a/Engine/Source/Editor/UnrealEd/Private/EditorEngine.cpp
+++ b/Engine/Source/Editor/UnrealEd/Private/EditorEngine.cpp
@@ -6431,29 +6431,6 @@ void FActorLabelUtilities::RenameExistingActor(AActor* Actor, const FString& New
 void UEditorEngine::AutomationLoadMap(const FString& MapName, bool bForceReload, FString* OutError)
 {
 #if WITH_AUTOMATION_TESTS
-	struct FFailedGameStartHandler
-	{
-		bool bCanProceed;
-
-		FFailedGameStartHandler()
-		{
-			bCanProceed = true;
-			FEditorDelegates::EndPIE.AddRaw(this, &FFailedGameStartHandler::OnEndPIE);
-		}
-
-		~FFailedGameStartHandler()
-		{
-			FEditorDelegates::EndPIE.RemoveAll(this);
-		}
-
-		bool CanProceed() const { return bCanProceed; }
-
-		void OnEndPIE(const bool bInSimulateInEditor)
-		{
-			bCanProceed = false;
-		}
-	};
-
 	bool bLoadAsTemplate = false;
 	bool bShowProgress = false;
 
@@ -6493,19 +6470,19 @@ void UEditorEngine::AutomationLoadMap(const FString& MapName, bool bForceReload,
 		bNeedPieStart = true;
 	}
 
-	// special precaution needs to be taken while triggering PIE since it can
-	// fail if there are BP compilation issues
 	if (bNeedPieStart)
 	{
 		UE_LOG(LogEditor, Log, TEXT("Starting PIE for the automation tests for world, %s"), *GWorld->GetMapName());
 
-		FFailedGameStartHandler FailHandler;
-
 		FRequestPlaySessionParams RequestParams;
 		ULevelEditorPlaySettings* EditorPlaySettings = NewObject<ULevelEditorPlaySettings>();
 		EditorPlaySettings->SetPlayNumberOfClients(1);
 		EditorPlaySettings->bLaunchSeparateServer = false;
 		RequestParams.EditorPlaySettings = EditorPlaySettings;
+		// Keep EditorPlaySettings alive across the deferred latent command boundary - it lives inside
+		// a non-UObject struct with no UPROPERTY, so the GC cannot reach it. Balanced by
+		// RemoveFromRoot() in ~FStartPIEForAutomationCommand(), which always runs.
+		EditorPlaySettings->AddToRoot();
 
 		// Make sure the player start location is a valid location.
 		if (CheckForPlayerStart() == nullptr)
@@ -6513,19 +6490,10 @@ void UEditorEngine::AutomationLoadMap(const FString& MapName, bool bForceReload,
 			FAutomationEditorCommonUtils::SetPlaySessionStartToActiveViewport(RequestParams);
 		}
 
-		RequestPlaySession(RequestParams);
-
-		// Immediately launch the session 
-		StartQueuedPlaySessionRequest();
-
-		if (!FailHandler.CanProceed())
-		{
-			*OutError = TEXT("Error encountered.");
-		}
-		else
-		{
-			ADD_LATENT_AUTOMATION_COMMAND(FWaitForMapToLoadCommand);
-		}
+		// Wait for the editor world's navmesh to finish building before duplicating into PIE,
+		// then start PIE and wait for the map to load (FStartPIEForAutomationCommand does both).
+		ADD_LATENT_AUTOMATION_COMMAND(FWaitForEditorNavMeshBuildCommand);
+		ADD_LATENT_AUTOMATION_COMMAND(FStartPIEForAutomationCommand(RequestParams));
 	}
 #endif
 	return;
diff --git a/Engine/Source/Editor/UnrealEd/Private/Tests/AutomationEditorCommon.cpp b/Engine/Source/Editor/UnrealEd/Private/Tests/AutomationEditorCommon.cpp
--- a/Engine/Source/Editor/UnrealEd/Private/Tests/AutomationEditorCommon.cpp
+++ b/Engine/Source/Editor/UnrealEd/Private/Tests/AutomationEditorCommon.cpp
@@ -16,6 +16,9 @@
 #include "Engine/StaticMesh.h"
 #include "Misc/ConfigCacheIni.h"
 #include "Serialization/ArchiveReplaceObjectRef.h"
+#include "NavigationSystem.h"
+#include "Settings/LevelEditorPlaySettings.h"
+#include "GameFramework/GameStateBase.h"
 #include "GameFramework/WorldSettings.h"
 #include "Engine/StaticMeshActor.h"
 #include "LevelEditorViewport.h"
@@ -1131,6 +1134,189 @@ bool FWaitToFinishBuildDeployCommand::Update()
 	return false;
 }
 
+bool FWaitForEditorNavMeshBuildCommand::Update()
+{
+	// Timeout to avoid hanging forever if the nav build never completes (e.g. build system stuck).
+	// Configurable via UAutomationTestSettings::NavMeshBuildTimeoutSeconds in DefaultEngine.ini.
+	const float ConfiguredTimeout = GetDefault<UAutomationTestSettings>()->NavMeshBuildTimeoutSeconds;
+	const float NavMeshBuildTimeoutSeconds = FMath::Max(ConfiguredTimeout, 1.0f);
+	if (ConfiguredTimeout < 1.0f)
+	{
+		static bool bWarnedAboutNavMeshTimeout = false;
+		if (!bWarnedAboutNavMeshTimeout)
+		{
+			bWarnedAboutNavMeshTimeout = true;
+			UE_LOG(LogAutomationEditorCommon, Warning, TEXT("FWaitForEditorNavMeshBuildCommand: NavMeshBuildTimeoutSeconds is %.1f, which is invalid. Using %.1f instead."), ConfiguredTimeout, NavMeshBuildTimeoutSeconds);
+		}
+	}
+	if (GetCurrentRunTime() > NavMeshBuildTimeoutSeconds)
+	{
+		UE_LOGF(LogAutomationEditorCommon, Warning, "FWaitForEditorNavMeshBuildCommand timed out after %.0f seconds waiting for editor navmesh build to complete.", NavMeshBuildTimeoutSeconds);
+		return true;
+	}
+
+	for (const FWorldContext& Context : GEngine->GetWorldContexts())
+	{
+		if (Context.WorldType == EWorldType::Editor && Context.World())
+		{
+			UNavigationSystemV1* NavigationSys = FNavigationSystem::GetCurrent<UNavigationSystemV1>(Context.World());
+			if (NavigationSys)
+			{
+				// Returns true immediately if nav is already clean — zero cost for maps where the navmesh
+				// is pre-built or where the test doesn't care about navigation.
+				return !NavigationSys->IsNavigationBuildInProgress() && !NavigationSys->HasDirtyAreasQueued();
+			}
+			break;
+		}
+	}
+
+	return true; // No editor world or no nav system — returns immediately on first tick, zero cost for maps that don't use navigation.
+}
+
+FStartPIEForAutomationCommand::FStartPIEForAutomationCommand(FRequestPlaySessionParams InRequestParams)
+	: RequestParams(InRequestParams)
+	, bPIEStarted(false)
+	, bPIEFailed(false)
+	, bSessionRequested(false)
+{
+	FEditorDelegates::PostPIEStarted.AddRaw(this, &FStartPIEForAutomationCommand::OnPostPIEStarted);
+	FEditorDelegates::EndPIE.AddRaw(this, &FStartPIEForAutomationCommand::OnEndPIE);
+}
+
+FStartPIEForAutomationCommand::~FStartPIEForAutomationCommand()
+{
+	FEditorDelegates::PostPIEStarted.RemoveAll(this);
+	FEditorDelegates::EndPIE.RemoveAll(this);
+	// Balance the AddToRoot() called in AutomationLoadMap. The destructor always runs regardless
+	// of whether Update() was ever called, so this covers test abort and framework teardown paths.
+	// The null check guards against the unlikely case where EditorPlaySettings was never set.
+	if (ULevelEditorPlaySettings* EditorPlaySettings = RequestParams.EditorPlaySettings.Get())
+	{
+		EditorPlaySettings->RemoveFromRoot();
+	}
+}
+
+bool FStartPIEForAutomationCommand::Update()
+{
+	// First tick: fire the PIE request
+	if (!bSessionRequested)
+	{
+		bSessionRequested = true;
+		GEditor->RequestPlaySession(RequestParams);
+		GEditor->StartQueuedPlaySessionRequest();
+		return false; // keep polling
+	}
+
+	const float ConfiguredPIETimeout = GetDefault<UAutomationTestSettings>()->PIEStartTimeoutSeconds;
+	const float TimeoutSeconds = FMath::Max(ConfiguredPIETimeout, 1.0f);
+	if (ConfiguredPIETimeout < 1.0f)
+	{
+		static bool bWarnedAboutPIETimeout = false;
+		if (!bWarnedAboutPIETimeout)
+		{
+			bWarnedAboutPIETimeout = true;
+			UE_LOG(LogAutomationEditorCommon, Warning, TEXT("FStartPIEForAutomationCommand: PIEStartTimeoutSeconds is %.1f, which is invalid. Using %.1f instead."), ConfiguredPIETimeout, TimeoutSeconds);
+		}
+	}
+
+	// PIE started successfully — wait for the world to be ready.
+	// Matches the condition used by FWaitForMapToLoadCommand. HasMatchStarted() checks World->bMatchStarted,
+	// which is set by AGameModeBase::StartPlay() -> GameState->HandleBeginPlay() -> NotifyMatchStarted()
+	// for all AGameModeBase-derived game modes, including AFunctionalTestGameMode.
+	if (bPIEStarted)
+	{
+		// PIE ended after starting — covers mid-initialization crashes (e.g. Blueprint load failure)
+		if (bPIEFailed)
+		{
+			if (FAutomationTestBase* CurrentTest = FAutomationTestFramework::Get().GetCurrentTest())
+			{
+				CurrentTest->AddError(TEXT("FStartPIEForAutomationCommand: PIE ended before the world became ready."));
+			}
+			else
+			{
+				UE_LOG(LogAutomationEditorCommon, Error, TEXT("FStartPIEForAutomationCommand: PIE ended before the world became ready."));
+			}
+			return true;
+		}
+
+		UWorld* GameWorld = AutomationCommon::GetAnyGameWorld();
+		if (GameWorld && GameWorld->AreActorsInitialized())
+		{
+			AGameStateBase* GameState = GameWorld->GetGameState();
+			if (GameState && GameState->HasMatchStarted())
+			{
+				return true; // ready
+			}
+		}
+
+		// Timeout guards against the world stalling after PIE has started
+		if (GetCurrentRunTime() > TimeoutSeconds)
+		{
+			const FText ErrorMsg = FText::Format(
+				FText::FromString(TEXT("FStartPIEForAutomationCommand: world did not become ready after {0} seconds (PIEStartTimeoutSeconds = {1}).")),
+				FText::AsNumber(FMath::RoundToInt(GetCurrentRunTime())),
+				FText::AsNumber(TimeoutSeconds));
+			if (FAutomationTestBase* CurrentTest = FAutomationTestFramework::Get().GetCurrentTest())
+			{
+				CurrentTest->AddError(ErrorMsg.ToString());
+			}
+			else
+			{
+				UE_LOG(LogAutomationEditorCommon, Error, TEXT("%s"), *ErrorMsg.ToString());
+			}
+			return true;
+		}
+
+		return false; // PIE is up but world not ready yet, keep waiting
+	}
+
+	// PIE failed before it started (EndPIE fired before PostPIEStarted)
+	if (bPIEFailed)
+	{
+		if (FAutomationTestBase* CurrentTest = FAutomationTestFramework::Get().GetCurrentTest())
+		{
+			CurrentTest->AddError(TEXT("FStartPIEForAutomationCommand: PIE failed to start (Blueprint compilation error or other startup failure)."));
+		}
+		else
+		{
+			UE_LOG(LogAutomationEditorCommon, Error, TEXT("FStartPIEForAutomationCommand: PIE failed to start (Blueprint compilation error or other startup failure)."));
+		}
+		return true;
+	}
+
+	// Timeout — guards against the session request silently failing before either delegate fires.
+	if (GetCurrentRunTime() > TimeoutSeconds)
+	{
+		const FText ErrorMsg = FText::Format(
+			FText::FromString(TEXT("FStartPIEForAutomationCommand: PIE did not start after {0} seconds (PIEStartTimeoutSeconds = {1}).")),
+			FText::AsNumber(FMath::RoundToInt(GetCurrentRunTime())),
+			FText::AsNumber(TimeoutSeconds));
+		if (FAutomationTestBase* CurrentTest = FAutomationTestFramework::Get().GetCurrentTest())
+		{
+			CurrentTest->AddError(ErrorMsg.ToString());
+		}
+		else
+		{
+			UE_LOG(LogAutomationEditorCommon, Error, TEXT("%s"), *ErrorMsg.ToString());
+		}
+		return true;
+	}
+
+	return false; // still waiting for PIE to come up or fail
+}
+
+void FStartPIEForAutomationCommand::OnPostPIEStarted(const bool bInSimulateInEditor)
+{
+	bPIEStarted = true;
+}
+
+void FStartPIEForAutomationCommand::OnEndPIE(const bool bInSimulateInEditor)
+{
+	// Set bPIEFailed unconditionally — covers both "PIE never started" and "PIE started but
+	// ended before the world became ready" (e.g. mid-initialization crash).
+	bPIEFailed = true;
+}
+
 // agrant-todo: Expose the version in AutomationCommon.cpp for 4.27
 namespace EditorAutomationPrivate
 {
diff --git a/Engine/Source/Editor/UnrealEd/Public/Tests/AutomationEditorCommon.h b/Engine/Source/Editor/UnrealEd/Public/Tests/AutomationEditorCommon.h
--- a/Engine/Source/Editor/UnrealEd/Public/Tests/AutomationEditorCommon.h
+++ b/Engine/Source/Editor/UnrealEd/Public/Tests/AutomationEditorCommon.h
@@ -12,6 +12,7 @@
 class UFactory;
 
 #if WITH_AUTOMATION_TESTS
+#include "Editor.h"
 
 //////////////////////////////////////////////////////////////////////////
 // FAutomationEditorCommonUtils
@@ -278,4 +279,37 @@ DEFINE_EXPORTED_LATENT_AUTOMATION_COMMAND(UNREALED_API, FWaitToFinishBuildDeploy
 */
 DEFINE_EXPORTED_LATENT_AUTOMATION_COMMAND_ONE_PARAMETER(UNREALED_API, FDeleteDirCommand, FString, InFolderLocation);
 
+/**
+ * Waits for the editor world's navmesh to finish building before proceeding.
+ * Ensures PIE duplicates a fully built navmesh rather than an empty one.
+ * Timeout duration is configurable via UAutomationTestSettings::NavMeshBuildTimeoutSeconds.
+ */
+DEFINE_EXPORTED_LATENT_AUTOMATION_COMMAND(UNREALED_API, FWaitForEditorNavMeshBuildCommand);
+
+/**
+ * Starts a PIE session for an automation test and waits for it to be ready.
+ * Defers PIE start so it can be preceded by latent wait commands (e.g. FWaitForEditorNavMeshBuildCommand).
+ * Handles failure at all stages: if PIE fails to start, starts but ends before the world is ready,
+ * or the world stalls during initialization, the test is marked as failed via AddError() and the
+ * command terminates. All failure paths are guarded by PIEStartTimeoutSeconds.
+ */
+class UNREALED_API FStartPIEForAutomationCommand : public IAutomationLatentCommand
+{
+public:
+	FStartPIEForAutomationCommand(FRequestPlaySessionParams InRequestParams);
+	virtual ~FStartPIEForAutomationCommand();
+
+	// Update returns false until its work is complete, then it returns true to indicate its work is complete.
+	virtual bool Update() override;
+
+private:
+	void OnPostPIEStarted(const bool bInSimulateInEditor);
+	void OnEndPIE(const bool bInSimulateInEditor);
+
+	FRequestPlaySessionParams RequestParams;
+	bool bPIEStarted;
+	bool bPIEFailed;
+	bool bSessionRequested;
+};
+
 #endif
diff --git a/Engine/Source/Runtime/Engine/Classes/Tests/AutomationTestSettings.h b/Engine/Source/Runtime/Engine/Classes/Tests/AutomationTestSettings.h
--- a/Engine/Source/Runtime/Engine/Classes/Tests/AutomationTestSettings.h
+++ b/Engine/Source/Runtime/Engine/Classes/Tests/AutomationTestSettings.h
@@ -463,4 +463,22 @@ public:
 	 */
 	UPROPERTY(EditAnywhere, config, Category = Automation)
 	float DefaultInteractiveFramerateDuration;
+
+	/**
+	 * Timeout in seconds for FWaitForEditorNavMeshBuildCommand. If the editor world's navmesh
+	 * has not finished building within this time, the command completes with a warning log and
+	 * the test proceeds — PIE may then start with incomplete navmesh data.
+	 * Increase this for projects with very large or complex navmeshes.
+	 */
+	UPROPERTY(EditAnywhere, config, Category = Automation)
+	float NavMeshBuildTimeoutSeconds;
+
+	/**
+	 * Timeout in seconds for FStartPIEForAutomationCommand. If PIE fails to start, or starts
+	 * but the game world does not reach a ready state within this time, the command completes
+	 * and marks the test as failed via AddError().
+	 * Increase this for maps with slow initialization or on slower machines.
+	 */
+	UPROPERTY(EditAnywhere, config, Category = Automation)
+	float PIEStartTimeoutSeconds;
 };
diff --git a/Engine/Source/Runtime/Engine/Private/Tests/AutomationTestSettings.cpp b/Engine/Source/Runtime/Engine/Private/Tests/AutomationTestSettings.cpp
--- a/Engine/Source/Runtime/Engine/Private/Tests/AutomationTestSettings.cpp
+++ b/Engine/Source/Runtime/Engine/Private/Tests/AutomationTestSettings.cpp
@@ -15,5 +15,7 @@ UAutomationTestSettings::UAutomationTestSettings(const FObjectInitializer& Objec
 	DefaultInteractiveFramerate = 10.f;
 	DefaultInteractiveFramerateWaitTime = 10.f * 60.f;		// 10 minutes - assume shaders and other DDC may need to be built
 	DefaultInteractiveFramerateDuration = 5.f;
+	NavMeshBuildTimeoutSeconds = 60.f;
+	PIEStartTimeoutSeconds = 60.f;
 }
 
