EditorUsability : UnrealEd
- Replaced delete content browser popups by dialogs to attract user attention (folders, favorites, collections)
- Improved delete dialog widget to include pending deletion folders
- If the folder is not empty and contains assets, a first dialog will be shown to avoid loading all assets recusively on a mis-click otherwise only the main delete dialog widget will be used

#Synced-CL 52147079

diff --git a/Engine/Source/Developer/AssetTools/Private/AssetViewUtils.cpp b/Engine/Source/Developer/AssetTools/Private/AssetViewUtils.cpp
--- a/Engine/Source/Developer/AssetTools/Private/AssetViewUtils.cpp
+++ b/Engine/Source/Developer/AssetTools/Private/AssetViewUtils.cpp
@@ -428,10 +428,18 @@ int32 AssetViewUtils::DeleteAssets(const TArray<UObject*>& AssetsToDelete)
 }
 
 bool AssetViewUtils::DeleteFolders(const TArray<FString>& PathsToDelete)
+{
+	FDeleteFolderParameters Options;
+	Options.bShowConfirmationBeforeLoadingAssets = true;
+	Options.PathsToDelete = PathsToDelete;
+	return DeleteFolders(Options);
+}
+
+bool AssetViewUtils::DeleteFolders(const FDeleteFolderParameters& InOptions)
 {
 	// Get a list of assets in the paths to delete
 	TArray<FAssetData> AssetDataList;
-	GetAssetsInPaths(PathsToDelete, AssetDataList);
+	GetAssetsInPaths(InOptions.PathsToDelete, AssetDataList);
 
 	// Add localized variants associated with the current assets to delete if the user wants to include them
 	if (!ObjectTools::AddExtraLocalizedVariantsToDelete(AssetDataList, AssetDataList))
@@ -441,14 +449,35 @@ bool AssetViewUtils::DeleteFolders(const TArray<FString>& PathsToDelete)
 	}
 
 	const int32 NumAssetsInPaths = AssetDataList.Num();
-	bool bAllowFolderDelete = false;
+	bool bSuccess = false;
 	if ( NumAssetsInPaths == 0 )
 	{
 		// There were no assets, allow the folder delete.
-		bAllowFolderDelete = true;
+		ObjectTools::FDeleteItemsParameters Options;
+		Options.FoldersToDelete = InOptions.PathsToDelete;
+		bSuccess = ObjectTools::DeleteItems(Options).FoldersDeleted == InOptions.PathsToDelete.Num();
 	}
 	else
 	{
+		if (InOptions.bShowConfirmationBeforeLoadingAssets)
+		{
+			FText Prompt;
+			if (InOptions.PathsToDelete.Num() == 1)
+			{
+				Prompt = FText::Format(LOCTEXT("FolderDeleteConfirm_Single", "Delete folder '{0}'?"), FText::FromString(InOptions.PathsToDelete[0]));
+			}
+			else
+			{
+				Prompt = FText::Format(LOCTEXT("FolderDeleteConfirm_Multiple", "Delete {0} folders?"), FText::AsNumber(InOptions.PathsToDelete.Num()));
+			}
+
+			// If we had any folders selected, ask the user whether they want to delete them 
+			// as it can be slow to build the deletion dialog on an accidental click (due to assets being loaded)
+			if (FMessageDialog::Open(EAppMsgType::YesNo, Prompt, LOCTEXT("FolderDeleteConfirmTitle", "Delete")) != EAppReturnType::Type::Yes)
+			{
+				return false;
+			}
+		}
 		// Load all the assets in the folder and attempt to delete them.
 		// If it was successful, allow the folder delete.
 
@@ -471,18 +500,13 @@ bool AssetViewUtils::DeleteFolders(const TArray<FString>& PathsToDelete)
 			// Make sure we loaded all of them
 			if ( LoadedAssets.Num() == NumAssetsInPaths )
 			{
-				TArray<UObject*> ToDelete = LoadedAssets;
-				ObjectTools::AddExtraObjectsToDelete(ToDelete);
-				const int32 NumAssetsDeleted = DeleteAssets(ToDelete);
-				if ( NumAssetsDeleted == ToDelete.Num() )
-				{
-					// Successfully deleted all assets in the specified path. Allow the folder to be removed.
-					bAllowFolderDelete = true;
-				}
-				else
-				{
-					// Not all the assets in the selected paths were deleted
-				}
+				ObjectTools::FDeleteItemsParameters Options;
+				Options.ObjectsToDelete = LoadedAssets;
+				Options.FoldersToDelete = InOptions.PathsToDelete;
+				ObjectTools::AddExtraObjectsToDelete(Options.ObjectsToDelete);
+				ObjectTools::FDeleteItemsResult Result = ObjectTools::DeleteItems(Options);
+				bSuccess = Result.FoldersDeleted == Options.FoldersToDelete.Num()
+					&& Result.ObjectsDeleted == Options.ObjectsToDelete.Num();
 			}
 			else
 			{
@@ -494,23 +518,8 @@ bool AssetViewUtils::DeleteFolders(const TArray<FString>& PathsToDelete)
 			// The user declined to load some assets or some assets failed to load
 		}
 	}
-	
-	if ( bAllowFolderDelete )
-	{
-		FAssetRegistryModule& AssetRegistryModule = FModuleManager::LoadModuleChecked<FAssetRegistryModule>("AssetRegistry");
-
-		for (const FString& PathToDelete : PathsToDelete)
-		{
-			if (DeleteEmptyFolderFromDisk(PathToDelete))
-			{
-				AssetRegistryModule.Get().RemovePath(PathToDelete);
-			}
-		}
-
-		return true;
-	}
 
-	return false;
+	return bSuccess;
 }
 
 bool AssetViewUtils::DeleteEmptyFolderFromDisk(const FString& InPathToDelete)
diff --git a/Engine/Source/Developer/AssetTools/Public/AssetViewUtils.h b/Engine/Source/Developer/AssetTools/Public/AssetViewUtils.h
--- a/Engine/Source/Developer/AssetTools/Public/AssetViewUtils.h
+++ b/Engine/Source/Developer/AssetTools/Public/AssetViewUtils.h
@@ -104,6 +104,14 @@ namespace AssetViewUtils
 
 	/** Attempts to delete the specified folders and all assets inside them. Returns true if the operation succeeded. */
 	ASSETTOOLS_API bool DeleteFolders(const TArray<FString>& PathsToDelete);
+	
+	/** Attempts to delete the specified folders and all assets inside them. Returns true if the operation succeeded. */
+	struct FDeleteFolderParameters
+	{
+		TArray<FString> PathsToDelete;
+		bool bShowConfirmationBeforeLoadingAssets = true;
+	};
+	ASSETTOOLS_API bool DeleteFolders(const FDeleteFolderParameters& InOptions);
 
 	/** Gets an array of assets inside the specified folders */
 	ASSETTOOLS_API void GetAssetsInPaths(const TArray<FString>& InPaths, TArray<FAssetData>& OutAssetDataList);
diff --git a/Engine/Source/Editor/ContentBrowser/Private/AssetContextMenu.cpp b/Engine/Source/Editor/ContentBrowser/Private/AssetContextMenu.cpp
--- a/Engine/Source/Editor/ContentBrowser/Private/AssetContextMenu.cpp
+++ b/Engine/Source/Editor/ContentBrowser/Private/AssetContextMenu.cpp
@@ -1219,28 +1219,9 @@ void FAssetContextMenu::ExecuteDelete()
 		SourceAndItemsPair.Key->BulkDeleteItems(SourceAndItemsPair.Value);
 	}
 
-	// If we had any folders selected, ask the user whether they want to delete them 
-	// as it can be slow to build the deletion dialog on an accidental click
 	if (SelectedFolders.Num() > 0)
 	{
-		FText Prompt;
-		if (SelectedFolders.Num() == 1)
-		{
-			Prompt = FText::Format(LOCTEXT("FolderDeleteConfirm_Single", "Delete folder '{0}'?"), SelectedFolders[0].GetDisplayName());
-		}
-		else
-		{
-			Prompt = FText::Format(LOCTEXT("FolderDeleteConfirm_Multiple", "Delete {0} folders?"), SelectedFolders.Num());
-		}
-
-		// Spawn a confirmation dialog since this is potentially a highly destructive operation
-		ContentBrowserUtils::DisplayConfirmationPopup(
-			Prompt,
-			LOCTEXT("FolderDeleteConfirm_Yes", "Delete"),
-			LOCTEXT("FolderDeleteConfirm_No", "Cancel"),
-			AssetView.Pin().ToSharedRef(),
-			FOnClicked::CreateSP(this, &FAssetContextMenu::ExecuteDeleteFolderConfirmed)
-		);
+		ExecuteDeleteFolderConfirmed();
 	}
 }
 
diff --git a/Engine/Source/Editor/ContentBrowser/Private/CollectionContextMenu.cpp b/Engine/Source/Editor/ContentBrowser/Private/CollectionContextMenu.cpp
--- a/Engine/Source/Editor/ContentBrowser/Private/CollectionContextMenu.cpp
+++ b/Engine/Source/Editor/ContentBrowser/Private/CollectionContextMenu.cpp
@@ -656,13 +656,10 @@ void FCollectionContextMenu::ExecuteDestroyCollection()
 		Prompt = FText::Format(LOCTEXT("CollectionDestroyConfirm_Multiple", "Delete {0} Collections?"), FText::AsNumber(SelectedCollections.Num()));
 	}
 
-	FOnClicked OnYesClicked = FOnClicked::CreateSP( this, &FCollectionContextMenu::ExecuteDestroyCollectionConfirmed, SelectedCollections );
-	ContentBrowserUtils::DisplayConfirmationPopup(
-		Prompt,
-		LOCTEXT("CollectionDestroyConfirm_Yes", "Delete"),
-		LOCTEXT("CollectionDestroyConfirm_No", "Cancel"),
-		CollectionView.Pin().ToSharedRef(),
-		OnYesClicked);
+	if (ContentBrowserUtils::DisplayConfirmationDialog(LOCTEXT("CollectionDestroyConfirm_Yes", "Delete"), Prompt))
+	{
+		ExecuteDestroyCollectionConfirmed(SelectedCollections);
+	}
 }
 
 FReply FCollectionContextMenu::ExecuteDestroyCollectionConfirmed(TArray<TSharedPtr<FCollectionItem>> CollectionList)
diff --git a/Engine/Source/Editor/ContentBrowser/Private/ContentBrowserUtils.cpp b/Engine/Source/Editor/ContentBrowser/Private/ContentBrowserUtils.cpp
--- a/Engine/Source/Editor/ContentBrowser/Private/ContentBrowserUtils.cpp
+++ b/Engine/Source/Editor/ContentBrowser/Private/ContentBrowserUtils.cpp
@@ -407,6 +407,11 @@ void ContentBrowserUtils::DisplayConfirmationPopup(const FText& Message, const F
 	Popup->OpenPopup(ParentContent);
 }
 
+bool ContentBrowserUtils::DisplayConfirmationDialog(const FText& InTitle, const FText& InMessage)
+{
+	return FMessageDialog::Open(EAppMsgType::YesNo, InMessage, InTitle) == EAppReturnType::Type::Yes;
+}
+
 FString ContentBrowserUtils::GetItemReferencesText(const TArray<FContentBrowserItem>& Items)
 {
 	TArray<FContentBrowserItem> SortedItems = Items;
diff --git a/Engine/Source/Editor/ContentBrowser/Private/ContentBrowserUtils.h b/Engine/Source/Editor/ContentBrowser/Private/ContentBrowserUtils.h
--- a/Engine/Source/Editor/ContentBrowser/Private/ContentBrowserUtils.h
+++ b/Engine/Source/Editor/ContentBrowser/Private/ContentBrowserUtils.h
@@ -56,6 +56,9 @@ namespace ContentBrowserUtils
 	/** Displays a modeless message asking yes or no type question */
 	void DisplayConfirmationPopup(const FText& Message, const FText& YesString, const FText& NoString, const TSharedRef<SWidget>& ParentContent, const FOnClicked& OnYesClicked, const FOnClicked& OnNoClicked = FOnClicked());
 
+	/** Displays a modal dialog asking yes or no question (blocking) */
+	bool DisplayConfirmationDialog(const FText& InTitle, const FText& InMessage);
+
 	/** Returns references to the specified items */
 	FString GetItemReferencesText(const TArray<FContentBrowserItem>& Items);
 
diff --git a/Engine/Source/Editor/ContentBrowser/Private/PathContextMenu.cpp b/Engine/Source/Editor/ContentBrowser/Private/PathContextMenu.cpp
--- a/Engine/Source/Editor/ContentBrowser/Private/PathContextMenu.cpp
+++ b/Engine/Source/Editor/ContentBrowser/Private/PathContextMenu.cpp
@@ -573,29 +573,10 @@ FText FPathContextMenu::GetDeleteToolTip() const
 
 void FPathContextMenu::ExecuteDelete()
 {
-	// If we had any folders selected, ask the user whether they want to delete them 
-	// as it can be slow to build the deletion dialog on an accidental click
 	TSharedPtr<SWidget> ParentContentPtr = ParentContent.Pin();
 	if (ParentContentPtr && SelectedFolders.Num() > 0)
 	{
-		FText Prompt;
-		if (SelectedFolders.Num() == 1)
-		{
-			Prompt = FText::Format(LOCTEXT("FolderDeleteConfirm_Single", "Delete folder '{0}'?"), SelectedFolders[0].GetDisplayName());
-		}
-		else
-		{
-			Prompt = FText::Format(LOCTEXT("FolderDeleteConfirm_Multiple", "Delete {0} folders?"), SelectedFolders.Num());
-		}
-
-		// Spawn a confirmation dialog since this is potentially a highly destructive operation
-		ContentBrowserUtils::DisplayConfirmationPopup(
-			Prompt,
-			LOCTEXT("FolderDeleteConfirm_Yes", "Delete"),
-			LOCTEXT("FolderDeleteConfirm_No", "Cancel"),
-			ParentContentPtr.ToSharedRef(),
-			FOnClicked::CreateSP(this, &FPathContextMenu::ExecuteDeleteFolderConfirmed)
-		);
+		ExecuteDeleteFolderConfirmed();
 	}
 }
 
diff --git a/Engine/Source/Editor/ContentBrowser/Private/SAssetDialog.cpp b/Engine/Source/Editor/ContentBrowser/Private/SAssetDialog.cpp
--- a/Engine/Source/Editor/ContentBrowser/Private/SAssetDialog.cpp
+++ b/Engine/Source/Editor/ContentBrowser/Private/SAssetDialog.cpp
@@ -644,30 +644,10 @@ void SAssetDialog::ExecuteDelete(EAssetDialogCommandContext CommandContext)
 		}
 	}
 
-	// If we had any folders selected, ask the user whether they want to delete them 
-	// as it can be slow to build the deletion dialog on an accidental click
 	if (SelectedFolderInternalPaths.Num() > 0)
 	{
-		FText Prompt;
-		if (SelectedFolderInternalPaths.Num() == 1)
-		{
-			Prompt = FText::Format(LOCTEXT("FolderDeleteConfirm_Single", "Delete folder '{0}'?"), FirstFolderDisplayName);
-		}
-		else
-		{
-			Prompt = FText::Format(LOCTEXT("FolderDeleteConfirm_Multiple", "Delete {0} folders?"), SelectedFolderInternalPaths.Num());
-		}
-
 		const bool bResetSelection = CommandContext == EAssetDialogCommandContext::PathView;
-
-		// Spawn a confirmation dialog since this is potentially a highly destructive operation
-		ContentBrowserUtils::DisplayConfirmationPopup(
-			Prompt,
-			LOCTEXT("FolderDeleteConfirm_Yes", "Delete"),
-			LOCTEXT("FolderDeleteConfirm_No", "Cancel"),
-			AssetPicker->GetAssetView().ToSharedRef(),
-			FOnClicked::CreateSP(this, &SAssetDialog::ExecuteDeleteFolderConfirmed, SelectedFolderInternalPaths, bResetSelection)
-		);
+		ExecuteDeleteFolderConfirmed(SelectedFolderInternalPaths, bResetSelection);
 	}
 }
 
diff --git a/Engine/Source/Editor/ContentBrowser/Private/SContentBrowser.cpp b/Engine/Source/Editor/ContentBrowser/Private/SContentBrowser.cpp
--- a/Engine/Source/Editor/ContentBrowser/Private/SContentBrowser.cpp
+++ b/Engine/Source/Editor/ContentBrowser/Private/SContentBrowser.cpp
@@ -4242,24 +4242,16 @@ void SContentBrowser::HandleDeleteFavorite(TSharedPtr<SWidget> ParentWidget)
 		}
 
 		// Spawn a confirmation dialog since this is potentially a highly destructive operation
-		ContentBrowserUtils::DisplayConfirmationPopup(
-			Prompt,
-			LOCTEXT("FavoriteRemoveConfirm_Yes", "Remove"),
-			LOCTEXT("FavoriteRemoveConfirm_No", "Cancel"),
-			ParentWidget.ToSharedRef(),
-			FOnClicked::CreateSPLambda(this, [this, SelectedFolders]() -> FReply
+		if (ContentBrowserUtils::DisplayConfirmationDialog(LOCTEXT("FavoriteRemoveConfirm_Yes", "Remove"), Prompt))
+		{
+			for (const FContentBrowserItem& Folder : SelectedFolders)
 			{
-				for (const FContentBrowserItem& Folder : SelectedFolders)
-				{
-					ContentBrowserUtils::RemoveFavoriteFolder(FContentBrowserItemPath(Folder.GetVirtualPath(), EContentBrowserPathType::Virtual));
-				}
-
-				GConfig->Flush(false, GEditorPerProjectIni);
-				LegacyContentSourceWidgets->FavoritePathViewPtr->Populate();
+				ContentBrowserUtils::RemoveFavoriteFolder(FContentBrowserItemPath(Folder.GetVirtualPath(), EContentBrowserPathType::Virtual));
+			}
 
-				return FReply::Handled();
-			})
-		);
+			GConfig->Flush(false, GEditorPerProjectIni);
+			LegacyContentSourceWidgets->FavoritePathViewPtr->Populate();
+		}
 	}
 }
 
diff --git a/Engine/Source/Editor/UnrealEd/Private/Dialogs/SDeleteAssetsDialog.cpp b/Engine/Source/Editor/UnrealEd/Private/Dialogs/SDeleteAssetsDialog.cpp
--- a/Engine/Source/Editor/UnrealEd/Private/Dialogs/SDeleteAssetsDialog.cpp
+++ b/Engine/Source/Editor/UnrealEd/Private/Dialogs/SDeleteAssetsDialog.cpp
@@ -35,6 +35,8 @@ namespace DeleteAssetsView
 	static const FName ColumnID_AssetClass( "Class" );
 	static const FName ColumnID_DiskReferences( "DiskReferences" );
 	static const FName ColumnID_MemoryReferences( "MemoryReferences" );
+	static const FName ColumnID_Folder( "Folder" );
+	static const FName ColumnID_Path( "Path" );
 }
 
 static FLinearColor DangerColor( 0.715465432, 0.034230207, 0 );
@@ -106,6 +108,50 @@ private:
 	TSharedPtr<FPendingDelete> Item;
 };
 
+class SPendingFolderDeleteRow : public SMultiColumnTableRow<TSharedPtr<FString>>
+{
+public:
+	SLATE_BEGIN_ARGS(SPendingFolderDeleteRow)
+	{}
+	SLATE_END_ARGS()
+
+	void Construct(const FArguments& InArgs, const TSharedRef<STableViewBase>& InOwnerTableView, TSharedPtr<FString> InItem)
+	{
+		Item = InItem;
+		SMultiColumnTableRow<TSharedPtr<FString>>::Construct(FSuperRowType::FArguments(), InOwnerTableView);
+	}
+
+	virtual TSharedRef<SWidget> GenerateWidgetForColumn(const FName& InColumnName)
+	{
+		if (!Item)
+		{
+			return SNullWidget::NullWidget;
+		}
+
+		if (InColumnName == DeleteAssetsView::ColumnID_Folder)
+		{
+			return SNew(SHorizontalBox)
+				+ SHorizontalBox::Slot()
+				.AutoWidth()
+				.Padding(3, 0, 0, 0)
+				[
+					SNew(STextBlock)
+					.Text(FText::FromString(FPaths::GetBaseFilename(*Item)))
+				];
+		}
+		else if (InColumnName == DeleteAssetsView::ColumnID_Path)
+		{
+			return SNew(STextBlock)
+				.Text(FText::FromString(*Item));
+		}
+
+		return SNullWidget::NullWidget;
+	}
+
+private:
+	TSharedPtr<FString> Item;
+};
+
 //////////////////////////////////////////////////////////////////////////
 // SDeleteAssetsDialog
 
@@ -120,11 +166,21 @@ void SDeleteAssetsDialog::Construct( const FArguments& InArgs, TSharedRef<FAsset
 	bIsActiveTimerRegistered = true;
 	RegisterActiveTimer( 0.f, FWidgetActiveTimerDelegate::CreateSP( this, &SDeleteAssetsDialog::TickDeleteModel ) );
 
+	// Model handles objects/assets to delete
 	DeleteModel = InDeleteModel;
 
+	// Save additional folders to delete
+	FoldersToDelete.Reset(InArgs._FoldersToDelete.Num());
+	for (const FString& Folder : InArgs._FoldersToDelete)
+	{
+		FoldersToDelete.Add(MakeShared<FString>(Folder));
+	}
+
 	// Save off the attributes
 	ParentWindow = InArgs._ParentWindow;
 
+	OnDialogClosed = InArgs._OnDialogClosed;
+
 	ReferencerCommands = TSharedPtr< FUICommandList >(new FUICommandList);
 
 	ReferencerCommands->MapAction(FGenericCommands::Get().Delete, FUIAction(
@@ -182,38 +238,71 @@ TSharedRef<SWidget> SDeleteAssetsDialog::BuildDeleteDialog()
 	FFormatNamedArguments Args;
 	Args.Add( TEXT( "OnDiskReferences" ), FText::AsNumber( DeleteModel->GetAssetReferences().Num() ) );
 
-	TSharedRef< SHeaderRow > HeaderRowWidget =
-		SNew( SHeaderRow )
-		+ SHeaderRow::Column( DeleteAssetsView::ColumnID_Asset )
-		.DefaultLabel( LOCTEXT( "Column_AssetName", "Asset" ) )
-		.HAlignHeader( EHorizontalAlignment::HAlign_Left )
-		.FillWidth(0.5f)
-		+ SHeaderRow::Column( DeleteAssetsView::ColumnID_AssetClass )
-		.DefaultLabel( LOCTEXT( "Column_AssetClass", "Class" ) )
-		.HAlignHeader( EHorizontalAlignment::HAlign_Left )
-		.FillWidth( 0.25f )
-		+ SHeaderRow::Column( DeleteAssetsView::ColumnID_DiskReferences )
-		.DefaultLabel( LOCTEXT( "Column_DiskReferences", "Asset Referencers" ) )
-		.HAlignHeader( EHorizontalAlignment::HAlign_Left )
-		.FillWidth( 0.25f )
-		+ SHeaderRow::Column( DeleteAssetsView::ColumnID_MemoryReferences )
-		.DefaultLabel( LOCTEXT( "Column_MemoryReferences", "Memory References" ) )
-		.HAlignHeader( EHorizontalAlignment::HAlign_Left )
-		.FillWidth( 0.25f );
+	const bool bShowFolderList = !FoldersToDelete.IsEmpty();
+	const bool bShowAssetList = !DeleteModel->GetPendingDeletedAssets()->IsEmpty();
 
 	return SNew( SVerticalBox )
 
-	// The to be deleted assets
+	// The folders to be deleted
 	+ SVerticalBox::Slot()
-	.FillHeight( 0.5f )
+	.FillHeight( bShowFolderList ? 0.4f : 0.f )
 	.Padding( 5.0f )
 	[
 		SNew( SBorder )
+		.Visibility(bShowFolderList ? EVisibility::Visible : EVisibility::Collapsed)
 		.BorderImage( FAppStyle::GetBrush( "ToolPanel.GroupBorder" ) )
 		.Padding( FMargin(0, 0, 0, 3) )
 		[
 			SNew( SVerticalBox )
 
+			// Attempting delete text
+			+ SVerticalBox::Slot()
+			.AutoHeight()
+			[
+				SNew( SBorder )
+				.BorderImage( FAppStyle::GetBrush( "DetailsView.CategoryTop" ) )
+				.BorderBackgroundColor( FLinearColor( .6, .6, .6, 1.0f ) )
+				.Padding(3.0f)
+				[
+					SNew( STextBlock )
+					.Text( LOCTEXT( "AttemptingFolderDelete", "Pending Deleted Folders" ) )
+					.Font( FAppStyle::GetFontStyle( "BoldFont" ) )
+					.ShadowOffset( FVector2D( 1.0f, 1.0f ) )
+				]
+			]
+
+			+ SVerticalBox::Slot()
+			.FillHeight(1.0f)
+			[
+				SAssignNew( FoldersToDeleteList, SListView< TSharedPtr<FString> > )
+				.ListItemsSource( &FoldersToDelete )
+				.OnGenerateRow( this, &SDeleteAssetsDialog::HandleGenerateFolderRow )
+				.HeaderRow(
+					SNew( SHeaderRow )
+                    + SHeaderRow::Column( DeleteAssetsView::ColumnID_Folder )
+                    .DefaultLabel( LOCTEXT( "Column_FolderName", "Folder" ) )
+                    .HAlignHeader( EHorizontalAlignment::HAlign_Left )
+                    .FillWidth(0.25f)
+                    + SHeaderRow::Column( DeleteAssetsView::ColumnID_Path )
+                    .DefaultLabel( LOCTEXT( "Column_FolderPath", "Path" ) )
+                    .HAlignHeader( EHorizontalAlignment::HAlign_Left )
+                    .FillWidth( 0.75f )
+				)
+			]
+		]
+	]
+
+	// The assets to be deleted
+	+ SVerticalBox::Slot()
+	.FillHeight( bShowAssetList ? 0.6f : 0.f )
+	.Padding( 5.0f )
+	[
+		SNew( SBorder )
+		.BorderImage( FAppStyle::GetBrush( "ToolPanel.GroupBorder" ) )
+		.Padding( FMargin(0, 0, 0, 3) )
+		[
+			SNew( SVerticalBox )
+			.Visibility(bShowAssetList ? EVisibility::Visible : EVisibility::Collapsed)
 			// Attempting delete text
 			+ SVerticalBox::Slot()
 			.AutoHeight()
@@ -236,7 +325,25 @@ TSharedRef<SWidget> SDeleteAssetsDialog::BuildDeleteDialog()
 				SAssignNew( ObjectsToDeleteList, SListView< TSharedPtr<FPendingDelete> > )
 				.ListItemsSource( DeleteModel->GetPendingDeletedAssets() )
 				.OnGenerateRow( this, &SDeleteAssetsDialog::HandleGenerateAssetRow )
-				.HeaderRow( HeaderRowWidget )
+				.HeaderRow(
+					SNew( SHeaderRow )
+                    + SHeaderRow::Column( DeleteAssetsView::ColumnID_Asset )
+                    .DefaultLabel( LOCTEXT( "Column_AssetName", "Asset" ) )
+                    .HAlignHeader( EHorizontalAlignment::HAlign_Left )
+                    .FillWidth(0.5f)
+                    + SHeaderRow::Column( DeleteAssetsView::ColumnID_AssetClass )
+                    .DefaultLabel( LOCTEXT( "Column_AssetClass", "Class" ) )
+                    .HAlignHeader( EHorizontalAlignment::HAlign_Left )
+                    .FillWidth( 0.25f )
+                    + SHeaderRow::Column( DeleteAssetsView::ColumnID_DiskReferences )
+                    .DefaultLabel( LOCTEXT( "Column_DiskReferences", "Asset Referencers" ) )
+                    .HAlignHeader( EHorizontalAlignment::HAlign_Left )
+                    .FillWidth( 0.25f )
+                    + SHeaderRow::Column( DeleteAssetsView::ColumnID_MemoryReferences )
+                    .DefaultLabel( LOCTEXT( "Column_MemoryReferences", "Memory References" ) )
+                    .HAlignHeader( EHorizontalAlignment::HAlign_Left )
+                    .FillWidth( 0.25f )
+				)
 			]
 		]
 	]
@@ -272,7 +379,7 @@ TSharedRef<SWidget> SDeleteAssetsDialog::BuildDeleteDialog()
 	]
 
 	+ SVerticalBox::Slot()
-	.FillHeight( 1.0f )
+	.FillHeight( 0.5f )
 	.Padding( 5.0f )
 	[
 		SNew( SBorder )
@@ -431,7 +538,7 @@ FText SDeleteAssetsDialog::GetHandleText() const
 {
 	if ( CanDelete() )
 	{
-		return LOCTEXT( "AreYouSure", "Are you sure you want to delete these assets?" );
+		return LOCTEXT( "AreYouSure", "Are you sure you want to delete these items?" );
 	}
 	else
 	{
@@ -688,6 +795,11 @@ TSharedRef<ITableRow> SDeleteAssetsDialog::HandleGenerateAssetRow( TSharedPtr<FP
 		.Visibility(InItem->IsInternal() ? EVisibility::Collapsed : EVisibility::Visible);
 }
 
+TSharedRef<ITableRow> SDeleteAssetsDialog::HandleGenerateFolderRow(TSharedPtr<FString> InItem, const TSharedRef<STableViewBase>& InOwnerTable)
+{
+	return SNew(SPendingFolderDeleteRow, InOwnerTable, InItem);
+}
+
 void SDeleteAssetsDialog::DeleteRelevantSourceContent()
 {
 	if (DeleteModel->HasAnySourceContentFilesToDelete())
@@ -707,6 +819,8 @@ void SDeleteAssetsDialog::DeleteRelevantSourceContent()
 
 FReply SDeleteAssetsDialog::Delete()
 {
+	OnDialogClosed.ExecuteIfBound(EResult::Delete);
+
 	if (TSharedPtr<SWindow> OwningWindow = ParentWindow.Pin())
 	{
 		OwningWindow->RequestDestroyWindow();
@@ -725,6 +839,8 @@ FReply SDeleteAssetsDialog::Delete()
 
 FReply SDeleteAssetsDialog::Cancel()
 {
+	OnDialogClosed.ExecuteIfBound(EResult::Cancel);
+
 	if (TSharedPtr<SWindow> OwningWindow = ParentWindow.Pin())
 	{
 		OwningWindow->RequestDestroyWindow();
@@ -735,6 +851,8 @@ FReply SDeleteAssetsDialog::Cancel()
 
 FReply SDeleteAssetsDialog::ForceDelete()
 {
+	OnDialogClosed.ExecuteIfBound(EResult::ForceDelete);
+
 	if (TSharedPtr<SWindow> OwningWindow = ParentWindow.Pin())
 	{
 		OwningWindow->RequestDestroyWindow();
diff --git a/Engine/Source/Editor/UnrealEd/Private/ObjectTools.cpp b/Engine/Source/Editor/UnrealEd/Private/ObjectTools.cpp
--- a/Engine/Source/Editor/UnrealEd/Private/ObjectTools.cpp
+++ b/Engine/Source/Editor/UnrealEd/Private/ObjectTools.cpp
@@ -3064,140 +3064,197 @@ namespace ObjectTools
 
 	int32 DeleteObjects( const TArray< UObject* >& InObjectsToDelete, bool bShowConfirmation, EAllowCancelDuringDelete AllowCancelDuringDelete )
 	{
-		const FScopedBusyCursor BusyCursor;
-
-		TArray<UObject*> ObjectsToDelete = InObjectsToDelete;
-		AddExtraObjectsToDelete(ObjectsToDelete);
+		FDeleteItemsParameters Options;
+		Options.ObjectsToDelete = InObjectsToDelete;
+		Options.bShowConfirmation = bShowConfirmation;
+		Options.AllowCancelDuringDelete = AllowCancelDuringDelete;
+		return DeleteItems(Options).ObjectsDeleted;
+	}
+	
+	FDeleteItemsResult DeleteItems(const FDeleteItemsParameters& InOptions)
+	{
+		auto DeleteObjects = [&InOptions](TArray<UObject*> ObjectsToDelete, bool& bCancelled)
+		{
+			const FScopedBusyCursor BusyCursor;
 
-		// Allows deleting of sounds after they have been previewed
-		GEditor->ClearPreviewComponents();
+			AddExtraObjectsToDelete(ObjectsToDelete);
 
-		// Ensure the audio manager is not holding on to any sounds
-		FAudioDeviceManager* AudioDeviceManager = GEditor->GetAudioDeviceManager();
-		if (AudioDeviceManager != nullptr)
-		{
-			AudioDeviceManager->UpdateActiveAudioDevices(false);
+			// Allows deleting of sounds after they have been previewed
+			GEditor->ClearPreviewComponents();
 
-			const int32 NumAudioDevices = AudioDeviceManager->GetNumActiveAudioDevices();
-			for (int32 DeviceIndex = 0; DeviceIndex < NumAudioDevices; DeviceIndex++)
+			// Ensure the audio manager is not holding on to any sounds
+			FAudioDeviceManager* AudioDeviceManager = GEditor->GetAudioDeviceManager();
+			if (AudioDeviceManager != nullptr)
 			{
-				FAudioDevice* AudioDevice = AudioDeviceManager->GetAudioDeviceRaw(DeviceIndex);
-				if (AudioDevice != nullptr)
+				AudioDeviceManager->UpdateActiveAudioDevices(false);
+
+				const int32 NumAudioDevices = AudioDeviceManager->GetNumActiveAudioDevices();
+				for (int32 DeviceIndex = 0; DeviceIndex < NumAudioDevices; DeviceIndex++)
 				{
-					AudioDevice->StopAllSounds();
+					FAudioDevice* AudioDevice = AudioDeviceManager->GetAudioDeviceRaw(DeviceIndex);
+					if (AudioDevice != nullptr)
+					{
+						AudioDevice->StopAllSounds();
+					}
 				}
 			}
-		}
 
-		// Query delegate hook to validate if the delete operation is available
-		FCanDeleteAssetResult CanDeleteResult;
-		FEditorDelegates::OnAssetsCanDelete.Broadcast(ObjectsToDelete, CanDeleteResult);
-		if (!CanDeleteResult.Get())
-		{
-			FMessageDialog::Open(EAppMsgType::Ok, NSLOCTEXT("UnrealEd", "CannotDelete", "Cannot currently delete selected objects. See log for details."));
-			return 0;
-		}
+			// Query delegate hook to validate if the delete operation is available
+			FCanDeleteAssetResult CanDeleteResult;
+			FEditorDelegates::OnAssetsCanDelete.Broadcast(ObjectsToDelete, CanDeleteResult);
+			if (!CanDeleteResult.Get())
+			{
+				FMessageDialog::Open(EAppMsgType::Ok, NSLOCTEXT("UnrealEd", "CannotDelete", "Cannot currently delete selected objects. See log for details."));
+				return 0;
+			}
 
-		// Make sure packages being saved are fully loaded.
-		if( !HandleFullyLoadingPackages( ObjectsToDelete, NSLOCTEXT("UnrealEd", "Delete", "Delete") ) )
-		{
-			return 0;
-		}
+			// Make sure packages being saved are fully loaded.
+			if( !HandleFullyLoadingPackages( ObjectsToDelete, NSLOCTEXT("UnrealEd", "Delete", "Delete") ) )
+			{
+				return 0;
+			}
 
-		FOnPreDestructiveAssetActionArgs PreDestructiveArgs(ObjectsToDelete, EDestructiveAssetActions::AssetDelete);
-		FEditorDelegates::OnPreDestructiveAssetActionWithArgs.Broadcast(PreDestructiveArgs);
+			FOnPreDestructiveAssetActionArgs PreDestructiveArgs(ObjectsToDelete, EDestructiveAssetActions::AssetDelete);
+			FEditorDelegates::OnPreDestructiveAssetActionWithArgs.Broadcast(PreDestructiveArgs);
 
-		if (!PreDestructiveArgs.Result.bSuccess)
-		{
-			UE_LOG(LogObjectTools, Warning, TEXT("%s"), *PreDestructiveArgs.Result.ErrorMessage);
-			return 0;
-		}
+			if (!PreDestructiveArgs.Result.bSuccess)
+			{
+				UE_LOG(LogObjectTools, Warning, TEXT("%s"), *PreDestructiveArgs.Result.ErrorMessage);
+				return 0;
+			}
 
-		// Load the asset registry module
-		FAssetRegistryModule& AssetRegistryModule = FModuleManager::LoadModuleChecked<FAssetRegistryModule>(TEXT("AssetRegistry"));
+			// Load the asset registry module
+			FAssetRegistryModule& AssetRegistryModule = FModuleManager::LoadModuleChecked<FAssetRegistryModule>(TEXT("AssetRegistry"));
 
-		// Don't delete anything if we're still building the asset registry, warn the user and don't delete.
-		if (AssetRegistryModule.Get().IsLoadingAssets())
-		{
-			FNotificationInfo Info( NSLOCTEXT("UnrealEd", "Warning_CantDeleteRebuildingAssetRegistry", "Unable To Delete While Discovering Assets") );
-			Info.ExpireDuration = 3.0f;
-			FSlateNotificationManager::Get().AddNotification(Info);
-			return 0;
-		}
+			// Don't delete anything if we're still building the asset registry, warn the user and don't delete.
+			if (AssetRegistryModule.Get().IsLoadingAssets())
+			{
+				FNotificationInfo Info( NSLOCTEXT("UnrealEd", "Warning_CantDeleteRebuildingAssetRegistry", "Unable To Delete While Discovering Assets") );
+				Info.ExpireDuration = 3.0f;
+				FSlateNotificationManager::Get().AddNotification(Info);
+				return 0;
+			}
 
-		if (ContainsWorldInUse(ObjectsToDelete))
-		{
-			FMessageDialog::Open(
-				EAppMsgType::Ok,
-				NSLOCTEXT("UnrealEd", "DeleteFailedWorldInUse", "Unable to delete level while it is open"),
-				NSLOCTEXT("UnrealEd", "DeleteFailedWorldInUseTitle", "Unable to delete level")
-			);
+			if (ContainsWorldInUse(ObjectsToDelete))
+			{
+				FMessageDialog::Open(
+					EAppMsgType::Ok,
+					NSLOCTEXT("UnrealEd", "DeleteFailedWorldInUse", "Unable to delete level while it is open"),
+					NSLOCTEXT("UnrealEd", "DeleteFailedWorldInUseTitle", "Unable to delete level")
+				);
 
-			return 0;
-		}
+				return 0;
+			}
 
-		// let systems clean up any unnecessary references that they may have
-		// (so that they're not flagged in the dialog)
-		FEditorDelegates::OnAssetsPreDelete.Broadcast(ObjectsToDelete);
+			// let systems clean up any unnecessary references that they may have
+			// (so that they're not flagged in the dialog)
+			FEditorDelegates::OnAssetsPreDelete.Broadcast(ObjectsToDelete);
 
-		TSharedRef<FAssetDeleteModel> DeleteModel = MakeShared<FAssetDeleteModel>(ObjectsToDelete);
+			TSharedRef<FAssetDeleteModel> DeleteModel = MakeShared<FAssetDeleteModel>(ObjectsToDelete);
 
-		if ( bShowConfirmation )
-		{
-			const FVector2D DEFAULT_WINDOW_SIZE = FVector2D( 600, 700 );
+			if ( InOptions.bShowConfirmation )
+			{
+				const FVector2D WindowSize = FVector2D( 600, 700 );
 
-			/** Create the window to host our package dialog widget */
-			TSharedRef<SCustomDialog> DeleteAssetsWindow = SNew(SCustomDialog)
-				.Title(FText::FromString("Delete Assets"))
-				.ClientSize(DEFAULT_WINDOW_SIZE)
-				.AutoCloseOnButtonPress(false);
+				/** Create the window to host our package dialog widget */
+				TSharedRef<SCustomDialog> DeleteAssetsWindow = SNew(SCustomDialog)
+					.Title(FText::FromString("Delete Assets"))
+					.ClientSize(WindowSize)
+					.AutoCloseOnButtonPress(false);
 
-			/** Set the content of the window to our package dialog widget */
-			TSharedRef< SDeleteAssetsDialog > DeleteDialog =
-				SNew(SDeleteAssetsDialog, DeleteModel)
-				.ParentWindow( DeleteAssetsWindow );
+				/** Set the content of the window to our package dialog widget */
+				TSharedRef< SDeleteAssetsDialog > DeleteDialog =
+					SNew(SDeleteAssetsDialog, DeleteModel)
+					.FoldersToDelete(InOptions.FoldersToDelete)
+					.OnDialogClosed(
+						SDeleteAssetsDialog::FOnDialogClosed::CreateLambda(
+						[&bCancelled](SDeleteAssetsDialog::EResult Result)
+							{
+								bCancelled = (Result == SDeleteAssetsDialog::EResult::Cancel);
+							})
+						)
+					.ParentWindow(DeleteAssetsWindow);
 
-			DeleteAssetsWindow->SetContent( DeleteDialog );
-			DeleteAssetsWindow->SetOnCancelHotkeyPressed(FExecuteAction::CreateSP(DeleteDialog, &SDeleteAssetsDialog::OnCancelHotkeyPressed));
-			DeleteAssetsWindow->SetOnConfirmHotkeyPressed(FExecuteAction::CreateSP(DeleteDialog, &SDeleteAssetsDialog::OnConfirmHotkeyPressed));
-			DeleteAssetsWindow->SetGetWidgetToFocusOnActivate(FGetFocusWidget::CreateSP(DeleteDialog, &SDeleteAssetsDialog::GetWidgetToFocusOnActivate));
+				DeleteAssetsWindow->SetContent( DeleteDialog );
+				DeleteAssetsWindow->SetOnCancelHotkeyPressed(FExecuteAction::CreateSP(DeleteDialog, &SDeleteAssetsDialog::OnCancelHotkeyPressed));
+				DeleteAssetsWindow->SetOnConfirmHotkeyPressed(FExecuteAction::CreateSP(DeleteDialog, &SDeleteAssetsDialog::OnConfirmHotkeyPressed));
+				DeleteAssetsWindow->SetGetWidgetToFocusOnActivate(FGetFocusWidget::CreateSP(DeleteDialog, &SDeleteAssetsDialog::GetWidgetToFocusOnActivate));
 
-			/** Show the package dialog window as a modal window */
-			GEditor->EditorAddModalWindow( DeleteAssetsWindow );
+				/** Show the package dialog window as a modal window */
+				GEditor->EditorAddModalWindow( DeleteAssetsWindow );
 
-			return DeleteModel->GetDeletedObjectCount();
-		}
+				return DeleteModel->GetDeletedObjectCount();
+			}
 
-		bool bUserCanceled = false;
-		const bool bAllowCancelDuringDelete = (AllowCancelDuringDelete == EAllowCancelDuringDelete::AllowCancel);
-		GWarn->BeginSlowTask(NSLOCTEXT("UnrealEd", "VerifyingDelete", "Verifying Delete"), true, bAllowCancelDuringDelete);
-		while ( !bUserCanceled && DeleteModel->GetState() != FAssetDeleteModel::Finished )
-		{
-			DeleteModel->Tick(0);
-			GWarn->StatusUpdate((int32)( DeleteModel->GetProgress() * 100 ), 100, DeleteModel->GetProgressText());
+			bool bUserCanceled = false;
+			const bool bAllowCancelDuringDelete = (InOptions.AllowCancelDuringDelete == EAllowCancelDuringDelete::AllowCancel);
+			GWarn->BeginSlowTask(NSLOCTEXT("UnrealEd", "VerifyingDelete", "Verifying Delete"), true, bAllowCancelDuringDelete);
+			while ( !bUserCanceled && DeleteModel->GetState() != FAssetDeleteModel::Finished )
+			{
+				DeleteModel->Tick(0);
+				GWarn->StatusUpdate((int32)( DeleteModel->GetProgress() * 100 ), 100, DeleteModel->GetProgressText());
+
+				if (bAllowCancelDuringDelete)
+				{
+					bUserCanceled = GWarn->ReceivedUserCancel();
+				}
+			}
+			GWarn->EndSlowTask();
 
-			if (bAllowCancelDuringDelete)
+			if ( bUserCanceled )
 			{
-				bUserCanceled = GWarn->ReceivedUserCancel();
+				UE_LOG(LogUObjectGlobals, Warning, TEXT("User canceled delete operation"));
+				return 0;
 			}
-		}
-		GWarn->EndSlowTask();
 
-		if ( bUserCanceled )
-		{
-			UE_LOG(LogUObjectGlobals, Warning, TEXT("User canceled delete operation"));
-			return 0;
-		}
+			if ( !DeleteModel->DoDelete() )
+			{
+				UE_LOG(LogUObjectGlobals, Warning, TEXT("Could not delete"));
+
+				//@todo ndarnell explain why the delete failed?  Maybe we should show the delete UI
+				// when this fails?
+			}
 
-		if ( !DeleteModel->DoDelete() )
+			return DeleteModel->GetDeletedObjectCount();
+		};
+
+		bool bCancelled = false;
+		FDeleteItemsResult Result;
+		Result.ObjectsDeleted = DeleteObjects(InOptions.ObjectsToDelete, bCancelled);
+
+		if (!bCancelled)
 		{
-			UE_LOG(LogUObjectGlobals, Warning, TEXT("Could not delete"));
+			// Look for files on disk in case the folder contains things not tracked by the asset registry
+			auto DeleteEmptyFolderFromDisk = [](const FString& InPathToDelete)->bool
+				{
+					FString PathToDeleteOnDisk;
+					if (FPackageName::TryConvertLongPackageNameToFilename(InPathToDelete, PathToDeleteOnDisk))
+					{
+						TArray<FString> FileNames;
+						IFileManager::Get().FindFilesRecursive(FileNames, *PathToDeleteOnDisk, TEXT(""), /** Files */true, /** Directory */false);
 
-			//@todo ndarnell explain why the delete failed?  Maybe we should show the delete UI
-			// when this fails?
+						if (FileNames.IsEmpty())
+						{
+							return IFileManager::Get().DeleteDirectory(*PathToDeleteOnDisk, /** RequireExists */false, /** Tree */true);
+						}
+					}
+			
+					return false;
+				};
+
+			// Remove folders specified only if they are empty (no files)
+			const FAssetRegistryModule& AssetRegistryModule = FModuleManager::LoadModuleChecked<FAssetRegistryModule>("AssetRegistry");
+			for (const FString& PathToDelete : InOptions.FoldersToDelete)
+			{
+				if (DeleteEmptyFolderFromDisk(PathToDelete))
+				{
+					AssetRegistryModule.Get().RemovePath(PathToDelete);
+					Result.FoldersDeleted++;
+				}
+			}
 		}
 
-		return DeleteModel->GetDeletedObjectCount();
+		return Result;
 	}
 
 	int32 PrivatizeObjects(const TArray<UObject*>& InObjectsToPrivatize, bool bShowConfirmation, EAllowCancelDuringPrivatize AllowCancelDuringPrivatize, const EAssetAccessSpecifier InAssetAccessSpecifier)
diff --git a/Engine/Source/Editor/UnrealEd/Public/Dialogs/SDeleteAssetsDialog.h b/Engine/Source/Editor/UnrealEd/Public/Dialogs/SDeleteAssetsDialog.h
--- a/Engine/Source/Editor/UnrealEd/Public/Dialogs/SDeleteAssetsDialog.h
+++ b/Engine/Source/Editor/UnrealEd/Public/Dialogs/SDeleteAssetsDialog.h
@@ -30,6 +30,14 @@ class SWindow;
 class SDeleteAssetsDialog : public SCompoundWidget
 {
 public:
+	enum class EResult : uint8
+	{
+		Cancel,
+		Delete,
+		ForceDelete,
+	};
+	
+	DECLARE_DELEGATE_OneParam(FOnDialogClosed, EResult);
 
 	SLATE_BEGIN_ARGS(SDeleteAssetsDialog)
 		: _Style(TEXT("DeleteAssetsDialog"))
@@ -44,6 +52,12 @@ public:
 		/** When specified, the path box will request this fixed size. */
 		SLATE_ATTRIBUTE(FOptionalSize, WidthOverride)
 
+		/** Additional folders to delete */
+		SLATE_ARGUMENT(TArray<FString>, FoldersToDelete)
+
+		/** Called before this dialog is destroyed */
+		SLATE_EVENT(FOnDialogClosed, OnDialogClosed)
+
 	SLATE_END_ARGS()
 
 public:
@@ -148,6 +162,8 @@ private:
 
 	TSharedRef<ITableRow> HandleGenerateAssetRow( TSharedPtr<FPendingDelete> InItem, const TSharedRef<STableViewBase>& OwnerTable );
 
+	TSharedRef<ITableRow> HandleGenerateFolderRow(TSharedPtr<FString> InItem, const TSharedRef<STableViewBase>& InOwnerTable);
+
 private:
 
 	/** Whether the active timer is currently registered */
@@ -178,4 +194,9 @@ private:
 
 	/** Command list for the context menu for the referencer assets */
 	TSharedPtr< FUICommandList > ReferencerCommands;
+	
+	/** Additional folders to delete */
+	TArray<TSharedPtr<FString>> FoldersToDelete;
+	TSharedPtr<SListView<TSharedPtr<FString>>> FoldersToDeleteList;
+	FOnDialogClosed OnDialogClosed;
 };
diff --git a/Engine/Source/Editor/UnrealEd/Public/ObjectTools.h b/Engine/Source/Editor/UnrealEd/Public/ObjectTools.h
--- a/Engine/Source/Editor/UnrealEd/Public/ObjectTools.h
+++ b/Engine/Source/Editor/UnrealEd/Public/ObjectTools.h
@@ -322,6 +322,26 @@ namespace ObjectTools
 	 */
 	UNREALED_API int32 DeleteObjects( const TArray< UObject* >& ObjectsToDelete, bool bShowConfirmation = true, EAllowCancelDuringDelete AllowCancelDuringDelete = EAllowCancelDuringDelete::AllowCancel);
 
+	struct FDeleteItemsParameters
+	{
+		TArray<UObject*> ObjectsToDelete;
+		TArray<FString> FoldersToDelete;
+		bool bShowConfirmation = true;
+		EAllowCancelDuringDelete AllowCancelDuringDelete = EAllowCancelDuringDelete::AllowCancel;
+	};
+
+	struct FDeleteItemsResult
+	{
+		int32 ObjectsDeleted = 0;
+		int32 FoldersDeleted = 0;
+	};
+	/** 
+	 * Deletes items provided (objects, folders)
+	 * @param InOptions The payload and options for the deletion operation
+	 * @return FDeleteItemsResult The items successfully deleted
+	 */ 
+	UNREALED_API FDeleteItemsResult DeleteItems(const FDeleteItemsParameters& InOptions);
+
 	/**
 	* Privatizes the list of objects (marks their packages as NotExternallyReferencable)
 	* 
