Correctly support (non-actor) external packages in the asset localization gather

These are now explicitly included in the gather alongside their outer package, rather than being implicitly included via the search criteria

#rnx
#Synced-CL 52054134

diff --git a/Engine/Source/Editor/UnrealEd/Classes/Commandlets/GatherTextFromAssetsCommandlet.h b/Engine/Source/Editor/UnrealEd/Classes/Commandlets/GatherTextFromAssetsCommandlet.h
--- a/Engine/Source/Editor/UnrealEd/Classes/Commandlets/GatherTextFromAssetsCommandlet.h
+++ b/Engine/Source/Editor/UnrealEd/Classes/Commandlets/GatherTextFromAssetsCommandlet.h
@@ -49,6 +49,12 @@ struct FGatherTextFromAssetsWorkerMessage_PackageRequest
 
 	UPROPERTY()
 	TSet<FGuid> ExternalActors;
+
+	UPROPERTY()
+	TSet<FName> ExternalPackages;
+
+	UPROPERTY()
+	bool bIsMapPackage = false;
 };
 
 USTRUCT()
@@ -128,13 +134,13 @@ private:
 	void ApplyExcludeExactClassesFilter(const FARFilter& InFilter, TArray<FAssetData>& InOutAssetDataArray) const;
 	void FilterAssetsBasedOnIncludeExcludePaths(TArray<FAssetData>& InOutAssetDataArray) const;
 	
-	bool DiscoverExternalActors(TArray<FAssetData>& InOutAssetDataArray);
-	void RemoveExistingExternalActors(TArray<FAssetData>& InOutAssetDataArray, const TSet<FName>* WorldPackageFilter, TSet<FName>& OutExternalActorsWorldPackageNames, TSet<FName>& OutGameFeatureDataPackageNames) const;
+	bool DiscoverExternalPackages(TArray<FAssetData>& InOutAssetDataArray);
+	void RemoveExistingExternalPackages(TArray<FAssetData>& InOutAssetDataArray, const TSet<FName>* WorldPackageFilter, TSet<FName>& OutExternalActorsWorldPackageNames, TSet<FName>& OutGameFeatureDataPackageNames) const;
 
-	TSet<FName> GetPackageNamesToGather(const TArray<FAssetData>& InAssetDataArray) const;
+	void ExtractPackageNamesToGather(const TArray<FAssetData>& InAssetDataArray, TSet<FName>& OutPackageNames, TMap<FName, FName>& OutExternalToOuterPackages) const;
 	void PopulatePackagesPendingGather(TSet<FName> PackageNamesToGather);
-	void ProcessAndRemoveCachedPackages(TMap<FName, TSet<FGuid>>& OutExternalActorsWithStaleOrMissingCaches);
-	void MergeInExternalActorsWithStaleOrMissingCaches(TMap<FName, TSet<FGuid>>& ExternalActorsWithStaleOrMissingCaches);
+	void ProcessAndRemoveCachedPackages(const TMap<FName, FName>& ExternalToOuterPackages, TMap<FName, TSet<FGuid>>& OutExternalActorsWithStaleOrMissingCaches, TMap<FName, TSet<FName>>& OutExternalPackagesWithStaleOrMissingCaches);
+	void MergeInExternalPackagesWithStaleOrMissingCaches(TMap<FName, TSet<FGuid>>& ExternalActorsWithStaleOrMissingCaches, TMap<FName, TSet<FName>>& ExternalPackagesWithStaleOrMissingCaches);
 
 	bool LoadAndProcessUncachedPackages(TArray<FName>& OutPackagesWithStaleGatherCache);
 
@@ -156,9 +162,15 @@ private:
 		/** The set of external actors to process for a world partition map package */
 		TSet<FGuid> ExternalActors;
 
+		/** The set of external packages to process */
+		TSet<FName> ExternalPackages;
+
 		/** The localization ID of this package, if any */
 		FString PackageLocalizationId;
 
+		/** True if this package contains a map */
+		bool bIsMapPackage = false;
+
 		/** Localization cache state of this package */
 		EPackageLocCacheState PackageLocCacheState;
 
@@ -254,6 +266,12 @@ private:
 		/** The set of external actors to process for a world partition map package */
 		TSet<FGuid> ExternalActors;
 
+		/** The set of external packages to process */
+		TSet<FName> ExternalPackages;
+
+		/** True if this package contains a map */
+		bool bIsMapPackage = false;
+
 		/** Address of the endpoint that requested this package to be gathered */
 		FMessageAddress RequestorAddress;
 	};
diff --git a/Engine/Source/Editor/UnrealEd/Private/Commandlets/GatherTextFromAssetsCommandlet.cpp b/Engine/Source/Editor/UnrealEd/Private/Commandlets/GatherTextFromAssetsCommandlet.cpp
--- a/Engine/Source/Editor/UnrealEd/Private/Commandlets/GatherTextFromAssetsCommandlet.cpp
+++ b/Engine/Source/Editor/UnrealEd/Private/Commandlets/GatherTextFromAssetsCommandlet.cpp
@@ -4,6 +4,7 @@
 
 #include "AssetRegistry/ARFilter.h"
 #include "AssetRegistry/AssetData.h"
+#include "AssetRegistry/AssetRegistryHelpers.h"
 #include "AssetRegistry/AssetRegistryModule.h"
 #include "Async/ParallelFor.h"
 #include "CollectionManagerModule.h"
@@ -346,7 +347,7 @@ namespace UE::Private::GatherTextFromAssetsCommandlet
 	static FAutoConsoleVariableRef CVarProcessAndRemoveCachedPackagesMaxThreads(TEXT("Localization.GatherTextFromAssetsCommandlet.ProcessAndRemoveCachedPackagesMaxThreads"), ProcessAndRemoveCachedPackagesMaxThreads, TEXT("Max number of threads to use if parallelizing the 'process and remove cached packages' step, or <= 0 to use as many threads as possible."));
 
 	// Bump whenever changing the data in any of the FGatherTextFromAssetsWorker*Message types
-	static constexpr int32 WorkerProtocolVersion = 1;
+	static constexpr int32 WorkerProtocolVersion = 2;
 
 	// Values controlling how packages are distributed to workers
 	static int32 MinPackagesToUseWorkers = 1000;
@@ -358,12 +359,12 @@ namespace UE::Private::GatherTextFromAssetsCommandlet
 	static int32 WorkerIdleThreshold = 20;
 	static FAutoConsoleVariableRef CVarWorkerIdleThreshold(TEXT("Localization.GatherTextFromAssetsCommandlet.WorkerIdleThreshold"), WorkerIdleThreshold, TEXT("How many packages can a worker have pending in its queue before we consider it 'idle' and try and assign it more packages? (to avoid workers becoming idle)"));
 
-	UPackage* LoadPackageToGather(const FString& PackageName, const bool bHasExternalActors)
+	UPackage* LoadPackageToGather(const FString& PackageName, const bool bIsMapPackage)
 	{
 		UPackage* Package = nullptr;
 		{
 			FLoadPackageLogOutputRedirector::FScopedCapture ScopedCapture(PackageName);
-			if (bHasExternalActors)
+			if (bIsMapPackage)
 			{
 				Package = LoadWorldPackageForEditor(PackageName, EWorldType::Editor, LOAD_NoWarn | LOAD_Quiet);
 			}
@@ -375,11 +376,12 @@ namespace UE::Private::GatherTextFromAssetsCommandlet
 		return Package;
 	}
 
-	bool GatherTextFromPackage(UPackage* Package, const TSet<FGuid>& ExternalActors, const bool bGatherFromPrimaryPackage, const bool bGatherFromExternalActorPackages, TArray<FGatherableTextData>& OutGatherableTextDataArray)
+	bool GatherTextFromPackage(UPackage* Package, const bool bIsMapPackage, const TSet<FGuid>& ExternalActors, const TSet<FName>& ExternalPackages, const bool bGatherFromPrimaryPackage, const bool bGatherFromExternalPackages, TArray<FGatherableTextData>& OutGatherableTextDataArray)
 	{
 		// Because packages may not have been resaved after this flagging was implemented, we may have added packages to load that weren't flagged - potential false positives.
 		// The loading process should have reflagged said packages so that only true positives will have this flag.
-		if (Package->RequiresLocalizationGather() || ExternalActors.Num() > 0)
+		const bool bHasExternalObjects = ExternalActors.Num() > 0 || ExternalPackages.Num() > 0;
+		if (Package->RequiresLocalizationGather() || bHasExternalObjects)
 		{
 			if (bGatherFromPrimaryPackage)
 			{
@@ -388,32 +390,33 @@ namespace UE::Private::GatherTextFromAssetsCommandlet
 				FPropertyLocalizationDataGatherer(OutGatherableTextDataArray, Package, GatherableTextResultFlags);
 			}
 
-			// If this is a WP world then query the localization for any external actors actors that were determined to be stale
-			if (bGatherFromExternalActorPackages && ExternalActors.Num() > 0)
+			if (bGatherFromExternalPackages && bHasExternalObjects)
 			{
-				if (UWorld* World = UWorld::FindWorldInPackage(Package))
+				// If this is a WP world then query the localization for any external actors actors that were determined to be stale
+				TOptional<FScopedEditorWorld> ScopeEditorWorld;
+				if (UWorld* World = bIsMapPackage ? UWorld::FindWorldInPackage(Package) : nullptr)
 				{
-					UWorld::InitializationValues IVS;
-					IVS.InitializeScenes(false);
-					IVS.AllowAudioPlayback(false);
-					IVS.RequiresHitProxies(false);
-					IVS.CreatePhysicsScene(false);
-					IVS.CreateNavigation(false);
-					IVS.CreateAISystem(false);
-					IVS.ShouldSimulatePhysics(false);
-					IVS.EnableTraceCollision(false);
-					IVS.SetTransactional(false);
-					IVS.CreateFXSystem(false);
-					IVS.CreateWorldPartition(true);
-
-					TOptional<FScopedEditorWorld> ScopeEditorWorld;
 					if (!World->IsInitialized())
 					{
+						UWorld::InitializationValues IVS;
+						IVS.InitializeScenes(false);
+						IVS.AllowAudioPlayback(false);
+						IVS.RequiresHitProxies(false);
+						IVS.CreatePhysicsScene(false);
+						IVS.CreateNavigation(false);
+						IVS.CreateAISystem(false);
+						IVS.ShouldSimulatePhysics(false);
+						IVS.EnableTraceCollision(false);
+						IVS.SetTransactional(false);
+						IVS.CreateFXSystem(false);
+						IVS.CreateWorldPartition(true);
+
 						FLoadPackageLogOutputRedirector::FScopedCapture ScopedCapture(Package->GetName());
 						ScopeEditorWorld.Emplace(World, IVS); // Initializing FScopedEditorWorld can log warnings, so capture those like we do with loading errors
 					}
 
-					if (UWorldPartition* WorldPartition = World->GetWorldPartition())
+					if (UWorldPartition* WorldPartition = World->GetWorldPartition();
+						WorldPartition && ExternalActors.Num() > 0)
 					{
 						// ForEachActorWithLoading may GC while running, so keep the world partition (and indirectly the world and its package) alive
 						TGCObjectScopeGuard<UWorldPartition> WorldPartitionGCGuard(WorldPartition);
@@ -426,13 +429,31 @@ namespace UE::Private::GatherTextFromAssetsCommandlet
 							{
 								if (const AActor* Actor = ActorDescInstance->GetActor())
 								{
-									EPropertyLocalizationGathererResultFlags ActorGatherableTextResultFlags = EPropertyLocalizationGathererResultFlags::Empty;
-									FPropertyLocalizationDataGatherer(OutGatherableTextDataArray, Actor->GetExternalPackage(), ActorGatherableTextResultFlags);
+									EPropertyLocalizationGathererResultFlags GatherableTextResultFlags = EPropertyLocalizationGathererResultFlags::Empty;
+									FPropertyLocalizationDataGatherer(OutGatherableTextDataArray, Actor->GetExternalPackage(), GatherableTextResultFlags);
 								}
 								return true;
 							}, ForEachActorParams);
 					}
 				}
+
+				// Other external packages are currently loaded with their outer package
+				for (const FName ExternalPackageName : ExternalPackages)
+				{
+					if (UPackage* ExternalPackage = FindPackage(nullptr, FNameBuilder(ExternalPackageName).ToString()))
+					{
+						EPropertyLocalizationGathererResultFlags GatherableTextResultFlags = EPropertyLocalizationGathererResultFlags::Empty;
+						FPropertyLocalizationDataGatherer(OutGatherableTextDataArray, ExternalPackage, GatherableTextResultFlags);
+					}
+					else
+					{
+						UE_LOGFMT(LogGatherTextFromAssetsCommandlet, Warning, "'{externalPackage}' expected to load with '{package}'.",
+							("externalPackage", ExternalPackageName.ToString()),
+							("package", Package->GetName()),
+							("id", ::GatherTextFromAssetsCommandlet::LocalizationLogIdentifier)
+						);
+					}
+				}
 			}
 
 			return true;
@@ -993,27 +1014,30 @@ void UGatherTextFromAssetsCommandlet::FilterAssetsBasedOnIncludeExcludePaths(TAr
 	}
 }
 
-/** Remove any external actors that currently exist in InOutAssetDataArray. OutExternalActorsWorldPackageNames is populated with the package paths of worlds using external actors.*/
-void UGatherTextFromAssetsCommandlet::RemoveExistingExternalActors(TArray<FAssetData>& InOutAssetDataArray, const TSet<FName>* WorldPackageFilter, TSet<FName>& OutExternalActorsWorldPackageNames, TSet<FName>& OutGameFeatureDataPackageNames) const
+/** Remove any external packages that currently exist in InOutAssetDataArray. OutExternalActorsWorldPackageNames is populated with the package paths of worlds using external actors.*/
+void UGatherTextFromAssetsCommandlet::RemoveExistingExternalPackages(TArray<FAssetData>& InOutAssetDataArray, const TSet<FName>* WorldPackageFilter, TSet<FName>& OutExternalActorsWorldPackageNames, TSet<FName>& OutGameFeatureDataPackageNames) const
 {
-	auto IsGameFeatureData = [](const UClass* AssetClass)
-	{
-		static const FTopLevelAssetPath GameFeatureDataClassPath = FTopLevelAssetPath("/Script/GameFeatures", "GameFeatureData"); // cannot include UGameFeatureData as it is inside a plugin
-		for (const UClass* Class = AssetClass; Class; Class = Class->GetSuperClass())
+	InOutAssetDataArray.RemoveAllSwap(
+		[WorldPackageFilter, &OutExternalActorsWorldPackageNames, &OutGameFeatureDataPackageNames](const FAssetData& AssetData)
 		{
-			if (Class->GetClassPathName() == GameFeatureDataClassPath)
-			{
-				return true;
-			}
-		}
-		return false;
-	};
+			static const TSet<FTopLevelAssetPath> GameFeatureDataClassNames =
+				[]()
+				{
+					TSet<FTopLevelAssetPath> GFDClassNames;
+					if (const UClass* GFDClass = FindObject<UClass>(FTopLevelAssetPath("/Script/GameFeatures", "GameFeatureData"))) // cannot include UGameFeatureData as it is inside a plugin
+					{
+						TArray<UClass*> DerivedGFDClasses;
+						GetDerivedClasses(GFDClass, DerivedGFDClasses);
 
-	InOutAssetDataArray.RemoveAll([WorldPackageFilter, &OutExternalActorsWorldPackageNames, &OutGameFeatureDataPackageNames, &IsGameFeatureData](const FAssetData& AssetData)
-		{
-			const FNameBuilder PackageNameStr(AssetData.PackageName);
+						GFDClassNames.Add(GFDClass->GetClassPathName());
+						Algo::Transform(DerivedGFDClasses, GFDClassNames, &UClass::GetClassPathName);
+					}
+					return GFDClassNames;
+				}();
+
+			static const FTopLevelAssetPath WorldClassName = UWorld::StaticClass()->GetClassPathName();
 
-			if (AssetData.AssetClassPath == UWorld::StaticClass()->GetClassPathName())
+			if (AssetData.AssetClassPath == WorldClassName)
 			{
 				if (WorldPackageFilter && !WorldPackageFilter->Contains(AssetData.PackageName))
 				{
@@ -1024,26 +1048,26 @@ void UGatherTextFromAssetsCommandlet::RemoveExistingExternalActors(TArray<FAsset
 					OutExternalActorsWorldPackageNames.Add(AssetData.PackageName);
 				}
 			}
-			else if (PackageNameStr.ToView().Contains(FPackagePath::GetExternalActorsFolderName()))
+			else if (!AssetData.GetOptionalOuterPathName().IsNone())
 			{
-				// Remove any external actors that are already in the list, as they will be re-added providing their owner world passed the gather criteria
-				// It is possible for an external actor to be directly specified for gather in the configs but have their world not pass the gather criteria.
+				// Remove any external packages that are already in the list, as they will be re-added providing their outer package passed the gather criteria
+				// It is possible for an external package to be directly specified for gather in the configs but have their outer package not pass the gather criteria.
 				return true;
 			}
-			else if (IsGameFeatureData(AssetData.GetClass()))
+			else if (GameFeatureDataClassNames.Contains(AssetData.AssetClassPath))
 			{
 				OutGameFeatureDataPackageNames.Add(AssetData.PackageName);
 			}
 
 			return false;
-		});
+		}, EAllowShrinking::No);
 }
 
-/** Appends any external actors that also need to be gathered to the InOutAssetDataArray. */
-bool UGatherTextFromAssetsCommandlet::DiscoverExternalActors(TArray<FAssetData>& InOutAssetDataArray)
+/** Appends any external packages that also need to be gathered to the InOutAssetDataArray. */
+bool UGatherTextFromAssetsCommandlet::DiscoverExternalPackages(TArray<FAssetData>& InOutAssetDataArray)
 {
-	UE_LOG(LogGatherTextFromAssetsCommandlet, Display, TEXT("Discovering external actors to gather..."));
-	const double DiscoveringExternalActorsStartTime = FPlatformTime::Seconds();
+	UE_LOG(LogGatherTextFromAssetsCommandlet, Display, TEXT("Discovering external packages to gather..."));
+	const double DiscoveringExternalPackagesStartTime = FPlatformTime::Seconds();
 
 	IAssetRegistry& AssetRegistry = IAssetRegistry::GetChecked();
 
@@ -1071,7 +1095,26 @@ bool UGatherTextFromAssetsCommandlet::DiscoverExternalActors(TArray<FAssetData>&
 
 	TSet<FName> ExternalActorsWorldPackageNames;
 	TSet<FName> GameFeatureDataPackageNames;
-	RemoveExistingExternalActors(InOutAssetDataArray, WorldPackageFilter.GetPtrOrNull(), ExternalActorsWorldPackageNames, GameFeatureDataPackageNames);
+	RemoveExistingExternalPackages(InOutAssetDataArray, WorldPackageFilter.GetPtrOrNull(), ExternalActorsWorldPackageNames, GameFeatureDataPackageNames);
+
+	// Discover external packages by their outer path
+	{
+		TArray<FAssetData> ExternalAssets;
+		UE::FMutex ExternalAssetsMutex;
+		ParallelFor(TEXT("UGatherTextFromAssetsCommandlet::DiscoverExternalPackages"), InOutAssetDataArray.Num(), 1,
+			[&InOutAssetDataArray, &ExternalAssets, &ExternalAssetsMutex](int32 Index)
+			{
+				const FAssetData& AssetData = InOutAssetDataArray[Index];
+
+				if (TArray<FAssetData> AssetsWithPackageAsOuter = UAssetRegistryHelpers::GetAssetsWithOuterForPaths({ AssetData.PackageName }, AssetData.PackageName, EAssetsWithOuterForPathsFlags::RecursivePaths | EAssetsWithOuterForPathsFlags::IncludeOnlyOnDiskAsset);
+					!AssetsWithPackageAsOuter.IsEmpty())
+				{
+					UE::TScopeLock _(ExternalAssetsMutex);
+					ExternalAssets.Append(MoveTemp(AssetsWithPackageAsOuter));
+				}
+			});
+		InOutAssetDataArray.Append(MoveTemp(ExternalAssets));
+	}
 
 	// Append the actors that are directly known by each world (by looking for their external actors under the expected path)
 	if (ExternalActorsWorldPackageNames.Num() > 0)
@@ -1143,7 +1186,7 @@ bool UGatherTextFromAssetsCommandlet::DiscoverExternalActors(TArray<FAssetData>&
 		}
 	}
 
-	UE_LOG(LogGatherTextFromAssetsCommandlet, Display, TEXT("Discovering external actors took %.2f seconds."), FPlatformTime::Seconds() - DiscoveringExternalActorsStartTime);
+	UE_LOG(LogGatherTextFromAssetsCommandlet, Display, TEXT("Discovering external packages took %.2f seconds."), FPlatformTime::Seconds() - DiscoveringExternalPackagesStartTime);
 	return true;
 }
 
@@ -1192,10 +1235,10 @@ void UGatherTextFromAssetsCommandlet::ApplyExcludeExactClassesFilter(const FARFi
 	// NOTE: The filter applied is actually the inverse, due to API limitations, so the resultant set must be removed from the current set.
 	TArray<FAssetData> AssetsToExclude = InOutAssetDataArray;
 	IAssetRegistry::GetChecked().RunAssetsThroughFilter(AssetsToExclude, InFilter);
-	InOutAssetDataArray.RemoveAll([&](const FAssetData& AssetData)
+	InOutAssetDataArray.RemoveAllSwap([AssetsToExcludeSet = TSet<FAssetData>(AssetsToExclude)](const FAssetData& AssetData)
 		{
-			return AssetsToExclude.Contains(AssetData);
-		});
+			return AssetsToExcludeSet.Contains(AssetData);
+		}, EAllowShrinking::No);
 }
 
 /**
@@ -1316,23 +1359,34 @@ UGatherTextFromAssetsCommandlet::FPackagePendingGather* UGatherTextFromAssetsCom
 	}
 	PackageFilename = FPaths::ConvertRelativePathToFull(PackageFilename);
 
+	bool bIsMapPackage = false;
+	{
+		FStringView PackageExt = FPathViews::GetExtension(PackageFilename, /*bIncludeDot*/true);
+		bIsMapPackage = PackageExt == FPackageName::GetMapPackageExtension() || PackageExt == FPackageName::GetTextMapPackageExtension();
+	}
+
 	FPackagePendingGather& PackagePendingGather = PackagesPendingGather.AddDefaulted_GetRef();
 	PackagePendingGather.PackageName = PackageNameToGather;
 	PackagePendingGather.PackageFilename = MoveTemp(PackageFilename);
+	PackagePendingGather.bIsMapPackage = bIsMapPackage;
 	PackagePendingGather.PackageLocCacheState = EPackageLocCacheState::Cached;
 	return &PackagePendingGather;
 }
 
-TSet<FName> UGatherTextFromAssetsCommandlet::GetPackageNamesToGather(const TArray<FAssetData>& InAssetDataArray) const
+void UGatherTextFromAssetsCommandlet::ExtractPackageNamesToGather(const TArray<FAssetData>& InAssetDataArray, TSet<FName>& OutPackageNames, TMap<FName, FName>& OutExternalToOuterPackages) const
 {
 	// Collapse the assets down to a set of packages
-	TSet<FName> PackageNamesToGather;
-	PackageNamesToGather.Reserve(InAssetDataArray.Num());
+	OutPackageNames.Empty(InAssetDataArray.Num());
 	for (const FAssetData& AssetData : InAssetDataArray)
 	{
-		PackageNamesToGather.Add(AssetData.PackageName);
+		OutPackageNames.Add(AssetData.PackageName);
+
+		if (!AssetData.GetOptionalOuterPathName().IsNone())
+		{
+			const FName OuterPackageName = FName(FPackageName::ObjectPathToPackageName(FNameBuilder(AssetData.GetOptionalOuterPathName()).ToView()));
+			OutExternalToOuterPackages.Add(AssetData.PackageName, OuterPackageName);
+		}
 	}
-	return PackageNamesToGather;
 }
 
 void UGatherTextFromAssetsCommandlet::PopulatePackagesPendingGather(TSet<FName> PackageNamesToGather)
@@ -1348,7 +1402,7 @@ void UGatherTextFromAssetsCommandlet::PopulatePackagesPendingGather(TSet<FName>
 }
 
 /** Process packages with loc data cached in its header and removes them from the pending packages.*/
-void UGatherTextFromAssetsCommandlet::ProcessAndRemoveCachedPackages(TMap<FName, TSet<FGuid>>& OutExternalActorsWithStaleOrMissingCaches)
+void UGatherTextFromAssetsCommandlet::ProcessAndRemoveCachedPackages(const TMap<FName, FName>& ExternalToOuterPackages, TMap<FName, TSet<FGuid>>& OutExternalActorsWithStaleOrMissingCaches, TMap<FName, TSet<FName>>& OutExternalPackagesWithStaleOrMissingCaches)
 {
 	UE_SCOPED_TIMER(TEXT("UGatherTextFromAssetsCommandlet::ProcessAndRemoveCachedPackages"), LogGatherTextFromAssetsCommandlet, Display);
 
@@ -1411,14 +1465,12 @@ void UGatherTextFromAssetsCommandlet::ProcessAndRemoveCachedPackages(TMap<FName,
 
 		int32 NumPackagesProcessed = 0;
 		TMap<FString, FName> AssignedPackageLocalizationIds;
-		PackagesPendingGather.RemoveAll([this, &OutExternalActorsWithStaleOrMissingCaches, &AssignedPackageLocalizationIds, &NumPackagesProcessed, InitialNumPackagesToGather](const FPackagePendingGather& PackagePendingGather) -> bool
+		PackagesPendingGather.RemoveAllSwap([this, &ExternalToOuterPackages, &OutExternalActorsWithStaleOrMissingCaches, &OutExternalPackagesWithStaleOrMissingCaches, &AssignedPackageLocalizationIds, &NumPackagesProcessed, InitialNumPackagesToGather](const FPackagePendingGather& PackagePendingGather) -> bool
 		{
 			FNameBuilder PackageNameStr(PackagePendingGather.PackageName);
 			const int32 CurrentPackageNum = ++NumPackagesProcessed;
 			const float PercentageComplete = static_cast<float>(CurrentPackageNum) / static_cast<float>(InitialNumPackagesToGather) * 100.0f;
 
-			const bool bIsExternalActorPackage = PackageNameStr.ToView().Contains(FPackagePath::GetExternalActorsFolderName());
-
 			// Track the package localization ID of this package (if known) and detect duplicates
 			bool bThisPackageHasLocIdConflictToFix = false;
 			if (!PackagePendingGather.PackageLocalizationId.IsEmpty())
@@ -1456,20 +1508,28 @@ void UGatherTextFromAssetsCommandlet::ProcessAndRemoveCachedPackages(TMap<FName,
 
 			if (PackagePendingGather.PackageLocCacheState != EPackageLocCacheState::Cached)
 			{
-				// External actors must be gathered via their owner world rather than via a raw LoadPackage call
-				// Remove them from PackagesToGather as the owner world is merged back in below
-				if (bIsExternalActorPackage)
+				if (const FName* OuterPackage = ExternalToOuterPackages.Find(PackagePendingGather.PackageName))
 				{
-					TArray<FAssetData> ActorsInPackage;
-					IAssetRegistry::GetChecked().GetAssetsByPackageName(PackagePendingGather.PackageName, ActorsInPackage);
-					for (const FAssetData& ActorInPackage : ActorsInPackage)
+					// External packages must be gathered via their outer package rather than via a raw LoadPackage call
+					// Remove them from PackagesToGather as they will be merged back in by MergeInExternalPackagesWithStaleOrMissingCaches
+					if (PackageNameStr.ToView().Contains(FPackagePath::GetExternalActorsFolderName()))
 					{
-						if (TUniquePtr<FWorldPartitionActorDesc> ActorDesc = FWorldPartitionActorDescUtils::GetActorDescriptorFromAssetData(ActorInPackage))
+						TArray<FAssetData> ActorsInPackage;
+						IAssetRegistry::GetChecked().GetAssetsByPackageName(PackagePendingGather.PackageName, ActorsInPackage);
+						for (const FAssetData& ActorInPackage : ActorsInPackage)
 						{
-							FName WorldPackageName = *FPackageName::ObjectPathToPackageName(ActorDesc->GetActorSoftPath().ToString());
-							OutExternalActorsWithStaleOrMissingCaches.FindOrAdd(WorldPackageName).Add(ActorDesc->GetGuid());
+							if (TUniquePtr<FWorldPartitionActorDesc> ActorDesc = FWorldPartitionActorDescUtils::GetActorDescriptorFromAssetData(ActorInPackage))
+							{
+								FName WorldPackageName = *FPackageName::ObjectPathToPackageName(ActorDesc->GetActorSoftPath().ToString());
+								ensureAlwaysMsgf(WorldPackageName == *OuterPackage, TEXT("Actor '%s' was expected to be part of '%s'"), *ActorDesc->GetActorSoftPath().ToString(), *OuterPackage->ToString());
+								OutExternalActorsWithStaleOrMissingCaches.FindOrAdd(WorldPackageName).Add(ActorDesc->GetGuid());
+							}
 						}
 					}
+					else
+					{
+						OutExternalPackagesWithStaleOrMissingCaches.FindOrAdd(*OuterPackage).Add(PackagePendingGather.PackageName);
+					}
 					return true;
 				}
 
@@ -1496,7 +1556,7 @@ void UGatherTextFromAssetsCommandlet::ProcessAndRemoveCachedPackages(TMap<FName,
 			}
 
 			return true;
-		});
+		}, EAllowShrinking::No);
 
 		UE_LOG(LogGatherTextFromAssetsCommandlet, Display, TEXT("Ingesting the cached localization data for %d package(s) took %.2f seconds."), InitialNumPackagesToGather, FPlatformTime::Seconds() - IngestStartTime);
 	}
@@ -1504,20 +1564,26 @@ void UGatherTextFromAssetsCommandlet::ProcessAndRemoveCachedPackages(TMap<FName,
 	UE_LOG(LogGatherTextFromAssetsCommandlet, Display, TEXT("Number of packages to load: %d"), PackagesPendingGather.Num());
 }
 
-void UGatherTextFromAssetsCommandlet::MergeInExternalActorsWithStaleOrMissingCaches(TMap<FName, TSet<FGuid>>& StaleExternalActors)
+void UGatherTextFromAssetsCommandlet::MergeInExternalPackagesWithStaleOrMissingCaches(TMap<FName, TSet<FGuid>>& ExternalActorsWithStaleOrMissingCaches, TMap<FName, TSet<FName>>& ExternalPackagesWithStaleOrMissingCaches)
 {
-	for (TTuple<FName, TSet<FGuid>>& StaleExternalActorsPair : StaleExternalActors)
-	{
-		FPackagePendingGather* WorldPackagePendingGather = PackagesPendingGather.FindByPredicate([&StaleExternalActorsPair](const FPackagePendingGather& PotentialPackagePendingGather)
-			{
-				return PotentialPackagePendingGather.PackageName == StaleExternalActorsPair.Key;
-			});
-		if (!WorldPackagePendingGather)
+	auto FindOrAddPackagePendingGather = 
+		[this](const FName PackageName) -> FPackagePendingGather*
 		{
-			WorldPackagePendingGather = AppendPackagePendingGather(StaleExternalActorsPair.Key);
-		}
+			FPackagePendingGather* PackagePendingGather = PackagesPendingGather.FindByPredicate(
+				[PackageName](const FPackagePendingGather& PotentialPackagePendingGather)
+				{
+					return PotentialPackagePendingGather.PackageName == PackageName;
+				});
+			if (!PackagePendingGather)
+			{
+				PackagePendingGather = AppendPackagePendingGather(PackageName);
+			}
+			return PackagePendingGather;
+		};
 
-		if (WorldPackagePendingGather)
+	for (TTuple<FName, TSet<FGuid>>& StaleExternalActorsPair : ExternalActorsWithStaleOrMissingCaches)
+	{
+		if (FPackagePendingGather* WorldPackagePendingGather = FindOrAddPackagePendingGather(StaleExternalActorsPair.Key))
 		{
 			WorldPackagePendingGather->ExternalActors = MoveTemp(StaleExternalActorsPair.Value);
 			WorldPackagePendingGather->PackageLocCacheState = EPackageLocCacheState::Uncached_TooOld;
@@ -1531,6 +1597,23 @@ void UGatherTextFromAssetsCommandlet::MergeInExternalActorsWithStaleOrMissingCac
 			);
 		}
 	}
+
+	for (TTuple<FName, TSet<FName>>& StaleExternalPackagesPair : ExternalPackagesWithStaleOrMissingCaches)
+	{
+		if (FPackagePendingGather* WorldPackagePendingGather = FindOrAddPackagePendingGather(StaleExternalPackagesPair.Key))
+		{
+			WorldPackagePendingGather->ExternalPackages = MoveTemp(StaleExternalPackagesPair.Value);
+			WorldPackagePendingGather->PackageLocCacheState = EPackageLocCacheState::Uncached_TooOld;
+		}
+		else
+		{
+			UE_LOGFMT(LogGatherTextFromAssetsCommandlet, Warning, "Failed to queue outer package '{package}' for {nbExternalPackages} external package(s).",
+				("package", *StaleExternalPackagesPair.Key.ToString()),
+				("nbExternalPackages", StaleExternalPackagesPair.Value.Num()),
+				("id", GatherTextFromAssetsCommandlet::LocalizationLogIdentifier)
+			);
+		}
+	}
 }
 
 /** Load the remaining pending packages for gather.*/
@@ -1593,7 +1676,7 @@ bool UGatherTextFromAssetsCommandlet::LoadAndProcessUncachedPackages(TArray<FNam
 			return false;
 		}
 
-		UPackage* Package = UE::Private::GatherTextFromAssetsCommandlet::LoadPackageToGather(*PackageNameStr, PackagePendingGather.ExternalActors.Num() > 0);
+		UPackage* Package = UE::Private::GatherTextFromAssetsCommandlet::LoadPackageToGather(*PackageNameStr, PackagePendingGather.bIsMapPackage);
 		if (!Package)
 		{
 			UE_LOGFMT(LogGatherTextFromAssetsCommandlet, Warning, "Failed to load package: '{package}'.",
@@ -1642,7 +1725,7 @@ bool UGatherTextFromAssetsCommandlet::LoadAndProcessUncachedPackages(TArray<FNam
 		}
 		WorkerDirector.TickWorkers();
 
-		if (UE::Private::GatherTextFromAssetsCommandlet::GatherTextFromPackage(Package, PackagePendingGather.ExternalActors, /*bGatherFromPrimaryPackage*/true, /*bGatherFromExternalActorPackages*/false, GatherableTextDataArray))
+		if (UE::Private::GatherTextFromAssetsCommandlet::GatherTextFromPackage(Package, PackagePendingGather.bIsMapPackage, PackagePendingGather.ExternalActors, PackagePendingGather.ExternalPackages, /*bGatherFromPrimaryPackage*/true, /*bGatherFromExternalPackages*/false, GatherableTextDataArray))
 		{
 			UE_LOG(LogGatherTextFromAssetsCommandlet, Display, TEXT("[%6.2f%%] Gathering package: '%s'..."), PercentageComplete, *PackageNameStr);
 
@@ -1694,10 +1777,10 @@ bool UGatherTextFromAssetsCommandlet::LoadAndProcessUncachedPackages(TArray<FNam
 				}
 			}
 
-			// If this is a WP world then query the localization for any external actors actors that were determined to be stale
-			if (PackagePendingGather.ExternalActors.Num() > 0)
+			// Include any external packages that were determined to be stale
+			if (PackagePendingGather.ExternalActors.Num() > 0 || PackagePendingGather.ExternalPackages.Num() > 0)
 			{
-				UE::Private::GatherTextFromAssetsCommandlet::GatherTextFromPackage(Package, PackagePendingGather.ExternalActors, /*bGatherFromPrimaryPackage*/false, /*bGatherFromExternalActorPackages*/true, GatherableTextDataArray);
+				UE::Private::GatherTextFromAssetsCommandlet::GatherTextFromPackage(Package, PackagePendingGather.bIsMapPackage, PackagePendingGather.ExternalActors, PackagePendingGather.ExternalPackages, /*bGatherFromPrimaryPackage*/false, /*bGatherFromExternalPackages*/true, GatherableTextDataArray);
 			}
 
 			// This package may have already been cached in cases where we're reporting or fixing assets with a stale gather cache
@@ -1897,14 +1980,16 @@ int32 UGatherTextFromAssetsCommandlet::Main(const FString& Params)
 		return 0;
 	}
 
-	// Discover the external actors for any worlds that are pending gather
-	if (!DiscoverExternalActors(AssetDataArray))
+	// Discover the external packages for any assets that are pending gather
+	if (!DiscoverExternalPackages(AssetDataArray))
 	{
 		return -1;
 	}
 
 	// Collect the basic information about the packages that we're going to gather from
-	TSet<FName> PackageNamesToGather = GetPackageNamesToGather(AssetDataArray);
+	TSet<FName> PackageNamesToGather;
+	TMap<FName, FName> ExternalToOuterPackages;
+	ExtractPackageNamesToGather(AssetDataArray, PackageNamesToGather, ExternalToOuterPackages);
 	AssetDataArray.Empty();
 	PopulatePackagesPendingGather(MoveTemp(PackageNamesToGather));
 
@@ -1913,14 +1998,15 @@ int32 UGatherTextFromAssetsCommandlet::Main(const FString& Params)
 	// These are external actor packages that are stale or are missing a gather cache from their package
 	// Map of world package name -> external actor Ids in the world 
 	TMap<FName, TSet<FGuid>> ExternalActorsWithStaleOrMissingCaches;
+	TMap<FName, TSet<FName>> ExternalPackagesWithStaleOrMissingCaches;
 	// Process all packages that do not need to be loaded. Remove processed packages from the list.
 	UE_LOG(LogGatherTextFromAssetsCommandlet, Display, TEXT("Processing assets to gather..."));
-	ProcessAndRemoveCachedPackages(ExternalActorsWithStaleOrMissingCaches);
+	ProcessAndRemoveCachedPackages(ExternalToOuterPackages, ExternalActorsWithStaleOrMissingCaches, ExternalPackagesWithStaleOrMissingCaches);
 
 	UE::Private::GatherTextFromAssetsCommandlet::AssetGatherCacheMetrics.LogMetrics();
 
-	// Merge any pending WP map requests back into PackagesPendingGather
-	MergeInExternalActorsWithStaleOrMissingCaches(ExternalActorsWithStaleOrMissingCaches);
+	// Merge any pending external package requests back into PackagesPendingGather
+	MergeInExternalPackagesWithStaleOrMissingCaches(ExternalActorsWithStaleOrMissingCaches, ExternalPackagesWithStaleOrMissingCaches);
 	ExternalActorsWithStaleOrMissingCaches.Reset();
 
 	// All packages left in PackagesPendingGather should now have to be loaded 
@@ -2318,6 +2404,8 @@ void UGatherTextFromAssetsCommandlet::AssignPackagesToWorkers(TConstArrayView<FG
 						PackageRequest.PackageName = PackagePendingGather.PackageName;
 						PackageRequest.Dependencies = PackagePendingGather.Dependencies;
 						PackageRequest.ExternalActors = PackagePendingGather.ExternalActors;
+						PackageRequest.ExternalPackages = PackagePendingGather.ExternalPackages;
+						PackageRequest.bIsMapPackage = PackagePendingGather.bIsMapPackage;
 
 						if (WorkerDirector.AssignPackageToWorker(WorkerId, PackageRequest))
 						{
@@ -2536,10 +2624,10 @@ int32 UGatherTextFromAssetsWorkerCommandlet::Main(const FString& Params)
 			PackageResultMessage->PackageName = PackagePendingGather->PackageName;
 
 			LogOutputTargetDevice.Reset();
-			if (UPackage* Package = UE::Private::GatherTextFromAssetsCommandlet::LoadPackageToGather(PackageNameStr, PackagePendingGather->ExternalActors.Num() > 0))
+			if (UPackage* Package = UE::Private::GatherTextFromAssetsCommandlet::LoadPackageToGather(PackageNameStr, PackagePendingGather->bIsMapPackage))
 			{
 				TArray<FGatherableTextData> GatherableTextDataArray;
-				UE::Private::GatherTextFromAssetsCommandlet::GatherTextFromPackage(Package, PackagePendingGather->ExternalActors, /*bGatherFromPrimaryPackage*/true, /*bGatherFromExternalActorPackages*/true, GatherableTextDataArray);
+				UE::Private::GatherTextFromAssetsCommandlet::GatherTextFromPackage(Package, PackagePendingGather->bIsMapPackage, PackagePendingGather->ExternalActors, PackagePendingGather->ExternalPackages, /*bGatherFromPrimaryPackage*/true, /*bGatherFromExternalPackages*/true, GatherableTextDataArray);
 
 				FMemoryWriter MemoryWriter(PackageResultMessage->GatherableTextData);
 				MemoryWriter << GatherableTextDataArray;
@@ -2646,7 +2734,7 @@ void UGatherTextFromAssetsWorkerCommandlet::HandlePackageRequestMessage(const FG
 	// Note: This isn't an optimal way to insert into the array, but the array is ordered so it can pop from the end, so inserting at element 0 keeps it FIFO
 	{
 		UE::TScopeLock _(PackagesPendingGatherMutex);
-		PackagesPendingGather.Insert(FPackagePendingGather{ Message.PackageName, Message.Dependencies, Message.ExternalActors, Context->GetSender() }, 0);
+		PackagesPendingGather.Insert(FPackagePendingGather{ Message.PackageName, Message.Dependencies, Message.ExternalActors, Message.ExternalPackages, Message.bIsMapPackage, Context->GetSender() }, 0);
 	}
 }
 
