Fixed double gather

#rnx
#Synced-CL 52054134

diff --git a/Engine/Source/Editor/UnrealEd/Private/Commandlets/GatherTextFromAssetsCommandlet.cpp b/Engine/Source/Editor/UnrealEd/Private/Commandlets/GatherTextFromAssetsCommandlet.cpp
--- a/Engine/Source/Editor/UnrealEd/Private/Commandlets/GatherTextFromAssetsCommandlet.cpp
+++ b/Engine/Source/Editor/UnrealEd/Private/Commandlets/GatherTextFromAssetsCommandlet.cpp
@@ -1729,19 +1729,20 @@ bool UGatherTextFromAssetsCommandlet::LoadAndProcessUncachedPackages(TArray<FNam
 		{
 			UE_LOG(LogGatherTextFromAssetsCommandlet, Display, TEXT("[%6.2f%%] Gathering package: '%s'..."), PercentageComplete, *PackageNameStr);
 
-			// Gathers from the given package
-			EPropertyLocalizationGathererResultFlags GatherableTextResultFlags = EPropertyLocalizationGathererResultFlags::Empty;
-			FPropertyLocalizationDataGatherer(GatherableTextDataArray, Package, GatherableTextResultFlags);
-
 			bool bSavePackage = false;
 
 			// Optionally check to see whether the clean gather we did is in-sync with the gather cache and deal with it accordingly
 			if ((bReportStaleGatherCache || bFixStaleGatherCache) && PackagePendingGather.PackageLocCacheState == EPackageLocCacheState::Cached)
 			{
+				// Gathers from the given package
+				TArray<FGatherableTextData> CurrentGatherableTextDataArray;
+				EPropertyLocalizationGathererResultFlags CurrentGatherableTextResultFlags = EPropertyLocalizationGathererResultFlags::Empty;
+				FPropertyLocalizationDataGatherer(CurrentGatherableTextDataArray, Package, CurrentGatherableTextResultFlags);
+
 				// Look for any structurally significant changes (missing, added, or changed texts) in the cache
 				// Ignore insignificant things (like source changes caused by assets moving or being renamed)
-				if (EnumHasAnyFlags(GatherableTextResultFlags, EPropertyLocalizationGathererResultFlags::HasTextWithInvalidPackageLocalizationID)
-					|| !IsGatherableTextDataIdentical(GatherableTextDataArray, PackagePendingGather.GatherableTextDataArray))
+				if (EnumHasAnyFlags(CurrentGatherableTextResultFlags, EPropertyLocalizationGathererResultFlags::HasTextWithInvalidPackageLocalizationID)
+					|| !IsGatherableTextDataIdentical(CurrentGatherableTextDataArray, PackagePendingGather.GatherableTextDataArray))
 				{
 					OutPackagesWithStaleGatherCache.Add(PackagePendingGather.PackageName);
 
