Fixed subobjects reinstancing with new class despite being part of the InstancesThatShouldUseOldClass set.

#Synced-CL 51956848

diff --git a/Engine/Source/Editor/UnrealEd/Private/Kismet2/KismetReinstanceUtilities.cpp b/Engine/Source/Editor/UnrealEd/Private/Kismet2/KismetReinstanceUtilities.cpp
--- a/Engine/Source/Editor/UnrealEd/Private/Kismet2/KismetReinstanceUtilities.cpp
+++ b/Engine/Source/Editor/UnrealEd/Private/Kismet2/KismetReinstanceUtilities.cpp
@@ -2217,7 +2217,7 @@ bool FBlueprintCompileReinstancer::IsReinstClass(const UClass* Class)
 	return Class && Class->GetFName().ToString().StartsWith(ReinstPrefix);
 }
 
-void FBlueprintCompileReinstancer::ReplaceObjectHelper(const FReinstancingState& State, UObject*& OldObject, UClass* OldClass, UObject*& NewUObject, UClass* NewClass, int32 OldObjIndex, TArray<UObject*>& ObjectsToReplace, TFunctionRef<TArray<TObjectPtr<USceneComponent>>& (USceneComponent*)> GetAttachChildrenArray, bool bIsComponent, bool bArchetypesAreUpToDate)
+void FBlueprintCompileReinstancer::ReplaceObjectHelper(const FReinstancingState& State, UObject*& OldObject, UClass* OldClass, UObject*& NewUObject, UClass* NewClass, int32 OldObjIndex, TArray<UObject*>& ObjectsToReplace, TFunctionRef<TArray<TObjectPtr<USceneComponent>>& (USceneComponent*)> GetAttachChildrenArray, bool bIsComponent, bool bArchetypesAreUpToDate, const TSet<UObject*>* InstancesThatShouldUseOldClass)
 {
 	SCOPED_LOADTIMER_ASSET_TEXT(*WriteToString<256>(TEXT("ReplaceObjectHelper "), *GetPathNameSafe(OldObject)));
 	// If the old object was spawned from an archetype (i.e. not the CDO), we must use the new version of that archetype as the template object when constructing the new instance.
@@ -2384,7 +2384,7 @@ void FBlueprintCompileReinstancer::ReplaceObjectHelper(const FReinstancingState&
 
 	TMap<UObject*, UObject*> CreatedInstanceMap;
 	TArray< TTuple<UObject*, UObject*>> OrderedListOfObjectToCopy;
-	PreCreateSubObjectsForReinstantiation(State.InOldToNewClassMap, OldObject, NewUObject, CreatedInstanceMap, &State.OldToNewInstanceMap, &OrderedListOfObjectToCopy);
+	PreCreateSubObjectsForReinstantiation(State.InOldToNewClassMap, OldObject, NewUObject, CreatedInstanceMap, &State.OldToNewInstanceMap, &OrderedListOfObjectToCopy, InstancesThatShouldUseOldClass);
 	State.OldToNewInstanceMap.Append(CreatedInstanceMap);
 
 	// Copy property values
@@ -2772,7 +2772,7 @@ void FBlueprintCompileReinstancer::ReplaceInstancesOfClass_Inner(const TMap<UCla
 			check(bArchetypesAreUpToDate); // not supported, can't be
 			check(Params.InstancesThatShouldUseOldClass == nullptr); // not supported, could be.. but why?
 
-			ReinstanceWithVersioningData(ReinstancingState);
+			ReinstanceWithVersioningData(ReinstancingState, InstancesThatShouldUseOldClass);
 		}
 		else
 		{
@@ -2893,7 +2893,7 @@ void FBlueprintCompileReinstancer::ReplaceInstancesOfClass_Inner(const TMap<UCla
 						CheckAndSaveOuterPackageToCleanList(CleanPackageList, OldObject);
 
 						UObject* NewUObject = nullptr;
-						ReplaceObjectHelper(ReinstancingState, OldObject, OldClass, NewUObject, NewClass, OldObjIndex, ObjectsToReplace, &FDirectAttachChildrenAccessor::Get, bIsComponent, bArchetypesAreUpToDate);
+						ReplaceObjectHelper(ReinstancingState, OldObject, OldClass, NewUObject, NewClass, OldObjIndex, ObjectsToReplace, &FDirectAttachChildrenAccessor::Get, bIsComponent, bArchetypesAreUpToDate, InstancesThatShouldUseOldClass);
 						UpdateObjectBeingDebugged(OldObject, NewUObject);
 						ObjectsReplaced.Add(OldObject);
 
@@ -2973,7 +2973,7 @@ void FBlueprintCompileReinstancer::ReplaceInstancesOfClass_Inner(const TMap<UCla
 						else
 						{
 							// Actors that are not in a level cannot be reconstructed, sequencer team decided to reinstance these as normal objects:
-							ReplaceObjectHelper(ReinstancingState, OldObject, OldClass, NewUObject, NewClass, OldObjIndex, ObjectsToReplace, &FDirectAttachChildrenAccessor::Get, false, bArchetypesAreUpToDate);
+							ReplaceObjectHelper(ReinstancingState, OldObject, OldClass, NewUObject, NewClass, OldObjIndex, ObjectsToReplace, &FDirectAttachChildrenAccessor::Get, false, bArchetypesAreUpToDate, InstancesThatShouldUseOldClass);
 						}
 						UpdateObjectBeingDebugged(OldObject, NewUObject);
 						ObjectsReplaced.Add(OldObject);
@@ -3369,7 +3369,8 @@ void FBlueprintCompileReinstancer::PreCreateSubObjectsForReinstantiation(
 	UObject* OldObject, UObject* NewUObject, 
 	TMap<UObject*, UObject*>& CreatedInstanceMap, 
 	const TMap<UObject*, UObject*>* OldToNewInstanceMap/* = nullptr*/, 
-	TArray< TTuple<UObject*, UObject*>>* OrderedListOfObjectToCopy /*= nullptr*/)
+	TArray< TTuple<UObject*, UObject*>>* OrderedListOfObjectToCopy /*= nullptr*/,
+	const TSet<UObject*>* InstancesThatShouldUseOldClass)
 {
 	// When handling only referenced owned subobjects, we absolutely need to have done the postload so that we are pointing 
 	// to all the correct loaded subobjects, otherwise we are skip important objects. Ex: Added subobjects in containers during the load.
@@ -3384,7 +3385,7 @@ void FBlueprintCompileReinstancer::PreCreateSubObjectsForReinstantiation(
 
 	// Add the mapping from the old to the new object exists...
 	CreatedInstanceMap.Add(OldObject, NewUObject);
-	PreCreateSubObjectsForReinstantiation_Inner(OldInstancedSubObjectsPtr, OldToNewClassMap, OldObject, NewUObject, CreatedInstanceMap, OldToNewInstanceMap, OrderedListOfObjectToCopy);
+	PreCreateSubObjectsForReinstantiation_Inner(OldInstancedSubObjectsPtr, OldToNewClassMap, OldObject, NewUObject, CreatedInstanceMap, OldToNewInstanceMap, OrderedListOfObjectToCopy, InstancesThatShouldUseOldClass);
 	if (OrderedListOfObjectToCopy)
 	{
 		// Post add for deep first order
@@ -3533,7 +3534,8 @@ void FBlueprintCompileReinstancer::PreCreateSubObjectsForReinstantiation_Inner(
 	UObject* OldObject, UObject* NewUObject, 
 	TMap<UObject*, UObject*>& CreatedInstanceMap, 
 	const TMap<UObject*, UObject*>* OldToNewInstanceMap,
-	TArray< TTuple<UObject*, UObject*>>* OrderedListOfObjectToCopy)
+	TArray< TTuple<UObject*, UObject*>>* OrderedListOfObjectToCopy,
+	const TSet<UObject*>* InstancesThatShouldUseOldClass)
 {
 	// Gather subobjects on old object and pre-create them if needed
 	TArray<UObject*> ContainedOldSubObjects;
@@ -3594,7 +3596,7 @@ void FBlueprintCompileReinstancer::PreCreateSubObjectsForReinstantiation_Inner(
 				}
 
 				UClass* SubObjectClass = OldSubObjectClass;
-				if (UClass* const* NewSubObjectClass = OldToNewClassMap.Find(OldSubObjectClass))
+				if (UClass* const* NewSubObjectClass = (!InstancesThatShouldUseOldClass || !InstancesThatShouldUseOldClass->Contains(OldSubObject)) ? OldToNewClassMap.Find(OldSubObjectClass) : nullptr)
 				{
 					SubObjectClass = *NewSubObjectClass;
 				}
@@ -3643,7 +3645,7 @@ void FBlueprintCompileReinstancer::PreCreateSubObjectsForReinstantiation_Inner(
 					}
 
 					CreatedInstanceMap.Add(OldSubObject, NewSubObject);
-					PreCreateSubObjectsForReinstantiation_Inner(OldInstancedSubObjects, OldToNewClassMap, OldSubObject, NewSubObject, CreatedInstanceMap, OldToNewInstanceMap, OrderedListOfObjectToCopy);
+					PreCreateSubObjectsForReinstantiation_Inner(OldInstancedSubObjects, OldToNewClassMap, OldSubObject, NewSubObject, CreatedInstanceMap, OldToNewInstanceMap, OrderedListOfObjectToCopy, InstancesThatShouldUseOldClass);
 					if (OrderedListOfObjectToCopy)
 					{
 						// Post add for deep first order
@@ -3662,7 +3664,7 @@ void FBlueprintCompileReinstancer::PreCreateSubObjectsForReinstantiation_Inner(
 				OldSubObject->SetExternalPackage(nullptr); // move old subobject outside of external package (it'll be in the transient package) so that it's doesn't cause the package to stick around
 			}
 
-			PreCreateSubObjectsForReinstantiation_Inner(OldInstancedSubObjects, OldToNewClassMap, OldSubObject, *NewSubObject, CreatedInstanceMap, OldToNewInstanceMap, OrderedListOfObjectToCopy);
+			PreCreateSubObjectsForReinstantiation_Inner(OldInstancedSubObjects, OldToNewClassMap, OldSubObject, *NewSubObject, CreatedInstanceMap, OldToNewInstanceMap, OrderedListOfObjectToCopy, InstancesThatShouldUseOldClass);
 		}
 	}
 }
@@ -3732,7 +3734,7 @@ FBlake3Hash FBlueprintCompileReinstancer::CalculateConstructionHash(const UClass
 	return Hash.Finalize();
 }
 
-void FBlueprintCompileReinstancer::ReinstanceWithVersioningData(const FReinstancingState& ReinstancingState)
+void FBlueprintCompileReinstancer::ReinstanceWithVersioningData(const FReinstancingState& ReinstancingState, const TSet<UObject*>* InstancesThatShouldUseOldClass)
 {
 	// if an OldToNewTemplates map has been provided, include all the instances in the replacement map
 	// this is used to make reinstancing during PIE more robust:
@@ -3761,8 +3763,8 @@ void FBlueprintCompileReinstancer::ReinstanceWithVersioningData(const FReinstanc
 		ActorReinstancingClasses, 
 		ReinstancingState);
 	RestoreClass(UnchangedClasses, ReinstancingState);
-	ReinstanceNormalObjects(NormalReinstancingClasses, ReinstancingState);
-	ReinstanceActors(ActorReinstancingClasses, ReinstancingState);
+	ReinstanceNormalObjects(NormalReinstancingClasses, ReinstancingState, InstancesThatShouldUseOldClass);
+	ReinstanceActors(ActorReinstancingClasses, ReinstancingState, InstancesThatShouldUseOldClass);
 }
 
 void FBlueprintCompileReinstancer::CategorizeClassesAndInstances(
@@ -3891,7 +3893,8 @@ void FBlueprintCompileReinstancer::RestoreClass(
 
 void FBlueprintCompileReinstancer::ReinstanceNormalObjects(
 	TMap<UClass*, UClass*>& OldToNewClassMap,
-	const FBlueprintCompileReinstancer::FReinstancingState& ReinstancingState)
+	const FBlueprintCompileReinstancer::FReinstancingState& ReinstancingState,
+	const TSet<UObject*>* InstancesThatShouldUseOldClass)
 {
 	FEditorCacheArchetypeManager& CacheManager = FEditorCacheArchetypeManager::Get();
 	TArray<UObject*> EmptyList;
@@ -3930,7 +3933,8 @@ void FBlueprintCompileReinstancer::ReinstanceNormalObjects(
 				EmptyList, 
 				&FDirectAttachChildrenAccessor::Get, 
 				bIsComponent, 
-				bArchetypesAreUpToDate);
+				bArchetypesAreUpToDate,
+				InstancesThatShouldUseOldClass);
 			UpdateObjectBeingDebugged(OldObject, NewUObject);
 			ReinstancingState.ObjectsReplaced.Add(OldObject);
 		}
@@ -3939,7 +3943,8 @@ void FBlueprintCompileReinstancer::ReinstanceNormalObjects(
 
 void FBlueprintCompileReinstancer::ReinstanceActors(
 	TMap<UClass*, UClass*>& OldToNewClassMap, 
-	const FBlueprintCompileReinstancer::FReinstancingState& ReinstancingState)
+	const FBlueprintCompileReinstancer::FReinstancingState& ReinstancingState,
+	const TSet<UObject*>* InstancesThatShouldUseOldClass)
 {
 	TArray<UObject*> EmptyList;
 	const int32 OldObjIndex = INDEX_NONE; // not needed as archetypes are up to date in this codepath
@@ -4006,7 +4011,8 @@ void FBlueprintCompileReinstancer::ReinstanceActors(
 					EmptyList,
 					&FDirectAttachChildrenAccessor::Get,
 					bIsComponent,
-					bArchetypesAreUpToDate);
+					bArchetypesAreUpToDate,
+					InstancesThatShouldUseOldClass);
 			}
 			UpdateObjectBeingDebugged(OldObject, NewUObject);
 			ReinstancingState.ObjectsReplaced.Add(OldObject);
diff --git a/Engine/Source/Editor/UnrealEd/Public/Kismet2/KismetReinstanceUtilities.h b/Engine/Source/Editor/UnrealEd/Public/Kismet2/KismetReinstanceUtilities.h
--- a/Engine/Source/Editor/UnrealEd/Public/Kismet2/KismetReinstanceUtilities.h
+++ b/Engine/Source/Editor/UnrealEd/Public/Kismet2/KismetReinstanceUtilities.h
@@ -325,8 +325,9 @@ public:
 	 * @param CreatedInstanceMap in/out of the result of all of the pre-created objects
 	 * @param OldToNewInstanceMap optional parameter of the possible re-instanced sub objects if any
 	 * @param OrderedListOfObjectToCopy List of object to copy in order
+	 * @param InstancesThatShouldUseOldClass Subobjects that should keep their old class during reinstantiation, even if a new one is provided by OldToNewClassMap. This shouldn't be used for simple reinstantiation as the old class would be dead. One use-case would be changing the class of an object.
 	 */
-	static UNREALED_API void PreCreateSubObjectsForReinstantiation(const TMap<UClass*, UClass*>& OldToNewClassMap, UObject* OldObject, UObject* NewUObject, TMap<UObject*, UObject*>& CreatedInstanceMap, const TMap<UObject*, UObject*>* OldToNewInstanceMap = nullptr, TArray< TTuple<UObject*, UObject*>>* OrderedListOfObjectToCopy = nullptr);
+	static UNREALED_API void PreCreateSubObjectsForReinstantiation(const TMap<UClass*, UClass*>& OldToNewClassMap, UObject* OldObject, UObject* NewUObject, TMap<UObject*, UObject*>& CreatedInstanceMap, const TMap<UObject*, UObject*>* OldToNewInstanceMap = nullptr, TArray< TTuple<UObject*, UObject*>>* OrderedListOfObjectToCopy = nullptr, const TSet<UObject*>* InstancesThatShouldUseOldClass = nullptr);
 
 private:
 	struct FReinstancingState
@@ -350,7 +351,7 @@ private:
 	};
 
 	/** Helper to replace an instance of an object*/
-	static UNREALED_API void ReplaceObjectHelper(const FReinstancingState& State, UObject*& OldObject, UClass* OldClass, UObject*& NewUObject, UClass* NewClass, int32 OldObjIndex, TArray<UObject*>& ObjectsToReplace, TFunctionRef<TArray<TObjectPtr<USceneComponent>>&(USceneComponent*)> GetAttachChildrenArray, bool bIsComponent, bool bArchetypesAreUpToDate);
+	static UNREALED_API void ReplaceObjectHelper(const FReinstancingState& State, UObject*& OldObject, UClass* OldClass, UObject*& NewUObject, UClass* NewClass, int32 OldObjIndex, TArray<UObject*>& ObjectsToReplace, TFunctionRef<TArray<TObjectPtr<USceneComponent>>&(USceneComponent*)> GetAttachChildrenArray, bool bIsComponent, bool bArchetypesAreUpToDate, const TSet<UObject*>* InstancesThatShouldUseOldClass);
 
 	/**
 	 * Pre-populate instancing graph for a specified instance if needed
@@ -366,7 +367,7 @@ private:
 
 
 	/** Handles the sub object pre-creation recursively */
-	static UNREALED_API void PreCreateSubObjectsForReinstantiation_Inner(const TSet<UObject*>* OldInstancedSubObjects, const TMap<UClass*, UClass*>& OldToNewClassMap, UObject* OldObject, UObject* NewUObject, TMap<UObject*, UObject*>& CreatedInstanceMap, const TMap<UObject*, UObject*>* OldToNewInstanceMap, TArray< TTuple<UObject*, UObject*>>* OrderedListOfObjectToCopy);
+	static UNREALED_API void PreCreateSubObjectsForReinstantiation_Inner(const TSet<UObject*>* OldInstancedSubObjects, const TMap<UClass*, UClass*>& OldToNewClassMap, UObject* OldObject, UObject* NewUObject, TMap<UObject*, UObject*>& CreatedInstanceMap, const TMap<UObject*, UObject*>* OldToNewInstanceMap, TArray< TTuple<UObject*, UObject*>>* OrderedListOfObjectToCopy, const TSet<UObject*>* InstancesThatShouldUseOldClass);
 
 	/** Handles the work of ReplaceInstancesOfClass, handling both normal replacement of instances and batch */
 	static UNREALED_API void ReplaceInstancesOfClass_Inner(const TMap<UClass*, UClass*>& InOldToNewClassMap, const FReplaceInstancesOfClassParameters& Params);
@@ -384,7 +385,7 @@ private:
 	/** Returns the above, but only if there are instances of the class that would require reinstancing */
 	static UNREALED_API FBlake3Hash CalculateConstructionHashIfNecessary(const UClass* ForClass);
 
-	static void ReinstanceWithVersioningData(const FReinstancingState& ReinstancingState);
+	static void ReinstanceWithVersioningData(const FReinstancingState& ReinstancingState, const TSet<UObject*>* InstancesThatShouldUseOldClass);
 	static void CategorizeClassesAndInstances(
 		const TMap<UClass*, UClass*>& InOldToNewClassMap
 		, TMap<UClass*, UClass*>& UnchangedClasses
@@ -392,8 +393,8 @@ private:
 		, TMap<UClass*, UClass*>& ActiveActorReinstancingClasses
 		, const FReinstancingState& ReinstancingState);
 	static void RestoreClass(TMap<UClass*, UClass*>& OldToNewClassMap, const FReinstancingState& ReinstancingState);
-	static void ReinstanceNormalObjects(TMap<UClass*, UClass*>& OldToNewClassMap, const FReinstancingState& ReinstancingState);
-	static void ReinstanceActors(TMap<UClass*, UClass*>& OldToNewClassMap, const FReinstancingState& ReinstancingState);
+	static void ReinstanceNormalObjects(TMap<UClass*, UClass*>& OldToNewClassMap, const FReinstancingState& ReinstancingState, const TSet<UObject*>* InstancesThatShouldUseOldClass);
+	static void ReinstanceActors(TMap<UClass*, UClass*>& OldToNewClassMap, const FReinstancingState& ReinstancingState, const TSet<UObject*>* InstancesThatShouldUseOldClass);
 	static void GetObjectsToReinstance(const UClass* OldClass, TArray<UObject*>& ObjectsToReplace);
 	static void UpdateObjectBeingDebugged(UObject* InOldObject, UObject* InNewObject);
 	static void CheckAndSaveOuterPackageToCleanList(TSet<UPackage*>& CleanPackageList, const UObject* InObject);
