IncrementalCook: ValidationCookArtifact: Update the artifact to read its previous data from the ICookArtifactReader instead of from disk. ICookArtifactReader is necessary when the previous cook files are stored in ZenStore rather than on disk.
#rnx

#Synced-CL 52123874

diff --git a/Engine/Source/Editor/UnrealEd/Private/Cooker/ValidationCookArtifact.cpp b/Engine/Source/Editor/UnrealEd/Private/Cooker/ValidationCookArtifact.cpp
--- a/Engine/Source/Editor/UnrealEd/Private/Cooker/ValidationCookArtifact.cpp
+++ b/Engine/Source/Editor/UnrealEd/Private/Cooker/ValidationCookArtifact.cpp
@@ -4,6 +4,7 @@
 
 #include "AssetRegistry/AssetData.h"
 #include "AssetRegistry/IAssetRegistry.h"
+#include "Cooker/CookArtifactReader.h"
 #include "CookOnTheSide/CookOnTheFlyServer.h"
 #include "Engine/Level.h"
 #include "JsonUtils/RapidJsonUtils.h"
@@ -72,7 +73,7 @@ void FValidationCookArtifact::CompareSettings(UE::Cook::Artifact::FCompareSettin
 
 	// Validation doesn't have any settings that affect incremental cooking
 	
-	if (LoadCookArtifactFile(PlatformArtifactData))
+	if (LoadCookArtifactFile(Context, PlatformArtifactData))
 	{
 		UE_LOGF(LogCook, Display, "CookArtifact '%ls' loaded.", *PlatformArtifactData.Filename);
 	}
@@ -207,10 +208,16 @@ void FValidationCookArtifact::InvalidateCache(const ITargetPlatform* TargetPlatf
 	PlatformArtifactData.WorldValidationCache.Reset();
 }
 
-bool FValidationCookArtifact::LoadCookArtifactFile(FArtifactData& PlatformArtifactData)
+bool FValidationCookArtifact::LoadCookArtifactFile(UE::Cook::Artifact::FCompareSettingsContext& Context, FArtifactData& PlatformArtifactData)
 {
+	TUniquePtr<FArchive> Reader(Context.GetArtifactReader().CreateFileReader(*PlatformArtifactData.Filename));
+	if (!Reader)
+	{
+		return false;
+	}
+
 	FString PlatformArtifactJsonString;
-	if (!FFileHelper::LoadFileToString(PlatformArtifactJsonString, *PlatformArtifactData.Filename))
+	if (!FFileHelper::LoadFileToString(PlatformArtifactJsonString, *Reader.Get()))
 	{
 		return false;
 	}
diff --git a/Engine/Source/Editor/UnrealEd/Private/Cooker/ValidationCookArtifact.h b/Engine/Source/Editor/UnrealEd/Private/Cooker/ValidationCookArtifact.h
--- a/Engine/Source/Editor/UnrealEd/Private/Cooker/ValidationCookArtifact.h
+++ b/Engine/Source/Editor/UnrealEd/Private/Cooker/ValidationCookArtifact.h
@@ -60,7 +60,7 @@ private:
 		FString Filename;
 	};
 
-	bool LoadCookArtifactFile(FArtifactData& PlatformArtifactData);
+	bool LoadCookArtifactFile(UE::Cook::Artifact::FCompareSettingsContext& Context, FArtifactData& PlatformArtifactData);
 	bool SaveCookArtifactFile(FArtifactData& PlatformArtifactData);
 
 	UCookOnTheFlyServer& COTFS;
