Remove check for PlatformManager->HasSessionPlatform() in GetCookMetadataOutputFolder.

During COTF, we don't populate session platforms until after creating the sandbox file and writing out CookedSettings.txt/CookedIniVersion.txt, so this check is causing us to write those files to the wrong place.

[RN] none
#Synced-CL 52124361

diff --git a/Engine/Source/Editor/UnrealEd/Classes/CookOnTheSide/CookOnTheFlyServer.h b/Engine/Source/Editor/UnrealEd/Classes/CookOnTheSide/CookOnTheFlyServer.h
--- a/Engine/Source/Editor/UnrealEd/Classes/CookOnTheSide/CookOnTheFlyServer.h
+++ b/Engine/Source/Editor/UnrealEd/Classes/CookOnTheSide/CookOnTheFlyServer.h
@@ -1369,6 +1369,7 @@ private:
 	FString GetSandboxAssetRegistryFilename();
 
 	FString GetCookedAssetRegistryFilename(const FString& PlatformName);
+	FString GetCookMetadataOutputFolderInternal(const ITargetPlatform* TargetPlatform) const;
 
 	// Return the filename to use for the cook metadata file, adjusted for DLC, sandbox, and platform.
 	FString GetCookedCookMetadataFilename(const ITargetPlatform* TargetPlatform);
diff --git a/Engine/Source/Editor/UnrealEd/Private/CookOnTheFlyServer.cpp b/Engine/Source/Editor/UnrealEd/Private/CookOnTheFlyServer.cpp
--- a/Engine/Source/Editor/UnrealEd/Private/CookOnTheFlyServer.cpp
+++ b/Engine/Source/Editor/UnrealEd/Private/CookOnTheFlyServer.cpp
@@ -6923,17 +6923,24 @@ FString UCookOnTheFlyServer::GetCookOutputFolder(const ITargetPlatform* TargetPl
 	return Result;
 }
 
-FString UCookOnTheFlyServer::GetCookMetadataOutputFolder(const ITargetPlatform* TargetPlatform) const
+FString UCookOnTheFlyServer::GetCookMetadataOutputFolderInternal(const ITargetPlatform* TargetPlatform) const
 {
-	if (!SandboxFile || !PlatformManager || !TargetPlatform)
+	if (!SandboxFile || !TargetPlatform)
 	{
 		return FString();
 	}
-	if (!PlatformManager->HasSessionPlatform(TargetPlatform))
+
+	return ConvertToFullSandboxPath(GetMetadataDirectory(), true, TargetPlatform->PlatformName());
+}
+
+FString UCookOnTheFlyServer::GetCookMetadataOutputFolder(const ITargetPlatform* TargetPlatform) const
+{
+	if (!PlatformManager || !PlatformManager->HasSessionPlatform(TargetPlatform))
 	{
 		return FString();
 	}
-	return ConvertToFullSandboxPath(GetMetadataDirectory(), true, TargetPlatform->PlatformName());
+
+	return GetCookMetadataOutputFolderInternal(TargetPlatform);
 }
 
 void UCookOnTheFlyServer::RegisterCollector(UE::Cook::IMPCollector* Collector, UE::Cook::EProcessType ProcessType)
@@ -7541,7 +7548,7 @@ bool UCookOnTheFlyServer::GetCurrentIniVersionStrings( const ITargetPlatform* Ta
 
 bool UCookOnTheFlyServer::GetCookedIniVersionStrings(const ITargetPlatform* TargetPlatform, UE::Cook::FIniSettingContainer& OutIniSettings, TMap<FString,FString>& OutAdditionalSettings) const
 {
-	const FString PlatformSandboxEditorIni = GetCookMetadataOutputFolder(TargetPlatform) / TEXT("CookedIniVersion.txt"); 
+	const FString PlatformSandboxEditorIni = GetCookMetadataOutputFolderInternal(TargetPlatform) / TEXT("CookedIniVersion.txt"); 
 
 	FConfigFile ConfigFile;
 	ConfigFile.Read(*PlatformSandboxEditorIni);
@@ -7608,7 +7615,7 @@ FString UCookOnTheFlyServer::GetCookSettingsFileName(const ITargetPlatform* Targ
 	FString LeafName = ArtifactName != TEXT("global")
 		? FString::Printf(TEXT("CookedSettings_%s.txt"), *ArtifactName)
 		: TEXT("CookedSettings.txt");
-	return GetCookMetadataOutputFolder(TargetPlatform) / LeafName;
+	return GetCookMetadataOutputFolderInternal(TargetPlatform) / LeafName;
 }
 
 FConfigFile UCookOnTheFlyServer::LoadCookSettings(const ITargetPlatform* TargetPlatform, const FString& ArtifactName)
@@ -7786,7 +7793,7 @@ bool UCookOnTheFlyServer::SaveCurrentIniSettings(const ITargetPlatform* TargetPl
 	UE::Cook::FIniSettingContainer CurrentIniSettings;
 	GetCurrentIniVersionStrings(TargetPlatform, CurrentIniSettings);
 
-	const FString PlatformSandboxEditorIni = GetCookMetadataOutputFolder(TargetPlatform) / TEXT("CookedIniVersion.txt");
+	const FString PlatformSandboxEditorIni = GetCookMetadataOutputFolderInternal(TargetPlatform) / TEXT("CookedIniVersion.txt");
 
 	FConfigFile ConfigFile;
 	// ConfigFile.Read(*PlatformSandboxEditorIni);
@@ -8956,7 +8963,7 @@ FString UCookOnTheFlyServer::GetCookedAssetRegistryFilename(const FString& Platf
 
 FString UCookOnTheFlyServer::GetCookedCookMetadataFilename(const ITargetPlatform* TargetPlatform)
 {
-	return GetCookMetadataOutputFolder(TargetPlatform) / UE::Cook::GetCookMetadataFilename();
+	return GetCookMetadataOutputFolderInternal(TargetPlatform) / UE::Cook::GetCookMetadataFilename();
 }
 
 FString UCookOnTheFlyServer::GetSandboxCachedEditorThumbnailsFilename()
@@ -9254,7 +9261,7 @@ void UCookOnTheFlyServer::WriteCookMetadata(const ITargetPlatform* InTargetPlatf
 
 void UCookOnTheFlyServer::WriteReferencedSet(const ITargetPlatform* InTargetPlatform, TArray<FName>&& CookedPackageNames)
 {
-	const FString MetadataFilename = GetCookMetadataOutputFolder(InTargetPlatform) / UE::Cook::GetReferencedSetFilename();
+	const FString MetadataFilename = GetCookMetadataOutputFolderInternal(InTargetPlatform) / UE::Cook::GetReferencedSetFilename();
 
 	CookedPackageNames.Sort(FNameLexicalLess());
 
@@ -9720,7 +9727,7 @@ void UCookOnTheFlyServer::CookByTheBookFinishedInternal()
 		}
 		else
 		{
-			DiagnosticsDir = GetCookMetadataOutputFolder(TargetPlatform);
+			DiagnosticsDir = GetCookMetadataOutputFolderInternal(TargetPlatform);
 		}
 		IncrementallyModifiedDiagnostics->WriteFinal(DiagnosticsDir);
 	}
@@ -10779,7 +10786,7 @@ UE::Cook::FCookSavePackageContext* UCookOnTheFlyServer::CreateSaveContext(const
 	const FString RootPathSandbox = ConvertToFullSandboxPath(FPaths::RootDir(), true);
 	const FString PlatformString = TargetPlatform->PlatformName();
 	const FString ResolvedRootPath = RootPathSandbox.Replace(TEXT("[Platform]"), *PlatformString);
-	const FString ResolvedMetadataPath = GetCookMetadataOutputFolder(TargetPlatform);
+	const FString ResolvedMetadataPath = GetCookMetadataOutputFolderInternal(TargetPlatform);
 
 	TUniquePtr<FDeterminismManager> DeterminismManager;
 	TSharedPtr<ICookArtifactReader> CookArtifactReader = nullptr;
