[Engine] Bump shader version to fix issue 1283985

diff --git a/Engine/Source/Editor/UnrealEd/Private/Cooker/CookIncrementalVersion.cpp b/Engine/Source/Editor/UnrealEd/Private/Cooker/CookIncrementalVersion.cpp
--- a/Engine/Source/Editor/UnrealEd/Private/Cooker/CookIncrementalVersion.cpp
+++ b/Engine/Source/Editor/UnrealEd/Private/Cooker/CookIncrementalVersion.cpp
@@ -8,6 +8,6 @@ namespace UE::Cook
 {
 
 // Change CookIncrementalVersion to a new guid when all packages in an incremental cook need to be invalidated
-FGuid CookIncrementalVersion(0xFFB471DC, 0x82224FE2, 0xBEDF627C, 0x8371D307);
+FGuid CookIncrementalVersion(0xA3F92C7B, 0x5D1E4A90, 0xC8B74F21, 0x19E6D3AB);
 
 }
