Add some flow phase logging delimiters

#Synced-CL 52051466

diff --git a/Engine/Source/Editor/UnrealEd/Private/CookOnTheFlyServer.cpp b/Engine/Source/Editor/UnrealEd/Private/CookOnTheFlyServer.cpp
--- a/Engine/Source/Editor/UnrealEd/Private/CookOnTheFlyServer.cpp
+++ b/Engine/Source/Editor/UnrealEd/Private/CookOnTheFlyServer.cpp
@@ -9296,6 +9296,7 @@ void UCookOnTheFlyServer::WriteReferencedSet(const ITargetPlatform* InTargetPlat
 
 void UCookOnTheFlyServer::CookByTheBookFinished()
 {
+	UE_LOGF(LogCook, Display, "---- Finalisation: Start ----");
 	{
 		// Add a timer around most of CookByTheBookFinished, but the timer can not exist during or after
 		// ShutdownCookSession because it deletes memory for the timers
@@ -9308,6 +9309,7 @@ void UCookOnTheFlyServer::CookByTheBookFinished()
 		ClearCookInProgressFlagFromGlobalCookSettings(TargetPlatform);
 	}
 	ShutdownCookSession();
+	UE_LOGF(LogCook, Display, "---- Finalisation: End ----");
 	UE_LOG(LogCook, Display, TEXT("Done!"));
 }
 
@@ -11103,6 +11105,7 @@ void UCookOnTheFlyServer::StartCookByTheBook( const FCookByTheBookStartupOptions
 		return;
 	}
 
+	UE_LOGF(LogCook, Display, "---- Post Content Initialisation: Start ----");
 	UE_SCOPED_COOKTIMER(StartCookByTheBook);
 	LLM_SCOPE_BYTAG(Cooker);
 	check(IsInGameThread());
@@ -11154,6 +11157,7 @@ void UCookOnTheFlyServer::StartCookByTheBook( const FCookByTheBookStartupOptions
 	BeginCookDirector(BeginContext);
 
 	// BroadcastCookStarted by contract is sent after all internal CookByTheBookStarted is complete.
+	UE_LOGF(LogCook, Display, "---- Post Content Initialisation: End ----");
 	BroadcastCookStarted();
 }
 
@@ -11166,6 +11170,7 @@ const UCookOnTheFlyServer::FCookByTheBookStartupOptions& UCookOnTheFlyServer::Bl
 	
 	if (IsCookByTheBookMode() && !IsCookingInEditor())
 	{
+		UE_LOGF(LogCook, Display, "---- Pre Content Initialisation: End ----");
 		ConditionalWaitOnCommandFile(TEXTVIEW("Cook"),
 			[this, &CookByTheBookStartupOptions, &ModifiedStartupOptions, &EffectiveStartupOptions, &bOutAbortCook] (FStringView CommandContents)
 		{
@@ -13309,6 +13314,7 @@ void ConditionalWaitOnCommandFile(FStringView GateName, TFunctionRef<void (FStri
 		return;
 	}
 
+	UE_LOGF(LogCook, Display, "---- Waiting For Content: Start ----");
 	uint64 WaitStartTime = FPlatformTime::Cycles64();
 	uint64 LastMessageTime = WaitStartTime;
 	FString CommandContents;
@@ -13330,6 +13336,7 @@ void ConditionalWaitOnCommandFile(FStringView GateName, TFunctionRef<void (FStri
 	}
 
 	CommandHandler(CommandContents);
+	UE_LOGF(LogCook, Display, "---- Waiting For Content: End ----");
 }
 
 void UCookOnTheFlyServer::BroadcastCookStarted()
