Added missing property names to UnrealNames.inl.

#Synced-CL 52250244

diff --git a/Engine/Source/Editor/UnrealEd/Private/Commandlets/TextAssetCommandlet.cpp b/Engine/Source/Editor/UnrealEd/Private/Commandlets/TextAssetCommandlet.cpp
--- a/Engine/Source/Editor/UnrealEd/Private/Commandlets/TextAssetCommandlet.cpp
+++ b/Engine/Source/Editor/UnrealEd/Private/Commandlets/TextAssetCommandlet.cpp
@@ -131,7 +131,6 @@ TSet<FName> GMissingThings;
 void GeneratePropertySchema(FProperty* Property, FStructuredArchiveRecord Record, TArray<FString>& Required)
 {
 	static const FName NAME_ClassProperty(TEXT("ClassProperty"));
-	static const FName NAME_WeakObjectProperty(TEXT("WeakObjectProperty"));
 
 	const FFieldClass* PropertyClass = Property->GetClass();
 	const FName PropertyClassName = PropertyClass->GetFName();
diff --git a/Engine/Source/Runtime/Core/Public/UObject/UnrealNames.inl b/Engine/Source/Runtime/Core/Public/UObject/UnrealNames.inl
--- a/Engine/Source/Runtime/Core/Public/UObject/UnrealNames.inl
+++ b/Engine/Source/Runtime/Core/Public/UObject/UnrealNames.inl
@@ -53,6 +53,9 @@ REGISTER_NAME(34,EnumProperty)
 REGISTER_NAME(35,OptionalProperty)
 REGISTER_NAME(36,Utf8StrProperty)
 REGISTER_NAME(37,AnsiStrProperty)
+REGISTER_NAME(38,WeakObjectProperty)
+REGISTER_NAME(39,MulticastSparseDelegateProperty)
+REGISTER_NAME(40,FieldPathProperty)
 
 
 // Special types.
