ImportObjectProperties now starts a new instancing graph for each non-DefaultSubObject.

#Synced-CL 52250841

diff --git a/Engine/Source/Editor/UnrealEd/Private/EditorObject.cpp b/Engine/Source/Editor/UnrealEd/Private/EditorObject.cpp
--- a/Engine/Source/Editor/UnrealEd/Private/EditorObject.cpp
+++ b/Engine/Source/Editor/UnrealEd/Private/EditorObject.cpp
@@ -307,6 +307,12 @@ static const TCHAR* ImportProperties(
 
 	TArray<FDefinedProperty> DefinedProperties;
 
+	// Local graphs to use for subobjects that require dedicated instancing graphs.
+	TMap<UObject*, FObjectInstancingGraph> DedicatedGraphs;
+
+	// ObjectRemapper to use if none is passed in. Use TOptional to avoid heap allocations and only pay for construction/destruction if Emplace is called.
+	TOptional<TMap<FSoftObjectPath, UObject*>> ExportedPathToObjectCreated;
+
 	// Parse all objects stored in the actor.
 	// Build list of all text properties.
 	bool ImportedBrush = 0;
@@ -536,8 +542,13 @@ static const TCHAR* ImportProperties(
 				if (bUseOverrideSemantics
 					&& FOverridableManager::Get().IsEnabled(BaseTemplate))
 				{
+					if (!ObjectRemapper)
+					{
+						// When using override T3D, always map exported subobject paths so they can be resolved at redefinition time
+						ObjectRemapper = &ExportedPathToObjectCreated.Emplace();
+					}
+
 					FSoftObjectPath ExportedObjectPath(ExportedObjectFullName);
-					checkf(ObjectRemapper, TEXT("Object remapper expected to be non-NULL when importing with override semantics."));
 					if (!ensureMsgf(bExportPath, TEXT("Subobject redefinition is missing the object's export path while importing it.")))
 					{
 						// Fall back to the resolved subobject, however its path can differ from the exported path that it was mapped
@@ -552,9 +563,22 @@ static const TCHAR* ImportProperties(
 					}
 				}
 
+				FObjectInstancingGraph* InstanceGraphToUse = DedicatedGraphs.Find(BaseTemplate);
+				if (!InstanceGraphToUse)
+				{
+					if (FOverridableSerializationLogic::ShouldUseDedicatedInstancingGraph(BaseTemplate))
+					{
+						InstanceGraphToUse = &DedicatedGraphs.Emplace(BaseTemplate, FObjectInstancingGraph(BaseTemplate));
+					}
+					else
+					{
+						InstanceGraphToUse = &InstanceGraph;
+					}
+				}
+
 				// since we're redefining an object in the same text block, only need to import properties again
 				SourceText = ImportObjectProperties( (uint8*)BaseTemplate, SourceText, TemplateClass, SubobjectRoot, BaseTemplate,
-													Warn, Depth + 1, ContextSupplier ? ContextSupplier->CurrentLine : 0, &InstanceGraph, ObjectRemapper );
+													Warn, Depth + 1, ContextSupplier ? ContextSupplier->CurrentLine : 0, InstanceGraphToUse, ObjectRemapper );
 
 				// If the object import fails, early out and stop parsing.
 				if (SourceText == nullptr)
@@ -714,9 +738,12 @@ static const TCHAR* ImportProperties(
 					}
 				}
 
-
+				FObjectInstancingGraph* InstanceGraphToUse = &InstanceGraph;
 				if (!ComponentTemplate)
 				{
+					bool bUseDedicatedInstancingGraph = FOverridableSerializationLogic::ShouldUseDedicatedInstancingGraphForNewObject(TemplateClass, NewFlags, Archetype, bOverrideT3DSupported && OverriddenOperation);
+					
+					FObjectInstancingGraph TempGraph;
 					ComponentTemplate = NewObject<UObject>(
 						SubobjectOuter,
 						TemplateClass,
@@ -724,11 +751,30 @@ static const TCHAR* ImportProperties(
 						NewFlags,
 						Archetype,
 						!!SubobjectOuter,
-						&InstanceGraph
+						bUseDedicatedInstancingGraph ? &TempGraph : &InstanceGraph
 						);
+
+					if (bUseDedicatedInstancingGraph)
+					{
+						if (!TempGraph.HasDestinationRoot())
+						{
+							TempGraph.SetDestinationRoot(ComponentTemplate);
+						}
+
+						InstanceGraphToUse = &DedicatedGraphs.Emplace(ComponentTemplate, MoveTemp(TempGraph));
+					}
 				}
 				else
 				{
+					if (FObjectInstancingGraph* DedicatedGraph = DedicatedGraphs.Find(ComponentTemplate))
+					{
+						InstanceGraphToUse = DedicatedGraph;
+					}
+					else if (FOverridableSerializationLogic::ShouldUseDedicatedInstancingGraph(ComponentTemplate, bOverrideT3DSupported && OverriddenOperation))
+					{
+						InstanceGraphToUse = &DedicatedGraphs.Emplace(ComponentTemplate, FObjectInstancingGraph(ComponentTemplate));
+					}
+
 					// We do not want to set RF_Transactional for construction script created components, so we have to monkey with things here
 					if (NewFlags & RF_Transactional)
 					{
@@ -764,7 +810,7 @@ static const TCHAR* ImportProperties(
 				// Notes:
 				//	- It is ok if the archetype is NULL. Including it avoids a redundant GetArchetype() call when it has already been set.
 				//	- It is safe to assume that we're always inside a root object scope here, because we're importing a Begin Object block.
-				InstanceGraph.AddNewInstance(ComponentTemplate, Archetype);
+				InstanceGraphToUse->AddNewInstance(ComponentTemplate, Archetype);
 
 				// replace all properties in this subobject outer' class that point to the original subobject with the new subobject
 				TMap<UObject*, UObject*> ReplacementMap;
@@ -806,7 +852,7 @@ static const TCHAR* ImportProperties(
 					Warn, 
 					Depth+1,
 					ContextSupplier ? ContextSupplier->CurrentLine : 0,
-					&InstanceGraph,
+					InstanceGraphToUse,
 					ObjectRemapper
 					);
 
@@ -907,37 +953,19 @@ const TCHAR* ImportObjectProperties( FImportObjectParams& InParams )
 		InParams.SubobjectOuter->PreEditChange(NULL);
 	}
 
-	// When importing hierarchies of objects with instanced sub-objects we don't want to use the
-	// instancing graph of our outer objects since we need our dedicated instances.
-	// For example:
-	//
-	//	A_CDO { SubObj }
-	//	PastedObj { A_1{ SubObj_1 } A_2{ SubObj_1 } }
-	//
-	//	by sharing the instancing graph, we end up with both SubObj_1 being the same pointer
-	//	by using a dedicated instancing graph, both SubObj_1 are newly created object pointer 
-	//
-	// Ideally this rule should be applied more widely but for now we limit it to the Prefabs
-	// which is identified by using overridable serialization.
-	const bool bShouldUseDedicatedInstancingGraph = FOverridableManager::Get().IsEnabled(InParams.SubobjectRoot);
 	FObjectInstancingGraph TempGraph;
-	FObjectInstancingGraph& InstanceGraph = InParams.InInstanceGraph && !bShouldUseDedicatedInstancingGraph ? *InParams.InInstanceGraph : TempGraph;
+	FObjectInstancingGraph& InstanceGraph = InParams.InInstanceGraph ? *InParams.InInstanceGraph : TempGraph;
+	if ( InParams.SubobjectRoot && InParams.SubobjectRoot != UObject::StaticClass()->GetDefaultObject() )
+	{
+		InstanceGraph.SetDestinationRoot(InParams.SubobjectRoot);
+	}
 
 	// Use TOptional to avoid heap allocations and only pay for construction/destruction if Emplace is called.
 	TOptional<TMap<FSoftObjectPath, UObject*>> ExportedPathToObjectCreated;
-	if (bShouldUseDedicatedInstancingGraph)
+	if (!InParams.ObjectRemapper && FOverridableManager::Get().IsEnabled(InParams.SubobjectOuter))
 	{
-		InstanceGraph.SetDestinationRoot(InParams.SubobjectOuter);
-
 		// When using override T3D, always map exported subobject paths so they can be resolved at redefinition time
-		if (!InParams.ObjectRemapper)
-		{
-			InParams.ObjectRemapper = &ExportedPathToObjectCreated.Emplace();
-		}
-	}
-	else if ( InParams.SubobjectRoot && InParams.SubobjectRoot != UObject::StaticClass()->GetDefaultObject() )
-	{
-		InstanceGraph.SetDestinationRoot(InParams.SubobjectRoot);
+		InParams.ObjectRemapper = &ExportedPathToObjectCreated.Emplace();
 	}
 
 	// Parse the object properties.
@@ -1117,6 +1145,9 @@ static const TCHAR* ImportCreateSubObjectsStep(
 
 	FString StrLine;
 
+	// Local graphs to use for subobjects that require dedicated instancing graphs.
+	TMap<UObject*, FObjectInstancingGraph> DedicatedGraphs;
+
 	// Parse all objects stored in the actor.
 	// Build list of all text properties.
 	bool ImportedBrush = 0;
@@ -1240,6 +1271,19 @@ static const TCHAR* ImportCreateSubObjectsStep(
 			checkSlow(TemplateClass != nullptr);
 			if (bRedefiningSubobject)
 			{
+				FObjectInstancingGraph* InstanceGraphToUse = DedicatedGraphs.Find(BaseTemplate);
+				if (!InstanceGraphToUse)
+				{
+					if (FOverridableSerializationLogic::ShouldUseDedicatedInstancingGraph(BaseTemplate))
+					{
+						InstanceGraphToUse = &DedicatedGraphs.Emplace(BaseTemplate, FObjectInstancingGraph(BaseTemplate));
+					}
+					else
+					{
+						InstanceGraphToUse = &InstanceGraph;
+					}
+				}
+				
 				// Simply recreate any missing sub object
 				FMultiStepsImportObjectParams Params;
 				Params.SourceText = FStringView(CurrentSourceText, SourceText.Len() - (CurrentSourceText - SourceText.GetData()));
@@ -1249,7 +1293,7 @@ static const TCHAR* ImportCreateSubObjectsStep(
 				Params.Warn = Warn;
 				Params.Depth = Depth + 1;
 				Params.LineNumber = ContextSupplier ? ContextSupplier->CurrentLine : 0;
-				Params.InInstanceGraph = &InstanceGraph;
+				Params.InInstanceGraph = InstanceGraphToUse;
 				Params.ObjectRemapper = ObjectRemapper;
 
 				CurrentSourceText = ImportCreateObjectsStep(Params);
@@ -1361,8 +1405,12 @@ static const TCHAR* ImportCreateSubObjectsStep(
 					}
 				}
 				
+				FObjectInstancingGraph* InstanceGraphToUse = &InstanceGraph;
 				if (!ComponentTemplate)
 				{
+					bool bUseDedicatedInstancingGraph = FOverridableSerializationLogic::ShouldUseDedicatedInstancingGraphForNewObject(TemplateClass, NewFlags, Archetype, bOverrideT3DSupported && OverriddenOperation);
+					
+					FObjectInstancingGraph TempGraph;
 					ComponentTemplate = NewObject<UObject>(
 						SubobjectOuter,
 						TemplateClass,
@@ -1370,9 +1418,19 @@ static const TCHAR* ImportCreateSubObjectsStep(
 						NewFlags,
 						Archetype,
 						!!SubobjectOuter,
-						&InstanceGraph
+						bUseDedicatedInstancingGraph ? &TempGraph : &InstanceGraph
 						);
 
+					if (bUseDedicatedInstancingGraph)
+					{
+						if (!TempGraph.HasDestinationRoot())
+						{
+							TempGraph.SetDestinationRoot(ComponentTemplate);
+						}
+
+						InstanceGraphToUse = &DedicatedGraphs.Emplace(ComponentTemplate, MoveTemp(TempGraph));
+					}
+
 					if (!ExportedObjectFullName.IsEmpty() && ObjectRemapper)
 					{
 						ObjectRemapper->Add(MoveTemp(ExportedObjectFullName), ComponentTemplate);
@@ -1380,6 +1438,15 @@ static const TCHAR* ImportCreateSubObjectsStep(
 				}
 				else
 				{
+					if (FObjectInstancingGraph* DedicatedGraph = DedicatedGraphs.Find(ComponentTemplate))
+					{
+						InstanceGraphToUse = DedicatedGraph;
+					}
+					else if (FOverridableSerializationLogic::ShouldUseDedicatedInstancingGraph(ComponentTemplate, bOverrideT3DSupported && OverriddenOperation))
+					{
+						InstanceGraphToUse = &DedicatedGraphs.Emplace(ComponentTemplate, FObjectInstancingGraph(ComponentTemplate));
+					}
+
 					// We do not want to set RF_Transactional for construction script created components, so we have to monkey with things here
 					if (NewFlags & RF_Transactional)
 					{
@@ -1412,7 +1479,7 @@ static const TCHAR* ImportCreateSubObjectsStep(
 				// Notes:
 				//	- It is ok if the archetype is NULL. Including it avoids a redundant GetArchetype() call when it has already been set.
 				//	- It is safe to assume that we're always inside a root object scope here, because we're importing a Begin Object block.
-				InstanceGraph.AddNewInstance(ComponentTemplate, Archetype);
+				InstanceGraphToUse->AddNewInstance(ComponentTemplate, Archetype);
 
 				// replace all properties in this subobject outer' class that point to the original subobject with the new subobject
 				TMap<UObject*, UObject*> ReplacementMap;
@@ -1436,7 +1503,7 @@ static const TCHAR* ImportCreateSubObjectsStep(
 				Params.Warn = Warn;
 				Params.Depth = Depth + 1;
 				Params.LineNumber = ContextSupplier ? ContextSupplier->CurrentLine : 0;
-				Params.InInstanceGraph = &InstanceGraph;
+				Params.InInstanceGraph = InstanceGraphToUse;
 				Params.ObjectRemapper = ObjectRemapper;
 
 				if (!ExportedObjectFullName.IsEmpty() && ObjectRemapper)
@@ -1552,6 +1619,9 @@ static const TCHAR* ImportPropertiesStep(
 
 	TArray<FDefinedProperty> DefinedProperties;
 
+	// Local graphs to use for subobjects that require dedicated instancing graphs.
+	TMap<UObject*, FObjectInstancingGraph> DedicatedGraphs;
+
 	// Parse all objects stored in the actor.
 	// Build list of all text properties.
 	bool ImportedBrush = 0;
@@ -1610,6 +1680,19 @@ static const TCHAR* ImportPropertiesStep(
 
 				if (UObject* BaseTemplate = StaticFindObjectFast(UObject::StaticClass(), SubobjectOuter, TemplateName))
 				{
+					FObjectInstancingGraph* InstanceGraphToUse = DedicatedGraphs.Find(BaseTemplate);
+					if (!InstanceGraphToUse)
+					{
+						if (FOverridableSerializationLogic::ShouldUseDedicatedInstancingGraph(BaseTemplate))
+						{
+							InstanceGraphToUse = &DedicatedGraphs.Emplace(BaseTemplate, FObjectInstancingGraph(BaseTemplate));
+						}
+						else
+						{
+							InstanceGraphToUse = &InstanceGraph;
+						}
+					}
+
 					FMultiStepsImportObjectParams Params;
 					Params.DestData = static_cast<uint8*>(static_cast<void*>(BaseTemplate));
 					Params.SourceText = FStringView(CurrentSourceText, SourceText.Len() - (CurrentSourceText - SourceText.GetData()));
@@ -1619,7 +1702,7 @@ static const TCHAR* ImportPropertiesStep(
 					Params.Warn = Warn;
 					Params.Depth = Depth + 1;
 					Params.LineNumber = ContextSupplier ? ContextSupplier->CurrentLine : 0;
-					Params.InInstanceGraph = &InstanceGraph;
+					Params.InInstanceGraph = InstanceGraphToUse;
 					Params.ObjectRemapper = ObjectRemapper;
 					Params.PropertiesToSkip = PropertiesToSkip;
 
diff --git a/Engine/Source/Editor/UnrealEd/Private/Factories/EditorFactories.cpp b/Engine/Source/Editor/UnrealEd/Private/Factories/EditorFactories.cpp
--- a/Engine/Source/Editor/UnrealEd/Private/Factories/EditorFactories.cpp
+++ b/Engine/Source/Editor/UnrealEd/Private/Factories/EditorFactories.cpp
@@ -5359,6 +5359,9 @@ void FCustomizableTextObjectFactory::ProcessBuffer(UObject* InParent, EObjectFla
 
 	const bool bT3DOverrideEnabled = FOverridableSerializationLogic::HasCapabilities(FOverridableSerializationLogic::ECapabilities::T3DSerialization);
 
+	// Local graphs to use for objects that require dedicated instancing graphs.
+	TMap<UObject*, FObjectInstancingGraph> DedicatedGraphs;
+
 	FString StrLine;
 	while( FParse::Line(&Buffer,StrLine) )
 	{
@@ -5422,7 +5425,17 @@ void FCustomizableTextObjectFactory::ProcessBuffer(UObject* InParent, EObjectFla
 				ClearObjectNameUsage(ObjectParent, ObjName);
 
 				// Spawn the object and reset it's archetype
-				UObject* CreatedObject = NewObject<UObject>(ObjectParent, ObjClass, ObjName, Flags, ObjArchetype, !!ObjectParent, &InstanceGraph);
+				FObjectInstancingGraph TempGraph;
+				bool bUseDedicatedInstancingGraph = FOverridableSerializationLogic::ShouldUseDedicatedInstancingGraphForNewObject(ObjClass, Flags, ObjArchetype, bT3DOverrideEnabled && OverriddenOperation);
+				UObject* CreatedObject = NewObject<UObject>(ObjectParent, ObjClass, ObjName, Flags, ObjArchetype, !!ObjectParent, bUseDedicatedInstancingGraph ? &TempGraph : &InstanceGraph);
+				if (bUseDedicatedInstancingGraph)
+				{
+					if (!TempGraph.HasDestinationRoot())
+					{
+						TempGraph.SetDestinationRoot(CreatedObject);
+					}
+					DedicatedGraphs.Emplace(CreatedObject, MoveTemp(TempGraph));
+				}
 
 				// Overridable handling
 				if (bT3DOverrideEnabled)
@@ -5492,8 +5505,14 @@ void FCustomizableTextObjectFactory::ProcessBuffer(UObject* InParent, EObjectFla
 		UObject* CreatedObject = NewObjects[i];
 		const FString& PropText = PropMap.FindChecked(CreatedObject);
 
+		FObjectInstancingGraph* InstanceGraphToUse = DedicatedGraphs.Find(CreatedObject);
+		if (!InstanceGraphToUse)
+		{
+			InstanceGraphToUse = &InstanceGraph;
+		}
+
 		// Import the properties and give the derived factory a shot at it
-		if (ImportObjectProperties((uint8*)CreatedObject, *PropText, CreatedObject->GetClass(), CreatedObject, CreatedObject, WarningContext, 0, 0, &InstanceGraph))
+		if (ImportObjectProperties((uint8*)CreatedObject, *PropText, CreatedObject->GetClass(), CreatedObject, CreatedObject, WarningContext, 0, 0, InstanceGraphToUse))
 		{
 			ProcessConstructedObject(CreatedObject);
 		}
diff --git a/Engine/Source/Runtime/CoreUObject/Private/UObject/OverriddenPropertySet.cpp b/Engine/Source/Runtime/CoreUObject/Private/UObject/OverriddenPropertySet.cpp
--- a/Engine/Source/Runtime/CoreUObject/Private/UObject/OverriddenPropertySet.cpp
+++ b/Engine/Source/Runtime/CoreUObject/Private/UObject/OverriddenPropertySet.cpp
@@ -75,6 +75,44 @@ namespace Private
 	};
 
 	static FCapabilitiesAutoInitializer CapabilitiesAutoInitializer;
+	
+	/** Custom test if the given object is a default subobject. Overrides support non-derived subobject instances, which is not covered by UObject::IsDefaultSubObject. */
+	static bool IsDefaultSubObject(TNotNull<const UObject*> SubObject)
+	{
+		if (SubObject->HasAnyFlags(RF_DefaultSubObject))
+		{
+			return true;
+		}
+
+		const UObject* Archetype = SubObject->GetArchetype();
+		if (Archetype && !Archetype->HasAnyFlags(RF_ClassDefaultObject))
+		{
+			return true;
+		}
+	
+		return false;
+	}
+
+	/** Custom test if the object created by the given params will be a default subobject. */
+	static bool IsNewObjectDefaultSubObject(TNotNull<const UClass*> Class, EObjectFlags Flags, const UObject* Template)
+	{
+		if (EnumHasAnyFlags(Flags, RF_DefaultSubObject))
+		{
+			return true;
+		}
+	
+		if (!Template)
+		{
+			Template = Class->GetDefaultObject();
+		}
+	
+		if (Template && !Template->HasAnyFlags(RF_ClassDefaultObject))
+		{
+			return true;
+		}
+	
+		return false;
+	}
 }
 
 //----------------------------------------------------------------------//
@@ -158,6 +196,23 @@ bool FOverridableSerializationLogic::ShouldPropertyShadowSerializeSubObject(TNot
 	return false;
 }
 
+bool FOverridableSerializationLogic::ShouldUseDedicatedInstancingGraph(TNotNull<const UObject*> Object, TOptional<bool> UsesOverrides)
+{
+	if (!UsesOverrides.IsSet())
+	{
+		UsesOverrides = HasCapabilities(ECapabilities::T3DSerialization) && FOverridableManager::Get().IsEnabled(Object);
+	}
+	
+	// Ideally this rule should be applied more widely but for now we limit it to objects that use overridable serialization.
+	return UsesOverrides.GetValue() && !Private::IsDefaultSubObject(Object);
+}
+
+bool FOverridableSerializationLogic::ShouldUseDedicatedInstancingGraphForNewObject(TNotNull<const UClass*> Class, EObjectFlags Flags,
+	const UObject* Template, bool bUsesOverrides)
+{
+	return bUsesOverrides && !Private::IsNewObjectDefaultSubObject(Class, Flags, Template);
+}
+
 bool FOverridableSerializationLogic::HasCapabilities(ECapabilities InCapabilities)
 {
 	return (Capabilities & InCapabilities) == InCapabilities;
diff --git a/Engine/Source/Runtime/CoreUObject/Public/UObject/OverriddenPropertySet.h b/Engine/Source/Runtime/CoreUObject/Public/UObject/OverriddenPropertySet.h
--- a/Engine/Source/Runtime/CoreUObject/Public/UObject/OverriddenPropertySet.h
+++ b/Engine/Source/Runtime/CoreUObject/Public/UObject/OverriddenPropertySet.h
@@ -149,6 +149,33 @@ public:
 
 	/** To know if the specified property should shadow serialize its values for OS */
 	COREUOBJECT_API static bool ShouldPropertyShadowSerializeSubObject(TNotNull<const FProperty*> Property);
+	
+	/** Method to determine whether the object should use its own instancing object graph e.g. during import.
+	* When importing hierarchies of objects with instanced sub-objects we don't want to use the
+	* instancing graph of our outer objects since we need our dedicated instances.
+	* For example:
+	*
+	*	A_CDO { SubObj }
+	*	PastedObj { A_1{ SubObj_1 } A_2{ SubObj_1 } }
+	*
+	*	by sharing the instancing graph, we end up with both SubObj_1 being the same pointer
+	*	by using a dedicated instancing graph, both SubObj_1 are newly created object pointer 
+	* 
+	* @param Object The object that should be tested.
+	* @param UsesOverrides Whether the given object has or will have overrides enabled. If unset, this will be determined from the object itself.
+	* @returns True if this object should use its own instancing object graph e.g. during import. Otherwise, false.
+	*/
+	COREUOBJECT_API static bool ShouldUseDedicatedInstancingGraph(TNotNull<const UObject*> Object, TOptional<bool> UsesOverrides = {});
+	/** Method to determine whether the object should use its own instancing object graph e.g. during import.
+	 * @param Class The class that will be used to create the new object.
+	 * @param Flags The flags that will be used to create the new object.
+	 * @param Template The template that will be used to create the new object.
+	 * @param bUsesOverrides Whether the given object has or will have overrides enabled.
+	 * @see The other ShouldUseDedicatedInstancingGraph overload provides more information.
+	 * @returns True if a dedicated instancing object graph should be used when creating and loading the object. Otherwise, false.
+	 */
+	COREUOBJECT_API static bool ShouldUseDedicatedInstancingGraphForNewObject(TNotNull<const UClass*> Class, EObjectFlags Flags,
+		const UObject* Template, bool bUsesOverrides);
 
 private:
 
