LevelEditor: Enable Pivot actions on non-actors. Also add a helper for getting the selection set in UnrealEdSrv.cpp and use it.

#Synced-CL 52191381

diff --git a/Engine/Source/Editor/LevelEditor/Private/LevelEditorActions.cpp b/Engine/Source/Editor/LevelEditor/Private/LevelEditorActions.cpp
--- a/Engine/Source/Editor/LevelEditor/Private/LevelEditorActions.cpp
+++ b/Engine/Source/Editor/LevelEditor/Private/LevelEditorActions.cpp
@@ -4451,11 +4451,11 @@ void FLevelEditorCommands::RegisterCommands()
 	UI_COMMAND( HoldToEnablePivotVertexSnapping, "Hold to Enable Pivot Vertex Snapping", "Hold to enable vertex snapping while dragging a pivot. Alt must be a modifier in this command or it will not work.", EUserInterfaceActionType::ToggleButton, FInputChord(EModifierKey::Alt, EKeys::V) );
 
 	//@ todo Slate better tooltips for pivot options
-	UI_COMMAND( SavePivotToPrePivot, "Set as Pivot Offset", "Sets the current pivot location as the pivot offset for this actor", EUserInterfaceActionType::Button, FInputChord() );
-	UI_COMMAND( ResetPrePivot, "Reset Pivot Offset", "Resets the pivot offset for this actor", EUserInterfaceActionType::Button, FInputChord() );
+	UI_COMMAND( SavePivotToPrePivot, "Set as Pivot Offset", "Sets the current pivot location as the pivot offset for this element", EUserInterfaceActionType::Button, FInputChord() );
+	UI_COMMAND( ResetPrePivot, "Reset Pivot Offset", "Resets the pivot offset for this element", EUserInterfaceActionType::Button, FInputChord() );
 	UI_COMMAND( ResetPivot, "Reset Pivot", "Resets the pivot", EUserInterfaceActionType::Button, FInputChord() );
-	UI_COMMAND( MovePivotHere, "Set Pivot Offset Here", "Sets the pivot offset to the clicked location", EUserInterfaceActionType::Button, FInputChord() );
-	UI_COMMAND( MovePivotHereSnapped, "Set Pivot Offset Here (Snapped)", "Sets the pivot offset to the nearest grid point to the clicked location", EUserInterfaceActionType::Button, FInputChord() );
+	UI_COMMAND( MovePivotHere, "Set Pivot Offset Here", "Sets the pivot to the clicked location", EUserInterfaceActionType::Button, FInputChord() );
+	UI_COMMAND( MovePivotHereSnapped, "Set Pivot Offset Here (Snapped)", "Sets the pivot to the nearest grid point to the clicked location", EUserInterfaceActionType::Button, FInputChord() );
 	UI_COMMAND( MovePivotToCenter, "Center on Selection", "Centers the pivot to the middle of the selection", EUserInterfaceActionType::Button, FInputChord() );
 
 	UI_COMMAND( ConvertToAdditive, "Additive", "Converts the selected brushes to additive brushes", EUserInterfaceActionType::Button, FInputChord() );
diff --git a/Engine/Source/Editor/LevelEditor/Private/LevelEditorContextMenu.cpp b/Engine/Source/Editor/LevelEditor/Private/LevelEditorContextMenu.cpp
--- a/Engine/Source/Editor/LevelEditor/Private/LevelEditorContextMenu.cpp
+++ b/Engine/Source/Editor/LevelEditor/Private/LevelEditorContextMenu.cpp
@@ -728,14 +728,13 @@ void FLevelEditorContextMenu::RegisterElementContextMenu()
 				false, // default value
 				FSlateIcon(FAppStyle::GetAppStyleSetName(), "Icons.Snap"));
 
-			// TODO: Some pivot options work, but others still require additional support for non-actors
-			//Section.AddSubMenu(
-			//	"PivotSubMenu",
-			//	LOCTEXT("PivotElementSubMenu", "Pivot"),
-			//	LOCTEXT("PivotElementSubMenu_ToolTip", "Element pivoting utils"),
-			//	FNewToolMenuDelegate::CreateStatic(&FLevelEditorContextMenuImpl::FillPivotMenu),
-			//	false, // default value
-			//	FSlateIcon(FAppStyle::GetAppStyleSetName(), "Icons.SetShowPivot"));
+			Section.AddSubMenu(
+				"PivotSubMenu",
+				LOCTEXT("PivotElementSubMenu", "Pivot"),
+				LOCTEXT("PivotElementSubMenu_ToolTip", "Element pivoting utils"),
+				FNewToolMenuDelegate::CreateStatic(&FLevelEditorContextMenuImpl::FillPivotMenu),
+				false, // default value
+				FSlateIcon(FAppStyle::GetAppStyleSetName(), "Icons.SetShowPivot"));
 
 			// TODO: Attach/detach requires additional support. Attach, in particular, requires ability to show
 			//  elements to which attachment is possible.
diff --git a/Engine/Source/Editor/UnrealEd/Private/EditorSelectUtils.cpp b/Engine/Source/Editor/UnrealEd/Private/EditorSelectUtils.cpp
--- a/Engine/Source/Editor/UnrealEd/Private/EditorSelectUtils.cpp
+++ b/Engine/Source/Editor/UnrealEd/Private/EditorSelectUtils.cpp
@@ -84,13 +84,6 @@ void UUnrealEdEngine::NoteActorMovement()
 			}
 		}
 
-		if( SelectedActor == NULL )
-		{
-			USelection* SelectedActors = GetSelectedActors();
-			SelectedActors->Modify();
-			SelectActor( GWorld->GetDefaultBrush(), true, true );
-		}
-
 		// Look for an actor that requires snapping.
 		{
 			FSelectionIterator It(GetSelectedActorIterator());
diff --git a/Engine/Source/Editor/UnrealEd/Private/Elements/Actor/ActorElementEditorWorldInterface.cpp b/Engine/Source/Editor/UnrealEd/Private/Elements/Actor/ActorElementEditorWorldInterface.cpp
--- a/Engine/Source/Editor/UnrealEd/Private/Elements/Actor/ActorElementEditorWorldInterface.cpp
+++ b/Engine/Source/Editor/UnrealEd/Private/Elements/Actor/ActorElementEditorWorldInterface.cpp
@@ -38,7 +38,9 @@ bool UActorElementEditorWorldInterface::SetPivotOffset(const FTypedElementHandle
 {
 	if (AActor* Actor = ActorElementDataUtil::GetActorFromHandle(InElementHandle))
 	{
+		Actor->Modify();
 		Actor->SetPivotOffset(InPivotOffset);
+		Actor->PostEditMove(true);
 		return true;
 	}
 
diff --git a/Engine/Source/Editor/UnrealEd/Private/UnrealEdSrv.cpp b/Engine/Source/Editor/UnrealEd/Private/UnrealEdSrv.cpp
--- a/Engine/Source/Editor/UnrealEd/Private/UnrealEdSrv.cpp
+++ b/Engine/Source/Editor/UnrealEd/Private/UnrealEdSrv.cpp
@@ -83,9 +83,11 @@
 #include "AssetRegistry/AssetRegistryModule.h"
 #include "FbxExporter.h"
 #include "DesktopPlatformModule.h"
+#include "Elements/Actor/ActorElementData.h"
 #include "Elements/Framework/TypedElementList.h"
 #include "Elements/Framework/TypedElementRegistry.h"
 #include "Elements/Framework/TypedElementCommonActions.h"
+#include "Elements/Interfaces/TypedElementWorldInterface.h"
 #include "SnappingUtils.h"
 #include "AssetSelection.h"
 #include "HighResScreenshot.h"
@@ -115,6 +117,32 @@ DEFINE_LOG_CATEGORY_STATIC(LogUnrealEdSrv, Log, All);
 
 #define LOCTEXT_NAMESPACE "UnrealEdSrv"
 
+namespace UnrealEdSrvLocals
+{
+	UTypedElementSelectionSet* GetLevelEditorSelectionSet()
+	{
+		if (GCurrentLevelEditingViewportClient)
+		{
+			if (TSharedPtr<ILevelEditor> LevelEditor = GCurrentLevelEditingViewportClient->ParentLevelEditor.Pin())
+			{
+				return LevelEditor->GetMutableElementSelectionSet();
+			}
+		}
+		return nullptr;
+	}
+	UTypedElementCommonActions* GetLevelEditorCommonActions()
+	{
+		if (GCurrentLevelEditingViewportClient)
+		{
+			if (TSharedPtr<ILevelEditor> LevelEditor = GCurrentLevelEditingViewportClient->ParentLevelEditor.Pin())
+			{
+				return LevelEditor->GetCommonActions();
+			}
+		}
+		return nullptr;
+	}
+}
+
 /**
  * Dumps a set of selected objects to debugf.
  */
@@ -1488,6 +1516,8 @@ bool UUnrealEdEngine::DoDirtyPackagesNeedCheckout() const
 
 bool UUnrealEdEngine::Exec_Edit( UWorld* InWorld, const TCHAR* Str, FOutputDevice& Ar )
 {
+	using namespace UnrealEdSrvLocals;
+
 	const FText TransactionName = NSLOCTEXT("UnrealEd", "Cut", "Cut");
 
 	const bool bComponentsSelected = GetSelectedComponentCount() > 0;
@@ -1500,28 +1530,24 @@ bool UUnrealEdEngine::Exec_Edit( UWorld* InWorld, const TCHAR* Str, FOutputDevic
 		}
 
 		bool bExportedViaCommonActions = false;
-		if (GCurrentLevelEditingViewportClient)
+		
+		if (UTypedElementCommonActions* CommonActions = GetLevelEditorCommonActions())
 		{
-			if (TSharedPtr<ILevelEditor> LevelEditor = GCurrentLevelEditingViewportClient->ParentLevelEditor.Pin())
-			{
-				if (UTypedElementCommonActions* CommonActions = LevelEditor->GetCommonActions())
-				{
-					// End drags to avoid deleting something out from under one.
-					FSlateApplication::Get().CancelDragDrop();
+			// End drags to avoid deleting something out from under one.
+			FSlateApplication::Get().CancelDragDrop();
 					
-					const FScopedTransaction Transaction(TransactionName);
+			const FScopedTransaction Transaction(TransactionName);
 
-					const bool bCheckRef = GetDefault<ULevelEditorMiscSettings>()->bCheckReferencesOnDelete;
-					FTypedElementDeletionOptions Options;
-					Options
-						.SetWarnAboutReferences(bCheckRef)
-						.SetWarnAboutSoftReferences(bCheckRef);
+			const bool bCheckRef = GetDefault<ULevelEditorMiscSettings>()->bCheckReferencesOnDelete;
+			FTypedElementDeletionOptions Options;
+			Options
+				.SetWarnAboutReferences(bCheckRef)
+				.SetWarnAboutSoftReferences(bCheckRef);
 
-					UTypedElementSelectionSet* SelectionSet = LevelEditor->GetMutableElementSelectionSet();
-					bExportedViaCommonActions = CommonActions->CutSelectedElements(SelectionSet, InWorld, Options);
-				}
-			}
+			UTypedElementSelectionSet* SelectionSet = GetLevelEditorSelectionSet();
+			bExportedViaCommonActions = CommonActions->CutSelectedElements(SelectionSet, InWorld, Options);
 		}
+	
 
 		// As with copy, bExportedViaCommonActions would be false when TypedElementCommonActionsUtils::IsElementCopyAndPasteEnabled()
 		//  returns false inside an actor/component common actions customization, falling into this legacy path.
@@ -1555,16 +1581,10 @@ bool UUnrealEdEngine::Exec_Edit( UWorld* InWorld, const TCHAR* Str, FOutputDevic
 		}
 		
 		bool bExportedViaCommonActions = false;
-		if (GCurrentLevelEditingViewportClient)
+		if (UTypedElementCommonActions* CommonActions = GetLevelEditorCommonActions())
 		{
-			if (TSharedPtr<ILevelEditor> LevelEditor = GCurrentLevelEditingViewportClient->ParentLevelEditor.Pin())
-			{
-				if (UTypedElementCommonActions* CommonActions = LevelEditor->GetCommonActions())
-				{
-					UTypedElementSelectionSet* SelectionSet = LevelEditor->GetMutableElementSelectionSet();
-					bExportedViaCommonActions = CommonActions->CopySelectedElements(SelectionSet);
-				}
-			}
+			UTypedElementSelectionSet* SelectionSet = GetLevelEditorSelectionSet();
+			bExportedViaCommonActions = CommonActions->CopySelectedElements(SelectionSet);
 		}
 
 		// Legacy code path
@@ -1593,61 +1613,55 @@ bool UUnrealEdEngine::Exec_Edit( UWorld* InWorld, const TCHAR* Str, FOutputDevic
 		}
 		
 		bool bPastedViaCommonActions = false;
-		if (GCurrentLevelEditingViewportClient)
+		if (UTypedElementCommonActions* CommonActions = GetLevelEditorCommonActions())
 		{
-			if (TSharedPtr<ILevelEditor> LevelEditor = GCurrentLevelEditingViewportClient->ParentLevelEditor.Pin())
+			FString ClipboardPaste;
+			FPlatformApplicationMisc::ClipboardPaste(ClipboardPaste);
+			if (!ClipboardPaste.IsEmpty())
 			{
-				if (UTypedElementCommonActions* CommonActions = LevelEditor->GetCommonActions())
+				UTypedElementSelectionSet* SelectionSet = GetLevelEditorSelectionSet();
+						
+				bool bPasteAtLocation = false;
+				FVector PasteLocation = FVector::ZeroVector;
+				FString TempStr;
+				FText TransDescription = NSLOCTEXT("UnrealEd", "Paste", "Paste");
+				if (FParse::Value(Str, TEXT("TO="), TempStr))
 				{
-					FString ClipboardPaste;
-					FPlatformApplicationMisc::ClipboardPaste(ClipboardPaste);
-					if (!ClipboardPaste.IsEmpty())
+					if (!FCString::Strcmp(*TempStr, TEXT("HERE")))
 					{
-						UTypedElementSelectionSet* SelectionSet = LevelEditor->GetMutableElementSelectionSet();
-						
-						bool bPasteAtLocation = false;
-						FVector PasteLocation = FVector::ZeroVector;
-						FString TempStr;
-						FText TransDescription = NSLOCTEXT("UnrealEd", "Paste", "Paste");
-						if (FParse::Value(Str, TEXT("TO="), TempStr))
-						{
-							if (!FCString::Strcmp(*TempStr, TEXT("HERE")))
-							{
-								bPasteAtLocation = true;
-								const FSnappedPositioningData PositioningData = FSnappedPositioningData(GCurrentLevelEditingViewportClient, ClickLocation, ClickPlane)
-									.AlignToSurfaceRotation(false);
-								FVector SaveClickLocation = FActorPositioning::GetSnappedSurfaceAlignedTransform(PositioningData).GetLocation();
+						bPasteAtLocation = true;
+						const FSnappedPositioningData PositioningData = FSnappedPositioningData(GCurrentLevelEditingViewportClient, ClickLocation, ClickPlane)
+							.AlignToSurfaceRotation(false);
+						FVector SaveClickLocation = FActorPositioning::GetSnappedSurfaceAlignedTransform(PositioningData).GetLocation();
 								
-								PasteLocation = SaveClickLocation;
-								TransDescription = NSLOCTEXT("UnrealEd", "PasteHere", "Paste Here");
-							}
-							else
-							{
-								if (!FCString::Strcmp(*TempStr, TEXT("ORIGIN")))
-								{
-									bPasteAtLocation = false;
-									TransDescription = NSLOCTEXT("UnrealEd", "PasteToWorldOrigin", "Paste To World Origin");
-								}
-							}
-						}
-						
-						
-						const FScopedTransaction Transaction(TransDescription);
-						FTypedElementPasteOptions Options;
-						Options.SelectionSetToModify = SelectionSet;
-						Options.bPasteAtLocation = bPasteAtLocation;
-						Options.PasteLocation = PasteLocation;
-						TArray<FScriptTypedElementHandle> PastedElements = CommonActions->PasteElementsFromString(SelectionSet, InWorld, Options, ClipboardPaste);
-						bPastedViaCommonActions = !PastedElements.IsEmpty();
+						PasteLocation = SaveClickLocation;
+						TransDescription = NSLOCTEXT("UnrealEd", "PasteHere", "Paste Here");
 					}
 					else
 					{
-						// Grabbing clipboard contents can fail on some platforms if another program has the lock.
-						// Our HAL doesn't expose that though.  If the clipboard is empty then fail to go the legacy path (which internally will
-						// retry the clipboard - this MAY succeed and we don't want to go down that path if it does).
-						bPastedViaCommonActions = true;
+						if (!FCString::Strcmp(*TempStr, TEXT("ORIGIN")))
+						{
+							bPasteAtLocation = false;
+							TransDescription = NSLOCTEXT("UnrealEd", "PasteToWorldOrigin", "Paste To World Origin");
+						}
 					}
 				}
+						
+						
+				const FScopedTransaction Transaction(TransDescription);
+				FTypedElementPasteOptions Options;
+				Options.SelectionSetToModify = SelectionSet;
+				Options.bPasteAtLocation = bPasteAtLocation;
+				Options.PasteLocation = PasteLocation;
+				TArray<FScriptTypedElementHandle> PastedElements = CommonActions->PasteElementsFromString(SelectionSet, InWorld, Options, ClipboardPaste);
+				bPastedViaCommonActions = !PastedElements.IsEmpty();
+			}
+			else
+			{
+				// Grabbing clipboard contents can fail on some platforms if another program has the lock.
+				// Our HAL doesn't expose that though.  If the clipboard is empty then fail to go the legacy path (which internally will
+				// retry the clipboard - this MAY succeed and we don't want to go down that path if it does).
+				bPastedViaCommonActions = true;
 			}
 		}
 
@@ -1699,6 +1713,8 @@ bool UUnrealEdEngine::Exec_Edit( UWorld* InWorld, const TCHAR* Str, FOutputDevic
 
 bool UUnrealEdEngine::Exec_Pivot( const TCHAR* Str, FOutputDevice& Ar )
 {
+	using namespace UnrealEdSrvLocals;
+
 	if( FParse::Command(&Str,TEXT("HERE")) )
 	{
 		NoteActorMovement();
@@ -1724,48 +1740,56 @@ bool UUnrealEdEngine::Exec_Pivot( const TCHAR* Str, FOutputDevice& Ar )
 		int32 Count = 0;
 		FVector Center(0,0,0);
 
-		for ( FSelectionIterator It( GetSelectedActorIterator() ) ; It ; ++It )
+		if (UTypedElementSelectionSet* ElementSelectionSet = GetLevelEditorSelectionSet())
 		{
-			AActor* Actor = CastChecked<AActor>(*It);
-
-			if (ABrush* Brush = Cast<ABrush>(Actor))
+			ElementSelectionSet->ForEachSelectedElement<ITypedElementWorldInterface>(
+				[&Center, &Count](const TTypedElement<ITypedElementWorldInterface>& InElement)
 			{
-				// Treat brushes as a special case; calculate an effective position from the center point of the vertices.
-				// This way, "Center on Selection" has a special meaning for brushes.
-				TSet<FVector> UniqueVertices;
-				FVector VertexCenter = FVector::ZeroVector;
-
-				if (Brush->Brush && Brush->Brush->Polys)
+				// Special case legacy ABrush handling
+				if (ABrush* Brush = Cast<ABrush>(ActorElementDataUtil::GetActorFromHandle(InElement)))
 				{
-					for (const auto& Element : Brush->Brush->Polys->Element)
+					// Treat brushes as a special case; calculate an effective position from the center point of the vertices.
+					// This way, "Center on Selection" has a special meaning for brushes.
+					TSet<FVector> UniqueVertices;
+					FVector VertexCenter = FVector::ZeroVector;
+
+					if (Brush->Brush && Brush->Brush->Polys)
 					{
-						for (const auto& Vertex : Element.Vertices)
+						for (const auto& Poly : Brush->Brush->Polys->Element)
 						{
-							UniqueVertices.Add((FVector)Vertex);
+							for (const auto& Vertex : Poly.Vertices)
+							{
+								UniqueVertices.Add((FVector)Vertex);
+							}
 						}
-					}
 
-					for (const auto& Vertex : UniqueVertices)
-					{
-						VertexCenter += Vertex;
-					}
+						for (const auto& Vertex : UniqueVertices)
+						{
+							VertexCenter += Vertex;
+						}
 
-					if (UniqueVertices.Num() > 0)
-					{
-						VertexCenter /= UniqueVertices.Num();
+						if (UniqueVertices.Num() > 0)
+						{
+							VertexCenter /= UniqueVertices.Num();
+						}
 					}
+
+					Center += Brush->GetTransform().TransformPosition(VertexCenter);
+					++Count;
+					return true;
 				}
+				// If we got here, it's not a brush
 
-				Center += Brush->GetTransform().TransformPosition(VertexCenter);
-			}
-			else
-			{
-				Center += Actor->GetActorLocation();
-			}
+				FTransform WorldTransform;
+				if (InElement.GetWorldTransform(WorldTransform))
+				{
+					Center += WorldTransform.GetLocation();
+					++Count;
+				}
 
-			Count++;
+				return true;
+			});
 		}
-
 		if( Count > 0 )
 		{
 			FVector CenterLocation = Center / Count;
@@ -1978,6 +2002,8 @@ FPoly* CreateHugeTrianglePolygonOnPlane( const FPlane* InPlane )
 
 bool UUnrealEdEngine::Exec_Actor( UWorld* InWorld, const TCHAR* Str, FOutputDevice& Ar )
 {
+	using namespace UnrealEdSrvLocals;
+
 	// Keep a pointer to the beginning of the string to use for message displaying purposes
 	const TCHAR* const FullStr = Str;
 
@@ -2612,43 +2638,37 @@ bool UUnrealEdEngine::Exec_Actor( UWorld* InWorld, const TCHAR* Str, FOutputDevi
 	}
 	else if( FParse::Command(&Str,TEXT("DELETE")) )		// ACTOR SELECT DELETE
 	{
-		if (GCurrentLevelEditingViewportClient)
+		if (UTypedElementCommonActions* CommonActions = GetLevelEditorCommonActions())
 		{
-			if (TSharedPtr<ILevelEditor> LevelEditor = GCurrentLevelEditingViewportClient->ParentLevelEditor.Pin())
-			{
-				if (UTypedElementCommonActions* CommonActions = LevelEditor->GetCommonActions())
-				{
-					UTypedElementSelectionSet* SelectionSet = LevelEditor->GetMutableElementSelectionSet();
+			UTypedElementSelectionSet* SelectionSet = GetLevelEditorSelectionSet();
 				
-					// End drags to avoid deleting something out from under one.
-					FSlateApplication::Get().CancelDragDrop();
+			// End drags to avoid deleting something out from under one.
+			FSlateApplication::Get().CancelDragDrop();
 
-					const FScopedTransaction Transaction(NSLOCTEXT("UnrealEd", "DeleteElements", "Delete Elements"));
-					if (SelectionSet->GetNumSelectedElements() == 0)
-					{
-						// HACK: Not all modes will select elements, so allow them a shot at deletion if we don't think anything else is selected
-						// TODO: Move this logic into FLevelEditorActionCallbacks and have it call into the mode directly
-						if (!GLevelEditorModeTools().ProcessEditDelete())
-						{
-							// HACK: Call these directly for an empty selection so that folder deletion in the outliner still works
-							// TODO: Move this logic into FLevelEditorActionCallbacks and have it call into the outliner directly
-							FEditorDelegates::OnDeleteActorsBegin.Broadcast();
-							FEditorDelegates::OnDeleteActorsEnd.Broadcast();
-						}
-					}
-					else
-					{
-						FEditorDelegates::OnDeleteActorsBegin.Broadcast();
-						const bool bCheckRef = GetDefault<ULevelEditorMiscSettings>()->bCheckReferencesOnDelete;
-						FTypedElementDeletionOptions Options;
-						Options
-							.SetWarnAboutReferences(bCheckRef)
-							.SetWarnAboutSoftReferences(bCheckRef);
-						CommonActions->DeleteSelectedElements(SelectionSet, InWorld, Options);
-						FEditorDelegates::OnDeleteActorsEnd.Broadcast();
-					}
+			const FScopedTransaction Transaction(NSLOCTEXT("UnrealEd", "DeleteElements", "Delete Elements"));
+			if (!ensure(SelectionSet) || SelectionSet->GetNumSelectedElements() == 0)
+			{
+				// HACK: Not all modes will select elements, so allow them a shot at deletion if we don't think anything else is selected
+				// TODO: Move this logic into FLevelEditorActionCallbacks and have it call into the mode directly
+				if (!GLevelEditorModeTools().ProcessEditDelete())
+				{
+					// HACK: Call these directly for an empty selection so that folder deletion in the outliner still works
+					// TODO: Move this logic into FLevelEditorActionCallbacks and have it call into the outliner directly
+					FEditorDelegates::OnDeleteActorsBegin.Broadcast();
+					FEditorDelegates::OnDeleteActorsEnd.Broadcast();
 				}
 			}
+			else
+			{
+				FEditorDelegates::OnDeleteActorsBegin.Broadcast();
+				const bool bCheckRef = GetDefault<ULevelEditorMiscSettings>()->bCheckReferencesOnDelete;
+				FTypedElementDeletionOptions Options;
+				Options
+					.SetWarnAboutReferences(bCheckRef)
+					.SetWarnAboutSoftReferences(bCheckRef);
+				CommonActions->DeleteSelectedElements(SelectionSet, InWorld, Options);
+				FEditorDelegates::OnDeleteActorsEnd.Broadcast();
+			}
 		}
 		return true;
 	}
@@ -2685,49 +2705,52 @@ bool UUnrealEdEngine::Exec_Actor( UWorld* InWorld, const TCHAR* Str, FOutputDevi
 	}
 	else if( FParse::Command(&Str,TEXT("BAKEPREPIVOT")) )
 	{
-		FScopedLevelDirtied				LevelDirtyCallback;
-		FScopedActorPropertiesChange	ActorPropertiesChangeCallback;
-
-		// Bakes the current pivot position into all selected actors
-
-		FEditorModeTools& EditorModeTools = GLevelEditorModeTools();
-
-		for ( FSelectionIterator It( GetSelectedActorIterator() ) ; It ; ++It )
+		UTypedElementSelectionSet* ElementSelectionSet = GetLevelEditorSelectionSet();
+		if (ensure(ElementSelectionSet))
 		{
-			AActor* Actor = static_cast<AActor*>( *It );
-			checkSlow( Actor->IsA(AActor::StaticClass()) );
+			FScopedLevelDirtied				LevelDirtyCallback;
+			FScopedActorPropertiesChange	ActorPropertiesChangeCallback;
+
+			// Bakes the current pivot position into all selected elements
 
-			FVector Delta( EditorModeTools.PivotLocation - Actor->GetActorLocation() );
+			FEditorModeTools& EditorModeTools = GLevelEditorModeTools();
 
-			Actor->Modify();
-			Actor->SetPivotOffset(Actor->GetTransform().InverseTransformVector(Delta));
+			ElementSelectionSet->ForEachSelectedElement<ITypedElementWorldInterface>(
+				[&EditorModeTools](const TTypedElement<ITypedElementWorldInterface>& InElement)
+			{
+				FTransform WorldTransform;
+				if (InElement.GetWorldTransform(WorldTransform))
+				{
+					FVector Delta(EditorModeTools.PivotLocation - WorldTransform.GetLocation());
+					InElement.SetPivotOffset(WorldTransform.InverseTransformVector(Delta));
+				}
+				return true; // continue
+			});
 			SetPivotMovedIndependently(false);
-			Actor->PostEditMove(true);
-		}
 
-		GUnrealEd->NoteSelectionChange();
+			GUnrealEd->NoteSelectionChange();
+		}
 	} 
 	else if( FParse::Command(&Str,TEXT("UNBAKEPREPIVOT")) )
 	{
-		FScopedLevelDirtied		LevelDirtyCallback;
-		FScopedActorPropertiesChange	ActorPropertiesChangeCallback;
-
-		// Resets the PrePivot of the selected actors to 0,0,0 while leaving them in the same world location.
-
-		FEditorModeTools& EditorModeTools = GLevelEditorModeTools();
-
-		for ( FSelectionIterator It( GetSelectedActorIterator() ) ; It ; ++It )
+		UTypedElementSelectionSet* ElementSelectionSet = GetLevelEditorSelectionSet();
+		if (ensure(ElementSelectionSet))
 		{
-			AActor* Actor = static_cast<AActor*>( *It );
-			checkSlow( Actor->IsA(AActor::StaticClass()) );
+			FScopedLevelDirtied		LevelDirtyCallback;
+			FScopedActorPropertiesChange	ActorPropertiesChangeCallback;
 
-			Actor->Modify();
-			Actor->SetPivotOffset(FVector::ZeroVector);
+			// Resets the pivot offset of the selected elements to zero while leaving them in the same world location.
+
+			ElementSelectionSet->ForEachSelectedElement<ITypedElementWorldInterface>(
+			[](const TTypedElement<ITypedElementWorldInterface>& InElement)
+			{
+				InElement.SetPivotOffset(FVector::ZeroVector);
+				return true; // continue
+			});
 			SetPivotMovedIndependently(false);
-			Actor->PostEditMove(true);
-		}
 
-		GUnrealEd->NoteSelectionChange();
+			GUnrealEd->NoteSelectionChange();
+		}
 	}
 	else if( FParse::Command(&Str,TEXT("RESET")) )
 	{
@@ -2818,41 +2841,35 @@ bool UUnrealEdEngine::Exec_Actor( UWorld* InWorld, const TCHAR* Str, FOutputDevi
 	}
 	else if( FParse::Command(&Str,TEXT("DUPLICATE")) )
 	{
-		if (GCurrentLevelEditingViewportClient)
+		if (UTypedElementCommonActions* CommonActions = GetLevelEditorCommonActions())
 		{
-			if (TSharedPtr<ILevelEditor> LevelEditor = GCurrentLevelEditingViewportClient->ParentLevelEditor.Pin())
+			UTypedElementSelectionSet* SelectionSet = GetLevelEditorSelectionSet();
+
+			const FScopedTransaction Transaction(NSLOCTEXT("UnrealEd", "DuplicateElements", "Duplicate Elements"));
+			if (!ensure(SelectionSet) || SelectionSet->GetNumSelectedElements() == 0)
 			{
-				if (UTypedElementCommonActions* CommonActions = LevelEditor->GetCommonActions())
+				// HACK: Not all modes will select elements, so allow them a shot at duplication if we don't think anything else is selected
+				// TODO: Move this logic into FLevelEditorActionCallbacks and have it call into the mode directly
+				if (!GLevelEditorModeTools().ProcessEditDuplicate())
 				{
-					UTypedElementSelectionSet* SelectionSet = LevelEditor->GetMutableElementSelectionSet();
-
-					const FScopedTransaction Transaction(NSLOCTEXT("UnrealEd", "DuplicateElements", "Duplicate Elements"));
-					if (SelectionSet->GetNumSelectedElements() == 0)
-					{
-						// HACK: Not all modes will select elements, so allow them a shot at duplication if we don't think anything else is selected
-						// TODO: Move this logic into FLevelEditorActionCallbacks and have it call into the mode directly
-						if (!GLevelEditorModeTools().ProcessEditDuplicate())
-						{
-							// HACK: Call these directly for an empty selection so that folder duplication in the outliner still works
-							// TODO: Move this logic into FLevelEditorActionCallbacks and have it call into the outliner directly
-							FEditorDelegates::OnDuplicateActorsBegin.Broadcast();
-							FEditorDelegates::OnDuplicateActorsEnd.Broadcast();
-						}
-					}
-					else
-					{
-						const FVector DuplicateOffset = GEditor->GetGridLocationOffset(/*bUniformOffset*/false);
-						const TArray<FTypedElementHandle> DuplicatedElements = CommonActions->DuplicateSelectedElements(SelectionSet, InWorld, DuplicateOffset);
-						if (DuplicatedElements.Num() > 0)
-						{
-							SelectionSet->SetSelection(DuplicatedElements, FTypedElementSelectionOptions());
-							SelectionSet->NotifyPendingChanges();
+					// HACK: Call these directly for an empty selection so that folder duplication in the outliner still works
+					// TODO: Move this logic into FLevelEditorActionCallbacks and have it call into the outliner directly
+					FEditorDelegates::OnDuplicateActorsBegin.Broadcast();
+					FEditorDelegates::OnDuplicateActorsEnd.Broadcast();
+				}
+			}
+			else
+			{
+				const FVector DuplicateOffset = GEditor->GetGridLocationOffset(/*bUniformOffset*/false);
+				const TArray<FTypedElementHandle> DuplicatedElements = CommonActions->DuplicateSelectedElements(SelectionSet, InWorld, DuplicateOffset);
+				if (DuplicatedElements.Num() > 0)
+				{
+					SelectionSet->SetSelection(DuplicatedElements, FTypedElementSelectionOptions());
+					SelectionSet->NotifyPendingChanges();
 
-							// notify the global mode tools, the selection set should be identical to the new actors at this point
-							TArray<AActor*> SelectedActors = SelectionSet->GetSelectedObjects<AActor>();
-							GLevelEditorModeTools().ActorsDuplicatedNotify(SelectedActors, SelectedActors, DuplicateOffset != FVector::ZeroVector);
-						}
-					}
+					// notify the global mode tools, the selection set should be identical to the new actors at this point
+					TArray<AActor*> SelectedActors = SelectionSet->GetSelectedObjects<AActor>();
+					GLevelEditorModeTools().ActorsDuplicatedNotify(SelectedActors, SelectedActors, DuplicateOffset != FVector::ZeroVector);
 				}
 			}
 		}
