LevelEditor: Enable visibility context menu actions on non-actors.

#Synced-CL 52260419

#Synced-CL 52260419

diff --git a/Engine/Source/Editor/LevelEditor/Private/LevelEditorActions.cpp b/Engine/Source/Editor/LevelEditor/Private/LevelEditorActions.cpp
--- a/Engine/Source/Editor/LevelEditor/Private/LevelEditorActions.cpp
+++ b/Engine/Source/Editor/LevelEditor/Private/LevelEditorActions.cpp
@@ -6,6 +6,12 @@
 #include "Factories/Factory.h"
 #include "Animation/AnimSequence.h"
 #include "Components/LightComponent.h"
+#include "DataStorage/Features.h"
+#include "DataStorage/Handles.h"
+#include "Elements/Columns/TypedElementCompatibilityColumns.h" // FTypedElementActorTag
+#include "Elements/Columns/TypedElementSelectionColumns.h"
+#include "Elements/Columns/TypedElementVisibilityColumns.h"
+#include "Elements/Framework/TypedElementQueryBuilder.h"
 #include "Engine/SkeletalMesh.h"
 #include "HAL/PlatformFileManager.h"
 #include "Misc/MessageDialog.h"
@@ -195,6 +201,90 @@ namespace LevelEditorActionsHelpers
 		}
 		return false;
 	}
+
+	// Call VisitFn on each selected non-actor with a FVisibleInEditorColumn. Returns true if completed fully.
+	bool ForEachSelectedHideableNonActor(UWorld* WorldIn,
+		TFunctionRef<void(UE::Editor::DataStorage::Queries::EFlowControl& Flow, const FVisibleInEditorColumn& VisibilityColumn)> VisitFn)
+	{
+		using namespace UE::Editor::DataStorage;
+		using namespace UE::Editor::DataStorage::Queries;
+
+		ICoreProvider* DataStorage = GetMutableDataStorageFeature<ICoreProvider>(StorageFeatureName);
+		if (!DataStorage)
+		{
+			return false;
+		}
+		// Note: this query currently never gets unregistered because it is constantly used in CanExecute calls
+		//  while the visibility submenu is open, and we can't unregister on module shutdown because the level
+		//  editor isn't really a module that would be deactivated. 
+		static QueryHandle Query = DataStorage->RegisterQuery(
+			Select()
+			.ReadOnly<FTypedElementWorldColumn, FVisibleInEditorColumn>()
+			.Where()
+			.None<FTypedElementActorTag>()
+			.All<FTypedElementSelectionColumn>()
+			.Compile());
+
+		FQueryResult QueryResult = DataStorage->RunQuery(Query, ERunQueryFlags::Default,  
+			[WorldIn, &VisitFn](EFlowControl& Flow, const FTypedElementWorldColumn& WorldColumn, const FVisibleInEditorColumn& VisibilityColumn)
+				{
+					if (WorldColumn.World.Get() != WorldIn)
+					{
+						return;
+					}
+					VisitFn(Flow, VisibilityColumn);
+				});
+
+		return QueryResult.Completed == FQueryResult::ECompletion::Fully;
+	}
+
+	bool HaveSelectedVisibleNonActors(UWorld* WorldIn)
+	{
+		using namespace UE::Editor::DataStorage;
+		using namespace UE::Editor::DataStorage::Queries;
+
+		bool bToReturn = false;
+		bool bSuccess = ForEachSelectedHideableNonActor(WorldIn, [&bToReturn](EFlowControl& Flow, const FVisibleInEditorColumn& VisibilityColumn)
+		{
+			if (VisibilityColumn.bIsVisibleInEditor)
+			{
+				bToReturn = true;
+				Flow = EFlowControl::Break;
+			}
+		});
+		return bSuccess && bToReturn;
+	}
+
+	bool HaveSelectedHiddenNonActors(UWorld* WorldIn)
+	{
+		using namespace UE::Editor::DataStorage;
+		using namespace UE::Editor::DataStorage::Queries;
+
+		bool bToReturn = false;
+		bool bSuccess = ForEachSelectedHideableNonActor(WorldIn, [&bToReturn](EFlowControl& Flow, const FVisibleInEditorColumn& VisibilityColumn)
+		{
+			if (!VisibilityColumn.bIsVisibleInEditor)
+			{
+				bToReturn = true;
+				Flow = EFlowControl::Break;
+			}
+		});
+		return bSuccess && bToReturn;
+	}
+
+	bool HaveSelectedHideableNonActors(UWorld* WorldIn)
+	{
+		using namespace UE::Editor::DataStorage;
+		using namespace UE::Editor::DataStorage::Queries;
+
+		bool bToReturn = false;
+		bool bSuccess = ForEachSelectedHideableNonActor(WorldIn, [&bToReturn](EFlowControl& Flow, const FVisibleInEditorColumn& VisibilityColumn)
+		{
+			bToReturn = true;
+			Flow = EFlowControl::Break;
+		});
+		return bSuccess && bToReturn;
+	}
 }
 
 bool FLevelEditorActionCallbacks::DefaultCanExecuteAction()
@@ -2384,7 +2474,8 @@ bool FLevelEditorActionCallbacks::CanHideSelectedActors()
 			return true;
 		}
 	}
-	return false;
+
+	return LevelEditorActionsHelpers::HaveSelectedVisibleNonActors(GetWorld());
 }
 
 bool FLevelEditorActionCallbacks::CanShowSelectedActors()
@@ -2397,12 +2488,14 @@ bool FLevelEditorActionCallbacks::CanShowSelectedActors()
 			return true;
 		}
 	}
-	return false;
+
+	return LevelEditorActionsHelpers::HaveSelectedHiddenNonActors(GetWorld());
 }
 
 bool FLevelEditorActionCallbacks::CanToggleSelectedActorsVisibility()
 {
-	return GEditor->GetSelectedActorCount() > 0;
+	return GEditor->GetSelectedActorCount() > 0
+		|| LevelEditorActionsHelpers::HaveSelectedHideableNonActors(GetWorld());
 }
 
 void FLevelEditorActionCallbacks::OnSelectOwningHLODCluster()
@@ -4483,17 +4576,17 @@ void FLevelEditorCommands::RegisterCommands()
 	UI_COMMAND( FixupGroupActor, "Fixup Group Actor", "Removes null actors and deletes the GroupActor if it is empty.", EUserInterfaceActionType::Button, FInputChord());
 
 #if PLATFORM_MAC
-	UI_COMMAND( ShowAll, "Show All Actors", "Shows all actors", EUserInterfaceActionType::Button, FInputChord( EModifierKey::Command, EKeys::H ) );
+	UI_COMMAND( ShowAll, "Show All", "Shows all items", EUserInterfaceActionType::Button, FInputChord( EModifierKey::Command, EKeys::H ) );
 #else
-	UI_COMMAND( ShowAll, "Show All Actors", "Shows all actors", EUserInterfaceActionType::Button, FInputChord( EModifierKey::Control, EKeys::H ) );
+	UI_COMMAND( ShowAll, "Show All", "Shows all items", EUserInterfaceActionType::Button, FInputChord( EModifierKey::Control, EKeys::H ) );
 #endif
-	UI_COMMAND( ShowSelectedOnly, "Show Only Selected", "Shows only the selected actors", EUserInterfaceActionType::Button, FInputChord() );
-	UI_COMMAND( ShowSelected, "Show Selected", "Shows the selected actors", EUserInterfaceActionType::Button, FInputChord() );
-	UI_COMMAND( HideSelected, "Hide Selected", "Hides the selected actors", EUserInterfaceActionType::Button, FInputChord() );
-	UI_COMMAND( ToggleSelectedVisibility, "Toggle Selected Visibility", "Hides or shows the selected actors", EUserInterfaceActionType::Button, FInputChord() );
-	UI_COMMAND( ShowSelectedHierarchy, "Show Selected Hierarchy", "Shows the selected actors and their children", EUserInterfaceActionType::Button, FInputChord( EModifierKey::Shift, EKeys::H ) );
-	UI_COMMAND( HideSelectedHierarchy, "Hide Selected Hierarchy", "Hides the selected actors and their children", EUserInterfaceActionType::Button, FInputChord() );
-	UI_COMMAND( ToggleSelectedHierarchyVisibility, "Toggle Selected Hierarchy Visibility", "Hides or shows the selected actors and their children", EUserInterfaceActionType::Button, FInputChord( EKeys::H ) );
+	UI_COMMAND( ShowSelectedOnly, "Show Only Selected", "Shows only the selected items", EUserInterfaceActionType::Button, FInputChord() );
+	UI_COMMAND( ShowSelected, "Show Selected", "Shows the selected items", EUserInterfaceActionType::Button, FInputChord() );
+	UI_COMMAND( HideSelected, "Hide Selected", "Hides the selected items", EUserInterfaceActionType::Button, FInputChord() );
+	UI_COMMAND( ToggleSelectedVisibility, "Toggle Selected Visibility", "Hides or shows the selected items", EUserInterfaceActionType::Button, FInputChord() );
+	UI_COMMAND( ShowSelectedHierarchy, "Show Selected Hierarchy", "Shows the selected items and their children", EUserInterfaceActionType::Button, FInputChord( EModifierKey::Shift, EKeys::H ) );
+	UI_COMMAND( HideSelectedHierarchy, "Hide Selected Hierarchy", "Hides the selected items and their children", EUserInterfaceActionType::Button, FInputChord() );
+	UI_COMMAND( ToggleSelectedHierarchyVisibility, "Toggle Selected Hierarchy Visibility", "Hides or shows the selected items and their children", EUserInterfaceActionType::Button, FInputChord( EKeys::H ) );
 	UI_COMMAND( ShowAllStartup, "Show All At Startup", "Shows all actors at startup", EUserInterfaceActionType::Button, FInputChord() );
 	UI_COMMAND( ShowSelectedStartup, "Show Selected At Startup", "Shows selected actors at startup", EUserInterfaceActionType::Button, FInputChord() );
 	UI_COMMAND( HideSelectedStartup, "Hide Selected At Startup", "Hide selected actors at startup", EUserInterfaceActionType::Button, FInputChord() );
diff --git a/Engine/Source/Editor/LevelEditor/Private/LevelEditorContextMenu.cpp b/Engine/Source/Editor/LevelEditor/Private/LevelEditorContextMenu.cpp
--- a/Engine/Source/Editor/LevelEditor/Private/LevelEditorContextMenu.cpp
+++ b/Engine/Source/Editor/LevelEditor/Private/LevelEditorContextMenu.cpp
@@ -73,6 +73,9 @@ public:
 	 */
 	static void FillActorVisibilityMenu(UToolMenu* Menu);
 
+	/** Fills in menu options for typed element visibility menu */
+	static void FillElementVisibilityMenu(UToolMenu* Menu);
+
 	/**
 	 * Fills in menu options for the actor level menu
 	 *
@@ -703,14 +706,13 @@ void FLevelEditorContextMenu::RegisterElementContextMenu()
 				false, // default value
 				FSlateIcon(FAppStyle::GetAppStyleSetName(), "Icons.Edit"));
 
-			// TODO: The visibility actions don't yet properly support elements, but should
-			//Section.AddSubMenu(
-			//	"VisibilitySubMenu",
-			//	LOCTEXT("VisibilityElementSubMenu", "Visibility"),
-			//	LOCTEXT("VisibilityElementSubMenu_ToolTip", "Selected element visibility options"),
-			//	FNewToolMenuDelegate::CreateStatic(&FLevelEditorContextMenuImpl::FillActorVisibilityMenu),
-			//	false, // default value
-			//	FSlateIcon(FAppStyle::GetAppStyleSetName(), "Icons.Visibility"));
+			Section.AddSubMenu(
+				"VisibilitySubMenu",
+				LOCTEXT("VisibilityElementSubMenu", "Visibility"),
+				LOCTEXT("VisibilityElementSubMenu_ToolTip", "Selected element visibility options"),
+				FNewToolMenuDelegate::CreateStatic(&FLevelEditorContextMenuImpl::FillElementVisibilityMenu),
+				false, // default value
+				FSlateIcon(FAppStyle::GetAppStyleSetName(), "Icons.Visibility"));
 
 			Section.AddSubMenu(
 				"TransformSubMenu",
@@ -1236,7 +1238,7 @@ void FLevelEditorContextMenu::BuildGroupMenu(FToolMenuSection& Section, const FS
 	}
 }
 
-void FLevelEditorContextMenuImpl::FillActorVisibilityMenu(UToolMenu* Menu)
+void FLevelEditorContextMenuImpl::FillElementVisibilityMenu(UToolMenu* Menu)
 {
 	{
 		FToolMenuSection& Section = Menu->AddSection("VisibilitySelected", LOCTEXT("VisibilitySelectedHeader", "Selection"));
@@ -1253,6 +1255,11 @@ void FLevelEditorContextMenuImpl::FillActorVisibilityMenu(UToolMenu* Menu)
 		Section.AddMenuEntry(FLevelEditorCommands::Get().ShowSelectedOnly);
 		Section.AddMenuEntry(FLevelEditorCommands::Get().ShowAll);
 	}
+}
+
+void FLevelEditorContextMenuImpl::FillActorVisibilityMenu(UToolMenu* Menu)
+{
+	FillElementVisibilityMenu(Menu);
 
 	{
 		FToolMenuSection& Section = Menu->AddSection("VisibilityStartup", LOCTEXT("VisibilityStartupHeader", "Startup"));
diff --git a/Engine/Source/Editor/UnrealEd/Private/EditorActor.cpp b/Engine/Source/Editor/UnrealEd/Private/EditorActor.cpp
--- a/Engine/Source/Editor/UnrealEd/Private/EditorActor.cpp
+++ b/Engine/Source/Editor/UnrealEd/Private/EditorActor.cpp
@@ -2,6 +2,7 @@
 
 
 #include "CoreMinimal.h"
+#include "Misc/ITransaction.h" // GUndo
 #include "Misc/MessageDialog.h"
 #include "Modules/ModuleManager.h"
 #include "UObject/ObjectMacros.h"
@@ -79,13 +80,247 @@
 #include "SDeleteReferencedActorDialog.h"
 #include "DataStorage/Features.h"
 #include "Elements/Interfaces/TypedElementDataStorageCompatibilityInterface.h"
+#include "Elements/Interfaces/TypedElementDataStorageInterface.h"
+#include "Elements/Columns/TypedElementCompatibilityColumns.h"
+#include "Elements/Columns/TypedElementSelectionColumns.h"
 #include "Elements/Columns/TypedElementVisibilityColumns.h"
 #include "Elements/Columns/TypedElementMiscColumns.h"
+#include "Elements/Framework/TypedElementQueryBuilder.h"
 
 #define LOCTEXT_NAMESPACE "UnrealEd.EditorActor"
 
 DEFINE_LOG_CATEGORY_STATIC(LogEditorActor, Log, All);
 
+namespace UE::EditorActor::Private
+{
+	using namespace UE::Editor::DataStorage;
+	using namespace UE::Editor::DataStorage::Queries;
+
+	/** 
+	 * Takes a query that provides readonly world and visibility columns, and sets the visibility on the returned columns 
+	 *  if the world matches. Handles undo if a transaction is open.
+	 */ 
+	void RunSetVisibilityQuery(ICoreProvider& DataStorage, QueryHandle Query, UWorld* WorldIn, bool bVisible)
+	{
+		// We'll gather the rows to change in one pass and actually change them in another, to
+		//  avoid trying to add a tag (FVisibilityChangedTag) from within a callback 
+		FRowHandleArray RowsToChange;
+		DataStorage.RunQuery(Query, CreateDirectQueryCallbackBinding(
+			[WorldIn, bVisible, &RowsToChange](RowHandle Row, const FTypedElementWorldColumn& WorldColumn, const FVisibleInEditorColumn& VisibilityColumn)
+		{
+			if (WorldColumn.World.Get() != WorldIn)
+			{
+				return;
+			}
+			if (VisibilityColumn.bIsVisibleInEditor != bVisible)
+			{
+				RowsToChange.Add(Row);
+			}
+		}));
+
+		// Apply the change
+		for (RowHandle Row : RowsToChange.GetRows())
+		{
+			FVisibleInEditorColumn* VisibilityColumn = DataStorage.GetColumn<FVisibleInEditorColumn>(Row);
+			if (ensure(VisibilityColumn))
+			{
+				VisibilityColumn->bIsVisibleInEditor = bVisible;
+				DataStorage.AddColumn<FVisibilityChangedTag>(Row);
+			}
+		}
+
+		// Emit an undo change if a transaction is open
+		if (GUndo && WorldIn)
+		{
+			GUndo->StoreUndo(WorldIn, MakeUnique<FVisibilityColumnChange>(MoveTemp(RowsToChange), bVisible));
+		}
+	}
+
+	/**
+	 * Walks across editor object hierarchies for all selected items and calls VisitFn
+	 */
+	void WalkSelectedHierarchies(ICoreProvider& DataStorage, UWorld* WorldIn, 
+		TFunction<void(const ICoreProvider& Context, RowHandle Owner, RowHandle Target)> VisitFn)
+	{
+		// Gather selected rows with matching world
+		QueryHandle Query = DataStorage.RegisterQuery(
+			Select()
+			.ReadOnly<FTypedElementWorldColumn>()
+			.Where()
+			.None<FTypedElementActorTag>()
+			.All<FTypedElementSelectionColumn>()
+			.Compile());
+		ON_SCOPE_EXIT{ DataStorage.UnregisterQuery(Query); };
+
+		TArray<RowHandle> SelectedRows;
+		DataStorage.RunQuery(Query, CreateDirectQueryCallbackBinding(
+			[WorldIn, &DataStorage, &SelectedRows](RowHandle Row, const FTypedElementWorldColumn& WorldColumn)
+			{
+				if (WorldColumn.World.Get() != WorldIn)
+				{
+					return;
+				}
+				SelectedRows.Add(Row);
+			}));
+
+		// TODO: This should match UTedsEditorHierarchyFactory::EditorObjectHierarchyName when that is not experimental
+		FName EditorObjectHierarchyName("EditorObjectHierarchy");
+		FHierarchyHandle Hierarchy = DataStorage.FindHierarchyByName(EditorObjectHierarchyName);
+		if (!ensure(DataStorage.IsValidHierarchyHandle(Hierarchy)))
+		{
+			return;
+		}
+
+		// Walk the hierarchies for the selected rows
+		for (RowHandle Row : SelectedRows)
+		{
+			DataStorage.WalkDepthFirst(Hierarchy, Row, VisitFn);
+		}
+	}
+
+	void SetSelectedNonActorVisibility(UWorld* WorldIn, bool bVisible, bool bApplyToHierarchy)
+	{
+		ICoreProvider* DataStorage = GetMutableDataStorageFeature<ICoreProvider>(StorageFeatureName);
+		if (!DataStorage)
+		{
+			return;
+		}
+
+		if (!bApplyToHierarchy)
+		{
+			// This is the simpler case- just change the selected items.
+
+			// Query gets world and visibility columns for all selected non-actors
+			QueryHandle Query = DataStorage->RegisterQuery(
+				Select()
+				.ReadOnly<FTypedElementWorldColumn, FVisibleInEditorColumn>()
+				.Where()
+				.None<FTypedElementActorTag>()
+				.All<FTypedElementSelectionColumn>()
+				.Compile());
+			ON_SCOPE_EXIT{ DataStorage->UnregisterQuery(Query); };
+
+			RunSetVisibilityQuery(*DataStorage, Query, WorldIn, bVisible);
+			return;
+		}
+
+		// If we got here, we're processing the hierarchy.
+
+		FRowHandleArray RowsToChange;
+		WalkSelectedHierarchies(*DataStorage, WorldIn, [&RowsToChange, bVisible](const ICoreProvider& Context, RowHandle Owner, RowHandle Target)
+		{
+			if (const FVisibleInEditorColumn* VisibilityColumn = Context.GetColumn<FVisibleInEditorColumn>(Target))
+			{
+				if (VisibilityColumn->bIsVisibleInEditor != bVisible)
+				{
+					RowsToChange.Add(Target);
+				}
+			}
+		});
+
+		// Apply the change
+		for (RowHandle Row : RowsToChange.GetRows())
+		{
+			FVisibleInEditorColumn* VisibilityColumn = DataStorage->GetColumn<FVisibleInEditorColumn>(Row);
+			if (ensure(VisibilityColumn))
+			{
+				VisibilityColumn->bIsVisibleInEditor = bVisible;
+				DataStorage->AddColumn<FVisibilityChangedTag>(Row);
+			}
+		}
+
+		// Emit an undo change if a transaction is open
+		if (GUndo && WorldIn)
+		{
+			GUndo->StoreUndo(WorldIn, MakeUnique<FVisibilityColumnChange>(MoveTemp(RowsToChange), bVisible));
+		}
+	}
+
+	void SetAllNonActorVisiblity(UWorld* WorldIn, bool bVisible)
+	{
+		ICoreProvider* DataStorage = GetMutableDataStorageFeature<ICoreProvider>(StorageFeatureName);
+		if (!DataStorage)
+		{
+			return;
+		}
+		
+		QueryHandle Query = DataStorage->RegisterQuery(
+			Select()
+			.ReadOnly<FTypedElementWorldColumn, FVisibleInEditorColumn>()
+			.Where()
+			.None<FTypedElementActorTag>()
+			.Compile());
+		ON_SCOPE_EXIT{ DataStorage->UnregisterQuery(Query); };
+		RunSetVisibilityQuery(*DataStorage, Query, WorldIn, bVisible);
+	}
+
+	void SetUnselectedNonActorVisibility(UWorld* WorldIn, bool bVisible)
+	{
+		ICoreProvider* DataStorage = GetMutableDataStorageFeature<ICoreProvider>(StorageFeatureName);
+		if (!DataStorage)
+		{
+			return;
+		}
+
+		// Note that it is insufficient to just set the visibility of unselected items, even
+		//  if the selected ones are already visible, because in some cases the visibility handling
+		//  needs to propagate the visibility to some children that might otherwise remain invisible,
+		//  as they are not directly selected. For this reason, we go with this two step approach.
+
+		SetAllNonActorVisiblity(WorldIn, false);
+		SetSelectedNonActorVisibility(WorldIn, true, 
+			// This seems to be the current behavior for actors, so we match it
+			/*bApplyToHierarchy*/ false);
+	}
+
+	bool HaveAnySelectedVisibleNonActor(UWorld* WorldIn, bool bIncludeHierarchy)
+	{
+		ICoreProvider* DataStorage = GetMutableDataStorageFeature<ICoreProvider>(UE::Editor::DataStorage::StorageFeatureName);
+		if (!DataStorage)
+		{
+			return false;
+		}
+
+		bool bHaveVisible = false;
+		if (!bIncludeHierarchy)
+		{
+			// Query gets world and visibility columns for all selected non-actors
+			QueryHandle Query = DataStorage->RegisterQuery(
+				Select()
+				.ReadOnly<FTypedElementWorldColumn, FVisibleInEditorColumn>()
+				.Where()
+				.None<FTypedElementActorTag>()
+				.All<FTypedElementSelectionColumn>()
+				.Compile());
+			ON_SCOPE_EXIT{ DataStorage->UnregisterQuery(Query); };
+
+			DataStorage->RunQuery(Query, CreateDirectQueryCallbackBinding(
+				[WorldIn, &bHaveVisible](RowHandle Row, const FTypedElementWorldColumn& WorldColumn, const FVisibleInEditorColumn& VisibilityColumn)
+			{
+				if (WorldColumn.World.Get() != WorldIn)
+				{
+					return;
+				}
+
+				bHaveVisible = bHaveVisible || VisibilityColumn.bIsVisibleInEditor;
+			}));
+		}
+		else
+		{
+			WalkSelectedHierarchies(*DataStorage, WorldIn, [&bHaveVisible](const ICoreProvider& Context, RowHandle Owner, RowHandle Target)
+			{
+				if (const FVisibleInEditorColumn* VisibilityColumn = Context.GetColumn<FVisibleInEditorColumn>(Target))
+				{
+					bHaveVisible = bHaveVisible || VisibilityColumn->bIsVisibleInEditor;
+				}
+			});
+		}
+
+		return bHaveVisible;
+	}
+} // namespace UE::EditorActor::Private
+
+
 struct FGarbageEliminationScopeGuard
 {
 	bool PreviousValue;
@@ -1596,6 +1831,10 @@ void UUnrealEdEngine::Internal_HideSelected( UWorld* InWorld, bool bApplyToChild
 			}
 		}
 	}
+
+	// Hide selected non-actor elements via TEDS
+	UE::EditorActor::Private::SetSelectedNonActorVisibility(InWorld, false, bApplyToChildHierarchy);
+
 	RedrawLevelEditingViewports();
 }
 
@@ -1640,6 +1879,10 @@ void UUnrealEdEngine::edactHideUnselected( UWorld* InWorld )
 			}
 		}
 	}
+
+	// Hide unselected non-actor elements via TEDS
+	UE::EditorActor::Private::SetUnselectedNonActorVisibility(InWorld, false);
+
 	RedrawLevelEditingViewports();
 }
 
@@ -1684,6 +1927,12 @@ void UUnrealEdEngine::Internal_ToggleSelectedVisibility(UWorld* InWorld, bool bA
 			}
 		}
 	}
+
+	// Check non-actor element visibility to include in the toggle decision
+	if (bTargetVisibility)
+	{
+		bTargetVisibility = !UE::EditorActor::Private::HaveAnySelectedVisibleNonActor(InWorld, bApplyToChildHierarchy);
+	}
 	
 	// Apply the target visibility to the target actors
 	if ( ActorsToConsider.Num() > 0 )
@@ -1715,6 +1964,10 @@ void UUnrealEdEngine::Internal_ToggleSelectedVisibility(UWorld* InWorld, bool bA
 			}
 		}
 	}
+
+	// Apply the target visibility to non-actor elements
+	UE::EditorActor::Private::SetSelectedNonActorVisibility(InWorld, bTargetVisibility, bApplyToChildHierarchy);
+	
 	RedrawLevelEditingViewports();
 }
 
@@ -1757,6 +2010,10 @@ void UUnrealEdEngine::edactUnHideAll( UWorld* InWorld )
 			}
 		}
 	}
+
+	// Unhide all non-actor elements via TEDS
+	UE::EditorActor::Private::SetAllNonActorVisiblity(InWorld, true);
+
 	RedrawLevelEditingViewports();
 }
 
@@ -1988,6 +2245,10 @@ void UUnrealEdEngine::Internal_UnhideSelected(UWorld* InWorld, bool bApplyToChil
 			}
 		}
 	}
+
+	// Show selected non-actor elements via TEDS
+	UE::EditorActor::Private::SetSelectedNonActorVisibility(InWorld, true, bApplyToChildHierarchy);
+
 	RedrawLevelEditingViewports();
 }
 
