[Editor ITF] ViewportInteractions - allow both canvas and PDI drawing

#Synced-CL 52260466

diff --git a/Engine/Plugins/Animation/SkeletalMeshModelingTools/Source/SkeletalMeshModelingTools/Private/ViewportInteractions/SKMModelingToolsGeometrySoftSelectRadiusDrag.cpp b/Engine/Plugins/Animation/SkeletalMeshModelingTools/Source/SkeletalMeshModelingTools/Private/ViewportInteractions/SKMModelingToolsGeometrySoftSelectRadiusDrag.cpp
--- a/Engine/Plugins/Animation/SkeletalMeshModelingTools/Source/SkeletalMeshModelingTools/Private/ViewportInteractions/SKMModelingToolsGeometrySoftSelectRadiusDrag.cpp
+++ b/Engine/Plugins/Animation/SkeletalMeshModelingTools/Source/SkeletalMeshModelingTools/Private/ViewportInteractions/SKMModelingToolsGeometrySoftSelectRadiusDrag.cpp
@@ -31,18 +31,14 @@ void USkeletalMeshModelingToolsGeometrySoftSelectRadiusDragInteraction::BindSele
 	SelectionManager = InSelectionManager;
 }
 
-void USkeletalMeshModelingToolsGeometrySoftSelectRadiusDragInteraction::Render(
-	const FSceneView* View,
-	FCanvas* Canvas,
-    EViewInteractionState InteractionState
-    )
+void USkeletalMeshModelingToolsGeometrySoftSelectRadiusDragInteraction::Draw(FCanvas* InCanvas, IToolsContextRenderAPI* InRenderAPI)
 {
-	if (!IsDragging() || InteractionState != EViewInteractionState::Focused)
+	if (!ensure(InRenderAPI) || !ensure(InCanvas))
 	{
 		return;
 	}
 
-	if (!Canvas || !View)
+	if (!IsDragging() || InRenderAPI->GetViewInteractionState() != EViewInteractionState::Focused)
 	{
 		return;
 	}
@@ -56,18 +52,23 @@ void USkeletalMeshModelingToolsGeometrySoftSelectRadiusDragInteraction::Render(
 	SelectionManager->GetSelectionWorldFrame(SelectionFrame);
 
 	double Radius = SelectionManager->GetSoftSelectionRadius();
+
+	const FSceneView* SceneView = InRenderAPI->GetSceneView();
+	if (!SceneView)
+	{
+		return;
+	}
 	
 	FVector2D OriginPixel; 
-	View->WorldToPixel(SelectionFrame.Origin, OriginPixel);
-	OriginPixel /= Canvas->GetDPIScale();
+	SceneView->WorldToPixel(SelectionFrame.Origin, OriginPixel);
+	OriginPixel /= InCanvas->GetDPIScale();
 
-	FVector RadiusSamplePoint = SelectionFrame.Origin + View->GetViewUp() * Radius;
+	FVector RadiusSamplePoint = SelectionFrame.Origin + SceneView->GetViewUp() * Radius;
 	FVector2D RadiusSamplePointPixel;
-	View->WorldToPixel(RadiusSamplePoint, RadiusSamplePointPixel);
-	RadiusSamplePointPixel /= Canvas->GetDPIScale();
-	
-	double PixelRadius = FMath::Abs(RadiusSamplePointPixel.Y - OriginPixel.Y);
+	SceneView->WorldToPixel(RadiusSamplePoint, RadiusSamplePointPixel);
+	RadiusSamplePointPixel /= InCanvas->GetDPIScale();
 
+	double PixelRadius = FMath::Abs(RadiusSamplePointPixel.Y - OriginPixel.Y);
 
 	const int32 NumSides = FMath::Max(20, (int)PixelRadius / 25);
 	const double AngleDelta = 2.0 * UE_PI / NumSides;
@@ -82,7 +83,7 @@ void USkeletalMeshModelingToolsGeometrySoftSelectRadiusDragInteraction::Render(
 
 		FCanvasLineItem LineItem{ FVector2D(LastVertex), FVector2D(Vertex) };
 		LineItem.SetColor(FLinearColor::Green);
-		Canvas->DrawItem(LineItem);
+		InCanvas->DrawItem(LineItem);
 
 		LastVertex = Vertex;
 	}
diff --git a/Engine/Plugins/Animation/SkeletalMeshModelingTools/Source/SkeletalMeshModelingTools/Private/ViewportInteractions/SKMModelingToolsGeometrySoftSelectRadiusDrag.h b/Engine/Plugins/Animation/SkeletalMeshModelingTools/Source/SkeletalMeshModelingTools/Private/ViewportInteractions/SKMModelingToolsGeometrySoftSelectRadiusDrag.h
--- a/Engine/Plugins/Animation/SkeletalMeshModelingTools/Source/SkeletalMeshModelingTools/Private/ViewportInteractions/SKMModelingToolsGeometrySoftSelectRadiusDrag.h
+++ b/Engine/Plugins/Animation/SkeletalMeshModelingTools/Source/SkeletalMeshModelingTools/Private/ViewportInteractions/SKMModelingToolsGeometrySoftSelectRadiusDrag.h
@@ -18,9 +18,9 @@ public:
 	USkeletalMeshModelingToolsGeometrySoftSelectRadiusDragInteraction();
 
 	void BindSelectionManager(UGeometrySelectionManager* InSelectionManager);
-	
-	virtual void Render(const FSceneView* View, FCanvas* Canvas, EViewInteractionState InteractionState) override;
-	
+
+	virtual void Draw(FCanvas* InCanvas, IToolsContextRenderAPI* InRenderAPI) override;
+
 	//~ Begin IViewportClickDragBehaviorTarget
 	virtual FInputRayHit CanBeginClickDragSequence(const FInputDeviceRay& InPressPos) override;
 	virtual void OnDragStart(const FInputDeviceRay& InPressPos) override;
@@ -36,4 +36,3 @@ protected:
 	
 	TWeakObjectPtr<UGeometrySelectionManager> SelectionManager;
 };
-
diff --git a/Engine/Source/Editor/Experimental/EditorInteractiveToolsFramework/Private/EditorDragTools/MarqueeSelectInteraction.cpp b/Engine/Source/Editor/Experimental/EditorInteractiveToolsFramework/Private/EditorDragTools/MarqueeSelectInteraction.cpp
--- a/Engine/Source/Editor/Experimental/EditorInteractiveToolsFramework/Private/EditorDragTools/MarqueeSelectInteraction.cpp
+++ b/Engine/Source/Editor/Experimental/EditorInteractiveToolsFramework/Private/EditorDragTools/MarqueeSelectInteraction.cpp
@@ -18,14 +18,14 @@ UMarqueeSelectInteraction::UMarqueeSelectInteraction()
 	RegisterInputBehavior(ClickDragBehavior);
 }
 
-void UMarqueeSelectInteraction::Render(const FSceneView* View, FCanvas* Canvas, EViewInteractionState InteractionState)
+void UMarqueeSelectInteraction::Draw(FCanvas* InCanvas, IToolsContextRenderAPI* InRenderAPI)
 {
-	if (!IsDragging() || InteractionState != EViewInteractionState::Focused)
+	if (!ensure(InRenderAPI) || !ensure(InCanvas))
 	{
 		return;
 	}
 
-	if (!Canvas || !View)
+	if (!IsDragging() || InRenderAPI->GetViewInteractionState() != EViewInteractionState::Focused)
 	{
 		return;
 	}
@@ -33,12 +33,12 @@ void UMarqueeSelectInteraction::Render(const FSceneView* View, FCanvas* Canvas,
 	if (IsWindowSelection())
 	{
 		FCanvasBoxItem BoxItem(
-			FVector2D(Start.X, Start.Y) / Canvas->GetDPIScale(),
-			FVector2D(End.X - Start.X, End.Y - Start.Y) / Canvas->GetDPIScale()
+			FVector2D(Start.X, Start.Y) / InCanvas->GetDPIScale(),
+			FVector2D(End.X - Start.X, End.Y - Start.Y) / InCanvas->GetDPIScale()
 		);
 
 		BoxItem.SetColor(FLinearColor::White);
-		Canvas->DrawItem(BoxItem);
+		InCanvas->DrawItem(BoxItem);
 	}
 	else
 	{
@@ -57,15 +57,15 @@ void UMarqueeSelectInteraction::Render(const FSceneView* View, FCanvas* Canvas,
 		}
 
 		FCanvasDashedBoxItem BoxItem(
-			FVector2D(LocalStart.X, LocalStart.Y) / Canvas->GetDPIScale(),
-			FVector2D(LocalEnd.X - LocalStart.X, LocalEnd.Y - LocalStart.Y) / Canvas->GetDPIScale()
+			FVector2D(LocalStart.X, LocalStart.Y) / InCanvas->GetDPIScale(),
+			FVector2D(LocalEnd.X - LocalStart.X, LocalEnd.Y - LocalStart.Y) / InCanvas->GetDPIScale()
 		);
 
 		BoxItem.DashLength = 12.0f;
 		BoxItem.DashGap = 6.0f;
 		BoxItem.SetColor(FLinearColor::White);
 
-		Canvas->DrawItem(BoxItem);
+		InCanvas->DrawItem(BoxItem);
 	}
 }
 
diff --git a/Engine/Source/Editor/Experimental/EditorInteractiveToolsFramework/Private/EditorDragTools/MeasureToolInteraction.cpp b/Engine/Source/Editor/Experimental/EditorInteractiveToolsFramework/Private/EditorDragTools/MeasureToolInteraction.cpp
--- a/Engine/Source/Editor/Experimental/EditorInteractiveToolsFramework/Private/EditorDragTools/MeasureToolInteraction.cpp
+++ b/Engine/Source/Editor/Experimental/EditorInteractiveToolsFramework/Private/EditorDragTools/MeasureToolInteraction.cpp
@@ -32,14 +32,19 @@ UMeasureToolInteraction::UMeasureToolInteraction()
 	RegisterInputBehavior(MouseBehavior);
 }
 
-void UMeasureToolInteraction::Render(const FSceneView* InView, FCanvas* InCanvas, EViewInteractionState InteractionState)
+void UMeasureToolInteraction::Draw(FCanvas* InCanvas, IToolsContextRenderAPI* InRenderAPI)
 {
-	if (!IsDragging() || !InCanvas || !InView)
+	if (!ensure(InRenderAPI) || !ensure(InCanvas))
 	{
 		return;
 	}
 
-	if (InteractionState != EViewInteractionState::Focused)
+	if (!IsDragging())
+	{
+		return;
+	}
+
+	if (InRenderAPI->GetViewInteractionState() != EViewInteractionState::Focused)
 	{
 		return;
 	}
diff --git a/Engine/Source/Editor/Experimental/EditorInteractiveToolsFramework/Private/EditorDragTools/ViewportChangeInteraction.cpp b/Engine/Source/Editor/Experimental/EditorInteractiveToolsFramework/Private/EditorDragTools/ViewportChangeInteraction.cpp
--- a/Engine/Source/Editor/Experimental/EditorInteractiveToolsFramework/Private/EditorDragTools/ViewportChangeInteraction.cpp
+++ b/Engine/Source/Editor/Experimental/EditorInteractiveToolsFramework/Private/EditorDragTools/ViewportChangeInteraction.cpp
@@ -29,20 +29,25 @@ UViewportChangeInteraction::UViewportChangeInteraction()
 	RegisterInputBehavior(MouseBehavior);
 }
 
-void UViewportChangeInteraction::Render(const FSceneView* View, FCanvas* Canvas, EViewInteractionState InteractionState)
+void UViewportChangeInteraction::Draw(FCanvas* InCanvas, IToolsContextRenderAPI* InRenderAPI)
 {
-	if (!IsDragging() || !Canvas || !View)
+	if (!ensure(InRenderAPI) || !ensure(InCanvas))
+	{
+		return;
+	}
+	
+	if (!IsDragging())
 	{
 		return;
 	}
 
-	if (InteractionState != EViewInteractionState::Focused)
+	if (InRenderAPI->GetViewInteractionState() != EViewInteractionState::Focused)
 	{
 		return;
 	}
 
 	FCanvasLineItem LineItem(Start, End);
-	Canvas->DrawItem(LineItem);
+	InCanvas->DrawItem(LineItem);
 
 	const FLinearColor ToolColor = GetDefault<UEditorStyleSettings>()->ViewportToolOverlayColor;
 	FCanvasTextItem TextItem(
@@ -52,7 +57,7 @@ void UViewportChangeInteraction::Render(const FSceneView* View, FCanvas* Canvas,
 		ToolColor
 	);
 	TextItem.bCentreX = true;
-	Canvas->DrawItem(TextItem);
+	InCanvas->DrawItem(TextItem);
 }
 
 FInputRayHit UViewportChangeInteraction::CanBeginClickDragSequence(const FInputDeviceRay& InPressPos)
diff --git a/Engine/Source/Editor/Experimental/EditorInteractiveToolsFramework/Private/ViewportInteractions/ViewportInteractionsBehaviorSource.cpp b/Engine/Source/Editor/Experimental/EditorInteractiveToolsFramework/Private/ViewportInteractions/ViewportInteractionsBehaviorSource.cpp
--- a/Engine/Source/Editor/Experimental/EditorInteractiveToolsFramework/Private/ViewportInteractions/ViewportInteractionsBehaviorSource.cpp
+++ b/Engine/Source/Editor/Experimental/EditorInteractiveToolsFramework/Private/ViewportInteractions/ViewportInteractionsBehaviorSource.cpp
@@ -310,28 +310,24 @@ void UViewportInteractionsBehaviorSource::Tick(float InDeltaTime) const
 	}
 }
 
-void UViewportInteractionsBehaviorSource::RenderTools(
-	FViewportClient* InViewportClient, const FSceneView* InSceneView, FCanvas* InCanvas
-) const
+void UViewportInteractionsBehaviorSource::RenderTools(IToolsContextRenderAPI* InRenderAPI) const
 {
-	const FViewportClient* Focused = GetEditorViewportClient();
-	const FViewportClient* Hovered = GetHoveredEditorViewportClient();
-
-	EViewInteractionState InteractionState = EViewInteractionState::None;
-	if (InViewportClient == Focused)
-	{
-		InteractionState = EViewInteractionState::Focused;
-	}
-	else if (InViewportClient == Hovered)
+	for (const TPair<FName, TObjectPtr<UViewportInteraction>>& Pair : ViewportInteractions)
 	{
-		InteractionState = EViewInteractionState::Hovered;
+		if (UViewportInteraction* Interaction = Pair.Value)
+		{
+			Interaction->Render(InRenderAPI);
+		}
 	}
+}
 
+void UViewportInteractionsBehaviorSource::DrawTools(FCanvas* InCanvas, IToolsContextRenderAPI* InRenderAPI) const
+{
 	for (const TPair<FName, TObjectPtr<UViewportInteraction>>& Pair : ViewportInteractions)
 	{
 		if (UViewportInteraction* Interaction = Pair.Value)
 		{
-			Interaction->Render(InSceneView, InCanvas, InteractionState);
+			Interaction->Draw(InCanvas, InRenderAPI);
 		}
 	}
 }
diff --git a/Engine/Source/Editor/Experimental/EditorInteractiveToolsFramework/Public/EditorDragTools/MarqueeSelectInteraction.h b/Engine/Source/Editor/Experimental/EditorInteractiveToolsFramework/Public/EditorDragTools/MarqueeSelectInteraction.h
--- a/Engine/Source/Editor/Experimental/EditorInteractiveToolsFramework/Public/EditorDragTools/MarqueeSelectInteraction.h
+++ b/Engine/Source/Editor/Experimental/EditorInteractiveToolsFramework/Public/EditorDragTools/MarqueeSelectInteraction.h
@@ -23,7 +23,7 @@ class UMarqueeSelectInteraction : public UDragToolInteraction
 public:
 	UMarqueeSelectInteraction();
 
-	UE_API virtual void Render(const FSceneView* View, FCanvas* Canvas, EViewInteractionState InteractionState) override;
+	UE_API virtual void Draw(FCanvas* InCanvas, IToolsContextRenderAPI* InRenderAPI) override;
 	
 protected:
 
diff --git a/Engine/Source/Editor/Experimental/EditorInteractiveToolsFramework/Public/EditorDragTools/MeasureToolInteraction.h b/Engine/Source/Editor/Experimental/EditorInteractiveToolsFramework/Public/EditorDragTools/MeasureToolInteraction.h
--- a/Engine/Source/Editor/Experimental/EditorInteractiveToolsFramework/Public/EditorDragTools/MeasureToolInteraction.h
+++ b/Engine/Source/Editor/Experimental/EditorInteractiveToolsFramework/Public/EditorDragTools/MeasureToolInteraction.h
@@ -16,7 +16,7 @@ class UMeasureToolInteraction : public UDragToolInteraction
 public:
 	UMeasureToolInteraction();
 
-	virtual void Render(const FSceneView* InView,FCanvas* InCanvas, EViewInteractionState InteractionState) override;
+	virtual void Draw(FCanvas* InCanvas, IToolsContextRenderAPI* InRenderAPI) override;
 
 	//~ Begin IViewportClickDragBehaviorTarget
 	virtual FInputRayHit CanBeginClickDragSequence(const FInputDeviceRay& InPressPos) override;
diff --git a/Engine/Source/Editor/Experimental/EditorInteractiveToolsFramework/Public/EditorDragTools/ViewportChangeInteraction.h b/Engine/Source/Editor/Experimental/EditorInteractiveToolsFramework/Public/EditorDragTools/ViewportChangeInteraction.h
--- a/Engine/Source/Editor/Experimental/EditorInteractiveToolsFramework/Public/EditorDragTools/ViewportChangeInteraction.h
+++ b/Engine/Source/Editor/Experimental/EditorInteractiveToolsFramework/Public/EditorDragTools/ViewportChangeInteraction.h
@@ -19,7 +19,8 @@ class UViewportChangeInteraction : public UDragToolInteraction
 
 public:
 	UViewportChangeInteraction();
-	virtual void Render(const FSceneView* View, FCanvas* Canvas, EViewInteractionState InteractionState) override;
+
+	virtual void Draw(FCanvas* InCanvas, IToolsContextRenderAPI* InRenderAPI) override;
 
 	//~ Begin IViewportClickDragBehaviorTarget
 	virtual FInputRayHit CanBeginClickDragSequence(const FInputDeviceRay& InPressPos) override;
diff --git a/Engine/Source/Editor/Experimental/EditorInteractiveToolsFramework/Public/ViewportInteractions/ViewportInteraction.h b/Engine/Source/Editor/Experimental/EditorInteractiveToolsFramework/Public/ViewportInteractions/ViewportInteraction.h
--- a/Engine/Source/Editor/Experimental/EditorInteractiveToolsFramework/Public/ViewportInteractions/ViewportInteraction.h
+++ b/Engine/Source/Editor/Experimental/EditorInteractiveToolsFramework/Public/ViewportInteractions/ViewportInteraction.h
@@ -129,8 +129,11 @@ public:
 	/** Called every frame. */
 	virtual void Tick(float InDeltaTime) const {}
 	
+	/** Called every frame to render tools to world space */
+	virtual void Render(IToolsContextRenderAPI* InRenderAPI) { }
+
 	/** Called every frame to render tools to a 2D canvas */
-	virtual void Render(const FSceneView* InView, FCanvas* InCanvas, EViewInteractionState InteractionState) {}
+	virtual void Draw(FCanvas* InCanvas, IToolsContextRenderAPI* InRenderAPI) { }
 
 	TArray<TObjectPtr<UInputBehavior>> GetInputBehaviors()
 	{
diff --git a/Engine/Source/Editor/Experimental/EditorInteractiveToolsFramework/Public/ViewportInteractions/ViewportInteractionsBehaviorSource.h b/Engine/Source/Editor/Experimental/EditorInteractiveToolsFramework/Public/ViewportInteractions/ViewportInteractionsBehaviorSource.h
--- a/Engine/Source/Editor/Experimental/EditorInteractiveToolsFramework/Public/ViewportInteractions/ViewportInteractionsBehaviorSource.h
+++ b/Engine/Source/Editor/Experimental/EditorInteractiveToolsFramework/Public/ViewportInteractions/ViewportInteractionsBehaviorSource.h
@@ -119,9 +119,14 @@ public:
 	UE_API void Tick(float InDeltaTime) const;
 
 	/**
-	 * Renders the active tool on the specified View/Canvas
+	 * Renders the active tool on the specified View/World
 	 */
-	UE_API void RenderTools(FViewportClient* InViewportClient, const FSceneView* InSceneView, FCanvas* InCanvas) const;
+	UE_API void RenderTools(IToolsContextRenderAPI* InRenderAPI) const;
+
+	/**
+	 * Draws the active tool on the specified View/Canvas
+	 */
+	UE_API void DrawTools(FCanvas* InCanvas, IToolsContextRenderAPI* InRenderAPI) const;
 
 	/**
 	 * Add the specified interactions to this Behavior Source
diff --git a/Engine/Source/Editor/UnrealEd/Private/Tools/EdModeInteractiveToolsContext.cpp b/Engine/Source/Editor/UnrealEd/Private/Tools/EdModeInteractiveToolsContext.cpp
--- a/Engine/Source/Editor/UnrealEd/Private/Tools/EdModeInteractiveToolsContext.cpp
+++ b/Engine/Source/Editor/UnrealEd/Private/Tools/EdModeInteractiveToolsContext.cpp
@@ -658,17 +658,8 @@ void UEditorInteractiveToolsContext::Render(const FSceneView* View, FViewport* V
 	}
 
 	// Render Tool and Gizmos
-	const FEditorViewportClient* Focused = EditorModeManager->GetFocusedViewportClient();
-	const FEditorViewportClient* Hovered = EditorModeManager->GetHoveredViewportClient();
-	EViewInteractionState InteractionState = EViewInteractionState::None;
-	if (ViewportClient == Focused )
-	{
-		InteractionState |= EViewInteractionState::Focused;
-	}
-	if (ViewportClient == Hovered )
-	{
-		InteractionState |= EViewInteractionState::Hovered;
-	}
+	const EViewInteractionState InteractionState = GetViewInteractionStateFromViewportClient(ViewportClient);
+
 	FEdModeTempRenderContext RenderContext(View, Viewport, ViewportClient, PDI, InteractionState);
 	ToolManager->Render(&RenderContext);
 	GizmoManager->Render(&RenderContext);
@@ -678,17 +669,8 @@ void UEditorInteractiveToolsContext::Render(const FSceneView* View, FViewport* V
 void UEditorInteractiveToolsContext::DrawHUD(FViewportClient* ViewportClient,FViewport* Viewport,const FSceneView* View, FCanvas* Canvas)
 {
 	FEditorViewportClient* EditorViewportClient = static_cast<FEditorViewportClient*>(Viewport->GetClient());
-	const FViewportClient* Focused = EditorModeManager->GetFocusedViewportClient();
-	const FViewportClient* Hovered = EditorModeManager->GetHoveredViewportClient();
-	EViewInteractionState InteractionState = EViewInteractionState::None;
-	if (ViewportClient == Focused )
-	{
-		InteractionState |= EViewInteractionState::Focused;
-	}
-	if (ViewportClient == Hovered )
-	{
-		InteractionState |= EViewInteractionState::Hovered;
-	}
+	const EViewInteractionState InteractionState = GetViewInteractionStateFromViewportClient(EditorViewportClient);
+
 	FEdModeTempRenderContext RenderContext(View, Viewport, EditorViewportClient, nullptr /*PDI*/, InteractionState);
 	ToolManager->DrawHUD(Canvas, &RenderContext);
 	GizmoManager->DrawHUD(Canvas, &RenderContext);
@@ -904,6 +886,28 @@ void UEditorInteractiveToolsContext::OnToolPostBuild(UInteractiveToolManager* In
 	// todo: Add any shared tool targets for the mode toolkit
 }
 
+EViewInteractionState UEditorInteractiveToolsContext::GetViewInteractionStateFromViewportClient(const FEditorViewportClient* InViewportClient) const
+{
+	if (!ensure(InViewportClient))
+	{
+		return EViewInteractionState::None;
+	}
+
+	const FEditorViewportClient* Focused = EditorModeManager->GetFocusedViewportClient();
+	const FEditorViewportClient* Hovered = EditorModeManager->GetHoveredViewportClient();
+	EViewInteractionState InteractionState = EViewInteractionState::None;
+	if (InViewportClient == Focused )
+	{
+		InteractionState |= EViewInteractionState::Focused;
+	}
+	if (InViewportClient == Hovered )
+	{
+		InteractionState |= EViewInteractionState::Hovered;
+	}
+
+	return InteractionState;
+}
+
 void UEditorInteractiveToolsContext::SetEnableRenderingDuringHitProxyPass(bool bEnabled)
 {
 	bEnableRenderingDuringHitProxyPass = bEnabled;
@@ -978,6 +982,15 @@ void UModeManagerInteractiveToolsContext::Render(const FSceneView* View, FViewpo
 	{
 		EdModeContext->Render(View, Viewport, PDI);
 	}
+
+	if (ViewportInteractionsBehaviorSource)
+	{
+		FEditorViewportClient* EditorViewportClient = static_cast<FEditorViewportClient*>(Viewport->GetClient());
+		const EViewInteractionState InteractionState = GetViewInteractionStateFromViewportClient(EditorViewportClient);
+
+		FEdModeTempRenderContext RenderContext(View, Viewport, EditorViewportClient, PDI, InteractionState);
+		ViewportInteractionsBehaviorSource->RenderTools(&RenderContext);
+	}
 }
 
 void UModeManagerInteractiveToolsContext::DrawHUD(FViewportClient* ViewportClient, FViewport* Viewport, const FSceneView* View, FCanvas* Canvas)
@@ -990,7 +1003,11 @@ void UModeManagerInteractiveToolsContext::DrawHUD(FViewportClient* ViewportClien
 
 	if (ViewportInteractionsBehaviorSource)
 	{
-		ViewportInteractionsBehaviorSource->RenderTools(ViewportClient, View, Canvas);
+		FEditorViewportClient* EditorViewportClient = static_cast<FEditorViewportClient*>(Viewport->GetClient());
+		const EViewInteractionState InteractionState = GetViewInteractionStateFromViewportClient(EditorViewportClient);
+
+		FEdModeTempRenderContext RenderContext(View, Viewport, EditorViewportClient, nullptr /* canvas-based, no PDI */, InteractionState);
+		ViewportInteractionsBehaviorSource->DrawTools(Canvas, &RenderContext);
 	}
 }
 
diff --git a/Engine/Source/Editor/UnrealEd/Public/Tools/EdModeInteractiveToolsContext.h b/Engine/Source/Editor/UnrealEd/Public/Tools/EdModeInteractiveToolsContext.h
--- a/Engine/Source/Editor/UnrealEd/Public/Tools/EdModeInteractiveToolsContext.h
+++ b/Engine/Source/Editor/UnrealEd/Public/Tools/EdModeInteractiveToolsContext.h
@@ -257,6 +257,9 @@ protected:
 	UE_API void OnToolEnded(UInteractiveToolManager* InToolManager, UInteractiveTool* InEndedTool);
 	UE_API void OnToolPostBuild(UInteractiveToolManager* InToolManager, EToolSide InSide, UInteractiveTool* InBuiltTool, UInteractiveToolBuilder* InToolBuilder, const FToolBuilderState& ToolState);
 
+	/** Returns a ViewInteractionState based on the input ViewportClient, and context's ModeManager. */
+	UE_API EViewInteractionState GetViewInteractionStateFromViewportClient(const FEditorViewportClient* InViewportClient) const;
+
 	TOptional<FString> PendingToolToStart = {};
 	TOptional<EToolShutdownType> PendingToolShutdownType = {};
 
