[ITF Cursor]
- Add the ability for UToolsContextCursorAPI to override hardware/software cursor visiblity. Controlled by the CVar "Editor.ITF.SoftwareCursorOverride". The default value of 2 automatically uses the software cursor for "Grab Hand" overrides.
- Update the Software Grab Hand image with an SVG at the intended size. Add support for a software Open Hand.
- Deprecate `FEditorModeTools::GetOverrideCursorVisibility()` for `FEditorModeTools::GetCursorVisibilityOverride()` so that the third boolean argument is actually a reference and works.

#rnx
#Synced-CL 52300119

diff --git a/Engine/Source/Editor/EditorStyle/Private/StarshipStyle.cpp b/Engine/Source/Editor/EditorStyle/Private/StarshipStyle.cpp
--- a/Engine/Source/Editor/EditorStyle/Private/StarshipStyle.cpp
+++ b/Engine/Source/Editor/EditorStyle/Private/StarshipStyle.cpp
@@ -720,7 +720,8 @@ void FStarshipEditorStyle::FStyle::SetupGeneralStyles()
 	Set( "GenericPause", new IMAGE_BRUSH( "Icons/generic_pause_16x", Icon16x16 ) );
 	Set( "GenericStop", new IMAGE_BRUSH( "Icons/generic_stop_16x", Icon16x16 ) );
 
-	Set( "SoftwareCursor_Grab", new IMAGE_BRUSH( "Icons/cursor_grab", Icon24x24 ) );
+	Set( "SoftwareCursor_Hand", new IMAGE_BRUSH_SVG( "Icons/HandOpen_24", Icon24x24 ) );
+	Set( "SoftwareCursor_Grab", new IMAGE_BRUSH_SVG( "Icons/HandGrab_24", Icon24x24 ) );
 	Set( "SoftwareCursor_CardinalCross", new IMAGE_BRUSH( "Icons/cursor_cardinal_cross", Icon24x24 ) );
 	Set( "SoftwareCursor_UpDown", new IMAGE_BRUSH( "Icons/cursor_updown", Icon16x20 ) );
 
diff --git a/Engine/Source/Editor/UnrealEd/Private/EditorModeManager.cpp b/Engine/Source/Editor/UnrealEd/Private/EditorModeManager.cpp
--- a/Engine/Source/Editor/UnrealEd/Private/EditorModeManager.cpp
+++ b/Engine/Source/Editor/UnrealEd/Private/EditorModeManager.cpp
@@ -1788,10 +1788,24 @@ TSharedPtr<SWidget> FEditorModeTools::MapCursor() const
 
 bool FEditorModeTools::GetOverrideCursorVisibility(bool& bWantsOverride, bool& bHardwareCursorVisible, bool bSoftwareCursorVisible) const
 {
+	return GetCursorVisibilityOverride(bWantsOverride, bHardwareCursorVisible, bSoftwareCursorVisible);
+}
+
+bool FEditorModeTools::GetCursorVisibilityOverride(bool& bOutWantsOverride, bool& bOutHardwareCursorVisible, bool& bOutSoftwareCursorVisible) const
+{
+	if (CVarEnableITFCursorOverrideSupport.GetValueOnGameThread() > 0)
+	{
+		if (InteractiveToolsContext->GetCursorVisibilityOverride(bOutHardwareCursorVisible, bOutSoftwareCursorVisible))
+		{
+			bOutWantsOverride = true;
+			return true;
+		}
+	}
+
 	bool bHandled = false;
 	for (const UEdMode* Mode : ActiveScriptableModes)
 	{
-		bHandled |= Mode->GetOverrideCursorVisibility(bWantsOverride, bHardwareCursorVisible, bSoftwareCursorVisible);
+		bHandled |= Mode->GetOverrideCursorVisibility(bOutWantsOverride, bOutHardwareCursorVisible, bOutSoftwareCursorVisible);
 	}
 	return bHandled;
 }
diff --git a/Engine/Source/Editor/UnrealEd/Private/EditorViewportClient.cpp b/Engine/Source/Editor/UnrealEd/Private/EditorViewportClient.cpp
--- a/Engine/Source/Editor/UnrealEd/Private/EditorViewportClient.cpp
+++ b/Engine/Source/Editor/UnrealEd/Private/EditorViewportClient.cpp
@@ -3648,6 +3648,10 @@ bool FEditorViewportClient::Internal_InputKey(const FInputKeyEventArgs& EventArg
 	// input router and we won't need this bIsTracking flag to begin with.
 	if (ModeTools->InputKey(this, Viewport, Key, Event, bRouteToToolsContext))
 	{
+		RequiredCursorVisibiltyAndAppearance.bDontResetCursor = true;
+		UpdateRequiredCursorVisibility();
+		ApplyRequiredCursorVisibility(true);
+	
 		return true;
 	}
 
@@ -5851,7 +5855,7 @@ void FEditorViewportClient::UpdateRequiredCursorVisibility()
 	bool bOverrideCursorVisibility = false;
 	bool bHardwareCursorVisible = false;
 	bool bSoftwareCursorVisible = false;
-	if (ModeTools->GetOverrideCursorVisibility(bOverrideCursorVisibility, bHardwareCursorVisible, bSoftwareCursorVisible))
+	if (ModeTools->GetCursorVisibilityOverride(bOverrideCursorVisibility, bHardwareCursorVisible, bSoftwareCursorVisible))
 	{
 		if (bOverrideCursorVisibility)
 		{
@@ -6676,6 +6680,14 @@ void FEditorViewportClient::CapturedMouseMove( FViewport* InViewport, int32 InMo
 	// Let the current editor mode know about the mouse movement.
 	if (ModeTools->CapturedMouseMove(this, InViewport, InMouseX, InMouseY))
 	{
+		RequiredCursorVisibiltyAndAppearance.bDontResetCursor = true;
+		UpdateRequiredCursorVisibility();
+		ApplyRequiredCursorVisibility();
+		
+		if (Viewport->IsSoftwareCursorVisible())
+		{
+			Viewport->SetSoftwareCursorPosition(FVector2D(InMouseX, InMouseY));
+		}
 		return;
 	}
 
diff --git a/Engine/Source/Editor/UnrealEd/Private/Tools/EdModeInteractiveToolsContext.cpp b/Engine/Source/Editor/UnrealEd/Private/Tools/EdModeInteractiveToolsContext.cpp
--- a/Engine/Source/Editor/UnrealEd/Private/Tools/EdModeInteractiveToolsContext.cpp
+++ b/Engine/Source/Editor/UnrealEd/Private/Tools/EdModeInteractiveToolsContext.cpp
@@ -67,6 +67,13 @@ namespace UE::Editor::Private
 		false,
 		TEXT("Whether to force on viewport interactions, ignoring legacy requests from editor modes.")
 	);
+	
+	TAutoConsoleVariable<int32> CVarEnableITFSoftwareCursor(
+		TEXT("Editor.ITF.SoftwareCursorOverride"),
+		2,
+		TEXT("0 - No software cursor will be used")
+		TEXT("1 - Software cursor will be used when requested")
+		TEXT("2 - Software cursor will be used when requested, defaulting to On when set to \"Grab Hand\""));
 }
 
 class FEdModeToolsContextQueriesImpl : public IToolsContextQueriesAPI
@@ -1370,6 +1377,44 @@ TSharedPtr<SWidget> UModeManagerInteractiveToolsContext::GetCursorWidget() const
 	return nullptr;
 }
 
+bool UModeManagerInteractiveToolsContext::GetCursorVisibilityOverride(bool& bOutHardwareCursorVisible, bool& bOutSoftwareCursorVisible) const
+{
+	const int32 ITFSoftwareCursor = UE::Editor::Private::CVarEnableITFSoftwareCursor.GetValueOnAnyThread();
+	if (ITFSoftwareCursor > 0 && InputRouter->HasActiveMouseCapture())
+	{
+		if (const UToolsContextCursorAPI* CursorAPI = ToolManager->GetContextObjectStore()->FindContext<UToolsContextCursorAPI>())
+		{
+			if (CursorAPI->IsCursorOverridden())
+			{
+				const TOptional<bool> SoftwareCursorOverride = CursorAPI->GetUseSoftwareCursorOverride();
+				if (SoftwareCursorOverride.IsSet())
+				{
+					bOutSoftwareCursorVisible = SoftwareCursorOverride.GetValue();
+					bOutHardwareCursorVisible = !bOutSoftwareCursorVisible;
+					return true;
+				}
+				
+				if (ITFSoftwareCursor == 2 && CursorAPI->GetCurrentCursorOverride() == EMouseCursor::GrabHandClosed)
+				{
+					// The logic here is twofold:
+					// 1 - A "grab hand" is usually manipulating something visually in the viewport
+					// 2 - We want the connection between the cursor and the thing its manipulating to be rock-solid.
+					// The hardware cursor is disconnected from the frame rate & rendering latency of the viewport.
+					// This causes dragged objects to visually lag behind the cursor. The hardware cursor rendering has less latency than the full scene. 
+					// Using the software cursor helps lend a sense of solidity.
+					// There is no discrepancy between the position of the cursor and the position of the object
+					// as both are rendered within the same system with the same latency.
+					// Thus, defaulting the GrabHand to using the software cursor should result in a net-benefit to interaction solidity.
+					bOutHardwareCursorVisible = false;
+					bOutSoftwareCursorVisible = true;
+					return true;
+				}
+			}
+		}
+	}
+	return false;
+}
+
 FRay UModeManagerInteractiveToolsContext::GetLastWorldRay() const
 {
 	return CurrentMouseState.Mouse.WorldRay;
diff --git a/Engine/Source/Editor/UnrealEd/Public/EditorModeManager.h b/Engine/Source/Editor/UnrealEd/Public/EditorModeManager.h
--- a/Engine/Source/Editor/UnrealEd/Public/EditorModeManager.h
+++ b/Engine/Source/Editor/UnrealEd/Public/EditorModeManager.h
@@ -315,8 +315,12 @@ public:
 	 */
 	UNREALED_API TSharedPtr<SWidget> MapCursor() const;
 
+	UE_DEPRECATED(5.8, "Use GetCursorVisibilityOverride() instead. `bSoftwareCursorVisible` is not a ref and is hence broken")
 	/** Get override cursor visibility settings */
 	UNREALED_API bool GetOverrideCursorVisibility(bool& bWantsOverride, bool& bHardwareCursorVisible, bool bSoftwareCursorVisible) const;
+	
+	/** Get override cursor visibility settings */
+	UNREALED_API bool GetCursorVisibilityOverride(bool& bOutWantsOverride, bool& bOutHardwareCursorVisible, bool& bOutSoftwareCursorVisible) const;
 
 	/** Called before converting mouse movement to drag/rot */
 	UNREALED_API bool PreConvertMouseMovement(FEditorViewportClient* InViewportClient);
diff --git a/Engine/Source/Editor/UnrealEd/Public/Tools/EdModeInteractiveToolsContext.h b/Engine/Source/Editor/UnrealEd/Public/Tools/EdModeInteractiveToolsContext.h
--- a/Engine/Source/Editor/UnrealEd/Public/Tools/EdModeInteractiveToolsContext.h
+++ b/Engine/Source/Editor/UnrealEd/Public/Tools/EdModeInteractiveToolsContext.h
@@ -336,6 +336,10 @@ public:
 
 	/** @return The Slate widget to use for the cursor if any. Only called if GetCursor returns Custom. */
 	UE_API TSharedPtr<SWidget> GetCursorWidget() const;
+	
+	UE_EXPERIMENTAL(5.8, "The extension of the visibility override interface into the editor context is experimental")
+	/** @return true when the out values have been set. */
+	UE_API bool GetCursorVisibilityOverride(bool& bOutHardwareCursorVisible, bool& bOutSoftwareCursorVisible) const;
 
 	/** @return Ray into 3D scene at last mouse event */
 	UE_API virtual FRay GetLastWorldRay() const override;
diff --git a/Engine/Source/Runtime/InteractiveToolsFramework/Public/ToolContextInterfaces.h b/Engine/Source/Runtime/InteractiveToolsFramework/Public/ToolContextInterfaces.h
--- a/Engine/Source/Runtime/InteractiveToolsFramework/Public/ToolContextInterfaces.h
+++ b/Engine/Source/Runtime/InteractiveToolsFramework/Public/ToolContextInterfaces.h
@@ -230,6 +230,8 @@ private:
 	EMouseCursor::Type CursorOverrideShape;
 	TSharedPtr<SWidget> CursorOverrideWidget;
 	bool bIsCursorOverridden;
+	
+	TOptional<bool> UseSoftwareCursorOverride;
 
 public:
 
@@ -269,6 +271,16 @@ public:
 		CursorOverrideShape = EMouseCursor::Custom;
 		CursorOverrideWidget = InCursorOverrideWidget;
 	}
+	
+	void SetUseSoftwareCursor(const TOptional<bool>& InUseSoftwareCursorOverride)
+	{
+		UseSoftwareCursorOverride = InUseSoftwareCursorOverride;
+	}
+	
+	TOptional<bool> GetUseSoftwareCursorOverride() const
+	{
+		return UseSoftwareCursorOverride;
+	}
 
 	/** Clears the custom widget, if any, resets the shape to Default and resets the overridden flag. */
 	virtual void ClearCursorOverride() 
@@ -276,6 +288,7 @@ public:
 		bIsCursorOverridden = false;
 		CursorOverrideShape = EMouseCursor::Default;
 		CursorOverrideWidget = nullptr;
+		UseSoftwareCursorOverride.Reset();
 	}
 };
 
diff --git a/Engine/Source/Runtime/Slate/Private/Widgets/SViewport.cpp b/Engine/Source/Runtime/Slate/Private/Widgets/SViewport.cpp
--- a/Engine/Source/Runtime/Slate/Private/Widgets/SViewport.cpp
+++ b/Engine/Source/Runtime/Slate/Private/Widgets/SViewport.cpp
@@ -187,11 +187,20 @@ int32 SViewport::OnPaint( const FPaintArgs& Args, const FGeometry& AllottedGeome
 		 );
 		EMouseCursor::Type CursorType = Reply.GetCursorType();
 
-		const FSlateBrush* Brush = FCoreStyle::Get().GetBrush(TEXT("SoftwareCursor_Grab"));
-		if( CursorType == EMouseCursor::CardinalCross )
+		const FSlateBrush* Brush = [CursorType]
 		{
-			Brush = FCoreStyle::Get().GetBrush(TEXT("SoftwareCursor_CardinalCross"));
-		}
+			switch (CursorType)
+			{
+			case EMouseCursor::CardinalCross:
+				return FCoreStyle::Get().GetBrush(TEXT("SoftwareCursor_CardinalCross"));
+			case EMouseCursor::GrabHand:
+				return FCoreStyle::Get().GetBrush(TEXT("SoftwareCursor_Hand"));
+			case EMouseCursor::GrabHandClosed:
+			default: // Default is included here for legacy compatability; grab was the default.
+				return FCoreStyle::Get().GetBrush(TEXT("SoftwareCursor_Grab"));
+			}
+		}();
+		
 		const FVector2f CursorSize = Brush->ImageSize / AllottedGeometry.Scale;
 
 		const FVector2D ToViewportPixelScale = UE::Slate::Viewport::GetViewportLocalToPixelScale(AllottedGeometry, ViewportInterfacePin->GetSize());
