In a SlowTask, in the editor, if slate cannot display e.g. because global shaders are being recompiled, then fall back on the splash screen to display progress for the SlowTask.

#rnx
#Synced-CL 52284826

diff --git a/Engine/Source/Editor/UnrealEd/Private/FeedbackContextEditor.cpp b/Engine/Source/Editor/UnrealEd/Private/FeedbackContextEditor.cpp
--- a/Engine/Source/Editor/UnrealEd/Private/FeedbackContextEditor.cpp
+++ b/Engine/Source/Editor/UnrealEd/Private/FeedbackContextEditor.cpp
@@ -444,6 +444,15 @@ FFeedbackContextEditor::FFeedbackContextEditor()
 	
 }
 
+static bool IsFallbackSplashScreenAllowed()
+{
+#if WITH_EDITOR
+	return !FApp::IsUnattended() && !IsRunningDedicatedServer() && FPlatformProcess::SupportsMultithreading() && !FParse::Param(FCommandLine::Get(), TEXT("RenderOffScreen"));
+#else
+	return false;
+#endif
+}
+
 void FFeedbackContextEditor::StartSlowTask( const FText& Task, bool bShowCancelButton )
 {
 	FFeedbackContext::StartSlowTask( Task, bShowCancelButton );
@@ -505,40 +514,48 @@ void FFeedbackContextEditor::StartSlowTask( const FText& Task, bool bShowCancelB
 				UE_LOG(LogSlate, Log, TEXT("Prevented a slow task dialog from being summoned while a context menu was open"));
 			}
 
-			if (!bProgressWindowShown && !bHaveOpenMenu && FSlateApplication::Get().CanDisplayWindows())
+			if (!bProgressWindowShown && !bHaveOpenMenu)
 			{
-				FOnCancelClickedDelegate OnCancelClicked;
-				if (bShowCancelButton)
+				if (FSlateApplication::Get().CanDisplayWindows())
 				{
-					// The cancel button is only displayed if a delegate is bound to it.
-					OnCancelClicked = FOnCancelClickedDelegate::CreateRaw(this, &FFeedbackContextEditor::OnUserCancel);
-				}
+					FOnCancelClickedDelegate OnCancelClicked;
+					if (bShowCancelButton)
+					{
+						// The cancel button is only displayed if a delegate is bound to it.
+						OnCancelClicked = FOnCancelClickedDelegate::CreateRaw(this, &FFeedbackContextEditor::OnUserCancel);
+					}
 
-				const bool bFocus = FApp::HasFocus();
-				FReceiveUserCancelDelegate ReceiveUserCancelDelegate = FReceiveUserCancelDelegate::CreateRaw(this, &FFeedbackContextEditor::ReceivedUserCancel);
-				TSharedRef<SWindow> SlowTaskWindowRef = SNew(SWindow)
-					.SizingRule(ESizingRule::Autosized)
-					.AutoCenter(EAutoCenter::PreferredWorkArea)
-					.IsPopupWindow(true)
-					.CreateTitleBar(true)
-					.ActivationPolicy(bFocus ? EWindowActivationPolicy::Always : EWindowActivationPolicy::Never)
-					.FocusWhenFirstShown(bFocus);
+					const bool bFocus = FApp::HasFocus();
+					FReceiveUserCancelDelegate ReceiveUserCancelDelegate = FReceiveUserCancelDelegate::CreateRaw(this, &FFeedbackContextEditor::ReceivedUserCancel);
+					TSharedRef<SWindow> SlowTaskWindowRef = SNew(SWindow)
+						.SizingRule(ESizingRule::Autosized)
+						.AutoCenter(EAutoCenter::PreferredWorkArea)
+						.IsPopupWindow(true)
+						.CreateTitleBar(true)
+						.ActivationPolicy(bFocus ? EWindowActivationPolicy::Always : EWindowActivationPolicy::Never)
+						.FocusWhenFirstShown(bFocus);
 
-				SlowTaskWindowRef->SetContent(
-					SNew(SSlowTaskWidget)
-					.ScopeStack(GetScopeStackSharedPtr())
-					.OnCancelClickedDelegate(OnCancelClicked)
-					.ReceiveUserCancelDelegate(ReceiveUserCancelDelegate)
-				);
+					SlowTaskWindowRef->SetContent(
+						SNew(SSlowTaskWidget)
+						.ScopeStack(GetScopeStackSharedPtr())
+						.OnCancelClickedDelegate(OnCancelClicked)
+						.ReceiveUserCancelDelegate(ReceiveUserCancelDelegate)
+					);
 
-				SlowTaskWindow = SlowTaskWindowRef;
+					SlowTaskWindow = SlowTaskWindowRef;
 
-				const bool bSlowTask = true;
-				FSlateApplication::Get().AddModalWindow(SlowTaskWindowRef, ParentWindow, bSlowTask);
+					const bool bSlowTask = true;
+					FSlateApplication::Get().AddModalWindow(SlowTaskWindowRef, ParentWindow, bSlowTask);
 
-				SlowTaskWindowRef->ShowWindow();
+					SlowTaskWindowRef->ShowWindow();
 
-				TickSlate(SlowTaskWindow.Pin());
+					TickSlate(SlowTaskWindow.Pin());
+				}
+				else if(IsFallbackSplashScreenAllowed() && !FPlatformSplash::IsShown())
+				{
+					bFallbackSplashScreenVisible = true;
+					FPlatformSplash::Show();
+				}
 			}
 
 			FPlatformSplash::SetSplashText(SplashTextType::StartupProgress, Task);
@@ -555,6 +572,10 @@ void FFeedbackContextEditor::FinalizeSlowTask()
 		Window->RequestDestroyWindow();
 		SlowTaskWindow.Reset();
 	}
+	if (bFallbackSplashScreenVisible)
+	{
+		FPlatformSplash::Hide();
+	}
 
 	FFeedbackContext::FinalizeSlowTask( );
 }
diff --git a/Engine/Source/Editor/UnrealEd/Public/FeedbackContextEditor.h b/Engine/Source/Editor/UnrealEd/Public/FeedbackContextEditor.h
--- a/Engine/Source/Editor/UnrealEd/Public/FeedbackContextEditor.h
+++ b/Engine/Source/Editor/UnrealEd/Public/FeedbackContextEditor.h
@@ -28,6 +28,7 @@ class FFeedbackContextEditor : public FFeedbackContext
 	TSharedPtr<class SBuildProgressWidget> BuildProgressWidget;
 
 	bool HasTaskBeenCancelled;
+	bool bFallbackSplashScreenVisible = false;
 
 	double SlowTaskStartTime = 0;
 
