Improved hashing for external packages

Updated the hashing method for incremental validation to consider more than just the package save hash. In addition, it now considers:
* The schema of any C++ or Verse classes imported by the package.
* The package save hash of any direct dependencies (which also includes BPGC imported by this package).

This allows us to incrementally validate all external packages, rather than only external actors. Note that the CVar for enabling this feature has also been renamed to 'cook.validation.IncrementalExternalPackageValidationEnabled' since it is no longer actor-centric.

#Synced-CL 52288889

diff --git a/Engine/Source/Editor/UnrealEd/Private/CookOnTheFlyServer.cpp b/Engine/Source/Editor/UnrealEd/Private/CookOnTheFlyServer.cpp
--- a/Engine/Source/Editor/UnrealEd/Private/CookOnTheFlyServer.cpp
+++ b/Engine/Source/Editor/UnrealEd/Private/CookOnTheFlyServer.cpp
@@ -3435,17 +3435,16 @@ EDataValidationResult UCookOnTheFlyServer::ValidateSourcePackage(UE::Cook::FPack
 		{
 			const FString ExternalActorsPathForWorld = ULevel::GetExternalActorsPath(Package);
 			AssetRegistry->GetAssetsByPath(*ExternalActorsPathForWorld, ExternalObjects, /*bRecursive*/true, /*bIncludeOnlyOnDiskAssets*/true);
-
-			if (ValidationArtifact)
-			{
-				ValidationArtifact->FilterExternalActorsToValidate(Package->GetFName(), ExternalObjects);
-			}
 		}
 		if (TArray<FAssetData> AssetsWithPackageAsOuter = UAssetRegistryHelpers::GetAssetsWithOuterForPaths({ Package->GetFName() }, Package->GetFName(), EAssetsWithOuterForPathsFlags::RecursivePaths | EAssetsWithOuterForPathsFlags::IncludeOnlyOnDiskAsset);
 			!AssetsWithPackageAsOuter.IsEmpty())
 		{
 			ExternalObjects.Append(MoveTemp(AssetsWithPackageAsOuter));
 		}
+		if (ValidationArtifact)
+		{
+			ValidationArtifact->FilterExternalPackagesToValidate(Package->GetFName(), ExternalObjects);
+		}
 
 		FMessageLogScopedOverride AssetCheckLogOverride(UE::DataValidation::MessageLogName);
 		if (bLogErrorsAsWarnings)
diff --git a/Engine/Source/Editor/UnrealEd/Private/Cooker/ValidationCookArtifact.cpp b/Engine/Source/Editor/UnrealEd/Private/Cooker/ValidationCookArtifact.cpp
--- a/Engine/Source/Editor/UnrealEd/Private/Cooker/ValidationCookArtifact.cpp
+++ b/Engine/Source/Editor/UnrealEd/Private/Cooker/ValidationCookArtifact.cpp
@@ -7,11 +7,14 @@
 #include "Cooker/CookArtifactReader.h"
 #include "CookOnTheSide/CookOnTheFlyServer.h"
 #include "Engine/Level.h"
+#include "Hash/Blake3.h"
 #include "JsonUtils/RapidJsonUtils.h"
 #include "HAL/PlatformFileManager.h"
 #include "Interfaces/ITargetPlatform.h"
 #include "Misc/FileHelper.h"
 #include "Misc/Paths.h"
+#include "UObject/Class.h"
+#include "UObject/UObjectGlobals.h"
 
 namespace UE::Cook
 {
@@ -20,14 +23,16 @@ namespace Private
 {
 
 static const FString ConfigSectionName = TEXT("ValidationSettings");
-static const TCHAR JsonKey_WorldValidationCache[] = TEXT("worldValidationCache");
+static const TCHAR JsonKey_PackageValidationCache[] = TEXT("packageValidationCache");
+static const TCHAR JsonKey_ExternalActors[] = TEXT("externalActors");
+static const TCHAR JsonKey_ExternalPackages[] = TEXT("externalPackages");
 static const TCHAR JsonKey_LastAssetValidationBuildVersion[] = TEXT("lastAssetValidationBuildVersion");
 
-static bool bIncrementalWorldValidationEnabled = false;
-static FAutoConsoleVariableRef CVarIncrementalWorldValidationEnabled(
-	TEXT("cook.validation.IncrementalWorldValidationEnabled"),
-	bIncrementalWorldValidationEnabled,
-	TEXT("True if incremental world validation is enabled, where we only validate changed external actors in partitioned worlds"));
+static bool bIncrementalExternalPackageValidationEnabled = false;
+static FAutoConsoleVariableRef CVarIncrementalExternalPackageValidationEnabled(
+	TEXT("cook.validation.IncrementalExternalPackageValidationEnabled"),
+	bIncrementalExternalPackageValidationEnabled,
+	TEXT("True if incremental external package validation is enabled, where we only validate changed external packages (eg, external actors in a partitioned world)"));
 
 static bool bIncrementalValidationRequiresSameBuildVersion = false;
 static FAutoConsoleVariableRef CVarIncrementalValidationRequiresSameBuildVersion(
@@ -40,13 +45,13 @@ static FAutoConsoleVariableRef CVarIncrementalValidationRequiresSameBuildVersion
 FValidationCookArtifact::FValidationCookArtifact(UCookOnTheFlyServer& InCOTFS)
 	: COTFS(InCOTFS)
 	, CurrentBuildVersion(FApp::GetBuildVersion())
-	, bIsIncrementalWorldValidationEnabled(Private::bIncrementalWorldValidationEnabled)
+	, bIsIncrementalExternalPackageValidationEnabled(Private::bIncrementalExternalPackageValidationEnabled)
 {
-	// We currently only support incremental world validation for SP ByTheBook cooks (MP cooks would require a IMPCollector)
-	if (bIsIncrementalWorldValidationEnabled && (COTFS.GetProcessType() != UE::Cook::EProcessType::SingleProcess || COTFS.GetCookType() != UE::Cook::ECookType::ByTheBook))
+	// We currently only support incremental external package validation for SP ByTheBook cooks (MP cooks would require a IMPCollector)
+	if (bIsIncrementalExternalPackageValidationEnabled && (COTFS.GetProcessType() != UE::Cook::EProcessType::SingleProcess || COTFS.GetCookType() != UE::Cook::ECookType::ByTheBook))
 	{
-		UE_LOGF(LogCook, Warning, "Incremental world validation was requested, but it is only supported by single-process cook-by-the-book cooks. Incremental world validation will be disabled for this cook.");
-		bIsIncrementalWorldValidationEnabled = false;
+		UE_LOGF(LogCook, Warning, "Incremental external package validation was requested, but it is only supported by single-process cook-by-the-book cooks. Incremental external package validation will be disabled for this cook.");
+		bIsIncrementalExternalPackageValidationEnabled = false;
 	}
 }
 
@@ -88,9 +93,9 @@ void FValidationCookArtifact::CompareSettings(UE::Cook::Artifact::FCompareSettin
 			return;
 		}
 
-		if (!bIsIncrementalWorldValidationEnabled)
+		if (!bIsIncrementalExternalPackageValidationEnabled)
 		{
-			PlatformArtifactData.WorldValidationCache.Empty();
+			PlatformArtifactData.PackageValidationCache.Empty();
 		}
 	}
 }
@@ -113,11 +118,11 @@ void FValidationCookArtifact::UpdateOplogPackages(UE::Cook::Artifact::FUpdateOpl
 	// Remove any stale data from the persistent artifact
 	{
 		const TMap<FName, UE::Cook::Artifact::FOplogPackageData>& OplogPackages = Context.GetOplogPackages();
-		for (auto It = PlatformArtifactData.WorldValidationCache.CreateIterator(); It; ++It)
+		for (auto It = PlatformArtifactData.PackageValidationCache.CreateIterator(); It; ++It)
 		{
 			if (!OplogPackages.Contains(It->Key))
 			{
-				// World is no longer known to the cook; remove it
+				// Package is no longer known to the cook; remove it
 				It.RemoveCurrent();
 				continue;
 			}
@@ -140,39 +145,101 @@ void FValidationCookArtifact::UpdateOplogPackages(UE::Cook::Artifact::FUpdateOpl
 	}
 }
 
-void FValidationCookArtifact::FilterExternalActorsToValidate(const FName WorldPackageName, TArray<FAssetData>& InOutExternalActors, const bool bUpdateCache)
+void FValidationCookArtifact::FilterExternalPackagesToValidate(const FName PackageName, TArray<FAssetData>& InOutExternalPackages, const bool bUpdateCache)
 {
-	if (!bIsIncrementalWorldValidationEnabled || PerPlatformData.IsEmpty())
+	if (!bIsIncrementalExternalPackageValidationEnabled || PerPlatformData.IsEmpty())
 	{
 		return;
 	}
 
-	// Read the current hashes for every external actor package
-	TMap<FName, FIoHash> NewExternalActorHashes;
+	// Read the current hashes for all external packages (external actors and non-external-actor packages combined)
+	TMap<FName, FIoHash> NewExternalPackageHashes;
+	if (InOutExternalPackages.Num() > 0)
 	{
+		FCookLoadScope CookLoadScope(ECookLoadType::EditorOnly);
 		IAssetRegistry& AssetRegistry = IAssetRegistry::GetChecked();
-		for (const FAssetData& ExternalActor : InOutExternalActors)
+
+		auto ComputeExternalPackageValidationHash = 
+			[&AssetRegistry](const FName PackageName, const FAssetPackageData& PackageData) -> FIoHash
+			{
+				FBlake3 Hasher;
+
+				// Include the package saved hash
+				{
+					const FIoHash SavedHash = PackageData.GetPackageSavedHash();
+					Hasher.Update(MakeMemoryView(SavedHash.GetBytes(), sizeof(FIoHash::ByteArray)));
+				}
+
+				// Include the schema hash of each imported class to detect property layout changes without a resave
+				{
+					TArray<FName> ImportedClasses = PackageData.ImportedClasses;
+					ImportedClasses.Sort(FNameLexicalLess()); // Sort for determinism
+					for (const FName ClassName : ImportedClasses)
+					{
+						FNameBuilder ClassNameStr(ClassName);
+
+						// Note: This only includes C++ and Verse classes, as BP classes will be handled via the dependency code below (without needing to load them)
+						if (FPackageName::IsScriptPackage(ClassNameStr.ToView()) || FPackageName::IsVersePackage(ClassNameStr.ToView()))
+						{
+							if (const UClass* ImportedClass = LoadObject<UClass>(nullptr, ClassNameStr.ToView()))
+							{
+								const FBlake3Hash& SchemaHash = ImportedClass->GetSchemaHash(/*bSkipEditorOnly*/false);
+								Hasher.Update(MakeMemoryView(SchemaHash.GetBytes(), sizeof(FBlake3Hash)));
+							}
+						}
+					}
+				}
+
+				// Include the saved hash of each direct hard package dependency to detect changes in referenced assets
+				{
+					TArray<FName> Dependencies;
+					if (AssetRegistry.GetDependencies(PackageName, Dependencies, UE::AssetRegistry::EDependencyCategory::Package, UE::AssetRegistry::EDependencyQuery::Hard))
+					{
+						Dependencies.Sort(FNameLexicalLess()); // Sort for determinism
+						for (const FName Dependency : Dependencies)
+						{
+							FNameBuilder DependencyStr(Dependency);
+
+							// Skip C++ and Verse packages, as they were handled via the class imports above
+							if (!FPackageName::IsScriptPackage(DependencyStr.ToView()) && !FPackageName::IsVersePackage(DependencyStr.ToView()))
+							{
+								if (FAssetPackageData DependencyPackageData;
+									AssetRegistry.TryGetAssetPackageData(Dependency, DependencyPackageData) == AssetRegistry::EExists::Exists)
+								{
+									const FIoHash DependencySavedHash = DependencyPackageData.GetPackageSavedHash();
+									Hasher.Update(MakeMemoryView(DependencySavedHash.GetBytes(), sizeof(FIoHash::ByteArray)));
+								}
+							}
+						}
+					}
+				}
+
+				return FIoHash(Hasher.Finalize());
+			};
+
+		for (const FAssetData& ExternalPackage : InOutExternalPackages)
 		{
 			// Note: TryGetAssetPackageData only takes a read-lock, so this could be parallelized if needed
-			FAssetPackageData ExternalActorPackageData;
-			if (AssetRegistry.TryGetAssetPackageData(ExternalActor.PackageName, ExternalActorPackageData) == AssetRegistry::EExists::Exists)
+			FAssetPackageData ExternalPackageData;
+			if (AssetRegistry.TryGetAssetPackageData(ExternalPackage.PackageName, ExternalPackageData) == AssetRegistry::EExists::Exists)
 			{
-				NewExternalActorHashes.Add(ExternalActor.PackageName, ExternalActorPackageData.GetPackageSavedHash());
+				NewExternalPackageHashes.Add(ExternalPackage.PackageName, ComputeExternalPackageValidationHash(ExternalPackage.PackageName, ExternalPackageData));
 			}
 		}
 	}
 
 	// Compare these hashes against any previous state to know what's changed
 	// Note: Validation is platform agnostic, so we just use the data for the first platform
+	if (InOutExternalPackages.Num() > 0)
 	{
 		FArtifactData& PlatformArtifactData = PerPlatformData.begin()->Value;
-		if (const TMap<FName, FIoHash>* ExternalActorHashes = PlatformArtifactData.WorldValidationCache.Find(WorldPackageName))
+		if (const TMap<FName, FIoHash>* ExternalPackageHashes = PlatformArtifactData.PackageValidationCache.Find(PackageName))
 		{
-			InOutExternalActors.RemoveAll(
-				[&NewExternalActorHashes, ExternalActorHashes](const FAssetData& ExternalActor)
+			InOutExternalPackages.RemoveAll(
+				[&NewExternalPackageHashes, ExternalPackageHashes](const FAssetData& ExternalPackage)
 				{
-					const FIoHash* CurrentHash = NewExternalActorHashes.Find(ExternalActor.PackageName);
-					const FIoHash* PreviousHash = ExternalActorHashes->Find(ExternalActor.PackageName);
+					const FIoHash* CurrentHash  = NewExternalPackageHashes.Find(ExternalPackage.PackageName);
+					const FIoHash* PreviousHash = ExternalPackageHashes->Find(ExternalPackage.PackageName);
 					return CurrentHash && PreviousHash && *CurrentHash == *PreviousHash && !CurrentHash->IsZero();
 				});
 		}
@@ -183,7 +250,15 @@ void FValidationCookArtifact::FilterExternalActorsToValidate(const FName WorldPa
 	{
 		for (TTuple<const ITargetPlatform*, FArtifactData>& PerPlatformDataPair : PerPlatformData)
 		{
-			PerPlatformDataPair.Value.WorldValidationCache.Add(WorldPackageName, NewExternalActorHashes);
+			if (NewExternalPackageHashes.IsEmpty())
+			{
+				// Packages without external packages should not be stored in the cache
+				PerPlatformDataPair.Value.PackageValidationCache.Remove(PackageName);
+			}
+			else
+			{
+				PerPlatformDataPair.Value.PackageValidationCache.Add(PackageName, NewExternalPackageHashes);
+			}
 		}
 	}
 }
@@ -192,7 +267,7 @@ void FValidationCookArtifact::MarkPackageRequiresFullValidation(const FName Pack
 {
 	for (TTuple<const ITargetPlatform*, FArtifactData>& PerPlatformDataPair : PerPlatformData)
 	{
-		PerPlatformDataPair.Value.WorldValidationCache.Remove(PackageName);
+		PerPlatformDataPair.Value.PackageValidationCache.Remove(PackageName);
 	}
 }
 
@@ -205,7 +280,7 @@ void FValidationCookArtifact::InvalidateCache(const ITargetPlatform* TargetPlatf
 {
 	FArtifactData& PlatformArtifactData = PerPlatformData.FindChecked(TargetPlatform);
 	FPlatformFileManager::Get().GetPlatformFile().DeleteFile(*PlatformArtifactData.Filename);
-	PlatformArtifactData.WorldValidationCache.Reset();
+	PlatformArtifactData.PackageValidationCache.Reset();
 }
 
 bool FValidationCookArtifact::LoadCookArtifactFile(UE::Cook::Artifact::FCompareSettingsContext& Context, FArtifactData& PlatformArtifactData)
@@ -234,32 +309,45 @@ bool FValidationCookArtifact::LoadCookArtifactFile(UE::Cook::Artifact::FCompareS
 		return false;
 	}
 
-	// Read WorldValidationCache
-	if (TOptional<UE::Json::FConstObject> WorldValidationCacheJsonObject = UE::Json::GetObjectField(PlatformArtifactJsonRootObject.GetValue(), Private::JsonKey_WorldValidationCache))
+	// Read PackageValidationCache
+	if (TOptional<UE::Json::FConstObject> PackageValidationCacheJsonObject = UE::Json::GetObjectField(PlatformArtifactJsonRootObject.GetValue(), Private::JsonKey_PackageValidationCache))
 	{
-		for (const UE::Json::FMember& WorldValidationCacheJsonEntry : WorldValidationCacheJsonObject.GetValue())
+		for (const UE::Json::FMember& PackageCacheJsonEntry : PackageValidationCacheJsonObject.GetValue())
 		{
-			if (WorldValidationCacheJsonEntry.value.IsObject())
+			if (!PackageCacheJsonEntry.value.IsObject())
 			{
-				const FString WorldPackageName = FString(UE::Json::ValueAsStringView(WorldValidationCacheJsonEntry.name));
-				const FString ExternalActorsPrefix = ULevel::GetExternalActorsPath(WorldPackageName) / TEXT("");
+				continue;
+			}
+
+			const FString PackageName            = FString(UE::Json::ValueAsStringView(PackageCacheJsonEntry.name));
+			const FString ExternalActorsPrefix   = ULevel::GetExternalActorsPath(PackageName) / TEXT("");
+			const FString ExternalPackagesPrefix = PackageName / TEXT("");
 
-				TMap<FName, FIoHash>& ExternalActorHashes = PlatformArtifactData.WorldValidationCache.Add(*WorldPackageName);
-				for (const UE::Json::FMember& ExternalActorHashesJsonEntry : WorldValidationCacheJsonEntry.value.GetObject())
+			TMap<FName, FIoHash>& ExternalPackageHashes = PlatformArtifactData.PackageValidationCache.Add(*PackageName);
+
+			// Read a hash sub-object into the combined map, re-adding the given prefix to each stored path
+			auto ReadHashesIntoMap =
+				[&ExternalPackageHashes](const UE::Json::FConstObject& ContainerObject, const TCHAR* JsonKey, const FString& Prefix)
 				{
-					if (ExternalActorHashesJsonEntry.value.IsString())
+					if (TOptional<UE::Json::FConstObject> HashesJsonObject = UE::Json::GetObjectField(ContainerObject, JsonKey))
 					{
-						// Re-add the prefix stripped off during save
-						FString ExternalActorPath = FString(UE::Json::ValueAsStringView(ExternalActorHashesJsonEntry.name));
-						if (ExternalActorPath.Len() > 0 && ExternalActorPath[0] != TEXT('/'))
+						for (const UE::Json::FMember& HashJsonEntry : HashesJsonObject.GetValue())
 						{
-							ExternalActorPath.InsertAt(0, ExternalActorsPrefix);
+							if (HashJsonEntry.value.IsString())
+							{
+								// Re-add the prefix stripped off during save
+								FString Path = FString(UE::Json::ValueAsStringView(HashJsonEntry.name));
+								if (Path.Len() > 0 && Path[0] != TEXT('/'))
+								{
+									Path.InsertAt(0, Prefix);
+								}
+								ExternalPackageHashes.Add(*Path, FIoHash(UE::Json::ValueAsStringView(HashJsonEntry.value)));
+							}
 						}
-
-						ExternalActorHashes.Add(*ExternalActorPath, FIoHash(UE::Json::ValueAsStringView(ExternalActorHashesJsonEntry.value)));
 					}
-				}
-			}
+				};
+			ReadHashesIntoMap(PackageCacheJsonEntry.value.GetObject(), Private::JsonKey_ExternalActors,   ExternalActorsPrefix);
+			ReadHashesIntoMap(PackageCacheJsonEntry.value.GetObject(), Private::JsonKey_ExternalPackages, ExternalPackagesPrefix);
 		}
 	}
 
@@ -279,35 +367,62 @@ bool FValidationCookArtifact::SaveCookArtifactFile(FArtifactData& PlatformArtifa
 		UE::Json::FObject PlatformArtifactJsonRootObject = PlatformArtifactJsonDocument.GetObject();
 		UE::Json::FAllocator& PlatformArtifactJsonAllocator = PlatformArtifactJsonDocument.GetAllocator();
 
-		// Add WorldValidationCache
+		// Add PackageValidationCache
 		{
-			UE::Json::FValue WorldValidationCacheJsonObject(rapidjson::kObjectType);
-			PlatformArtifactData.WorldValidationCache.KeySort(FNameLexicalLess());
-			for (TTuple<FName, TMap<FName, FIoHash>>& WorldValidationCachePair : PlatformArtifactData.WorldValidationCache)
+			UE::Json::FValue PackageValidationCacheJsonObject(rapidjson::kObjectType);
+			PlatformArtifactData.PackageValidationCache.KeySort(FNameLexicalLess());
+			for (TTuple<FName, TMap<FName, FIoHash>>& PackageValidationCachePair : PlatformArtifactData.PackageValidationCache)
 			{
-				const FString ExternalActorsPrefix = ULevel::GetExternalActorsPath(WorldValidationCachePair.Key.ToString()) / TEXT("");
+				const FString PackageName            = PackageValidationCachePair.Key.ToString();
+				const FString ExternalActorsPrefix   = ULevel::GetExternalActorsPath(PackageName) / TEXT("");
+				const FString ExternalPackagesPrefix = PackageName / TEXT("");
+
+				// Split the combined external package hashes into external actors and non-external-actor packages for compactness,
+				// stripping the common root path prefix from each group
+				UE::Json::FValue ExternalActorsJsonObject(rapidjson::kObjectType);
+				UE::Json::FValue ExternalPackagesJsonObject(rapidjson::kObjectType);
 
-				UE::Json::FValue ExternalActorHashesJsonObject(rapidjson::kObjectType);
-				WorldValidationCachePair.Value.KeySort(FNameLexicalLess());
-				for (const TTuple<FName, FIoHash>& ExternalActorHashesPair : WorldValidationCachePair.Value)
+				PackageValidationCachePair.Value.KeySort(FNameLexicalLess());
+				for (const TTuple<FName, FIoHash>& ExternalPackageHashesPair : PackageValidationCachePair.Value)
 				{
-					// All external actors in this world have a common prefix, so strip that off when saving for compactness
-					FString ExternalActorPath = ExternalActorHashesPair.Key.ToString();
-					ExternalActorPath.RemoveFromStart(ExternalActorsPrefix);
+					FNameBuilder PathBuilder(ExternalPackageHashesPair.Key);
+					FStringView Path = PathBuilder.ToView();
 
-					ExternalActorHashesJsonObject.AddMember(
-						UE::Json::MakeStringValue(ExternalActorPath, PlatformArtifactJsonAllocator),
-						UE::Json::MakeStringValue(LexToString(ExternalActorHashesPair.Value), PlatformArtifactJsonAllocator),
+					UE::Json::FValue* TargetJsonObject = &ExternalPackagesJsonObject;
+					if (Path.StartsWith(ExternalActorsPrefix))
+					{
+						Path.RemovePrefix(ExternalActorsPrefix.Len());
+						TargetJsonObject = &ExternalActorsJsonObject;
+					}
+					else if (Path.StartsWith(ExternalPackagesPrefix))
+					{
+						Path.RemovePrefix(ExternalPackagesPrefix.Len());
+					}
+
+					TargetJsonObject->AddMember(
+						UE::Json::MakeStringValue(Path, PlatformArtifactJsonAllocator),
+						UE::Json::MakeStringValue(LexToString(ExternalPackageHashesPair.Value), PlatformArtifactJsonAllocator),
 						PlatformArtifactJsonAllocator);
 				}
-				WorldValidationCacheJsonObject.AddMember(
-					UE::Json::MakeStringValue(FNameBuilder(WorldValidationCachePair.Key).ToView(), PlatformArtifactJsonAllocator),
-					MoveTemp(ExternalActorHashesJsonObject),
+
+				UE::Json::FValue PackageEntryJsonObject(rapidjson::kObjectType);
+				if (ExternalActorsJsonObject.MemberCount() > 0)
+				{
+					PackageEntryJsonObject.AddMember(UE::Json::MakeStringRef(Private::JsonKey_ExternalActors), MoveTemp(ExternalActorsJsonObject), PlatformArtifactJsonAllocator);
+				}
+				if (ExternalPackagesJsonObject.MemberCount() > 0)
+				{
+					PackageEntryJsonObject.AddMember(UE::Json::MakeStringRef(Private::JsonKey_ExternalPackages), MoveTemp(ExternalPackagesJsonObject), PlatformArtifactJsonAllocator);
+				}
+
+				PackageValidationCacheJsonObject.AddMember(
+					UE::Json::MakeStringValue(FNameBuilder(PackageValidationCachePair.Key).ToView(), PlatformArtifactJsonAllocator),
+					MoveTemp(PackageEntryJsonObject),
 					PlatformArtifactJsonAllocator);
 			}
 			PlatformArtifactJsonRootObject.AddMember(
-				UE::Json::MakeStringRef(Private::JsonKey_WorldValidationCache),
-				MoveTemp(WorldValidationCacheJsonObject),
+				UE::Json::MakeStringRef(Private::JsonKey_PackageValidationCache),
+				MoveTemp(PackageValidationCacheJsonObject),
 				PlatformArtifactJsonAllocator);
 		}
 
diff --git a/Engine/Source/Editor/UnrealEd/Private/Cooker/ValidationCookArtifact.h b/Engine/Source/Editor/UnrealEd/Private/Cooker/ValidationCookArtifact.h
--- a/Engine/Source/Editor/UnrealEd/Private/Cooker/ValidationCookArtifact.h
+++ b/Engine/Source/Editor/UnrealEd/Private/Cooker/ValidationCookArtifact.h
@@ -30,13 +30,14 @@ public:
 	virtual void UpdateOplogPackages(UE::Cook::Artifact::FUpdateOplogPackagesContext& Context) override;
 
 	/**
-	 * Given the complete set of external actors associated with the given world package, filter the set down to only the external actors that need validation.
-	 * 
-	 * @param WorldPackageName		The name of the world package that contains the actors objects.
-	 * @param InOutExternalActors	The set of external actors to filter.
-	 * @param bUpdateCache			If set, update the current state with the data for InOutExternalActors, so that the next call to FilterExternalActorsToValidate will use that as its base data for comparison.
+	 * Given the complete set of external packages (external actors and non-external-actor packages) associated with the given package,
+	 * filter the set down to only those that need validation.
+	 *
+	 * @param PackageName				The name of the package being validated.
+	 * @param InOutExternalPackages		The combined set of external actors and non-external-actor packages to filter.
+	 * @param bUpdateCache				If set, update the cache with the current hashes, so the next call uses them as the baseline for comparison.
 	 */
-	void FilterExternalActorsToValidate(const FName WorldPackageName, TArray<FAssetData>& InOutExternalActors, const bool bUpdateCache = true);
+	void FilterExternalPackagesToValidate(const FName PackageName, TArray<FAssetData>& InOutExternalPackages, const bool bUpdateCache = true);
 
 	/**
 	 * Clear any incremental validation state for the given package (eg, if it failed validation).
@@ -50,8 +51,8 @@ private:
 
 	struct FArtifactData
 	{
-		/** World Package -> External Actor Package -> Hash */
-		TMap<FName, TMap<FName, FIoHash>> WorldValidationCache;
+		/** Package -> (External Package Name -> Hash); external actors and non-external-actor packages stored together */
+		TMap<FName, TMap<FName, FIoHash>> PackageValidationCache;
 
 		/** FApp::GetBuildVersion() that was used when we last ran asset validation (@see IsAssetValidationEnabledForCook) */
 		FString LastAssetValidationBuildVersion;
@@ -65,7 +66,7 @@ private:
 
 	UCookOnTheFlyServer& COTFS;
 	FString CurrentBuildVersion;
-	bool bIsIncrementalWorldValidationEnabled = false;
+	bool bIsIncrementalExternalPackageValidationEnabled = false;
 
 	TSortedMap<const ITargetPlatform*, FArtifactData> PerPlatformData;
 };
