Fix missing brush creation when drag and dropping an AVolume blueprint in a world
- fix check when editing brush settings in Blueprint asset
- change ActorFactoryBlueprint to create brush when spawning an AVolume

#Synced-CL 52307222

diff --git a/Engine/Source/Editor/UnrealEd/Classes/ActorFactories/ActorFactoryBlueprint.h b/Engine/Source/Editor/UnrealEd/Classes/ActorFactories/ActorFactoryBlueprint.h
--- a/Engine/Source/Editor/UnrealEd/Classes/ActorFactories/ActorFactoryBlueprint.h
+++ b/Engine/Source/Editor/UnrealEd/Classes/ActorFactories/ActorFactoryBlueprint.h
@@ -18,6 +18,7 @@ class UActorFactoryBlueprint : public UActorFactory
 
 	//~ Begin UActorFactory Interface
 	UNREALED_API virtual bool PreSpawnActor( UObject* Asset, FTransform& InOutLocation ) override;
+	UNREALED_API virtual void PostSpawnActor(UObject* Asset, AActor* NewActor) override;
 	UNREALED_API virtual bool CanCreateActorFrom( const FAssetData& AssetData, FText& OutErrorMsg ) override;
 	UNREALED_API virtual AActor* GetDefaultActor( const FAssetData& AssetData ) override;
 	UNREALED_API virtual UClass* GetDefaultActorClass( const FAssetData& AssetData ) override;
diff --git a/Engine/Source/Editor/UnrealEd/Private/EditorBrushBuilder.cpp b/Engine/Source/Editor/UnrealEd/Private/EditorBrushBuilder.cpp
--- a/Engine/Source/Editor/UnrealEd/Private/EditorBrushBuilder.cpp
+++ b/Engine/Source/Editor/UnrealEd/Private/EditorBrushBuilder.cpp
@@ -58,8 +58,11 @@ bool UEditorBrushBuilder::EndBrush( UWorld* InWorld, ABrush* InBrush )
 		return false;
 	}
 
-	check( InWorld != nullptr );
-	ABrush* BuilderBrush = (InBrush != nullptr) ? InBrush : InWorld->GetDefaultBrush();
+	ABrush* BuilderBrush = (InBrush != nullptr) ? InBrush : (InWorld != nullptr) ? InWorld->GetDefaultBrush() : nullptr;
+	if (!ensure(BuilderBrush))
+	{
+		return false;
+	}
 
 	// Ensure the builder brush is unhidden.
 	BuilderBrush->SetHidden(false);
diff --git a/Engine/Source/Editor/UnrealEd/Private/Factories/ActorFactory.cpp b/Engine/Source/Editor/UnrealEd/Private/Factories/ActorFactory.cpp
--- a/Engine/Source/Editor/UnrealEd/Private/Factories/ActorFactory.cpp
+++ b/Engine/Source/Editor/UnrealEd/Private/Factories/ActorFactory.cpp
@@ -1599,6 +1599,19 @@ bool UActorFactoryBlueprint::PreSpawnActor( UObject* Asset, FTransform& InOutLoc
 	return true;
 }
 
+void UActorFactoryBlueprint::PostSpawnActor(UObject* Asset, AActor* NewActor)
+{
+	Super::PostSpawnActor(Asset, NewActor);
+
+	AVolume* VolumeActor = Cast<AVolume>(NewActor);
+	if ( VolumeActor != nullptr )
+	{
+		const UBrushBuilder* ActorBrushBuilder = VolumeActor->GetBrushBuilder();
+		UBrushBuilder* BrushBuilder = ActorBrushBuilder ? DuplicateObject<UBrushBuilder>(ActorBrushBuilder, this) : NewObject<UCubeBuilder>();
+		CreateBrushForVolumeActor( VolumeActor, BrushBuilder );
+	}
+}
+
 /*-----------------------------------------------------------------------------
 UActorFactoryDirectionalLight
 -----------------------------------------------------------------------------*/
