[Level Viewport selection brackets] Fix screen space bounding box offset

#Synced-CL 52295285

diff --git a/Engine/Source/Editor/UnrealEd/Private/EditorViewportClient.cpp b/Engine/Source/Editor/UnrealEd/Private/EditorViewportClient.cpp
--- a/Engine/Source/Editor/UnrealEd/Private/EditorViewportClient.cpp
+++ b/Engine/Source/Editor/UnrealEd/Private/EditorViewportClient.cpp
@@ -7130,6 +7130,8 @@ void FEditorViewportClient::DrawBoundingBox(FBox &Box, FCanvas* InCanvas, const
 			}
 		}
 
+		ScreenBoxMax/=InCanvas->GetDPIScale();
+		ScreenBoxMin/=InCanvas->GetDPIScale();
 
 		FCanvasLineItem LineItem( FVector2D( 0.0f, 0.0f ), FVector2D( 0.0f, 0.0f ) );
 		LineItem.SetColor( InColor );
