[EditorDomain]

Early out of incremental cook construction checks when LogCore log levels are set lower than the check logs report with.

#rnx
#Synced-CL 52104572

diff --git a/Engine/Source/Editor/UnrealEd/Private/EditorDomain/EditorDomainUtils.cpp b/Engine/Source/Editor/UnrealEd/Private/EditorDomain/EditorDomainUtils.cpp
--- a/Engine/Source/Editor/UnrealEd/Private/EditorDomain/EditorDomainUtils.cpp
+++ b/Engine/Source/Editor/UnrealEd/Private/EditorDomain/EditorDomainUtils.cpp
@@ -2575,6 +2575,12 @@ void FConstructionDuringUpgradeTracker::OnObjectConstructed(UObject* Constructed
 {
 	using namespace UE::EditorDomain;
 
+	if (LogCore.GetVerbosity() < ELogVerbosity::Verbose)
+	{
+		// We only log if LogCore is set to >= Verbose, so skip everything otherwise
+		return;
+	}
+
 	check(Constructed);
 	const PackageAccessTracking_Private::FTrackedData* AccumulatedScopeData =
 		PackageAccessTracking_Private::FPackageAccessRefScope::GetCurrentThreadAccumulatedData();
