Skeletal Mesh & Deformers: use a readback mechanism to compute deformed components bounds when focusing

#rnx

#Synced-CL 52331233

diff --git a/Engine/Source/Editor/UnrealEd/Private/EditorServer.cpp b/Engine/Source/Editor/UnrealEd/Private/EditorServer.cpp
--- a/Engine/Source/Editor/UnrealEd/Private/EditorServer.cpp
+++ b/Engine/Source/Editor/UnrealEd/Private/EditorServer.cpp
@@ -174,6 +174,8 @@
 #include "Subsystems/BrushEditingSubsystem.h"
 #include "EdMode.h"
 #include "LevelEditorSubsystem.h"
+#include "Elements/Actor/ActorElementData.h"
+#include "Elements/Component/ComponentElementData.h"
 
 #include "Serialization/StructuredArchive.h"
 #include "Serialization/Formatters/JsonArchiveInputFormatter.h"
@@ -4794,7 +4796,87 @@ void UEditorEngine::MoveViewportCamerasToActor(const TArray<AActor*> &Actors, bo
 	MoveViewportCamerasToActor( Actors, TArray<UPrimitiveComponent*>(), bActiveViewportOnly );
 }
 
-void UEditorEngine::MoveViewportCamerasToActor(const TArray<AActor*> &Actors, const TArray<UPrimitiveComponent*>& Components, bool bActiveViewportOnly)
+#if WITH_EDITORONLY_DATA
+
+namespace UE::EditorEngine::Private
+{
+	// For components with an active deformer graph the true world-space vertex positions live in a GPU buffer
+	// and are not yet reflected in PrimitiveComponent::Bounds. Request an async readback for each and defer the camera movement until all results
+	// have been computed. InInitialBox already contains the non-deformer bounds.
+	
+	void DeferDeformedMeshesFocus(
+		const TSet<USkinnedMeshComponent*>& InDeformedMeshes,
+		const FBox& InInitialBox,
+		TFunction<void(const FBox&)>&& InFocusFunction)
+	{
+		if (InDeformedMeshes.IsEmpty())
+		{
+			return;
+		}
+		
+		struct FPendingFocusState
+		{
+			FPendingFocusState(const int32 NumDeformedMeshes, const FBox& InBox)
+				: NumPendingMeshes(NumDeformedMeshes)
+				, Box(InBox)
+			{}
+			
+			int32 NumPendingMeshes = 0;			
+			FBox Box{ ForceInit };
+			FCriticalSection Mutex;
+		};
+		TSharedRef<FPendingFocusState> State = MakeShared<FPendingFocusState>(InDeformedMeshes.Num(), InInitialBox);
+
+		for (USkinnedMeshComponent* DeformedMeshComponent : InDeformedMeshes)
+		{
+			TUniquePtr<FMeshDeformerGeometryReadbackRequest> Request = MakeUnique<FMeshDeformerGeometryReadbackRequest>();
+			Request->VertexDataArraysCallback_AnyThread =
+				[WeakComp = TWeakObjectPtr<USkinnedMeshComponent>(DeformedMeshComponent), State, InFocusFunction]
+				(const FMeshDeformerGeometryReadbackVertexDataArrays& Data)
+					{
+						FBox ComponentBox(ForceInit);
+						if (USkinnedMeshComponent* DeformedMeshComponent = WeakComp.Get())
+						{
+							if (!Data.Positions.IsEmpty())
+							{
+								const FTransform& Transform = DeformedMeshComponent->GetComponentTransform();
+								for (const FVector3f& Pos : Data.Positions)
+								{
+									ComponentBox += Transform.TransformPosition(FVector(Pos));
+								}
+							}
+							else
+							{
+								ComponentBox = DeformedMeshComponent->Bounds.GetBox();
+							}
+						}
+
+						bool bLast = false;
+						FBox FinalBox(ForceInit);
+						{
+							FScopeLock Lock(&State->Mutex);
+							State->Box += ComponentBox;
+							bLast = (--State->NumPendingMeshes == 0);
+							FinalBox = State->Box;
+						}
+
+						if (bLast)
+						{
+							AsyncTask(ENamedThreads::GameThread, [InFocusFunction, FinalBox]()
+							{
+								InFocusFunction(FinalBox);
+							});
+						}
+					};
+
+				DeformedMeshComponent->RequestReadbackRenderGeometry(MoveTemp(Request));
+			}
+		}
+	}
+
+#endif
+
+void UEditorEngine::MoveViewportCamerasToActor(const TArray<AActor*>& Actors, const TArray<UPrimitiveComponent*>& Components, bool bActiveViewportOnly)
 {
 	if( Actors.Num() == 0 && Components.Num() == 0 )
 	{
@@ -4816,6 +4898,24 @@ void UEditorEngine::MoveViewportCamerasToActor(const TArray<AActor*> &Actors, co
 	// Create a bounding volume of all of the selected actors.
 	FBox BoundingBox(ForceInit);
 
+	// Skeletal mesh components deformed via a deformer graph report bounds derived from bone transforms only, the GPU-displaced vertex positions
+	// are not reflected in Bounds yet. Collect them separately so we can issue an async GPU readback after the loop.
+#if WITH_EDITORONLY_DATA
+	TSet<USkinnedMeshComponent*> DeformedComponents;
+	auto ShouldDeferFocus = [&DeformedComponents](UPrimitiveComponent* PrimitiveComponent)
+	{
+		USkinnedMeshComponent* SkinnedMeshComponent = Cast<USkinnedMeshComponent>(PrimitiveComponent);
+		if ( SkinnedMeshComponent && SkinnedMeshComponent->HasMeshDeformer())
+		{
+			DeformedComponents.Add(SkinnedMeshComponent);
+			return true;
+		}
+		return false;
+	};
+#else
+	auto ShouldDeferFocus = [](UPrimitiveComponent* PrimitiveComponent) { return false; }
+#endif
+
 	if( Components.Num() > 0 )
 	{
 		// First look at components
@@ -4833,9 +4933,11 @@ void UEditorEngine::MoveViewportCamerasToActor(const TArray<AActor*> &Actors, co
 
 				if(!bIgnore && PrimitiveComponent->IsRegistered())
 				{
-					BoundingBox += PrimitiveComponent->Bounds.GetBox();
+					if (!ShouldDeferFocus(PrimitiveComponent))
+					{
+						BoundingBox += PrimitiveComponent->Bounds.GetBox();
+					}
 				}
-
 			}
 		}
 	}
@@ -4884,14 +4986,17 @@ void UEditorEngine::MoveViewportCamerasToActor(const TArray<AActor*> &Actors, co
 
 								if (!bIgnore)
 								{
-									FBox LocalBox(ForceInit);
-									if (GLevelEditorModeTools().ComputeBoundingBoxForViewportFocus(AlignActor, PrimitiveComponent, LocalBox))
-									{
-										BoundingBox += LocalBox;
-									}
-									else
+									if (!ShouldDeferFocus(PrimitiveComponent))
 									{
-										BoundingBox += PrimitiveComponent->Bounds.GetBox();
+										FBox LocalBox(ForceInit);
+										if (GLevelEditorModeTools().ComputeBoundingBoxForViewportFocus(AlignActor, PrimitiveComponent, LocalBox))
+										{
+											BoundingBox += LocalBox;
+										}
+										else
+										{
+											BoundingBox += PrimitiveComponent->Bounds.GetBox();
+										}
 									}
 
 									bHasAtLeastOnePrimitiveComponent = true;
@@ -4910,8 +5015,6 @@ void UEditorEngine::MoveViewportCamerasToActor(const TArray<AActor*> &Actors, co
 		}
 	}
 
-	MoveViewportCamerasToBox(BoundingBox, bActiveViewportOnly);
-
 	// Warn the user with a suppressable dialog if they attempted to zoom to actors that are in an invisible level
 	if ( InvisLevelActors.Num() > 0 )
 	{
@@ -4930,6 +5033,19 @@ void UEditorEngine::MoveViewportCamerasToActor(const TArray<AActor*> &Actors, co
 		InvisLevelActorWarning.ShowModal();
 	}
 
+#if WITH_EDITORONLY_DATA
+	if (!DeformedComponents.IsEmpty())
+	{
+		UE::EditorEngine::Private::DeferDeformedMeshesFocus(DeformedComponents, BoundingBox, [this, bActiveViewportOnly, Actors](const FBox& FinalBox)
+		{
+			MoveViewportCamerasToBox(FinalBox, bActiveViewportOnly);
+			FEditorDelegates::OnFocusViewportOnActors.Broadcast(Actors);
+		});
+		return;
+	}
+#endif // WITH_EDITORONLY_DATA
+
+	MoveViewportCamerasToBox(BoundingBox, bActiveViewportOnly);
 	// Notify 'focus on actors' delegate
 	FEditorDelegates::OnFocusViewportOnActors.Broadcast(Actors);
 }
@@ -4973,7 +5089,8 @@ void UEditorEngine::MoveViewportCamerasToElement(const UTypedElementSelectionSet
 		return;
 	}
 
-	// Create a bounding volume of all of the selected elements.
+	// GetFocusBounds now returns false for any USkinnedMeshComponent that has an active deformer graph. BoundingBox therefore contains only the
+	// correctly-bounded non-deformer primitives.
 	FBox BoundingBox(ForceInit);
 
 	FTypedElementListConstRef ElementsToView = SelectionSet->GetNormalizedElementList(SelectionSet->GetElementList(), FTypedElementSelectionNormalizationOptions());
@@ -4999,6 +5116,48 @@ void UEditorEngine::MoveViewportCamerasToElement(const UTypedElementSelectionSet
 		}
 	}
 
+#if WITH_EDITORONLY_DATA
+	// Skeletal mesh components deformed via a deformer graph report bounds derived from bone transforms only, the GPU-displaced vertex positions
+	// are not reflected in Bounds yet. Collect them separately so we can issue an async GPU readback after the loop.
+	TSet<USkinnedMeshComponent*> DeformedComponents;
+
+	ElementsToView->ForEachElementHandle( [&DeformedComponents](const FTypedElementHandle& InHandle)
+    {
+		constexpr bool bSilent = true;
+	    if (AActor* Actor = ActorElementDataUtil::GetActorFromHandle(InHandle, bSilent))
+	    {
+		    TInlineComponentArray<USkinnedMeshComponent*> SkinnedMeshComponents(Actor);
+		    for (USkinnedMeshComponent* SkinnedMeshComponent : SkinnedMeshComponents)
+		    {
+			    if (SkinnedMeshComponent && SkinnedMeshComponent->HasMeshDeformer())
+			    {
+				    DeformedComponents.Add(SkinnedMeshComponent);
+			    }
+		    }
+	    }
+	    else if (UActorComponent* Component = ComponentElementDataUtil::GetComponentFromHandle(InHandle, bSilent))
+	    {
+		    if (USkinnedMeshComponent* SkinnedMeshComponent = Cast<USkinnedMeshComponent>(Component))
+		    {
+			    if (SkinnedMeshComponent->HasMeshDeformer())
+			    {
+				    DeformedComponents.Add(SkinnedMeshComponent);
+			    }
+		    }
+	    }
+	    return true;
+    });
+
+	if (!DeformedComponents.IsEmpty())
+	{
+		UE::EditorEngine::Private::DeferDeformedMeshesFocus(DeformedComponents, BoundingBox, [this, bActiveViewportOnly](const FBox& FinalBox)
+		{
+			MoveViewportCamerasToBox(FinalBox, bActiveViewportOnly);
+		});
+		return;
+	}
+#endif // WITH_EDITORONLY_DATA
+
 	MoveViewportCamerasToBox(BoundingBox, bActiveViewportOnly);
 }
 
diff --git a/Engine/Source/Editor/UnrealEd/Private/Elements/Actor/ActorElementEditorViewportInteractionCustomization.cpp b/Engine/Source/Editor/UnrealEd/Private/Elements/Actor/ActorElementEditorViewportInteractionCustomization.cpp
--- a/Engine/Source/Editor/UnrealEd/Private/Elements/Actor/ActorElementEditorViewportInteractionCustomization.cpp
+++ b/Engine/Source/Editor/UnrealEd/Private/Elements/Actor/ActorElementEditorViewportInteractionCustomization.cpp
@@ -8,6 +8,7 @@
 #include "EditorModeManager.h"
 #include "AI/NavigationSystemBase.h"
 #include "Components/BrushComponent.h"
+#include "Components/SkinnedMeshComponent.h"
 #include "EditorGizmos/EditorTransformGizmoUtil.h"
 #include "Engine/DocumentationActor.h"
 #include "Particles/Emitter.h"
@@ -320,6 +321,17 @@ bool FActorElementEditorViewportInteractionCustomization::GetFocusBounds( const
 
 				if (!bIgnore)
 				{
+					// For components with an active deformer graph the true world-space vertex positions live in a GPU buffer
+					// and are not yet reflected in PrimitiveComponent::Bounds. Their real bounds have to be computed via an async GPU readback.
+					// We still mark bHasAtLeastOnePrimitiveComponent so the root-location fallback below does not fire for actors that are all-deformer.
+#if WITH_EDITORONLY_DATA
+					const USkinnedMeshComponent* SkinnedMeshComponent =  Cast<USkinnedMeshComponent>(PrimitiveComponent);
+					if (SkinnedMeshComponent && SkinnedMeshComponent->HasMeshDeformer())
+					{
+						bHasAtLeastOnePrimitiveComponent = true;
+						continue;
+					}
+#endif
 					FBox LocalBox(ForceInit);
 					if (GLevelEditorModeTools().ComputeBoundingBoxForViewportFocus(Actor, PrimitiveComponent, LocalBox))
 					{
diff --git a/Engine/Source/Editor/UnrealEd/Private/Elements/Component/ComponentElementEditorViewportInteractionCustomization.cpp b/Engine/Source/Editor/UnrealEd/Private/Elements/Component/ComponentElementEditorViewportInteractionCustomization.cpp
--- a/Engine/Source/Editor/UnrealEd/Private/Elements/Component/ComponentElementEditorViewportInteractionCustomization.cpp
+++ b/Engine/Source/Editor/UnrealEd/Private/Elements/Component/ComponentElementEditorViewportInteractionCustomization.cpp
@@ -8,6 +8,7 @@
 #include "Editor.h"
 #include "EditorModeManager.h"
 #include "EditorSupportDelegates.h"
+#include "Components/SkinnedMeshComponent.h"
 #include "EditorGizmos/TransformGizmoInterfaces.h"
 
 #include "Toolkits/IToolkitHost.h"
@@ -100,6 +101,16 @@ bool FComponentElementEditorViewportInteractionCustomization::GetFocusBounds(con
 				return false;
 			}
 
+			// For components with an active deformer graph the true world-space vertex positions live in a GPU buffer
+			// and are not yet reflected in PrimitiveComponent::Bounds. Their real bounds have to be computed via an async GPU readback.
+#if WITH_EDITORONLY_DATA
+			const USkinnedMeshComponent* SkinnedMeshComponent = Cast<USkinnedMeshComponent>(PrimitiveComponent);
+			if (SkinnedMeshComponent && SkinnedMeshComponent->HasMeshDeformer())
+			{
+				return false;
+			}
+#endif
+
 			AActor* AlignActor = Component->GetOwner();
 			if (!AlignActor)
 			{
