Cooker: Preloader: AsyncLoadingMultithreaded: Fix for CoreRedirects and other robustness fixes.
1. Add detection of packages that were CoreRedirected on load, and complete the preload rather than waiting for the original-name package forever.
2. Add an error and completion of the preload if the package is not found for other reasons.
3. Add handling for the race condition where the package was preloaded normally, but then was garbage collected before we tried to complete the preload.
#rnx
#Synced-CL 52335726

diff --git a/Engine/Source/Editor/UnrealEd/Private/Cooker/CookPackagePreloader.cpp b/Engine/Source/Editor/UnrealEd/Private/Cooker/CookPackagePreloader.cpp
--- a/Engine/Source/Editor/UnrealEd/Private/Cooker/CookPackagePreloader.cpp
+++ b/Engine/Source/Editor/UnrealEd/Private/Cooker/CookPackagePreloader.cpp
@@ -209,7 +209,7 @@ bool FPackagePreloader::TryPreload()
 	{
 		if (AsyncRequest || IsPreloaded())
 		{
-			if (AsyncRequest && !AsyncRequest->bHasFinished)
+			if (AsyncRequest && !AsyncRequest->IsCompleted())
 			{
 				// In case of async loading, the object can be found while still being asynchronously serialized, we need
 				// to wait until the callback is called and the async request is completely done.
@@ -228,6 +228,11 @@ bool FPackagePreloader::TryPreload()
 		{
 			return true;
 		}
+		else if (AsyncRequest && AsyncRequest->IsFailed())
+		{
+			SetIsPreloadAttempted(true);
+			return true;
+		}
 	}
 
 	if (PackageData.IsGenerated())
@@ -239,6 +244,15 @@ bool FPackagePreloader::TryPreload()
 	}
 	if (IsAsyncLoadingMultithreaded())
 	{
+		if (AsyncRequest && AsyncRequest->IsCompleted() && !IsPackageLoaded())
+		{
+			// The AsyncRequest previously finished, but we garbage collected and removed the package from memory
+			// before we had a chance to test IsPackageLoaded. Restart the AsyncRequest.
+			AsyncRequest.Reset();
+			// We will bump it back up down below.
+			PackageData.GetPackageDatas().GetMonitor().OnPreloadAllocatedChanged(PackageData, false);
+		}
+
 		if (!AsyncRequest.IsValid())
 		{
 			// Declare the package as being referenced by the ScriptPackageNameEngine; we look for this in
@@ -249,16 +263,42 @@ bool FPackagePreloader::TryPreload()
 
 			PackageData.GetPackageDatas().GetMonitor().OnPreloadAllocatedChanged(PackageData, true);
 			AsyncRequest = MakeShared<FAsyncRequest>();
-			AsyncRequest->RequestID = LoadPackageAsync(
+			int32 RequestID = LoadPackageAsync(
 				PackageData.GetFileName().ToString(),
 				FLoadPackageAsyncDelegate::CreateLambda(
-					[AsyncRequest = AsyncRequest](const FName&, UPackage*, EAsyncLoadingResult::Type)
+					[AsyncRequest = AsyncRequest, PackageName = this->PackageData.GetPackageName()]
+					(const FName&, UPackage* LoadedPackage, EAsyncLoadingResult::Type)
 					{
-						AsyncRequest->bHasFinished = true;
+						if (!LoadedPackage)
+						{
+							// PackageLoad failed, mark the preload as complete; UCookOnTheFlyServer::LoadPackageInQueue will
+							// call LoadPackage and handle the error.
+							AsyncRequest->SetFailed();
+						}
+						else if (LoadedPackage->GetFName() != PackageName)
+						{
+							// The PackageName is not the name that we loaded. This can happen due to CoreRedirects.
+							// Do not attempt to handle this here in the preloader; just mark the preload complete.
+							// UCookOnTheFlyServer::LoadPackageInQueue will call LoadPackage and will detect the renamed
+							// package there and will handle the update of the cooker.
+							AsyncRequest->SetFailed();
+						}
+						else if (!FPackagePreloader::IsPackageLoaded(PackageName))
+						{
+							UE_LOG(LogCook, Error,
+								TEXT("Package %s finished loading asynchronously, but IsPackageLoaded returned false. We will try to recover, but this error is unexpected."),
+								*PackageName.ToString());
+							AsyncRequest->SetFailed();
+						}
+						else
+						{
+							AsyncRequest->SetSucceeded();
+						}
 					}
 				),
 				32 /* Use arbitrary higher priority for preload as we're going to need them very soon */
 			);
+			AsyncRequest->SetRequestID(RequestID);
 		}
 
 		// always return false so we continue to check the status of the load until IsPackageLoaded().
@@ -342,10 +382,10 @@ void FPackagePreloader::ClearPreload()
 {
 	if (AsyncRequest)
 	{
-		if (!AsyncRequest->bHasFinished)
+		if (!AsyncRequest->IsCompleted())
 		{
-			FlushAsyncLoading(AsyncRequest->RequestID);
-			check(AsyncRequest->bHasFinished);
+			FlushAsyncLoading(AsyncRequest->GetRequestID());
+			check(AsyncRequest->IsCompleted());
 		}
 		PackageData.GetPackageDatas().GetMonitor().OnPreloadAllocatedChanged(PackageData, false);
 		AsyncRequest.Reset();
@@ -442,12 +482,17 @@ void FPackagePreloader::OnPackageLeaveLoadState()
 
 bool FPackagePreloader::IsPackageLoaded() const
 {
-	return IsPackageLoaded(PackageData);
+	return IsPackageLoaded(PackageData.GetPackageName());
 }
 
 bool FPackagePreloader::IsPackageLoaded(const FPackageData& InPackageData)
 {
-	return FindObjectFast<UPackage>(nullptr, InPackageData.GetPackageName()) != nullptr;
+	return IsPackageLoaded(InPackageData.GetPackageName());
+}
+
+bool FPackagePreloader::IsPackageLoaded(FName InPackageName)
+{
+	return FindObjectFast<UPackage>(nullptr, InPackageName) != nullptr;
 }
 
 void FPackagePreloader::IncrementCountFromRequestedLoads()
diff --git a/Engine/Source/Editor/UnrealEd/Private/Cooker/CookPackagePreloader.h b/Engine/Source/Editor/UnrealEd/Private/Cooker/CookPackagePreloader.h
--- a/Engine/Source/Editor/UnrealEd/Private/Cooker/CookPackagePreloader.h
+++ b/Engine/Source/Editor/UnrealEd/Private/Cooker/CookPackagePreloader.h
@@ -45,7 +45,8 @@ public:
 	bool IsPackageLoaded() const;
 	/** Return whether the given PackageData's package is loaded. */
 	static bool IsPackageLoaded(const FPackageData& InPackageData);
-	int32 GetCountFromRequestedLoads(); 
+	static bool IsPackageLoaded(FName InPackageName);
+	int32 GetCountFromRequestedLoads();
 
 	/**
 	 * State transition function that calls OnEnter/OnExit functions and (optionally) removes and adds the
@@ -79,8 +80,26 @@ private:
 	 */
 	struct FAsyncRequest
 	{
+	public:
+		bool IsSucceeded() const;
+		bool IsFailed() const;
+		bool IsCompleted() const;
+
+		void SetSucceeded();
+		void SetFailed();
+
+		int32 GetRequestID() const;
+		void SetRequestID(int32 InRequestID);
+
+	private:
+		enum ECompletionState : uint8
+		{
+			Incomplete,
+			Failed,
+			Succeeded
+		};
 		int32 RequestID{ 0 };
-		std::atomic<bool> bHasFinished{ false };
+		std::atomic<ECompletionState> CompletionState{ ECompletionState::Incomplete };
 	};
 
 	/** Used during DFS of the import graph. */
@@ -304,4 +323,39 @@ inline void FPackagePreloader::SetIsImportsGathered(bool bValue)
 	bImportsGathered = bValue;
 }
 
+inline bool FPackagePreloader::FAsyncRequest::IsSucceeded() const
+{
+	return CompletionState.load() == ECompletionState::Succeeded;
+}
+
+inline bool FPackagePreloader::FAsyncRequest::IsFailed() const
+{
+	return CompletionState.load() == ECompletionState::Failed;
+}
+
+inline bool FPackagePreloader::FAsyncRequest::IsCompleted() const
+{
+	return CompletionState.load() != ECompletionState::Incomplete;
+}
+
+inline void FPackagePreloader::FAsyncRequest::SetSucceeded()
+{
+	CompletionState = ECompletionState::Succeeded;
+}
+
+inline void FPackagePreloader::FAsyncRequest::SetFailed()
+{
+	CompletionState = ECompletionState::Failed;
+}
+
+inline int32 FPackagePreloader::FAsyncRequest::GetRequestID() const
+{
+	return RequestID;
+}
+
+inline void FPackagePreloader::FAsyncRequest::SetRequestID(int32 InRequestID)
+{
+	RequestID = InRequestID;
+}
+
 }
\ No newline at end of file
