Rename bPreviewFeatureLevelActive  to bPreviewShaderPlatformActive

#Synced-CL 52380800

diff --git a/Engine/Source/Editor/LevelEditor/Private/LevelEditorToolBar.cpp b/Engine/Source/Editor/LevelEditor/Private/LevelEditorToolBar.cpp
--- a/Engine/Source/Editor/LevelEditor/Private/LevelEditorToolBar.cpp
+++ b/Engine/Source/Editor/LevelEditor/Private/LevelEditorToolBar.cpp
@@ -100,8 +100,8 @@ namespace PreviewModeFunctionality
 		FText FriendlyNamePreviewOff = FDataDrivenShaderPlatformInfo::GetFriendlyName(GMaxRHIShaderPlatform);
 		FText FriendlyNamePreviewOn = GEditor->PreviewPlatform.GetFriendlyName();
 		{
-			const FText& RenderingAsPlatformName = GEditor->PreviewPlatform.bPreviewFeatureLevelActive ? FriendlyNamePreviewOn : FriendlyNamePreviewOff;
-			const FText& SwitchToPlatformName = GEditor->PreviewPlatform.bPreviewFeatureLevelActive ? FriendlyNamePreviewOff : FriendlyNamePreviewOn;
+			const FText& RenderingAsPlatformName = GEditor->PreviewPlatform.bPreviewShaderPlatformActive ? FriendlyNamePreviewOn : FriendlyNamePreviewOff;
+			const FText& SwitchToPlatformName = GEditor->PreviewPlatform.bPreviewShaderPlatformActive ? FriendlyNamePreviewOff : FriendlyNamePreviewOn;
 			if (PreviewShaderPlatform == GMaxRHIShaderPlatform)
 			{
 				return FText::Format(LOCTEXT("PreviewModeViewingAs", "Viewing {0}."), RenderingAsPlatformName);
diff --git a/Engine/Source/Editor/UnrealEd/Classes/Editor/EditorEngine.h b/Engine/Source/Editor/UnrealEd/Classes/Editor/EditorEngine.h
--- a/Engine/Source/Editor/UnrealEd/Classes/Editor/EditorEngine.h
+++ b/Engine/Source/Editor/UnrealEd/Classes/Editor/EditorEngine.h
@@ -217,7 +217,7 @@ struct FPreviewPlatformInfo
 	,	ShaderPlatform(EShaderPlatform::SP_NumPlatforms)
 	,	PreviewPlatformName(NAME_None)
 	,	PreviewShaderFormatName(NAME_None)
-	,	bPreviewFeatureLevelActive(false)
+	,	bPreviewShaderPlatformActive(false)
 	,	PreviewShaderPlatformName(NAME_None)
 	,	PreviewShaderPlatformFriendlyName()
 	,	ConstrainedAspectRatio(0.0f)
@@ -229,7 +229,7 @@ struct FPreviewPlatformInfo
 		FName InPreviewPlatformName = NAME_None, 
 		FName InPreviewShaderFormatName = NAME_None, 
 		FName InDeviceProfileName = NAME_None, 
-		bool InbPreviewFeatureLevelActive = false, 
+		bool InbPreviewShaderPlatformActive = false, 
 		FName InShaderPlatformName = NAME_None, 
 		FText InPreviewShaderPlatformFriendlyName = FText(), 
 		float InConstrainedAspectRatio = 0.0f,
@@ -240,7 +240,7 @@ struct FPreviewPlatformInfo
 	,	PreviewPlatformName(InPreviewPlatformName)
 	,	PreviewShaderFormatName(InPreviewShaderFormatName)
 	,	DeviceProfileName(InDeviceProfileName)
-	,	bPreviewFeatureLevelActive(InbPreviewFeatureLevelActive)
+	,	bPreviewShaderPlatformActive(InbPreviewShaderPlatformActive)
 	,	PreviewShaderPlatformName(InShaderPlatformName)
 	,	PreviewShaderPlatformFriendlyName(InPreviewShaderPlatformFriendlyName)
 	,	ConstrainedAspectRatio(InConstrainedAspectRatio)
@@ -264,13 +264,13 @@ struct FPreviewPlatformInfo
 	/** The device profile to preview. */
 	FName DeviceProfileName;
 
-	/** Is feature level preview currently active */
-	bool bPreviewFeatureLevelActive;
-	
+	/** Is shader platform preview currently active */
+	bool bPreviewShaderPlatformActive;
+
 	/** Preview Shader Platform Name */
 	FName PreviewShaderPlatformName;
 
-	/** Checks if two FPreviewPlatformInfos are for the same preview platform. Note, this does NOT compare the bPreviewFeatureLevelActive flag */
+	/** Checks if two FPreviewPlatformInfos are for the same preview platform. Note, this does NOT compare the bPreviewShaderPlatformActive flag */
 	bool Matches(const FPreviewPlatformInfo& Other) const
 	{
 		return PreviewFeatureLevel == Other.PreviewFeatureLevel && ShaderPlatform == Other.ShaderPlatform && PreviewPlatformName == Other.PreviewPlatformName && PreviewShaderFormatName == Other.PreviewShaderFormatName && DeviceProfileName == Other.DeviceProfileName && PreviewShaderPlatformName == Other.PreviewShaderPlatformName;
@@ -279,7 +279,7 @@ struct FPreviewPlatformInfo
 	/** Return platform name like "Android", or NAME_None if none is set or the preview feature level is not active */
 	FName GetEffectivePreviewPlatformName() const
 	{
-		return bPreviewFeatureLevelActive ? PreviewPlatformName : NAME_None;
+		return bPreviewShaderPlatformActive ? PreviewPlatformName : NAME_None;
 	}
 
 	/** Return the Constrained Aspect Ratio of the preview platform */
diff --git a/Engine/Source/Editor/UnrealEd/Classes/Editor/EditorPerProjectUserSettings.h b/Engine/Source/Editor/UnrealEd/Classes/Editor/EditorPerProjectUserSettings.h
--- a/Engine/Source/Editor/UnrealEd/Classes/Editor/EditorPerProjectUserSettings.h
+++ b/Engine/Source/Editor/UnrealEd/Classes/Editor/EditorPerProjectUserSettings.h
@@ -216,9 +216,13 @@ public:
 	UPROPERTY(config)
 	FVector4f PreviewSafeZones;
 
-	/** Is feature level preview currently active */
+	/** Is shader platform preview currently active */
 	UPROPERTY(config)
-	bool bPreviewFeatureLevelActive;
+	bool bPreviewShaderPlatformActive;
+
+	/** Deprecated: renamed to bPreviewShaderPlatformActive */
+	UPROPERTY(config, meta=(DeprecatedProperty, DeprecationMessage="bPreviewFeatureLevelActive has been renamed to bPreviewShaderPlatformActive."))
+	bool bPreviewFeatureLevelActive = false;
 
 	UPROPERTY(config)
 	bool bPreviewFeatureLevelWasDefault;
diff --git a/Engine/Source/Editor/UnrealEd/Private/EditorEngine.cpp b/Engine/Source/Editor/UnrealEd/Private/EditorEngine.cpp
--- a/Engine/Source/Editor/UnrealEd/Private/EditorEngine.cpp
+++ b/Engine/Source/Editor/UnrealEd/Private/EditorEngine.cpp
@@ -458,7 +458,7 @@ static void CheckForMissingAdvancedRenderingRequirements()
 
 ERHIFeatureLevel::Type FPreviewPlatformInfo::GetEffectivePreviewFeatureLevel() const
 {
-	return bPreviewFeatureLevelActive ? PreviewFeatureLevel : GMaxRHIFeatureLevel;
+	return bPreviewShaderPlatformActive ? PreviewFeatureLevel : GMaxRHIFeatureLevel;
 }
 
 EShaderPlatform FPreviewPlatformInfo::GetShaderPlatform() const
@@ -6546,7 +6546,7 @@ void UEditorEngine::OnEffectivePreviewShaderPlatformChange()
 		if (StereoRenderTargetManager)
 		{
 			StereoRenderTargetManager->ReconfigureForShaderPlatform(
-				PreviewPlatform.bPreviewFeatureLevelActive ? PreviewPlatform.ShaderPlatform : GMaxRHIShaderPlatform);
+				PreviewPlatform.bPreviewShaderPlatformActive ? PreviewPlatform.ShaderPlatform : GMaxRHIShaderPlatform);
 		}
 	}
 }
@@ -6558,7 +6558,7 @@ static void SaveShaderPlatformAsDisabled(FPreviewPlatformInfo& PreviewPlatform)
 	Settings->PreviewFeatureLevel = 0;
 	Settings->PreviewPlatformName = NAME_None;
 	Settings->PreviewShaderFormatName = NAME_None;
-	Settings->bPreviewFeatureLevelActive = false;
+	Settings->bPreviewShaderPlatformActive = false;
 	Settings->bPreviewFeatureLevelWasDefault = true;
 	Settings->PreviewDeviceProfileName = NAME_None;
 	Settings->PreviewShaderPlatformName = NAME_None;
@@ -6569,7 +6569,7 @@ static void SaveShaderPlatformAsDisabled(FPreviewPlatformInfo& PreviewPlatform)
 	Settings->PreviewFeatureLevel = (int32)PreviewPlatform.PreviewFeatureLevel;
 	Settings->PreviewPlatformName = PreviewPlatform.PreviewPlatformName;
 	Settings->PreviewShaderFormatName = PreviewPlatform.PreviewShaderFormatName;
-	Settings->bPreviewFeatureLevelActive = PreviewPlatform.bPreviewFeatureLevelActive;
+	Settings->bPreviewShaderPlatformActive = PreviewPlatform.bPreviewShaderPlatformActive;
 	Settings->bPreviewFeatureLevelWasDefault = (PreviewPlatform.PreviewFeatureLevel == GMaxRHIFeatureLevel);
 	Settings->PreviewDeviceProfileName = PreviewPlatform.DeviceProfileName;
 	Settings->PreviewShaderPlatformName = PreviewPlatform.PreviewShaderPlatformName;
@@ -6682,8 +6682,8 @@ void UEditorEngine::SetPreviewPlatform(const FPreviewPlatformInfo& NewPreviewPla
 
 		const bool bChangedPreviewShaderPlatform = NewPreviewPlatform.ShaderPlatform != PreviewPlatform.ShaderPlatform;
 		const bool bChangedEffectiveShaderPlatform = bChangedPreviewShaderPlatform
-			? (PreviewPlatform.bPreviewFeatureLevelActive || NewPreviewPlatform.bPreviewFeatureLevelActive)
-			: (NewPreviewPlatform.bPreviewFeatureLevelActive != PreviewPlatform.bPreviewFeatureLevelActive);
+			? (PreviewPlatform.bPreviewShaderPlatformActive || NewPreviewPlatform.bPreviewShaderPlatformActive)
+			: (NewPreviewPlatform.bPreviewShaderPlatformActive != PreviewPlatform.bPreviewShaderPlatformActive);
 		const ERHIFeatureLevel::Type EffectiveFeatureLevel = NewPreviewPlatform.GetEffectivePreviewFeatureLevel();
 
 		if (NewPreviewPlatform.PreviewShaderFormatName != NAME_None)
@@ -6723,14 +6723,14 @@ void UEditorEngine::SetPreviewPlatform(const FPreviewPlatformInfo& NewPreviewPla
 		auto UpdateDeviceProfile = [&](bool bRestore)
 		{
 			// Always restore if no preview profile will be set, or if explicitly requested (compile path)
-			bRestore = bRestore || !PreviewPlatform.bPreviewFeatureLevelActive || PreviewPlatform.DeviceProfileName == NAME_None;
+			bRestore = bRestore || !PreviewPlatform.bPreviewShaderPlatformActive || PreviewPlatform.DeviceProfileName == NAME_None;
 
 			if (bRestore)
 			{
 				UDeviceProfileManager::Get().RestorePreviewDeviceProfile();
 			}
 
-			if (PreviewPlatform.bPreviewFeatureLevelActive && PreviewPlatform.DeviceProfileName != NAME_None)
+			if (PreviewPlatform.bPreviewShaderPlatformActive && PreviewPlatform.DeviceProfileName != NAME_None)
 			{
 				if (UDeviceProfile* DP = UDeviceProfileManager::Get().FindProfile(PreviewPlatform.DeviceProfileName.ToString(), false))
 				{
@@ -6763,7 +6763,7 @@ void UEditorEngine::SetPreviewPlatform(const FPreviewPlatformInfo& NewPreviewPla
 
 			// Restore editor's shader platform in case a previous SM6 preview left it stale
 			GShaderPlatformForFeatureLevel[GMaxRHIFeatureLevel] = GMaxRHIShaderPlatform;
-			if (PreviewPlatform.bPreviewFeatureLevelActive)
+			if (PreviewPlatform.bPreviewShaderPlatformActive)
 			{
 				GShaderPlatformForFeatureLevel[PreviewPlatform.PreviewFeatureLevel] = ShaderPlatform;
 			}
@@ -6823,7 +6823,7 @@ void UEditorEngine::SetPreviewPlatform(const FPreviewPlatformInfo& NewPreviewPla
 void UEditorEngine::TogglePreviewPlatform()
 {
 	FPreviewPlatformInfo NewInfo = PreviewPlatform;
-	NewInfo.bPreviewFeatureLevelActive ^= 1;
+	NewInfo.bPreviewShaderPlatformActive ^= 1;
 
 	// When the preview shares the editor's feature level, toggling changes GShaderPlatformForFeatureLevel which requires the compile path.
 	const bool bNeedsCompile = (NewInfo.PreviewFeatureLevel == GMaxRHIFeatureLevel);
@@ -6842,12 +6842,12 @@ bool UEditorEngine::IsPreviewPlatformEnabled() const
 
 bool UEditorEngine::IsPreviewPlatformActive() const
 {
- 	return PreviewPlatform.bPreviewFeatureLevelActive;
+ 	return PreviewPlatform.bPreviewShaderPlatformActive;
 }
 
 EShaderPlatform UEditorEngine::GetActiveShaderPlatform() const
 {
-	if (PreviewPlatform.bPreviewFeatureLevelActive)
+	if (PreviewPlatform.bPreviewShaderPlatformActive)
 	{
 		return GShaderPlatformForFeatureLevel[PreviewPlatform.PreviewFeatureLevel];
 	}
@@ -6859,7 +6859,7 @@ EShaderPlatform UEditorEngine::GetActiveShaderPlatform() const
 
 FText UEditorEngine::GetActiveShaderPlatformFriendlyName() const
 {
-	if (PreviewPlatform.bPreviewFeatureLevelActive)
+	if (PreviewPlatform.bPreviewShaderPlatformActive)
 	{
 		return PreviewPlatform.GetFriendlyName();
 	}
@@ -6871,7 +6871,7 @@ FText UEditorEngine::GetActiveShaderPlatformFriendlyName() const
 
 ERHIFeatureLevel::Type UEditorEngine::GetActiveFeatureLevelPreviewType() const
 {
-	return PreviewPlatform.bPreviewFeatureLevelActive ? PreviewPlatform.PreviewFeatureLevel : GMaxRHIFeatureLevel;
+	return PreviewPlatform.bPreviewShaderPlatformActive ? PreviewPlatform.PreviewFeatureLevel : GMaxRHIFeatureLevel;
 }
 
 // Deprecated in 5.7
@@ -6914,7 +6914,7 @@ void UEditorEngine::LoadEditorShaderPlatformPreview()
 			Settings->PreviewPlatformName, 
 			Settings->PreviewShaderFormatName, 
 			Settings->PreviewDeviceProfileName, 
-			Settings->bPreviewFeatureLevelActive, 
+			Settings->bPreviewShaderPlatformActive, 
 			PreviewShaderPlatformName, 
 			Settings->PreviewShaderPlatformFriendlyName,
 			Settings->PreviewConstrainedAspectRatio,
@@ -6929,7 +6929,7 @@ void UEditorEngine::SaveEditorShaderPlatformPreview()
 	Settings->PreviewFeatureLevel = (int32)PreviewPlatform.PreviewFeatureLevel;
 	Settings->PreviewPlatformName = PreviewPlatform.PreviewPlatformName;
 	Settings->PreviewShaderFormatName = PreviewPlatform.PreviewShaderFormatName;
-	Settings->bPreviewFeatureLevelActive = PreviewPlatform.bPreviewFeatureLevelActive;
+	Settings->bPreviewShaderPlatformActive = PreviewPlatform.bPreviewShaderPlatformActive;
 	Settings->bPreviewFeatureLevelWasDefault = (PreviewPlatform.PreviewFeatureLevel == GMaxRHIFeatureLevel);
 	Settings->PreviewDeviceProfileName = PreviewPlatform.DeviceProfileName;
 	Settings->PreviewShaderPlatformName = PreviewPlatform.PreviewShaderPlatformName;
diff --git a/Engine/Source/Editor/UnrealEd/Private/Settings/EditorPerProjectUserSettings.cpp b/Engine/Source/Editor/UnrealEd/Private/Settings/EditorPerProjectUserSettings.cpp
--- a/Engine/Source/Editor/UnrealEd/Private/Settings/EditorPerProjectUserSettings.cpp
+++ b/Engine/Source/Editor/UnrealEd/Private/Settings/EditorPerProjectUserSettings.cpp
@@ -40,7 +40,7 @@ UEditorPerProjectUserSettings::UEditorPerProjectUserSettings(const FObjectInitia
 	PreviewFeatureLevel = GMaxRHIFeatureLevel;
 	PreviewPlatformName = NAME_None;
 	PreviewShaderFormatName = NAME_None;
-	bPreviewFeatureLevelActive = false;
+	bPreviewShaderPlatformActive = false;
 	bPreviewFeatureLevelWasDefault = true;
 	PreviewDeviceProfileName = NAME_None;
 	bShowSelectionSubcomponents = true;
@@ -60,7 +60,7 @@ void UEditorPerProjectUserSettings::PostInitProperties()
 		PreviewShaderPlatformName = NAME_None;
 		PreviewShaderFormatName = NAME_None;
         PreviewPlatformName = NAME_None;
-		bPreviewFeatureLevelActive = false;
+		bPreviewShaderPlatformActive = false;
 		bPreviewFeatureLevelWasDefault = true;
 		PreviewDeviceProfileName = NAME_None;
 	}
diff --git a/Engine/Source/Editor/UnrealEd/Private/Subsystems/EditorSubsystemBlueprintLibrary.cpp b/Engine/Source/Editor/UnrealEd/Private/Subsystems/EditorSubsystemBlueprintLibrary.cpp
--- a/Engine/Source/Editor/UnrealEd/Private/Subsystems/EditorSubsystemBlueprintLibrary.cpp
+++ b/Engine/Source/Editor/UnrealEd/Private/Subsystems/EditorSubsystemBlueprintLibrary.cpp
@@ -101,7 +101,7 @@ FName UEditorSubsystemBlueprintLibrary::GetPreviewPlatformName()
 	if (GEditor)
 	{
 		FPreviewPlatformInfo PreviewPlatform = GEditor->PreviewPlatform;
-		if (PreviewPlatform.bPreviewFeatureLevelActive)
+		if (PreviewPlatform.bPreviewShaderPlatformActive)
 		{
 			PreviewPlatformName = GetPreviewPlatformDropdownName(PreviewPlatform.PreviewShaderPlatformName, PreviewPlatform.DeviceProfileName);
 		}
