EditorUsability : ContentBrowser
- Fixed pressing a keybind like DELETE in content browser when nothing is selected would still trigger the command for the outliner selection even if it does not have the focus
- Added option in editor preferences to hide/show quick jump text (red/blue text) in the bottom right of the content browser (hidden by default)

#Synced-CL 52347035

diff --git a/Engine/Source/Editor/ContentBrowser/Private/SAssetView.cpp b/Engine/Source/Editor/ContentBrowser/Private/SAssetView.cpp
--- a/Engine/Source/Editor/ContentBrowser/Private/SAssetView.cpp
+++ b/Engine/Source/Editor/ContentBrowser/Private/SAssetView.cpp
@@ -6454,6 +6454,14 @@ FText SAssetView::GetQuickJumpTerm() const
 
 EVisibility SAssetView::IsQuickJumpVisible() const
 {
+	if (const UContentBrowserSettings* ContentBrowserSettings = GetDefault<UContentBrowserSettings>())
+	{
+		if (!ContentBrowserSettings->bDisplayQuickJumpText)
+		{
+			return EVisibility::Collapsed;
+		}
+	}
+
 	return (QuickJumpData.JumpTerm.IsEmpty()) ? EVisibility::Collapsed : EVisibility::HitTestInvisible;
 }
 
diff --git a/Engine/Source/Editor/ContentBrowser/Private/SContentBrowser.cpp b/Engine/Source/Editor/ContentBrowser/Private/SContentBrowser.cpp
--- a/Engine/Source/Editor/ContentBrowser/Private/SContentBrowser.cpp
+++ b/Engine/Source/Editor/ContentBrowser/Private/SContentBrowser.cpp
@@ -4246,6 +4246,12 @@ bool SContentBrowser::HandleDeleteCommandCanExecute() const
 		return LegacyContentSourceWidgets->PathContextMenu->CanExecuteDelete();
 	}
 
+	// To avoid processing a delete operation for viewport/outliner when content browser is focus but nothing selected
+	if (LegacyContentSourceWidgets->PathViewPtr->HasFocusedDescendants() || LegacyContentSourceWidgets->AssetViewPtr->HasFocusedDescendants())
+	{
+		return true;
+	}
+
 	return false;
 }
 
diff --git a/Engine/Source/Editor/UnrealEd/Classes/Settings/ContentBrowserSettings.h b/Engine/Source/Editor/UnrealEd/Classes/Settings/ContentBrowserSettings.h
--- a/Engine/Source/Editor/UnrealEd/Classes/Settings/ContentBrowserSettings.h
+++ b/Engine/Source/Editor/UnrealEd/Classes/Settings/ContentBrowserSettings.h
@@ -64,6 +64,10 @@ public:
 	/** Enables the rendering of Material Instance thumbnail previews */
 	UPROPERTY(EditAnywhere, config, Category = ContentBrowser)
 	bool bEnableRealtimeMaterialInstanceThumbnails = true;
+	
+	/** Display hint text of keys pressed in content browser in bottom right of the asset view, used to find item by name in a folder */
+	UPROPERTY(Config, EditAnywhere, Category = ContentBrowser)
+	bool bDisplayQuickJumpText = false;
 
 	/** Displays collections in the Collections window that were excluded by project settings */
 	UPROPERTY(Config, EditAnywhere, Category=Collections, meta = (DisplayName="Display excluded collections in Content Browser"))
