HitProxy: Skip hit proxy invalidation when the user is controlling a viewport camera using the mouse.

* Extends fix applied in CL 44604783.
* When mouse-looking a piloted camera, skip hit proxy invalidation in MoveLockedActorToCamera and MoveCameraToLockedActor since the user is not clicking on anything during rotation. The bInvalidateHitProxies flag is threaded down from PerspectiveCameraMoved so that other viewports locked to the same actor also benefit from the optimization.

#Synced-CL 52423541

diff --git a/Engine/Source/Editor/UnrealEd/Private/EditorViewportClient.cpp b/Engine/Source/Editor/UnrealEd/Private/EditorViewportClient.cpp
--- a/Engine/Source/Editor/UnrealEd/Private/EditorViewportClient.cpp
+++ b/Engine/Source/Editor/UnrealEd/Private/EditorViewportClient.cpp
@@ -2383,7 +2383,7 @@ void FEditorViewportClient::UpdateCameraMovement( float DeltaTime )
 		if ((bTransformDifferent || bFOVDifferent) && !bDisableInput)
 		{
 			// Something has changed!
-			const bool bInvalidateChildViews=true;
+			constexpr bool bInvalidateChildViews=true;
 
 			// When flying the camera around the hit proxies don't need to be invalidated since we are flying around and not clicking on anything
 			const bool bInvalidateHitProxies = bIsITFInteractions
diff --git a/Engine/Source/Editor/UnrealEd/Private/LevelEditorViewport.cpp b/Engine/Source/Editor/UnrealEd/Private/LevelEditorViewport.cpp
--- a/Engine/Source/Editor/UnrealEd/Private/LevelEditorViewport.cpp
+++ b/Engine/Source/Editor/UnrealEd/Private/LevelEditorViewport.cpp
@@ -94,6 +94,7 @@
 #include "EditorInteractiveGizmoManager.h"
 #include "EditorWorldExtension.h"
 #include "VREditorMode.h"
+#include "ViewportClientNavigationHelper.h"
 #include "ViewportWorldInteraction.h"
 #include "Subsystems/BrushEditingSubsystem.h"
 #include "Engine/VolumeTexture.h"
@@ -2685,7 +2686,10 @@ void FLevelEditorViewportClient::PerspectiveCameraMoved()
 	// If any other viewports have this actor locked too, we need to update them
 	if( GetActiveActorLock().IsValid() )
 	{
-		UpdateLockedActorViewports(GetActiveActorLock().Get(), false);
+		// Skip hit proxy invalidation during mouse-look rotation since the user is not clicking on anything
+		constexpr bool bCheckRealtime = false;
+		const bool bInvalidateHitProxies = !(ViewportNavigationHelper && ViewportNavigationHelper->bIsMouseLooking);
+		UpdateLockedActorViewports(GetActiveActorLock().Get(), bCheckRealtime, bInvalidateHitProxies);
 	}
 
 	// Broadcast 'camera moved' delegate
@@ -4235,7 +4239,9 @@ void FLevelEditorViewportClient::MoveLockedActorToCamera()
 		FScopedLevelDirtied LevelDirtyCallback;
 		LevelDirtyCallback.Request();
 
-		RedrawAllViewportsIntoThisScene();
+		// Skip hit proxy invalidation during mouse-look rotation since the user is not clicking on anything
+		const bool bInvalidateHitProxies = !(ViewportNavigationHelper && ViewportNavigationHelper->bIsMouseLooking);
+		GEditor->RedrawLevelEditingViewports(bInvalidateHitProxies);
 	}
 }
 
@@ -4305,7 +4311,7 @@ UTypedElementSelectionSet* FLevelEditorViewportClient::GetMutableSelectionSet()
 	return LevelEditor ? LevelEditor->GetMutableElementSelectionSet() : GEditor->GetSelectedActors()->GetElementSelectionSet();
 }
 
-void FLevelEditorViewportClient::MoveCameraToLockedActor()
+void FLevelEditorViewportClient::MoveCameraToLockedActor(const bool bInvalidateHitProxies)
 {
 	// Do a slightly shorter version of the code in UpdateViewForLockedActor where we don't care about
 	// post-process settings.
@@ -4348,7 +4354,8 @@ void FLevelEditorViewportClient::MoveCameraToLockedActor()
 			}
 		}
 
-		Invalidate();
+		constexpr bool bInvalidateChildViews = true;
+		Invalidate(bInvalidateChildViews, bInvalidateHitProxies);
 	}
 }
 
@@ -4501,24 +4508,24 @@ bool FLevelEditorViewportClient::IsAnyActorLocked() const
 	return GetActiveActorLock().IsValid();
 }
 
-void FLevelEditorViewportClient::UpdateLockedActorViewports(const AActor* InActor, const bool bCheckRealtime)
+void FLevelEditorViewportClient::UpdateLockedActorViewports(const AActor* InActor, const bool bCheckRealtime, const bool bInvalidateHitProxies)
 {
 	// Loop through all the other viewports, checking to see if the camera needs updating based on the locked actor
 	for (FLevelEditorViewportClient* Client : GEditor->GetLevelViewportClients())
 	{
 		if( Client && Client != this )
 		{
-			Client->UpdateLockedActorViewport(InActor, bCheckRealtime);
+			Client->UpdateLockedActorViewport(InActor, bCheckRealtime, bInvalidateHitProxies);
 		}
 	}
 }
 
-void FLevelEditorViewportClient::UpdateLockedActorViewport(const AActor* InActor, const bool bCheckRealtime)
+void FLevelEditorViewportClient::UpdateLockedActorViewport(const AActor* InActor, const bool bCheckRealtime, const bool bInvalidateHitProxies)
 {
 	// If this viewport has the actor locked and we need to update the camera, then do so
 	if( IsActorLocked(InActor) && ( !bCheckRealtime || IsRealtime() ) )
 	{
-		MoveCameraToLockedActor();
+		MoveCameraToLockedActor(bInvalidateHitProxies);
 	}
 }
 
diff --git a/Engine/Source/Editor/UnrealEd/Public/LevelEditorViewport.h b/Engine/Source/Editor/UnrealEd/Public/LevelEditorViewport.h
--- a/Engine/Source/Editor/UnrealEd/Public/LevelEditorViewport.h
+++ b/Engine/Source/Editor/UnrealEd/Public/LevelEditorViewport.h
@@ -325,8 +325,9 @@ public:
 
 	/**
 	 * Moves the viewport camera according to the locked actors location and rotation
+	 * @param bInvalidateHitProxies Whether to invalidate hit proxies when refreshing the viewport. Set to false during mouse-look to avoid unnecessary rebuilds.
 	 */
-	UNREALED_API void MoveCameraToLockedActor();
+	UNREALED_API void MoveCameraToLockedActor(const bool bInvalidateHitProxies = true);
 
 	/**
 	 * Check to see if this actor is locked by the viewport
@@ -765,9 +766,19 @@ private:
 
 	/**
 	 * Checks to see the viewports locked actor need updating
+	 * @param InActor The actor to check against each viewport's locked actor.
+	 * @param bCheckRealtime If true, only update viewports that have realtime enabled.
+	 * @param bInvalidateHitProxies Whether to invalidate hit proxies when refreshing the viewport. Set to false during mouse-look to avoid unnecessary rebuilds.
 	 */
-	UNREALED_API void UpdateLockedActorViewports(const AActor* InActor, const bool bCheckRealtime);
-	UNREALED_API void UpdateLockedActorViewport(const AActor* InActor, const bool bCheckRealtime);
+	UNREALED_API void UpdateLockedActorViewports(const AActor* InActor, const bool bCheckRealtime, const bool bInvalidateHitProxies = true);
+
+	/**
+	 * Checks to see the viewports locked actor need updating
+	 * @param InActor The actor to check against this viewport's locked actor.
+	 * @param bCheckRealtime If true, only update this viewport if realtime is enabled.
+	 * @param bInvalidateHitProxies Whether to invalidate hit proxies when refreshing the viewport. Set to false during mouse-look to avoid unnecessary rebuilds.
+	 */
+	UNREALED_API void UpdateLockedActorViewport(const AActor* InActor, const bool bCheckRealtime, const bool bInvalidateHitProxies = true);
 
 	/**
 	 * Moves the locked actor according to the viewport cameras location and rotation
