Bumping cook incremental version.

#Synced-CL 52199356

diff --git a/Engine/Source/Editor/UnrealEd/Private/Cooker/CookIncrementalVersion.cpp b/Engine/Source/Editor/UnrealEd/Private/Cooker/CookIncrementalVersion.cpp
--- a/Engine/Source/Editor/UnrealEd/Private/Cooker/CookIncrementalVersion.cpp
+++ b/Engine/Source/Editor/UnrealEd/Private/Cooker/CookIncrementalVersion.cpp
@@ -8,6 +8,6 @@ namespace UE::Cook
 {
 
 // Change CookIncrementalVersion to a new guid when all packages in an incremental cook need to be invalidated
-FGuid CookIncrementalVersion(0xA3F92C7B, 0x5D1E4A90, 0xC8B74F21, 0x19E6D3AB);
+FGuid CookIncrementalVersion(0x047D6737, 0x9B2949E3, 0xBF0493AF, 0x2C75D282);
 
 }
