[Verse] Integrate property visibility system specifically for the details panel

Adds a details-panel-specific property visibility delegate to PropertyVisibilityOverrideSubsystem
(FShouldHideDetailsPanelPropertyDelegate) that receives an IPropertyHandle, allowing callbacks
to inspect both the property and the objects being edited. This is separate from the existing
global FShouldHidePropertyDelegate which only has an FProperty*.

The reason for this change is that override serialization was also using the PropertyVisibilityOverrideSubsystem
and there is a risk that modifying the behaviour of the callback override serialization is using could lead to
data corruption.  It is clearer to provide an explicit callback for the details panel in the mean time.

#Synced-CL 52461421

diff --git a/Engine/Source/Editor/PropertyEditor/Private/DetailItemNode.cpp b/Engine/Source/Editor/PropertyEditor/Private/DetailItemNode.cpp
--- a/Engine/Source/Editor/PropertyEditor/Private/DetailItemNode.cpp
+++ b/Engine/Source/Editor/PropertyEditor/Private/DetailItemNode.cpp
@@ -1031,11 +1031,23 @@ void FDetailItemNode::FilterNode(const FDetailFilter& InFilter)
 		{
 			if (FProperty* Property = PropertyNodePin->GetProperty())
 			{
-				if (Property->GetBoolMetaData(PropertyNames::PropertyVisibilityOverrideName))
+				if (UPropertyVisibilityOverrideSubsystem* PropertyVisibilityOverrideSubsystem = UPropertyVisibilityOverrideSubsystem::Get())
 				{
-					if (UPropertyVisibilityOverrideSubsystem* PropertyVisibilityOverrideSubsystem = UPropertyVisibilityOverrideSubsystem::Get())
+					if (Customization.PropertyRow.IsValid())
 					{
-						bForceHidden = PropertyVisibilityOverrideSubsystem->ShouldHideProperty(Property);
+						if (TSharedPtr<IPropertyHandle> PropertyHandle = Customization.PropertyRow->GetPropertyHandle())
+						{
+							bForceHidden = PropertyVisibilityOverrideSubsystem->ShouldHidePropertyForDetailsPanel(*PropertyHandle);
+						}
+					}
+					
+					if (Property->GetBoolMetaData(PropertyNames::PropertyVisibilityOverrideName))
+					{
+						// legacy path
+						if (!bForceHidden)
+						{
+							bForceHidden = PropertyVisibilityOverrideSubsystem->ShouldHideProperty(Property);
+						}
 					}
 				}
 			}
diff --git a/Engine/Source/Editor/UnrealEd/Private/Subsystems/PropertyVisibilityOverrideSubsystem.cpp b/Engine/Source/Editor/UnrealEd/Private/Subsystems/PropertyVisibilityOverrideSubsystem.cpp
--- a/Engine/Source/Editor/UnrealEd/Private/Subsystems/PropertyVisibilityOverrideSubsystem.cpp
+++ b/Engine/Source/Editor/UnrealEd/Private/Subsystems/PropertyVisibilityOverrideSubsystem.cpp
@@ -67,3 +67,27 @@ bool UPropertyVisibilityOverrideSubsystem::ShouldHideProperty(const FProperty* P
 
 	return false;
 }
+
+void UPropertyVisibilityOverrideSubsystem::RegisterShouldHideDetailsPanelPropertyDelegate(const FName& DelegateName, const FShouldHideDetailsPanelPropertyDelegate& Delegate)
+{
+	DetailsPanelShouldHidePropertyDelegates.Add(DelegateName, Delegate);
+}
+
+void UPropertyVisibilityOverrideSubsystem::UnregisterShouldHideDetailsPanelPropertyDelegate(const FName& DelegateName)
+{
+	DetailsPanelShouldHidePropertyDelegates.Remove(DelegateName);
+}
+
+bool UPropertyVisibilityOverrideSubsystem::ShouldHidePropertyForDetailsPanel(const IPropertyHandle& PropertyHandle) const
+{
+	for (auto&& DelegatePair : DetailsPanelShouldHidePropertyDelegates)
+	{
+		const FShouldHideDetailsPanelPropertyDelegate& Delegate = DelegatePair.Value;
+		if (Delegate.IsBound() && Delegate.Execute(PropertyHandle))
+		{
+			return true;
+		}
+	}
+
+	return false;
+}
diff --git a/Engine/Source/Editor/UnrealEd/Public/Subsystems/PropertyVisibilityOverrideSubsystem.h b/Engine/Source/Editor/UnrealEd/Public/Subsystems/PropertyVisibilityOverrideSubsystem.h
--- a/Engine/Source/Editor/UnrealEd/Public/Subsystems/PropertyVisibilityOverrideSubsystem.h
+++ b/Engine/Source/Editor/UnrealEd/Public/Subsystems/PropertyVisibilityOverrideSubsystem.h
@@ -7,6 +7,7 @@
 #include "PropertyVisibilityOverrideSubsystem.generated.h"
 
 class FProperty;
+class IPropertyHandle;
 class UObject;
 
 namespace UE::PropertyVisibility
@@ -14,9 +15,12 @@ namespace UE::PropertyVisibility
 	UNREALED_API bool ConsiderPropertyForOverriddenState(TNotNull<const FProperty*> Property);
 }
 
-/** Bindable event for external objects to hook into ControlRig-level execution */
+/** Delegate for global property visibility checks (FProperty* only, called from multiple contexts) */
 DECLARE_DELEGATE_RetVal_OneParam(bool, FShouldHidePropertyDelegate, const FProperty*);
 
+/** Delegate for details-panel-specific property visibility checks (IPropertyHandle, called only from the details panel) */
+DECLARE_DELEGATE_RetVal_OneParam(bool, FShouldHideDetailsPanelPropertyDelegate, const IPropertyHandle&);
+
 UCLASS(MinimalAPI)
 class UPropertyVisibilityOverrideSubsystem : public UEditorSubsystem
 {
@@ -32,6 +36,17 @@ public:
 
 	UNREALED_API virtual bool ShouldHideProperty(const FProperty* Property) const;
 
+	UNREALED_API void RegisterShouldHideDetailsPanelPropertyDelegate(const FName& DelegateName, const FShouldHideDetailsPanelPropertyDelegate& Delegate);
+	UNREALED_API void UnregisterShouldHideDetailsPanelPropertyDelegate(const FName& DelegateName);
+
+	/**
+	 * Details-panel-specific visibility check.
+	 * Only call from the details panel where FScopedPushContext ensures valid TLS context.
+	 * Returns true if any registered delegate says the property should be hidden.
+	 */
+	UNREALED_API bool ShouldHidePropertyForDetailsPanel(const IPropertyHandle& PropertyHandle) const;
+
 private:
 	TMap<FName, FShouldHidePropertyDelegate> ShouldHidePropertyDelegates;
+	TMap<FName, FShouldHideDetailsPanelPropertyDelegate> DetailsPanelShouldHidePropertyDelegates;
 };
