[Gizmos] Remove experiemntal gizmo option from level editor viewport.  This was rehomed to gizmo settings and profiles

diff --git a/Engine/Source/Editor/UnrealEd/Private/ViewportToolbar/UnrealEdViewportToolbar.cpp b/Engine/Source/Editor/UnrealEd/Private/ViewportToolbar/UnrealEdViewportToolbar.cpp
--- a/Engine/Source/Editor/UnrealEd/Private/ViewportToolbar/UnrealEdViewportToolbar.cpp
+++ b/Engine/Source/Editor/UnrealEd/Private/ViewportToolbar/UnrealEdViewportToolbar.cpp
@@ -982,15 +982,12 @@ FToolMenuEntry CreateViewportToolbarTransformsSection()
 
 FToolMenuEntry CreateTransformsSubmenu()
 {
-	// Cache this once per session
-	const bool bIsGizmoSettingsModuleLoaded = FModuleManager::Get().IsModuleLoaded("GizmoSettings");
-
 	FToolMenuEntry Entry = FToolMenuEntry::InitSubMenu(
 		"Transform",
 		LOCTEXT("TransformsSubmenuLabel", "Transform"),
 		LOCTEXT("TransformsSubmenuTooltip", "Viewport-related transforms tools"),
 		FNewToolMenuDelegate::CreateLambda(
-			[bIsGizmoSettingsModuleLoaded](UToolMenu* Submenu) -> void
+			[](UToolMenu* Submenu) -> void
 			{
 				{
 					FToolMenuEntryToolBarData ToolBarData;
@@ -1219,36 +1216,6 @@ FToolMenuEntry CreateTransformsSubmenu()
 																						  : ECheckBoxState::Unchecked;
 							}
 						);
-
-						bool bSupportsNewGizmos = false;
-
-						if (const UUnrealEdViewportToolbarContext* const Context =
-								Submenu->FindContext<UUnrealEdViewportToolbarContext>())
-						{
-							if (TSharedPtr<SEditorViewport> Viewport = Context->Viewport.Pin())
-							{
-								if (TSharedPtr<FEditorViewportClient> ViewportClient = Viewport->GetViewportClient())
-								{
-									if (FEditorModeTools* ModeTools = ViewportClient->GetModeTools())
-									{
-										bSupportsNewGizmos = ModeTools->SupportsViewportITF();
-									}
-								}
-							}
-						}
-
-						// Only show if GizmoSettings module is loaded (Experimental GizmoFramework plugin) and editor supports new gizmos
-						if (bIsGizmoSettingsModuleLoaded && bSupportsNewGizmos)
-						{
-							GizmoSection.AddMenuEntry(
-								"TRSGizmoToggle",
-								LOCTEXT("TRSGizmoToggleLabel", "Use Experimental Gizmos"),
-								LOCTEXT("TRSGizmoToggleTooltip", "Whether or not to use the new Transform Gizmos"),
-								FSlateIcon(FAppStyle::GetAppStyleSetName(), "LevelEditor.UseExperimentalGizmos"),
-								Action,
-								EUserInterfaceActionType::ToggleButton
-							);
-						}
 					}
 
 					
