Add option to use collections to define maps to rebuild during nightly builds

#Synced-CL 52474816

diff --git a/Engine/Source/Editor/UnrealEd/Private/Commandlets/ContentCommandlets.cpp b/Engine/Source/Editor/UnrealEd/Private/Commandlets/ContentCommandlets.cpp
--- a/Engine/Source/Editor/UnrealEd/Private/Commandlets/ContentCommandlets.cpp
+++ b/Engine/Source/Editor/UnrealEd/Private/Commandlets/ContentCommandlets.cpp
@@ -142,6 +142,7 @@ int32 UResavePackagesCommandlet::InitializeResaveParameters( const TArray<FStrin
 		FString PackageFolderArg;
 		FString Maps;
 		FString File;
+		FString CollectionArg;
 		const FString& CurrentSwitch = Switches[ SwitchIdx ];
 		if( FParse::Value( *CurrentSwitch, TEXT( "PACKAGE="), Package ) )
 		{
@@ -233,6 +234,35 @@ int32 UResavePackagesCommandlet::InitializeResaveParameters( const TArray<FStrin
 				UE_LOG(LogContentCommandlet, Error, TEXT("Failed to load file %s"), *File);
 			}
 		}
+		else if (FParse::Value(*CurrentSwitch, TEXT("COLLECTION="), CollectionArg))
+		{
+			const TSharedRef<ICollectionContainer>& CollectionContainer = FCollectionManagerModule::GetModule().Get().GetProjectCollectionContainer();
+			TArray<FSoftObjectPath> CollectionObjects;
+
+			TArray<FString> Collections;
+			CollectionArg.ParseIntoArray(Collections, TEXT("+"));
+			for (const FString& Collection : Collections)
+			{
+				if (!CollectionContainer->CollectionExists(*Collection, ECollectionShareType::CST_All))
+				{
+					UE_LOG(LogContentCommandlet, Error, TEXT("Unknown source collection '%s'"), *Collection);
+					continue;
+				}
+
+				CollectionObjects.Reset();
+				if (!CollectionContainer->GetAssetsInCollection(*Collection, ECollectionShareType::CST_All, CollectionObjects))
+				{
+					UE_LOG(LogContentCommandlet, Error, TEXT("Failed to get assets in collection '%s'"), *Collection);
+					continue;
+				}
+
+				for (const FSoftObjectPath& CollectionObject : CollectionObjects)
+				{
+					PackageNames.Add(CollectionObject.GetLongPackageName());
+					bExplicitPackages = true;
+				}
+			}
+		}
 	}
 
 	// Check for numeric settings
