Skip CDOs in FEditorViewportClient::OnObjectTransacted to fix ensure when editing NavigationSystem config properties in Project Settings

#Synced-CL 52469485

diff --git a/Engine/Source/Editor/UnrealEd/Private/EditorViewportClient.cpp b/Engine/Source/Editor/UnrealEd/Private/EditorViewportClient.cpp
--- a/Engine/Source/Editor/UnrealEd/Private/EditorViewportClient.cpp
+++ b/Engine/Source/Editor/UnrealEd/Private/EditorViewportClient.cpp
@@ -7411,7 +7411,14 @@ void FEditorViewportClient::OnObjectTransacted(UObject* Object, const FTransacti
 		return;
 	}
 
-	if (GetWorld() && (Object->GetWorld() == GetWorld()) && !GetWorld()->IsGameWorld())
+	// Skip CDOs - calling GetWorld() on a CDO with a Within specifier triggers an ensure.
+	if (Object->HasAnyFlags(RF_ClassDefaultObject))
+	{
+		return;
+	}
+
+	UWorld* const World = GetWorld();
+	if (World && (Object->GetWorld() == World) && !World->IsGameWorld())
 	{
 		CachedWorldBounds.Init();
 	}
