[Fix] Remove MAX_PIE_INSTANCES hard limit and clamp PrimaryPIEClientIndex on client count change

[Fixed] PIEGuidMap changed from a fixed-size static array to a dynamic TArray, removing the MAX_PIE_INSTANCES allocation ceiling (Causing all the problems)
[Fixed] Number of Players toolbar spinner MaxValue now uses MAX_PIE_INSTANCES - 1 instead of hardcoded 64
[Fixed] SetPlayNumberOfClients now calls ClampPIEClientValues() so PrimaryPIEClientIndex is re-clamped whenever the client count is reduced

#Synced-CL 52470154

diff --git a/Engine/Source/Editor/UnrealEd/Classes/Settings/LevelEditorPlaySettings.h b/Engine/Source/Editor/UnrealEd/Classes/Settings/LevelEditorPlaySettings.h
--- a/Engine/Source/Editor/UnrealEd/Classes/Settings/LevelEditorPlaySettings.h
+++ b/Engine/Source/Editor/UnrealEd/Classes/Settings/LevelEditorPlaySettings.h
@@ -390,7 +390,7 @@ private:
 	bool RunUnderOneProcess;
 
 	/** The number of client windows to open. The first one to open will respect the Play In Editor "Modes" option (PIE, PINW), additional clients respect the RunUnderOneProcess setting. */
-	UPROPERTY(config, EditAnywhere, Category="Multiplayer Options|Client", meta=(ClampMin = "1", UIMin = "1", UIMax = "10"))
+	UPROPERTY(config, EditAnywhere, Category="Multiplayer Options|Client", meta=(ClampMin = "1", ClampMax = "10", UIMin = "1", UIMax = "10"))
 	int32 PlayNumberOfClients;
 
 	/** In multiplayer PIE which client will be the 'primary'. (default = 0, the first client)*/
@@ -471,7 +471,7 @@ private:
 private:
 	UNREALED_API void PushDebugDrawingSettings();
 	
-	void ClampPIEClientValues();
+	UNREALED_API void ClampPIEClientValues();
 
 public:
 	bool ShowServerDebugDrawingByDefault() const
@@ -517,7 +517,11 @@ public:
 	bool IsRunUnderOneProcessActive() const { return true; }
 	bool GetRunUnderOneProcess( bool &OutRunUnderOneProcess ) const { OutRunUnderOneProcess = RunUnderOneProcess; return IsRunUnderOneProcessActive(); }
 	
-	void SetPlayNumberOfClients( const int32 InPlayNumberOfClients ) { PlayNumberOfClients = InPlayNumberOfClients; }
+	void SetPlayNumberOfClients( const int32 InPlayNumberOfClients )
+	{
+		PlayNumberOfClients = InPlayNumberOfClients;
+		ClampPIEClientValues();
+	}
 	bool IsPlayNumberOfClientsActive() const { return true; }
 	bool GetPlayNumberOfClients( int32 &OutPlayNumberOfClients ) const { OutPlayNumberOfClients = PlayNumberOfClients; return IsPlayNumberOfClientsActive(); }
 
diff --git a/Engine/Source/Editor/UnrealEd/Private/Kismet2/DebuggerCommands.cpp b/Engine/Source/Editor/UnrealEd/Private/Kismet2/DebuggerCommands.cpp
--- a/Engine/Source/Editor/UnrealEd/Private/Kismet2/DebuggerCommands.cpp
+++ b/Engine/Source/Editor/UnrealEd/Private/Kismet2/DebuggerCommands.cpp
@@ -850,9 +850,9 @@ TSharedRef< SWidget > FPlayWorldCommands::GeneratePlayMenuContent(TSharedRef<FUI
 			{
 				TSharedRef<SWidget> NumPlayers = SNew(SSpinBox<int32>)	// Copy limits from PlayNumberOfClients meta data
 					.MinValue(1)
-					.MaxValue(64)
+					.MaxValue(MAX_PIE_INSTANCES - 1)
 					.MinSliderValue(1)
-					.MaxSliderValue(4)
+					.MaxSliderValue(MAX_PIE_INSTANCES - 1)
 					.Delta(1)
 					.ToolTipText(LOCTEXT("NumberOfClientsToolTip", "How many client instances do you want to create? The first instance respects the Play Mode location (PIE/PINW) and additional instances respect the RunUnderOneProcess setting."))
 					.Value_Static(&FInternalPlayWorldCommandCallbacks::GetNumberOfClients)
diff --git a/Engine/Source/Editor/UnrealEd/Private/Settings/SettingsClasses.cpp b/Engine/Source/Editor/UnrealEd/Private/Settings/SettingsClasses.cpp
--- a/Engine/Source/Editor/UnrealEd/Private/Settings/SettingsClasses.cpp
+++ b/Engine/Source/Editor/UnrealEd/Private/Settings/SettingsClasses.cpp
@@ -537,7 +537,7 @@ void ULevelEditorPlaySettings::PushDebugDrawingSettings()
 
 void ULevelEditorPlaySettings::ClampPIEClientValues()
 {
-	PlayNumberOfClients = FMath::Clamp(PlayNumberOfClients, 1, MAX_PIE_INSTANCES);
+	PlayNumberOfClients = FMath::Clamp(PlayNumberOfClients, 1, MAX_PIE_INSTANCES - 1);
 	PrimaryPIEClientIndex = FMath::Clamp(PrimaryPIEClientIndex, INDEX_NONE, PlayNumberOfClients - 1);
 }
 
@@ -608,6 +608,7 @@ void ULevelEditorPlaySettings::PostEditChangeProperty(struct FPropertyChangedEve
 void ULevelEditorPlaySettings::PostInitProperties()
 {
 	Super::PostInitProperties();
+	ClampPIEClientValues();
 
 	// If the deprecated bOneHeadsetEachProcess setting is on, translate it to bHMDForPrimaryProcessOnly value
 	PRAGMA_DISABLE_DEPRECATION_WARNINGS
diff --git a/Engine/Source/Runtime/CoreUObject/Private/UObject/LazyObjectPtr.cpp b/Engine/Source/Runtime/CoreUObject/Private/UObject/LazyObjectPtr.cpp
--- a/Engine/Source/Runtime/CoreUObject/Private/UObject/LazyObjectPtr.cpp
+++ b/Engine/Source/Runtime/CoreUObject/Private/UObject/LazyObjectPtr.cpp
@@ -15,7 +15,7 @@
 /** Annotation associating objects with their guids **/
 static FUObjectAnnotationSparseSearchable<FUniqueObjectGuid,true> GuidAnnotation;
 
-static TMap<FGuid, FGuid> PIEGuidMap[MAX_PIE_INSTANCES];
+static TArray<TMap<FGuid, FGuid>> PIEGuidMap;
 static FTransactionallySafeMutex PIEGuidMapLock;
 
 static FGuid RemapGuid(const FUniqueObjectGuid& Guid, int32 PIEInstanceID)
@@ -27,10 +27,14 @@ static FGuid RemapGuid(const FUniqueObjectGuid& Guid, int32 PIEInstanceID)
 	}
 
 	check(PIEInstanceID != INDEX_NONE);
-	check(PIEInstanceID < MAX_PIE_INSTANCES);
 
 	UE::TScopeLock Lock(PIEGuidMapLock);
 
+	if (PIEInstanceID >= PIEGuidMap.Num())
+	{
+		PIEGuidMap.AddDefaulted(PIEInstanceID + 1 - PIEGuidMap.Num());
+	}
+
 	FGuid& FoundGuid = PIEGuidMap[PIEInstanceID].FindOrAdd(Guid.GetGuid());
 
 	if (!FoundGuid.IsValid())
@@ -200,7 +204,9 @@ void FLazyObjectPtr::ResetPIEFixups()
 {
 	const int32 PlayInEditorID = UE::GetPlayInEditorID();
 	check(PlayInEditorID != -1);
-	check(PlayInEditorID < MAX_PIE_INSTANCES);
 	UE::TScopeLock Lock(PIEGuidMapLock);
-	PIEGuidMap[PlayInEditorID].Reset();
+	if (PIEGuidMap.IsValidIndex(PlayInEditorID))
+	{
+		PIEGuidMap[PlayInEditorID].Reset();
+	}
 }
