ActorFolders: Add a cvar to register Actor Folders and FFolders in TEDS. Currently set to false by default.
UWorldFolders: Add code to sync folder data to TEDS whenever changes are detected (direct, undo/redo etc).
Add sync queries to sync folder data to and from TEDS.
Modify actor parent queries in TEDS to also include parent folders.
Modify the TEDS Outliner query to include folders.

#Synced-CL 52470086

diff --git a/Engine/Plugins/Experimental/EditorDataStorageFeatures/Source/TedsActorCompatibility/Private/Compatibility/TedsActorParentQueries.cpp b/Engine/Plugins/Experimental/EditorDataStorageFeatures/Source/TedsActorCompatibility/Private/Compatibility/TedsActorParentQueries.cpp
--- a/Engine/Plugins/Experimental/EditorDataStorageFeatures/Source/TedsActorCompatibility/Private/Compatibility/TedsActorParentQueries.cpp
+++ b/Engine/Plugins/Experimental/EditorDataStorageFeatures/Source/TedsActorCompatibility/Private/Compatibility/TedsActorParentQueries.cpp
@@ -2,6 +2,7 @@
 
 #include "Compatibility/TedsActorParentQueries.h"
 
+#include "ActorFolders/TedsActorFolderUtils.h"
 #include "DataStorage/MapKey.h"
 #include "Elements/Columns/TypedElementCompatibilityColumns.h"
 #include "Elements/Columns/TypedElementHiearchyColumns.h"
@@ -11,6 +12,7 @@
 #include "Engine/Level.h"
 #include "Factories/TedsEditorHierarchyFactory.h"
 #include "GameFramework/Actor.h"
+#include "LevelInstance/LevelInstanceInterface.h"
 
 #include UE_INLINE_GENERATED_CPP_BY_NAME(TedsActorParentQueries)
 
@@ -24,24 +26,50 @@ namespace UE::Editor::DataStorage::ActorParentQueries::Private
 		bAddParentColumnToActors,
 		TEXT("Mirror parent information for actors to TEDS (only works when set on startup)"));
 	
-	// Helper to get the first viable parent for an actor in the editor
-	const UObject* GetEditorParentForActor(const AActor* InActor)
+	// Helper to get the mapping domain and key for the first viable parent for an actor in the editor
+	bool GetEditorParentForActor(const AActor* InActor, FName& OutParentMappingDomain, FMapKey& OutParentMappingKey)
 	{
-		// For the editor hierarchy, the order to check is Parent Actor -> Owning Level -> Owning World
-		if (const AActor* ParentActor = InActor->GetAttachParentActor())
+		// For the editor hierarchy, the order to check is Parent Actor -> Folder -> Level Instance -> Owning Level -> Owning World
+		if (const AActor* ParentActor = InActor->GetSceneOutlinerParent())
 		{
-			return ParentActor;
+			OutParentMappingDomain = ICompatibilityProvider::ObjectMappingDomain;
+			OutParentMappingKey = FMapKey(ParentActor);
+			return true;
 		}
-		else if (const ULevel* ParentLevel = InActor->GetLevel())
+
+		FFolder Folder = InActor->GetFolder();
+
+		if (!Folder.IsNone())
+		{
+			OutParentMappingDomain = ActorFolders::GetFolderMappingDomain(Folder);
+			OutParentMappingKey = ActorFolders::GetFolderIndex(Folder);
+			return true;
+		}
+
+		// Check level instance through the folder's root object
+		if (ILevelInstanceInterface* OwningLevelInstance = Cast<ILevelInstanceInterface>(Folder.GetRootObjectPtr()))
 		{
-			return ParentLevel;
+			AActor* OwningActor = CastChecked<AActor>(OwningLevelInstance);
+			OutParentMappingDomain = ICompatibilityProvider::ObjectMappingDomain;
+			OutParentMappingKey = FMapKey(OwningActor);
+			return true;
 		}
-		else if (const UWorld* ParentWorld = InActor->GetWorld())
+
+		if (const ULevel* ParentLevel = InActor->GetLevel())
+		{
+			OutParentMappingDomain = ICompatibilityProvider::ObjectMappingDomain;
+			OutParentMappingKey = FMapKey(ParentLevel);
+			return true;
+		}
+
+		if (const UWorld* ParentWorld = InActor->GetWorld())
 		{
-			return ParentWorld;
+			OutParentMappingDomain = ICompatibilityProvider::ObjectMappingDomain;
+			OutParentMappingKey = FMapKey(ParentWorld);
+			return true;
 		}
 		
-		return nullptr;
+		return false;
 	}
 };
 
@@ -80,19 +108,23 @@ void UActorParentDataStorageFactory::RegisterActorParentQueries(UE::Editor::Data
 			{
 				if (const AActor* ActorInstance = Cast<AActor>(Actor.Object))
 				{
-					if (const UObject* Parent = ActorParentQueries::Private::GetEditorParentForActor(ActorInstance))
+					FMapKey IdKey;
+					FName MappingDomain;
+					
+					if (ActorParentQueries::Private::GetEditorParentForActor(ActorInstance, MappingDomain, IdKey))
 					{
-						FMapKeyView IdKey = FMapKeyView(Parent);
-						RowHandle ParentRow = Context.LookupMappedRow(ICompatibilityProvider::ObjectMappingDomain, IdKey);
+						RowHandle ParentRow = Context.LookupMappedRow(MappingDomain, IdKey);
+					
 						if (Context.IsRowAvailable(ParentRow))
 						{
 							Context.SetParentRow(Row, ParentRow);
 						}
 						else
 						{
-							Context.SetUnresolvedParent(Row, FMapKey(Parent), ICompatibilityProvider::ObjectMappingDomain);
+							Context.SetUnresolvedParent(Row, IdKey, MappingDomain);
 						}
 					}
+					
 				}
 			})
 		.AccessesHierarchy(UTedsEditorHierarchyFactory::EditorObjectHierarchyName)
@@ -110,10 +142,12 @@ void UActorParentDataStorageFactory::RegisterActorParentQueries(UE::Editor::Data
 			{
 				if (const AActor* ActorInstance = Cast<AActor>(Actor.Object))
 				{
-					if (const UObject* Parent = ActorParentQueries::Private::GetEditorParentForActor(ActorInstance))
+					FMapKey IdKey;
+					FName MappingDomain;
+					
+					if (ActorParentQueries::Private::GetEditorParentForActor(ActorInstance, MappingDomain, IdKey))
 					{
-						FMapKeyView IdKey = FMapKeyView(Parent);
-						RowHandle NewParentRow = Context.LookupMappedRow(ICompatibilityProvider::ObjectMappingDomain, IdKey);
+						RowHandle NewParentRow = Context.LookupMappedRow(MappingDomain, IdKey);
 						RowHandle CurrentParentRow = Context.GetParentRow(TargetRow);
 
 						if (CurrentParentRow != NewParentRow)
@@ -124,7 +158,7 @@ void UActorParentDataStorageFactory::RegisterActorParentQueries(UE::Editor::Data
 							}
 							else
 							{
-								Context.SetUnresolvedParent(TargetRow, FMapKey(Parent), ICompatibilityProvider::ObjectMappingDomain);
+								Context.SetUnresolvedParent(TargetRow, IdKey, MappingDomain);
 							}
 						}
 						return;
diff --git a/Engine/Plugins/Experimental/EditorDataStorageFeatures/Source/TedsOutliner/Private/ActorFolderFactory.cpp b/Engine/Plugins/Experimental/EditorDataStorageFeatures/Source/TedsOutliner/Private/ActorFolderFactory.cpp
deleted file mode 100644
diff --git a/Engine/Plugins/Experimental/EditorDataStorageFeatures/Source/TedsOutliner/Private/ActorFolderFactory.h b/Engine/Plugins/Experimental/EditorDataStorageFeatures/Source/TedsOutliner/Private/ActorFolderFactory.h
deleted file mode 100644
diff --git a/Engine/Plugins/Experimental/EditorDataStorageFeatures/Source/TedsOutliner/Private/Widgets/OutlinerFolderLabelWidget.cpp b/Engine/Plugins/Experimental/EditorDataStorageFeatures/Source/TedsOutliner/Private/Widgets/OutlinerFolderLabelWidget.cpp
--- a/Engine/Plugins/Experimental/EditorDataStorageFeatures/Source/TedsOutliner/Private/Widgets/OutlinerFolderLabelWidget.cpp
+++ b/Engine/Plugins/Experimental/EditorDataStorageFeatures/Source/TedsOutliner/Private/Widgets/OutlinerFolderLabelWidget.cpp
@@ -2,7 +2,7 @@
 
 #include "Widgets/OutlinerFolderLabelWidget.h"
 
-#include "Columns/ActorFolderColumns.h"
+#include "ActorFolders/ActorFolderColumns.h"
 #include "Columns/SceneOutlinerColumns.h"
 #include "Columns/SlateDelegateColumns.h"
 #include "Compatibility/SceneOutlinerTedsBridge.h"
@@ -12,6 +12,7 @@
 #include "Elements/Columns/TypedElementMiscColumns.h"
 #include "Elements/Columns/TypedElementSlateWidgetColumns.h"
 #include "Elements/Framework/TypedElementAttributeBinding.h"
+#include "Elements/Columns/TypedElementUIColumns.h"
 #include "ISceneOutliner.h"
 #include "SceneOutlinerHelpers.h"
 #include "TedsOutlinerHelpers.h"
@@ -61,7 +62,7 @@ TSharedPtr<SWidget> FOutlinerFolderLabelWidgetConstructor::CreateWidget(UE::Edit
 					SNew(SImage)
 						.Image_Lambda([DataStorage, TargetRow, WidgetRow]()
 						{
-							bool bIsExpanded = DataStorage->HasColumns<FFolderExpandedTag>(TargetRow);
+							bool bIsExpanded = DataStorage->HasColumns<FExpandedInUITag>(TargetRow);
 
 							TWeakPtr<ISceneOutlinerTreeItem> OutlinerTreeItem = GetTreeItemForRow(DataStorage, TargetRow, WidgetRow);
 							
diff --git a/Engine/Plugins/Experimental/EditorDataStorageFeatures/Source/TedsOutliner/Private/Widgets/SceneOutlinerWidget.cpp b/Engine/Plugins/Experimental/EditorDataStorageFeatures/Source/TedsOutliner/Private/Widgets/SceneOutlinerWidget.cpp
--- a/Engine/Plugins/Experimental/EditorDataStorageFeatures/Source/TedsOutliner/Private/Widgets/SceneOutlinerWidget.cpp
+++ b/Engine/Plugins/Experimental/EditorDataStorageFeatures/Source/TedsOutliner/Private/Widgets/SceneOutlinerWidget.cpp
@@ -2,6 +2,7 @@
 
 #include "Widgets/SceneOutlinerWidget.h"
 
+#include "ActorFolders/ActorFolderColumns.h"
 #include "Columns/SceneOutlinerColumns.h"
 #include "Elements/Columns/TypedElementUIColumns.h"
 #include "ILevelEditor.h"
@@ -146,9 +147,12 @@ UE::Editor::DataStorage::FQueryDescription FTedsSceneOutlinerWidgetConstructor::
 {
 	using namespace UE::Editor::DataStorage::Queries;
 	
-	return Select()
+	// Show all rows with a label and world that are either a UObject or a folder
+	return
+		Select()
 		.Where()
-			.All<FTypedElementUObjectColumn, FTypedElementLabelColumn, FTypedElementWorldColumn>()
+			.All<FTypedElementLabelColumn, FTypedElementWorldColumn>()
+			.Any<FTypedElementUObjectColumn, FFolderCompatibilityColumn>()
 			.None<FHideRowFromUITag>()
 		.Compile();
 }
diff --git a/Engine/Plugins/Experimental/EditorDataStorageFeatures/Source/TedsOutliner/Public/Columns/ActorFolderColumns.h b/Engine/Plugins/Experimental/EditorDataStorageFeatures/Source/TedsOutliner/Public/Columns/ActorFolderColumns.h
deleted file mode 100644
diff --git a/Engine/Source/Editor/SceneOutliner/Private/SceneOutlinerHelpers.cpp b/Engine/Source/Editor/SceneOutliner/Private/SceneOutlinerHelpers.cpp
--- a/Engine/Source/Editor/SceneOutliner/Private/SceneOutlinerHelpers.cpp
+++ b/Engine/Source/Editor/SceneOutliner/Private/SceneOutlinerHelpers.cpp
@@ -7,6 +7,7 @@
 #include "ComponentTreeItem.h"
 #include "EditorActorFolders.h"
 #include "EditorClassUtils.h"
+#include "EditorFolderUtils.h"
 #include "EditorModeManager.h"
 #include "Engine/Blueprint.h"
 #include "Engine/Engine.h"
@@ -219,23 +220,7 @@ namespace SceneOutliner
 
 	void FSceneOutlinerHelpers::RenameFolder(const FFolder& InFolder, const FText& NewFolderName, UWorld* World)
 	{
-		if (!World)
-		{
-			return;
-		}
-		
-		FName NewPath = InFolder.GetParent().GetPath();
-		if (NewPath.IsNone())
-		{
-			NewPath = FName(*NewFolderName.ToString());
-		}
-		else
-		{
-			NewPath = FName(*(NewPath.ToString() / NewFolderName.ToString()));
-		}
-		
-		const FFolder TreeItemNewFolder(InFolder.GetRootObject(), NewPath);
-		FActorFolders::Get().RenameFolderInWorld(*World, InFolder, TreeItemNewFolder);
+		FEditorFolderUtils::RenameFolder(InFolder, NewFolderName, World);
 	}
 
 	bool FSceneOutlinerHelpers::ValidateFolderName(const FFolder& InFolder, UWorld* World, const FText& InLabel, FText& OutErrorMessage)
diff --git a/Engine/Source/Editor/UnrealEd/Private/ActorFolders/ActorFolderFactory.cpp b/Engine/Source/Editor/UnrealEd/Private/ActorFolders/ActorFolderFactory.cpp
new file mode 100644
--- /dev/null
+++ b/Engine/Source/Editor/UnrealEd/Private/ActorFolders/ActorFolderFactory.cpp
@@ -0,0 +1,123 @@
+// Copyright Epic Games, Inc. All Rights Reserved.
+
+#include "ActorFolderFactory.h"
+
+#include "ActorFolders/TedsActorFolderUtils.h"
+#include "Editor.h"
+#include "EditorActorFolders.h"
+#include "Elements/Columns/TypedElementCompatibilityColumns.h"
+#include "Elements/Columns/TypedElementFolderColumns.h"
+#include "Elements/Columns/TypedElementLabelColumns.h"
+#include "Elements/Columns/TypedElementSlateWidgetColumns.h"
+#include "Elements/Columns/TypedElementUIColumns.h"
+#include "Elements/Framework/TypedElementRegistry.h"
+#include "Elements/Interfaces/TypedElementDataStorageInterface.h"
+#include "EngineUtils.h"
+#include "ActorFolders/ActorFolderColumns.h"
+#include "DataStorage/Features.h"
+
+#include UE_INLINE_GENERATED_CPP_BY_NAME(ActorFolderFactory)
+
+void UTedsActorFolderFactory::RegisterTables(UE::Editor::DataStorage::ICoreProvider& InDataStorage)
+{
+	using namespace UE::Editor::DataStorage;
+	using namespace UE::Editor::DataStorage::ActorFolders;
+
+	TableHandle Table = InDataStorage.RegisterTable(
+		UE::Editor::DataStorage::TTypedElementColumnTypeList<
+		FFolderTag, FFolderCompatibilityColumn, FTypedElementLabelColumn, FTypedElementLabelHashColumn, FSlateColorColumn>(),
+	GetActorFolderTableName());
+
+	// Actor folders are registered with TEDS Compat for proper undo/redo support, old FFolders are handled manually in UWorldFolders
+	if (ICompatibilityProvider* DataStorageCompat = GetMutableDataStorageFeature<ICompatibilityProvider>(CompatibilityFeatureName))
+	{
+		DataStorageCompat->RegisterTypeTableAssociation(UActorFolder::StaticClass(), Table);
+	}
+}
+
+void UTedsActorFolderFactory::PreRegister(UE::Editor::DataStorage::ICoreProvider& InDataStorage)
+{
+	using namespace UE::Editor::DataStorage;
+	Super::PreRegister(InDataStorage);
+
+	const IConsoleVariable* TedsFoldersCvar = IConsoleManager::Get().FindConsoleVariable(TEXT("TEDS.Feature.Folders"));
+	if (TedsFoldersCvar && TedsFoldersCvar->GetBool())
+	{
+		FWorldDelegates::OnPostWorldInitialization.AddUObject(this, &UTedsActorFolderFactory::PostWorldInitialized);
+		FEditorDelegates::EndPIE.AddUObject( this, &UTedsActorFolderFactory::OnEndPIE );
+	}
+}
+
+void UTedsActorFolderFactory::RegisterQueries(UE::Editor::DataStorage::ICoreProvider& InDataStorage)
+{
+	using namespace UE::Editor::DataStorage;
+	using namespace UE::Editor::DataStorage::Queries;
+	
+	const IConsoleVariable* TedsFoldersCvar = IConsoleManager::Get().FindConsoleVariable(TEXT("TEDS.Feature.Folders"));
+	if (TedsFoldersCvar && TedsFoldersCvar->GetBool())
+	{
+		// For old FFolders (not actor folders), there is actually no separate folder for PIE according to UWorldFolders. The Outliner simply queries
+		// an actor's parent folder and adds a row to the UI representing it. To match that behavior for now but make some improvements, we keep that same
+		// mechanism of actors in PIE registering their parent folders recursively but also add them to TEDS (but NOT to UWorldFolders so there's no risk 
+		// of breaking old workflows)
+	
+		InDataStorage.RegisterQuery(
+			Select(
+				TEXT("Add PIE actor's owning folder to TEDS"),
+				FObserver::OnAdd<FPieObjectTag>().SetExecutionMode(EExecutionMode::GameThread),
+				[&InDataStorage](IQueryContext& Context, RowHandle RowHandle, const FTypedElementUObjectColumn& ActorColumn, const FTypedElementWorldColumn& WorldColumn)
+				{
+					if (AActor* Actor = Cast<AActor>(ActorColumn.Object))
+					{
+						FFolder Folder = Actor->GetFolder();
+					
+						if(!Folder.IsNone())
+						{
+							// PIE folders are only registered through this observer and nowhere else, so we need to recurse till we reach the root
+							constexpr bool bRecursivelyAddParents = true;
+							ActorFolders::RegisterFolderInTeds(Context, &InDataStorage, Folder, WorldColumn.World, bRecursivelyAddParents);
+						}
+					}
+				})
+			.Where()
+				.All<FTypedElementActorTag>()
+				.None<FHideRowFromUITag>() // We only want folders for actors visible in the UI
+			.Compile());
+	}
+}
+
+void UTedsActorFolderFactory::PreShutdown(UE::Editor::DataStorage::ICoreProvider& InDataStorage)
+{
+	const IConsoleVariable* TedsFoldersCvar = IConsoleManager::Get().FindConsoleVariable(TEXT("TEDS.Feature.Folders"));
+	if (TedsFoldersCvar && TedsFoldersCvar->GetBool())
+	{
+		FWorldDelegates::OnPostWorldInitialization.RemoveAll(this);
+		FEditorDelegates::EndPIE.RemoveAll(this);
+	}
+	
+	Super::PreShutdown(InDataStorage);
+}
+
+void UTedsActorFolderFactory::PostWorldInitialized(UWorld* World, const UWorld::InitializationValues InitializationValues)
+{
+	// With old FFolders (non actor-folders), the legacy behavior was to only create the folders lazily when they were requested by the Outliner UI.
+	// This means if the Outliner is closed no folders are added to TEDS. Actor folders on the other hand go through ULevel::FixupActorFolders which
+	// implicitly loads them from disk into memory.
+	
+	// To work around this we explicitly load folders when the world is initialized when the TEDS integration is enabled
+	if (World)
+	{
+		FActorFolders::Get().InitializeForWorld(*World);
+	}
+}
+
+void UTedsActorFolderFactory::OnEndPIE(const bool)
+{
+	using namespace UE::Editor::DataStorage;
+	
+	// One ending PIE, we can simply remove all the PIE folders at once instead of needing to do it per actor to mirror how they were added
+	if (ICoreProvider* DataStorage = GetMutableDataStorageFeature<ICoreProvider>(StorageFeatureName))
+	{
+		DataStorage->RemoveAllRowsWithColumns({FPieObjectTag::StaticStruct(), FFolderCompatibilityColumn::StaticStruct()});
+	}
+}
diff --git a/Engine/Source/Editor/UnrealEd/Private/ActorFolders/ActorFolderFactory.h b/Engine/Source/Editor/UnrealEd/Private/ActorFolders/ActorFolderFactory.h
new file mode 100644
--- /dev/null
+++ b/Engine/Source/Editor/UnrealEd/Private/ActorFolders/ActorFolderFactory.h
@@ -0,0 +1,41 @@
+// Copyright Epic Games, Inc. All Rights Reserved.
+
+#pragma once
+
+#include "DataStorage/Handles.h"
+#include "Elements/Interfaces/TypedElementDataStorageFactory.h"
+#include "Engine/World.h"
+
+#include "ActorFolderFactory.generated.h"
+
+namespace UE::Editor::DataStorage
+{
+	struct IQueryContext;
+}
+
+class IEditorDataStorageProvider;
+class UWorld;
+struct FFolder;
+class UActorFolder;
+
+/*
+ * Factory that manages the registration of folders in TEDS and contains compatibility queries/delegates for feature parity with how legacy FFolders
+ * interact with legacy editor systems until we can use TEDS as the source of truth
+ */
+UCLASS()
+class UTedsActorFolderFactory : public UEditorDataStorageFactory
+{
+	GENERATED_BODY()
+
+public:
+	virtual ~UTedsActorFolderFactory() override = default;
+
+	virtual void RegisterTables(UE::Editor::DataStorage::ICoreProvider& DataStorage) override;
+	virtual void PreRegister(UE::Editor::DataStorage::ICoreProvider& DataStorage) override;
+	virtual void RegisterQueries(UE::Editor::DataStorage::ICoreProvider& DataStorage) override;
+	virtual void PreShutdown(UE::Editor::DataStorage::ICoreProvider& DataStorage) override;
+
+private:
+	void PostWorldInitialized(UWorld* World, const UWorld::InitializationValues InitializationValues);
+	void OnEndPIE( const bool bIsSimulating );
+};
diff --git a/Engine/Source/Editor/UnrealEd/Private/ActorFolders/ActorFolderHierarchyQueries.cpp b/Engine/Source/Editor/UnrealEd/Private/ActorFolders/ActorFolderHierarchyQueries.cpp
new file mode 100644
--- /dev/null
+++ b/Engine/Source/Editor/UnrealEd/Private/ActorFolders/ActorFolderHierarchyQueries.cpp
@@ -0,0 +1,134 @@
+// Copyright Epic Games, Inc. All Rights Reserved.
+
+#include "ActorFolderHierarchyQueries.h"
+#include "ActorFolders/TedsActorFolderUtils.h"
+#include "ActorFolders/ActorFolderColumns.h"
+#include "Elements/Columns/TypedElementCompatibilityColumns.h"
+#include "Elements/Columns/TypedElementFolderColumns.h"
+#include "Elements/Columns/TypedElementMiscColumns.h"
+#include "Elements/Framework/TypedElementQueryBuilder.h"
+#include "Elements/Interfaces/TypedElementDataStorageInterface.h"
+#include "HAL/IConsoleManager.h"
+#include "LevelInstance/LevelInstanceInterface.h"
+#include "GameFramework/Actor.h"
+#include "Engine/Level.h"
+
+namespace UE::Editor::DataStorage::ActorFolders::Private
+{
+	// Helper to get the mapping domain + key for the  first viable parent for a folder in the editor
+	bool GetEditorParentForFolder(const FFolder& Folder, FName& OutParentMappingDomain, FMapKey& OutParentMappingKey)
+	{
+		FFolder ParentFolder = Folder.GetParent();
+		
+		// Parent Folder
+		if (!ParentFolder.IsNone())
+		{
+			OutParentMappingDomain = GetFolderMappingDomain(ParentFolder);
+			OutParentMappingKey = GetFolderIndex(ParentFolder);
+			return true;
+		}
+		
+		UObject* RootObj = ParentFolder.GetRootObjectPtr();
+		// Parent Level Instance
+		if (ILevelInstanceInterface* OwningLevelInstance = Cast<ILevelInstanceInterface>(RootObj))
+		{
+			AActor* OwningActor = CastChecked<AActor>(OwningLevelInstance);
+			OutParentMappingDomain = ICompatibilityProvider::ObjectMappingDomain;
+			OutParentMappingKey = FMapKey(OwningActor);
+			return true;
+
+		}
+		// Parent Level Using Actor Folders
+		if (ULevel* OwningLevel = Cast<ULevel>(RootObj))
+		{
+			OutParentMappingDomain = ICompatibilityProvider::ObjectMappingDomain;
+			OutParentMappingKey = FMapKey(OwningLevel);
+		 	return true;
+		}
+		return false;
+	}
+}
+
+void UTedsActorFolderHierarchyFactory::RegisterQueries(UE::Editor::DataStorage::ICoreProvider& DataStorage)
+{
+	const IConsoleVariable* TedsFoldersCvar = IConsoleManager::Get().FindConsoleVariable(TEXT("TEDS.Feature.Folders"));
+	
+	if (TedsFoldersCvar && TedsFoldersCvar->GetBool())
+	{
+		RegisterFolderHierarchyQueries(DataStorage);
+	}
+}
+
+void UTedsActorFolderHierarchyFactory::RegisterFolderHierarchyQueries(UE::Editor::DataStorage::ICoreProvider& DataStorage)
+{
+	using namespace UE::Editor::DataStorage;
+	using namespace UE::Editor::DataStorage::Queries;
+	
+	FHierarchyHandle EditorFolderHierarchyHandle = DataStorage.FindHierarchyByName(TEXT("EditorObjectHierarchy"));
+
+	DataStorage.RegisterQuery(
+		Select(
+			TEXT("Add parent column to folder"),
+			FProcessor(EQueryTickPhase::PostPhysics, DataStorage.GetQueryTickGroupName(EQueryTickGroups::SyncExternalToDataStorage))
+			.SetExecutionMode(EExecutionMode::GameThread),
+			[](IQueryContext& Context, RowHandle Row, const FFolderCompatibilityColumn& FolderColumn)
+			{
+				FMapKey IdKey;
+				FName MappingDomain;
+			
+				if (ActorFolders::Private::GetEditorParentForFolder(FolderColumn.Folder, MappingDomain, IdKey))
+				{
+					RowHandle ParentRow = Context.LookupMappedRow(MappingDomain, IdKey);
+			
+					if (Context.IsRowAvailable(ParentRow))
+					{
+						Context.SetParentRow(Row, ParentRow);
+					}
+					else
+					{
+						Context.SetUnresolvedParent(Row, IdKey, MappingDomain);
+					}
+				}
+			})
+		.AccessesHierarchy(TEXT("EditorObjectHierarchy"))
+		.Where()
+			.All<FTypedElementSyncFromWorldTag, FFolderTag>()
+			.None(DataStorage.GetChildTagType(EditorFolderHierarchyHandle))
+		.Compile());
+
+	DataStorage.RegisterQuery(
+		Select(
+			TEXT("Sync folder's parent to column"),
+			FProcessor(EQueryTickPhase::PostPhysics, DataStorage.GetQueryTickGroupName(EQueryTickGroups::SyncExternalToDataStorage))
+			.SetExecutionMode(EExecutionMode::GameThread),
+			[](IQueryContext& Context, RowHandle TargetRow, const FFolderCompatibilityColumn& FolderColumn, const FTypedElementWorldColumn& WorldColumn)
+			{
+				FMapKey IdKey;
+				FName MappingDomain;
+			
+				if (ActorFolders::Private::GetEditorParentForFolder(FolderColumn.Folder, MappingDomain, IdKey))
+				{
+					RowHandle NewParentRow = Context.LookupMappedRow(MappingDomain, IdKey);
+					RowHandle CurrentParentRow = Context.GetParentRow(TargetRow);
+					
+					if (CurrentParentRow != NewParentRow)
+					{
+						if (Context.IsRowAvailable(NewParentRow))
+						{
+							Context.SetParentRow(TargetRow, NewParentRow);
+						}
+						else
+						{
+							Context.SetUnresolvedParent(TargetRow, IdKey, MappingDomain);
+						}
+					}
+					return;
+				}
+					
+				Context.SetParentRow(TargetRow, InvalidRowHandle);
+				
+			})
+		.AccessesHierarchy(TEXT("EditorObjectHierarchy"))
+		.Where(TColumn<FTypedElementSyncFromWorldTag>() && TColumn<FFolderTag>() && TColumn(DataStorage.GetChildTagType(EditorFolderHierarchyHandle)))
+		.Compile());
+}
\ No newline at end of file
diff --git a/Engine/Source/Editor/UnrealEd/Private/ActorFolders/ActorFolderHierarchyQueries.h b/Engine/Source/Editor/UnrealEd/Private/ActorFolders/ActorFolderHierarchyQueries.h
new file mode 100644
--- /dev/null
+++ b/Engine/Source/Editor/UnrealEd/Private/ActorFolders/ActorFolderHierarchyQueries.h
@@ -0,0 +1,27 @@
+// Copyright Epic Games, Inc. All Rights Reserved.
+
+#pragma once
+
+#include "DataStorage/Handles.h"
+#include "Elements/Interfaces/TypedElementDataStorageFactory.h"
+
+#include "ActorFolderHierarchyQueries.generated.h"
+
+/*
+ * Factory that contains queries that are related to hierarchies for folders
+ */
+UCLASS()
+class UTedsActorFolderHierarchyFactory : public UEditorDataStorageFactory
+{
+	GENERATED_BODY()
+
+public:
+	virtual ~UTedsActorFolderHierarchyFactory() override = default;
+	virtual void RegisterQueries(UE::Editor::DataStorage::ICoreProvider& DataStorage) override;
+
+protected:
+	void RegisterWorldFolderHierarchyQueries(UE::Editor::DataStorage::ICoreProvider& DataStorage);
+	void RegisterFolderHierarchyQueries(UE::Editor::DataStorage::ICoreProvider& DataStorage);
+};
+
+
diff --git a/Engine/Source/Editor/UnrealEd/Private/ActorFolders/ActorFolderSyncQueries.cpp b/Engine/Source/Editor/UnrealEd/Private/ActorFolders/ActorFolderSyncQueries.cpp
new file mode 100644
--- /dev/null
+++ b/Engine/Source/Editor/UnrealEd/Private/ActorFolders/ActorFolderSyncQueries.cpp
@@ -0,0 +1,116 @@
+// Copyright Epic Games, Inc. All Rights Reserved.
+
+#include "ActorFolderSyncQueries.h"
+
+#include "ActorFolders/ActorFolderColumns.h"
+#include "ActorFolderFactory.h"
+#include "EditorActorFolders.h"
+#include "EditorFolderUtils.h"
+#include "Elements/Columns/TypedElementCompatibilityColumns.h"
+#include "Elements/Columns/TypedElementFolderColumns.h"
+#include "Elements/Columns/TypedElementLabelColumns.h"
+#include "Elements/Columns/TypedElementMiscColumns.h"
+#include "Elements/Framework/TypedElementQueryBuilder.h"
+#include "Elements/Interfaces/TypedElementDataStorageInterface.h"
+
+void UTedsActorFolderSyncFactory::RegisterQueries(UE::Editor::DataStorage::ICoreProvider& InDataStorage)
+{
+	const IConsoleVariable* TedsFoldersCvar = IConsoleManager::Get().FindConsoleVariable(TEXT("TEDS.Feature.Folders"));
+
+	if (TedsFoldersCvar  && TedsFoldersCvar->GetBool())
+	{
+		RegisterLabelQueries(InDataStorage);
+		RegisterMiscQueries(InDataStorage);
+	}
+}
+
+
+void UTedsActorFolderSyncFactory::RegisterLabelQueries(UE::Editor::DataStorage::ICoreProvider& InDataStorage) const
+{
+	using namespace UE::Editor::DataStorage::Queries;
+
+	// We have to defer the rename because for legacy FFolders it deletes and re-creates the folder which re-registers them in TEDS. 
+	// Which cannot be done while we are in a query callback.
+
+	struct FRenameFolderCommand
+	{
+		void operator()() const
+		{
+			FEditorFolderUtils::RenameFolder(Folder,
+				FText::FromString(NewLabel), World.Get());
+			
+		}
+		FFolder Folder;
+		FString NewLabel;
+		TWeakObjectPtr<UWorld> World;
+	};
+	
+	InDataStorage.RegisterQuery(
+		Select(
+			TEXT("Sync folder name to label column"),
+			FProcessor(EQueryTickPhase::PostPhysics, InDataStorage.GetQueryTickGroupName(EQueryTickGroups::SyncExternalToDataStorage))
+				.SetExecutionMode(EExecutionMode::GameThread),
+			[this](IQueryContext& Context, RowHandle RowHandle, const FTypedElementLabelColumn& LabelColumn,
+				const FTypedElementLabelHashColumn& LabelHashColumn, const FFolderCompatibilityColumn& FolderColumn,
+				const FTypedElementWorldColumn& WorldColumn)
+			{
+				FString FolderLabel = FolderColumn.Folder.GetLeafName().ToString();
+				uint64 FolderLabelHash = CityHash64(reinterpret_cast<const char*>(*FolderLabel), FolderLabel.Len() * sizeof(**FolderLabel));
+
+				if (LabelHashColumn.LabelHash != FolderLabelHash)
+				{
+					FRenameFolderCommand Command{.Folder = FolderColumn.Folder, .NewLabel = LabelColumn.Label, .World = WorldColumn.World};
+					Context.PushCommand(Command);
+				}
+			}
+		)
+		.Where(TColumn<FFolderTag>() && TColumn<FTypedElementSyncBackToWorldTag>())
+		.Compile()
+	);
+
+	InDataStorage.RegisterQuery(
+		Select(
+			TEXT("Sync label column to folder name"),
+			FProcessor(EQueryTickPhase::PostPhysics, InDataStorage.GetQueryTickGroupName(EQueryTickGroups::SyncDataStorageToExternal))
+				.SetExecutionMode(EExecutionMode::GameThread),
+			[](IQueryContext& Context, RowHandle Row, const FFolderCompatibilityColumn& FolderColumn,
+				FTypedElementLabelColumn& LabelColumn, FTypedElementLabelHashColumn& LabelHashColumn)
+			{
+				FString FolderLabel = FolderColumn.Folder.GetLeafName().ToString();
+				uint64 FolderLabelHash = CityHash64(reinterpret_cast<const char*>(*FolderLabel), FolderLabel.Len() * sizeof(**FolderLabel));
+
+				if (LabelHashColumn.LabelHash != FolderLabelHash)
+				{
+					LabelColumn.Label = FolderLabel;
+					LabelHashColumn.LabelHash = FolderLabelHash;
+				}
+			}
+		)
+		.Where(TColumn<FFolderTag>() && TColumn<FTypedElementSyncFromWorldTag>())
+		.Compile()
+		);
+}
+
+void UTedsActorFolderSyncFactory::RegisterMiscQueries(UE::Editor::DataStorage::ICoreProvider& InDataStorage) const
+{
+	using namespace UE::Editor::DataStorage::Queries;
+
+	// For actor folders, the FFolder data can change without the row getting unregistered and re-registered in TEDS (unlike legacy FFolders)
+	// This query makes sure the FFolder in TEDS has the correct data
+	InDataStorage.RegisterQuery(
+		Select(
+			TEXT("Sync FFolder from ActorFolder"),
+			FProcessor(EQueryTickPhase::PrePhysics, InDataStorage.GetQueryTickGroupName(EQueryTickGroups::SyncDataStorageToExternal))
+				.SetExecutionMode(EExecutionMode::GameThread),
+			[this](RowHandle RowHandle, const FTypedElementUObjectColumn& ObjectColumn, FFolderCompatibilityColumn& FolderColumn)
+			{
+				if (UActorFolder* ActorFolder = Cast<UActorFolder>(ObjectColumn.Object.Get(true)))
+				{
+					FolderColumn.Folder = ActorFolder->GetFolder();
+				}
+			}
+		)
+		.Where(TColumn<FFolderTag>() && TColumn<FTypedElementSyncFromWorldTag>())
+		.Compile()
+	);
+}
\ No newline at end of file
diff --git a/Engine/Source/Editor/UnrealEd/Private/ActorFolders/ActorFolderSyncQueries.h b/Engine/Source/Editor/UnrealEd/Private/ActorFolders/ActorFolderSyncQueries.h
new file mode 100644
--- /dev/null
+++ b/Engine/Source/Editor/UnrealEd/Private/ActorFolders/ActorFolderSyncQueries.h
@@ -0,0 +1,26 @@
+// Copyright Epic Games, Inc. All Rights Reserved.
+
+#pragma once
+
+#include "Elements/Interfaces/TypedElementDataStorageFactory.h"
+
+#include "ActorFolderSyncQueries.generated.h"
+
+/*
+ * Factory that contains queries that are used to sync general data for actor folders to and from TEDS
+ */
+UCLASS()
+class UTedsActorFolderSyncFactory : public UEditorDataStorageFactory
+{
+	GENERATED_BODY()
+
+public:
+	virtual ~UTedsActorFolderSyncFactory() override = default;
+	virtual void RegisterQueries(UE::Editor::DataStorage::ICoreProvider& DataStorage) override;
+
+private:
+	void RegisterLabelQueries(UE::Editor::DataStorage::ICoreProvider& DataStorage) const;
+	void RegisterMiscQueries(UE::Editor::DataStorage::ICoreProvider& DataStorage) const;
+};
+
+
diff --git a/Engine/Source/Editor/UnrealEd/Private/ActorFolders/TedsActorFolderTests.cpp b/Engine/Source/Editor/UnrealEd/Private/ActorFolders/TedsActorFolderTests.cpp
new file mode 100644
--- /dev/null
+++ b/Engine/Source/Editor/UnrealEd/Private/ActorFolders/TedsActorFolderTests.cpp
@@ -0,0 +1,123 @@
+﻿// Copyright Epic Games, Inc. All Rights Reserved.
+
+#if WITH_TESTS
+
+#include "DataStorage/Features.h"
+#include "Editor.h"
+#include "EditorActorFolders.h"
+#include "ActorFolders/TedsActorFolderUtils.h"
+#include "Editor/EditorEngine.h"
+#include "Elements/Interfaces/TypedElementDataStorageInterface.h"
+#include "Misc/AutomationTest.h"
+#include "Tests/AutomationCommon.h"
+
+namespace UE::Editor::ActorFolders::Tests
+{
+	using namespace UE::Editor::DataStorage;
+	
+	BEGIN_DEFINE_SPEC(TedsActorFolderTestFixture, "Editor.ActorFolders.Teds", EAutomationTestFlags::EditorContext | EAutomationTestFlags::EngineFilter)
+	
+	ICoreProvider* TedsInterface = nullptr;
+	UWorld* World = nullptr;
+	
+	TSet<FFolder> Folders;
+	
+	FFolder CreateNewFolder(FFolder ParentFolder)
+	{
+		const FFolder NewFolderName = FActorFolders::Get().GetDefaultFolderName(*World, ParentFolder);
+		
+		// Keep a record of all the folders we create so we can delete them later
+		Folders.Add(NewFolderName);
+		
+		if (FActorFolders::Get().CreateFolder(*World, NewFolderName))
+		{
+			return NewFolderName;
+		}
+		
+		return FFolder::GetInvalidFolder();
+	}
+	
+	void DeleteFolder(FFolder InFolder)
+	{
+		// Not doing this on purpose as tests call undo/redo which might bring the folder back without adding them back to the Set
+		// It's safer to just always cleanup folders even if it was already deleted
+		//Folders.Remove(InFolder);
+		
+		FActorFolders::Get().DeleteFolder(*World, InFolder);
+	}
+	
+	void CleanupFolders()
+	{
+		for (const FFolder& Folder : Folders)
+		{
+			FActorFolders::Get().DeleteFolder(*World, Folder);
+		}
+	}
+	
+	END_DEFINE_SPEC(TedsActorFolderTestFixture)
+    
+    void TedsActorFolderTestFixture::Define()
+	{
+		BeforeEach([this]()
+		{
+			TedsInterface = GetMutableDataStorageFeature<ICoreProvider>(StorageFeatureName);
+			TestTrue("", TedsInterface != nullptr);
+			
+			World = GEditor->GetEditorWorldContext().World();
+		});
+		
+		Describe("Verify Immediate Teds State", [this]
+		{
+			It(TEXT("Create Folder"), EAsyncExecution::TaskGraphMainTick, [this]()
+			{
+				FFolder Folder = CreateNewFolder(FFolder::GetInvalidFolder());
+				TestTrue(TEXT("Expected folder to be created successfully"), Folder.IsValid());
+				
+				RowHandle FolderRow = UE::Editor::DataStorage::ActorFolders::LookupMappedRow(TedsInterface, Folder);
+				TestTrue(TEXT("Expected folder row to be valid in TEDS"), TedsInterface->IsRowAvailable(FolderRow));
+			});
+				
+			It(TEXT("Delete Folder"), EAsyncExecution::TaskGraphMainTick, [this]()
+			{
+				FFolder Folder = CreateNewFolder(FFolder::GetInvalidFolder());
+				TestTrue(TEXT("Expected folder to be created successfully"), Folder.IsValid());
+				
+				DeleteFolder(Folder);
+				
+				RowHandle FolderRow = UE::Editor::DataStorage::ActorFolders::LookupMappedRow(TedsInterface, Folder);
+				TestFalse(TEXT("Expected folder row to be invalid in TEDS after deletion"), TedsInterface->IsRowAvailable(FolderRow));
+			});	
+				
+			It(TEXT("Undo Folder Creation"), EAsyncExecution::TaskGraphMainTick, [this]()
+			{
+				FFolder Folder = CreateNewFolder(FFolder::GetInvalidFolder());
+				TestTrue(TEXT("Expected folder to be created successfully"), Folder.IsValid());
+			
+				GEditor->UndoTransaction();
+			
+				RowHandle FolderRow = UE::Editor::DataStorage::ActorFolders::LookupMappedRow(TedsInterface, Folder);
+				TestFalse(TEXT("Expected folder row to be invalid in TEDS after undo-ing creation"), TedsInterface->IsRowAvailable(FolderRow));
+			});
+				
+			It(TEXT("Undo Folder Deletion"), EAsyncExecution::TaskGraphMainTick, [this]()
+			{
+				FFolder Folder = CreateNewFolder(FFolder::GetInvalidFolder());
+				TestTrue(TEXT("Expected folder to be created successfully"), Folder.IsValid());
+					
+				DeleteFolder(Folder);
+				GEditor->UndoTransaction();
+					
+				RowHandle FolderRow = UE::Editor::DataStorage::ActorFolders::LookupMappedRow(TedsInterface, Folder);
+				TestTrue(TEXT("Expected folder row to be valid in TEDS after undo-ing deletion"), TedsInterface->IsRowAvailable(FolderRow));
+			});
+		});
+		
+		AfterEach([this]()
+		{
+			CleanupFolders();
+			TedsInterface = nullptr;
+		});
+	}
+}
+
+#endif
\ No newline at end of file
diff --git a/Engine/Source/Editor/UnrealEd/Private/ActorFolders/TedsActorFolderUtils.cpp b/Engine/Source/Editor/UnrealEd/Private/ActorFolders/TedsActorFolderUtils.cpp
new file mode 100644
--- /dev/null
+++ b/Engine/Source/Editor/UnrealEd/Private/ActorFolders/TedsActorFolderUtils.cpp
@@ -0,0 +1,226 @@
+// Copyright Epic Games, Inc. All Rights Reserved.
+
+#include "ActorFolders/TedsActorFolderUtils.h"
+
+#include "ActorFolder.h"
+#include "ActorFolders/ActorFolderColumns.h"
+#include "DataStorage/Features.h"
+#include "Elements/Columns/TypedElementCompatibilityColumns.h"
+#include "Elements/Columns/TypedElementMiscColumns.h"
+#include "Engine/World.h"
+
+namespace UE::Editor::DataStorage::ActorFolders
+{
+	namespace Private
+	{
+		static const FName TableName("Editor_ActorFolderTable");
+
+		// Kept private on purpose because actor folders and non-actor folders are mapped under different domains and keys and we want the user to go through
+		// LookupMappedRow instead of needing to know which one a folder is
+		static const FName MappingDomain("ActorFolders");
+		
+		RowHandle RegisterFolder_Internal(ICoreProvider* Storage, const FFolder& Folder)
+		{
+			// Register actor folders with TEDS compat since they are UObjects
+			if (Folder.GetActorFolder())
+			{
+				ICompatibilityProvider* StorageCompat = GetMutableDataStorageFeature<ICompatibilityProvider>(CompatibilityFeatureName);
+				return StorageCompat->AddCompatibleObjectExplicit(Folder.GetActorFolder());
+			}
+			
+			// Regular folders are registered and kept track of manually
+			else
+			{
+				TableHandle FolderTable = Storage->FindTable(GetActorFolderTableName());
+				checkf(Storage->IsTableValid(FolderTable), TEXT("Cannot register folders in TEDS if the folder table is not registered."));
+
+				RowHandle Row = Storage->AddRow(FolderTable);
+				Storage->MapRow(GetFolderMappingDomain(Folder), GetFolderIndex(Folder), Row);
+				return Row;
+			}
+
+		}
+		
+		void UnregisterFolder_Internal(ICoreProvider* Storage, const FFolder& Folder)
+		{
+			if (Folder.GetActorFolder())
+			{
+				ICompatibilityProvider* StorageCompat = GetMutableDataStorageFeature<ICompatibilityProvider>(CompatibilityFeatureName);
+				StorageCompat->RemoveCompatibleObjectExplicit(Folder.GetActorFolder());
+			}
+			else
+			{
+				Storage->RemoveRow(LookupMappedRow(Storage, Folder));
+			}
+		}
+	}
+	
+	FName GetActorFolderTableName()
+	{
+		return Private::TableName;
+	}
+
+	void SetFolderColumns(ICoreProvider* Storage, RowHandle Row, UWorld* World, const FFolder& Folder)
+	{
+		Storage->AddColumn(Row, FTypedElementWorldColumn{.World = TWeakObjectPtr<UWorld>(World)});
+		Storage->AddColumn(Row, FFolderCompatibilityColumn{.Folder = Folder});
+
+		if (UActorFolder* ActorFolder = Folder.GetActorFolder())
+		{
+			Storage->AddColumn(Row, FTypedElementUObjectColumn{.Object = ActorFolder});
+		}
+
+		// Add the sync tag so other columns are synced through processors correctly
+		Storage->AddColumn<FTypedElementSyncFromWorldTag>(Row);
+	}
+	
+	RowHandle RegisterFolderInTeds(ICoreProvider* Storage, const FFolder& Folder, UWorld* World, bool bRecursivelyAddParents)
+	{
+		// Walk up the parent chain and register parent folders if requested. The hierarchy is setup by a processor so we don't have to do it
+		// here
+		if (bRecursivelyAddParents)
+		{
+			if (FFolder ParentFolder = Folder.GetParent(); !ParentFolder.IsNone())
+			{
+				RegisterFolderInTeds(Storage, ParentFolder, World, bRecursivelyAddParents);
+			}
+		}
+		
+		RowHandle ExistingRow = LookupMappedRow(Storage, Folder);
+
+		if (Storage->IsRowAvailable(ExistingRow))
+		{
+			return ExistingRow;
+		}
+		
+		RowHandle FolderRow = Private::RegisterFolder_Internal(Storage, Folder);
+		
+		if (Storage->IsRowAvailable(FolderRow))
+		{
+			SetFolderColumns(Storage, FolderRow, World, Folder);
+			return FolderRow;
+		}
+
+		return InvalidRowHandle;
+	}
+
+	void RegisterFolderInTeds(IQueryContext& Context, ICoreProvider* Storage, const FFolder& Folder, TWeakObjectPtr<UWorld> World, bool bRecursivelyAddParents)
+	{
+		struct FRegisterFolderCommand
+		{
+			void operator()() const
+			{
+				if (UWorld* World = WorldPtr.Get())
+				{
+					RegisterFolderInTeds(Storage, Folder, World, bRecursivelyAddParents);
+				}
+			}
+
+			ICoreProvider* Storage;
+			FFolder Folder;
+			TWeakObjectPtr<UWorld> WorldPtr;
+			bool bRecursivelyAddParents;
+		};
+
+		FRegisterFolderCommand Command
+		{
+			.Storage = Storage,
+			.Folder = Folder,
+			.WorldPtr = World,
+			.bRecursivelyAddParents = bRecursivelyAddParents
+		};
+		Context.PushCommand(Command);
+	}
+
+	void UnregisterFolderFromTeds(ICoreProvider* Storage, const FFolder& Folder)
+	{
+		Private::UnregisterFolder_Internal(Storage, Folder);
+	}
+	
+	void UnregisterFolderFromTeds(IQueryContext& Context, ICoreProvider* Storage, const FFolder& Folder)
+	{
+		struct FUnregisterFolderCommand
+		{
+			void operator()() const
+			{
+				UnregisterFolderFromTeds(Storage, Folder);
+			}
+
+			ICoreProvider* Storage;
+			FFolder Folder;
+		};
+
+		FUnregisterFolderCommand Command
+		{
+			.Storage = Storage,
+			.Folder = Folder,
+		};
+		Context.PushCommand(Command);
+	}
+
+	RowHandle LookupMappedRow(const ICoreProvider* Storage, const FFolder& Folder)
+	{
+		return Storage->LookupMappedRow(GetFolderMappingDomain(Folder), GetFolderIndex(Folder));
+	}
+
+	RowHandle LookupMappedRow(const IQueryContext& Context, const FFolder& Folder)
+	{
+		return Context.LookupMappedRow(GetFolderMappingDomain(Folder), GetFolderIndex(Folder));
+	}
+
+	RowHandle RemapFolder(ICoreProvider* Storage, const FFolder& InOldFolder, const FFolder& InNewFolder, UWorld* World)
+	{
+		if (!InNewFolder.GetActorFolder())
+		{
+			// Remap this row to the new folder information
+            FMapKey OldIndex = GetFolderIndex(InOldFolder);
+            FMapKey NewIndex = GetFolderIndex(InNewFolder);
+    
+            // Actor folders are mapped by the UObject pointer which doesn't really change, so there is no need to remap them if the index didn't change
+            // OR if InOldFolder is not actually mapped to any row (i.e it's a stale FFolder that simply needs updating in FFolderCompatibilityColumn)
+            if (OldIndex != NewIndex && Storage->IsRowAvailable(Storage->LookupMappedRow(GetFolderMappingDomain(InOldFolder), OldIndex)))
+            {
+            	Storage->RemapRow(GetFolderMappingDomain(InNewFolder), OldIndex, NewIndex);
+            }
+		}
+		
+		RowHandle NewFolderRow = LookupMappedRow(Storage, InNewFolder);
+		
+		// We want to set the default columns again even if the index didn't change, e.g an actor folder is mapped by the UObject ptr but the FFolder
+		// associated with it has changed and FFolderCompatibilityColumn needs updating
+		SetFolderColumns(Storage, NewFolderRow, World, InNewFolder);
+		
+		return NewFolderRow;
+	}
+
+	void SetUnresolvedFolderParent(IQueryContext& Context, RowHandle ChildRow, const FFolder& ParentFolder)
+	{
+		FMapKey ParentIndex = GetFolderIndex(ParentFolder);
+		Context.SetUnresolvedParent(ChildRow, ParentIndex, GetFolderMappingDomain(ParentFolder));
+	}
+	
+	FName GetFolderMappingDomain(const FFolder& Folder)
+	{
+		// Actor folders are registered with TEDS Compat
+		if (Folder.GetActorFolder())
+		{
+			return ICompatibilityProvider::ObjectMappingDomain;
+		}
+
+		return Private::MappingDomain;
+	}
+
+	FMapKey GetFolderIndex(const FFolder& Folder)
+	{
+		// Actor folders are mapped by the UObject itself because it is stable across operations that create redirectors
+		if (UActorFolder* ActorFolder = Folder.GetActorFolder())
+		{
+			return FMapKey(ActorFolder);
+		}
+		else
+		{
+			// We also use the root object as a part of the hash since two folders in different worlds (e.g Editor and PIE) can have the same hash
+			return FMapKey(static_cast<uint64>(HashCombine(GetTypeHash(Folder.GetRootObject()), GetTypeHash(Folder))));
+		}
+	}
+}
diff --git a/Engine/Source/Editor/UnrealEd/Private/EditorActorFolders.cpp b/Engine/Source/Editor/UnrealEd/Private/EditorActorFolders.cpp
--- a/Engine/Source/Editor/UnrealEd/Private/EditorActorFolders.cpp
+++ b/Engine/Source/Editor/UnrealEd/Private/EditorActorFolders.cpp
@@ -41,7 +41,7 @@ FActorFolders::FActorFolders()
 	check(GEngine);
 	GEngine->OnLevelActorFolderChanged().AddRaw(this, &FActorFolders::OnActorFolderChanged);
 	GEngine->OnActorFolderAdded().AddRaw(this, &FActorFolders::OnActorFolderAdded);
-
+	
 	FEditorDelegates::PostSaveWorldWithContext.AddRaw(this, &FActorFolders::OnWorldSaved);
 	FEditorDelegates::PostSaveExternalActors.AddRaw(this, &FActorFolders::SaveWorldFoldersState);
 
@@ -87,6 +87,7 @@ void FActorFolders::Housekeeping()
 	{
 		if (!It.Key().Get())
 		{
+			It.Value()->OnRemoved();
 			It.RemoveCurrent();
 		}
 	}
@@ -133,6 +134,11 @@ bool FActorFolders::IsInitializedForWorld(UWorld& InWorld) const
 	return !!WorldFolders.Find(&InWorld);
 }
 
+void FActorFolders::InitializeForWorld(UWorld& InWorld)
+{
+	GetOrCreateWorldFolders(InWorld);
+}
+
 UWorldFolders& FActorFolders::GetOrCreateWorldFolders(UWorld& InWorld)
 {
 	if (auto* Folders = WorldFolders.Find(&InWorld))
diff --git a/Engine/Source/Editor/UnrealEd/Private/EditorFolderUtils.cpp b/Engine/Source/Editor/UnrealEd/Private/EditorFolderUtils.cpp
--- a/Engine/Source/Editor/UnrealEd/Private/EditorFolderUtils.cpp
+++ b/Engine/Source/Editor/UnrealEd/Private/EditorFolderUtils.cpp
@@ -2,6 +2,8 @@
 
 #include "EditorFolderUtils.h"
 
+#include "EditorActorFolders.h"
+
 FName FEditorFolderUtils::GetLeafName(const FName& InPath)
 {
 	FString PathString = InPath.ToString();
@@ -28,4 +30,25 @@ bool FEditorFolderUtils::PathIsChildOf(const FString& PotentialChild, const FStr
 bool FEditorFolderUtils::PathIsChildOf(const FName& PotentialChild, const FName& Parent)
 {
 	return PathIsChildOf(PotentialChild.ToString(), Parent.ToString());
+}
+
+void FEditorFolderUtils::RenameFolder(const FFolder& InFolder, const FText& NewFolderName, UWorld* World)
+{
+	if (!World)
+	{
+		return;
+	}
+		
+	FName NewPath = InFolder.GetParent().GetPath();
+	if (NewPath.IsNone())
+	{
+		NewPath = FName(*NewFolderName.ToString());
+	}
+	else
+	{
+		NewPath = FName(*(NewPath.ToString() / NewFolderName.ToString()));
+	}
+		
+	const FFolder TreeItemNewFolder(InFolder.GetRootObject(), NewPath);
+	FActorFolders::Get().RenameFolderInWorld(*World, InFolder, TreeItemNewFolder);
 }
\ No newline at end of file
diff --git a/Engine/Source/Editor/UnrealEd/Private/WorldFolders.cpp b/Engine/Source/Editor/UnrealEd/Private/WorldFolders.cpp
--- a/Engine/Source/Editor/UnrealEd/Private/WorldFolders.cpp
+++ b/Engine/Source/Editor/UnrealEd/Private/WorldFolders.cpp
@@ -1,19 +1,25 @@
 // Copyright Epic Games, Inc. All Rights Reserved.
 
 #include "WorldFolders.h"
+#include "ActorFolders/ActorFolderColumns.h"
+#include "ActorFolders/TedsActorFolderUtils.h"
 #include "Serialization/JsonReader.h"
 #include "Serialization/JsonSerializer.h"
+#include "DataStorage/Features.h"
 #include "Dom/JsonObject.h"
 #include "Engine/World.h"
 #include "HAL/FileManager.h"
 #include "Misc/Paths.h"
 #include "EngineUtils.h"
 #include "EditorActorFolders.h"
+#include "Elements/Columns/TypedElementMiscColumns.h"
+#include "Elements/Interfaces/TypedElementDataStorageCompatibilityInterface.h"
 #include "LevelInstance/LevelInstanceSubsystem.h"
 #include "LevelInstance/LevelInstanceInterface.h"
 #include "UObject/Package.h"
 #include "UObject/WeakObjectPtr.h"
 #include "Editor.h"
+#include "Elements/Columns/TypedElementUIColumns.h"
 
 #include UE_INLINE_GENERATED_CPP_BY_NAME(WorldFolders)
 
@@ -21,6 +27,17 @@
 
 DEFINE_LOG_CATEGORY(LogWorldFolders);
 
+namespace UE::Editor::WorldFolders
+{
+	static bool bRegisterFoldersinTEDS = false;
+	
+	// Cvar to control registration of actor folders in TEDS
+	static FAutoConsoleVariableRef CvarRegisterFoldersinTEDS(
+		TEXT("TEDS.Feature.Folders"),
+		bRegisterFoldersinTEDS,
+		TEXT("Populate FFolders and Actor Folders in TEDS. Must be set at startup."));
+}
+
 void FActorPlacementFolder::Reset()
 {
 	Path = NAME_None;
@@ -54,6 +71,11 @@ void UWorldFolders::Initialize(UWorld* InWorld)
 	
 	PersistentFolders = MakeUnique<FWorldPersistentFolders>(*this);
 	TransientFolders = MakeUnique<FWorldTransientFolders>(*this);
+	
+	if (UE::Editor::WorldFolders::CvarRegisterFoldersinTEDS->GetBool())
+	{
+		DataStorage = UE::Editor::DataStorage::GetMutableDataStorageFeature<UE::Editor::DataStorage::ICoreProvider>(UE::Editor::DataStorage::StorageFeatureName);
+	}
 
 	bQueuedRebuild = false;
 	RebuildList();
@@ -145,11 +167,61 @@ void UWorldFolders::RebuildListInternal()
 	bQueuedRebuild = false;
 }
 
+void UWorldFolders::AddFolder_Internal(const FFolder& InFolder, const FActorFolderProps& FolderProperties)
+{
+	FoldersProperties.Add(InFolder, FolderProperties);
+	
+	using namespace UE::Editor::DataStorage;
+	if (UE::Editor::WorldFolders::CvarRegisterFoldersinTEDS->GetBool())
+	{
+		ActorFolders::RegisterFolderInTeds(DataStorage, InFolder, GetWorld());
+	}
+}
+
+void UWorldFolders::RemoveFolder_Internal(const FFolder& InFolder)
+{
+	using namespace UE::Editor::DataStorage;
+	if (UE::Editor::WorldFolders::CvarRegisterFoldersinTEDS->GetBool())
+	{
+		ActorFolders::UnregisterFolderFromTeds(DataStorage, InFolder);
+	}
+	
+	FoldersProperties.Remove(InFolder);
+}
+
+void UWorldFolders::ReAddFolder_Internal(const FFolder& InOldFolder, const FFolder& InNewFolder, const FActorFolderProps& InNewFolderProperties)
+{
+	FoldersProperties.Remove(InOldFolder);
+	FoldersProperties.Add(InNewFolder, InNewFolderProperties);
+	
+	using namespace UE::Editor::DataStorage;
+	if (UE::Editor::WorldFolders::CvarRegisterFoldersinTEDS->GetBool())
+	{
+		ActorFolders::RemapFolder(DataStorage, InOldFolder, InNewFolder, GetWorld());
+	}
+}
+
 UWorld* UWorldFolders::GetWorld() const
 {
 	return World.Get();
 }
 
+void UWorldFolders::OnRemoved()
+{
+	// This function is called when the FActorFolders bookkeeping decides that a particular UWorldFolders object isn't needed anymore
+	// so it's a good spot to unregister all folders belonging to a world from TEDS without relying on GC timing
+	if (UE::Editor::WorldFolders::CvarRegisterFoldersinTEDS->GetBool())
+	{
+		using namespace UE::Editor::DataStorage;
+		
+		ForEachFolder([this](const FFolder& InFolder)
+		{
+			ActorFolders::UnregisterFolderFromTeds(DataStorage, InFolder);
+			return true;
+		});
+	}
+}
+
 bool UWorldFolders::AddFolder(const FFolder& InFolder)
 {
 	if (!InFolder.IsNone())
@@ -164,10 +236,13 @@ bool UWorldFolders::AddFolder(const FFolder& InFolder)
 			}
 
 			Modify();
-			FActorFolderProps* LoadedFolderProps = LoadedStateFoldersProperties.Find(InFolder);
-			FoldersProperties.Add(InFolder, LoadedFolderProps ? *LoadedFolderProps : FActorFolderProps());
+			
+			bool bAdded = GetImpl(InFolder).AddFolder(InFolder);
 
-			return GetImpl(InFolder).AddFolder(InFolder);
+			FActorFolderProps* LoadedFolderProps = LoadedStateFoldersProperties.Find(InFolder);
+			AddFolder_Internal(InFolder, LoadedFolderProps ? *LoadedFolderProps : FActorFolderProps());
+			
+			return bAdded;
 		}
 	}
 
@@ -179,8 +254,7 @@ bool UWorldFolders::RemoveFolder(const FFolder& InFolder, bool bShouldDeleteFold
 	if (FoldersProperties.Contains(InFolder))
 	{
 		Modify();
-		FoldersProperties.Remove(InFolder);
-
+		RemoveFolder_Internal(InFolder);
 		return GetImpl(InFolder).RemoveFolder(InFolder, bShouldDeleteFolder);
 	}
 	return false;
@@ -215,6 +289,14 @@ bool UWorldFolders::RenameFolder(const FFolder& InOldFolder, const FFolder& InNe
 		{
 			FActorFolders::Get().BroadcastOnActorEditorContextClientChanged(*World.Get());
 		}
+
+		// The label for folders in TEDS is updated by a processor, add the sync tag so that it triggers for legacy FFolders as well
+		using namespace UE::Editor::DataStorage;
+		if (UE::Editor::WorldFolders::CvarRegisterFoldersinTEDS->GetBool())
+		{
+			RowHandle NewFolderRow = ActorFolders::LookupMappedRow(DataStorage, InNewFolder);
+			DataStorage->AddColumn<FTypedElementSyncFromWorldTag>(NewFolderRow);
+		}
 	}
 	return bSuccess;
 }
@@ -248,6 +330,15 @@ bool UWorldFolders::SetIsFolderExpanded(const FFolder& InFolder, bool bIsExpande
 	if (FActorFolderProps* FolderProps = FoldersProperties.Find(InFolder))
 	{
 		FolderProps->bIsExpanded = bIsExpanded;
+		
+		if (UE::Editor::WorldFolders::CvarRegisterFoldersinTEDS->GetBool())
+		{
+			using namespace UE::Editor::DataStorage;
+			
+			RowHandle FolderRow = ActorFolders::LookupMappedRow(DataStorage, InFolder);
+			bIsExpanded ? DataStorage->AddColumn<FExpandedInUITag>(FolderRow) : DataStorage->RemoveColumn<FExpandedInUITag>(FolderRow);
+		}
+		
 		return true;
 	}
 	return false;
@@ -345,10 +436,86 @@ void UWorldFolders::Serialize(FArchive& Ar)
 
 	check(PersistentFolders.IsValid());
 	Ar << FoldersProperties;
-
+	
 	Super::Serialize(Ar);
 }
 
+void UWorldFolders::PreEditUndo()
+{
+	Super::PreEditUndo();
+	
+	// We are currently mirroring folder data in TEDS, however with current implementation of UWorldFolders the individual folders don't receive
+	// any notification when their data changes via undo/redo and neither are any delegates fired. The per-world UWorldFolders object is the one
+	// that is modified/restored using the transaction system.
+	
+	// To make sure these changes are tracked in TEDS, we need to keep a copy of all folders before and after the undo/redo to compare and figure out
+	// which folders changed
+	if (UE::Editor::WorldFolders::CvarRegisterFoldersinTEDS->GetBool())
+	{
+		PreUndoFolders.Empty();
+		
+		// Cache the list of folders we have before the undo/redo operation to compare after
+		for (const TPair<FFolder, FActorFolderProps>& FolderIt : FoldersProperties)
+		{
+			// Actor folders are handled by TEDS compat since they are UObjects, plus their key is not stable across renames so this logic cannot
+			// track them properly
+			if (!FolderIt.Key.GetActorFolder())
+			{
+				PreUndoFolders.Add(FolderIt.Key);
+			}
+		}
+	}
+}
+
+void UWorldFolders::PostEditUndo()
+{
+	Super::PostEditUndo();
+	
+	// Compare the pre-undo list with the post-undo list to find out which folders were added/removed because of the operation
+	if (UE::Editor::WorldFolders::CvarRegisterFoldersinTEDS->GetBool())
+	{
+		using namespace UE::Editor::DataStorage;
+		
+		// Get the list of folders we have after the undo/redo operation
+		TSet<FFolder> PostUndoFolders;
+		PostUndoFolders.Reserve(FoldersProperties.Num());
+		
+		for (const TPair<FFolder, FActorFolderProps>& FolderIt : FoldersProperties)
+		{
+			// Actor folders are handled by TEDS compat since they are UObjects, plus their key is not stable across renames so this logic cannot
+			// track them properly
+			if (!FolderIt.Key.GetActorFolder())
+			{
+				PostUndoFolders.Add(FolderIt.Key);
+			}
+			
+			// We also want to add the sync tag to all folders, since undo/redo can cause data changes (e.g label) and not just remove/add which
+			// needs to re-run processors (similar to how actors get this tag automatically on undo/redo because of TEDS compat)
+			
+			// While actor folders also go through TEDS compat and get this tag added, that isn't true for all operations since some data is still
+			// stored in the FFolder for actor folders which does not modify the UObject. So we add it manually just to be sure
+			RowHandle FolderRow = ActorFolders::LookupMappedRow(DataStorage, FolderIt.Key);
+			DataStorage->AddColumn<FTypedElementSyncFromWorldTag>(FolderRow);
+		}
+	
+		// Any items that exist post-undo but did not exist pre-undo were "added" because of the undo
+		TSet<FFolder> AddedFolders = PostUndoFolders.Difference(PreUndoFolders);
+		// Add those folders to TEDS
+		for (const FFolder& AddedFolder : AddedFolders)
+		{
+			ActorFolders::RegisterFolderInTeds(DataStorage, AddedFolder, GetWorld());
+		}
+		
+		// Any items that existed pre-undo but do not exist post-undo were "removed" because of the undo
+		TSet<FFolder> RemovedFolders = PreUndoFolders.Difference(PostUndoFolders);
+		// Remove those folders from TEDS
+		for (const FFolder& RemoveFolder : RemovedFolders)
+		{
+			ActorFolders::UnregisterFolderFromTeds(DataStorage, RemoveFolder);
+		}
+	}
+}
+
 FString UWorldFolders::GetWorldStateFilename() const
 {
 	if (World->IsGameWorld() || World->IsInstanced() || FPackageName::IsTempPackage(World->GetPackage()->GetName()))
diff --git a/Engine/Source/Editor/UnrealEd/Private/WorldPersistentFolders.cpp b/Engine/Source/Editor/UnrealEd/Private/WorldPersistentFolders.cpp
--- a/Engine/Source/Editor/UnrealEd/Private/WorldPersistentFolders.cpp
+++ b/Engine/Source/Editor/UnrealEd/Private/WorldPersistentFolders.cpp
@@ -122,8 +122,7 @@ void FWorldPersistentFolders::ModifyFolderAndDetectChanges(ULevel* InLevel, cons
 		FFolder NewFolder(InRootObject, ChangedFolder.Value);
 		FActorFolderProps* OldFolderProperties = Owner.FoldersProperties.Find(OldFolder);
 		FActorFolderProps FolderProperties = OldFolderProperties ? *OldFolderProperties : DefaultFolderProperties;
-		Owner.FoldersProperties.Remove(OldFolder);
-		Owner.FoldersProperties.Add(NewFolder, FolderProperties);
+		Owner.ReAddFolder_Internal(OldFolder, NewFolder, FolderProperties);
 			
 		Owner.BroadcastOnActorFolderMoved(OldFolder, NewFolder);
 	}
diff --git a/Engine/Source/Editor/UnrealEd/Private/WorldTransientFolders.cpp b/Engine/Source/Editor/UnrealEd/Private/WorldTransientFolders.cpp
--- a/Engine/Source/Editor/UnrealEd/Private/WorldTransientFolders.cpp
+++ b/Engine/Source/Editor/UnrealEd/Private/WorldTransientFolders.cpp
@@ -52,12 +52,12 @@ bool FWorldTransientFolders::RenameFolder(const FFolder& OldPath, const FFolder&
 				// Use the existing properties for the folder if we have them
 				if (FActorFolderProps* ExistingProperties = Owner.FoldersProperties.Find(Path))
 				{
-					Owner.FoldersProperties.Add(NewFolder, *ExistingProperties);
+					Owner.AddFolder_Internal(NewFolder, *ExistingProperties);
 				}
 				else
 				{
 					// Otherwise use default properties
-					Owner.FoldersProperties.Add(NewFolder);
+					Owner.AddFolder_Internal(NewFolder);
 				}
 				Owner.BroadcastOnActorFolderMoved(Path, NewFolder);
 				Owner.BroadcastOnActorFolderCreated(NewFolder);
@@ -107,7 +107,7 @@ bool FWorldTransientFolders::RenameFolder(const FFolder& OldPath, const FFolder&
 	// Cleanup any old folders
 	for (const FFolder& Folder : RenamedFolders)
 	{
-		Owner.FoldersProperties.Remove(Folder);
+		Owner.RemoveFolder_Internal(Folder);
 		Owner.BroadcastOnActorFolderDeleted(Folder);
 	}
 
diff --git a/Engine/Source/Editor/UnrealEd/Public/ActorFolders/ActorFolderColumns.h b/Engine/Source/Editor/UnrealEd/Public/ActorFolders/ActorFolderColumns.h
new file mode 100644
--- /dev/null
+++ b/Engine/Source/Editor/UnrealEd/Public/ActorFolders/ActorFolderColumns.h
@@ -0,0 +1,19 @@
+// Copyright Epic Games, Inc. All Rights Reserved.
+
+#pragma once
+
+#include "Elements/Interfaces/TypedElementDataStorageInterface.h"
+#include "Folder.h"
+
+#include "ActorFolderColumns.generated.h"
+
+/**
+ * Column that stores a constructed FFolder
+ */
+UE_EXPERIMENTAL(5.8, "The TEDS+Folders integration is experimental and the API subject to change.")
+USTRUCT(meta = (DisplayName = "FFolder compatibility"))
+struct FFolderCompatibilityColumn final : public FEditorDataStorageColumn
+{
+	GENERATED_BODY()
+	FFolder Folder;
+};
\ No newline at end of file
diff --git a/Engine/Source/Editor/UnrealEd/Public/ActorFolders/TedsActorFolderUtils.h b/Engine/Source/Editor/UnrealEd/Public/ActorFolders/TedsActorFolderUtils.h
new file mode 100644
--- /dev/null
+++ b/Engine/Source/Editor/UnrealEd/Public/ActorFolders/TedsActorFolderUtils.h
@@ -0,0 +1,64 @@
+// Copyright Epic Games, Inc. All Rights Reserved.
+
+#pragma once
+
+#include "DataStorage/Handles.h"
+#include "Elements/Interfaces/TypedElementDataStorageCompatibilityInterface.h"
+#include "Elements/Interfaces/TypedElementDataStorageInterface.h"
+#include "Folder.h"
+
+#define UE_API UNREALED_API
+
+// Utility functions for the actor folder integration with TEDS
+namespace UE::Editor::DataStorage::ActorFolders
+{
+	// Get the name of the table that folder rows are stored in
+	UE_EXPERIMENTAL(5.8, "The TEDS+Folders integration is experimental and the API subject to change.")
+	UE_API FName GetActorFolderTableName();
+
+	// Register the given folder in TEDS
+	UE_EXPERIMENTAL(5.8, "The TEDS+Folders integration is experimental and the API subject to change.")
+	UE_API RowHandle RegisterFolderInTeds(ICoreProvider* Storage, const FFolder& Folder, UWorld* World, bool bRecursivelyAddParents = false);
+
+	// Register the given folder in TEDS from a query callback
+	UE_EXPERIMENTAL(5.8, "The TEDS+Folders integration is experimental and the API subject to change.")
+	UE_API void RegisterFolderInTeds(IQueryContext& Context, ICoreProvider* Storage, const FFolder& Folder, TWeakObjectPtr<UWorld> World, bool bRecursivelyAddParents = false);
+
+	// Set the default columns (+ data) for a folder row
+	UE_EXPERIMENTAL(5.8, "The TEDS+Folders integration is experimental and the API subject to change.")
+	UE_API void SetFolderColumns(ICoreProvider* Storage, RowHandle Row, UWorld* World, const FFolder& Folder);
+
+	// Unregister a folder from TEDS by looking up the folder's row
+	UE_EXPERIMENTAL(5.8, "The TEDS+Folders integration is experimental and the API subject to change.")
+	UE_API void UnregisterFolderFromTeds(ICoreProvider* Storage, const FFolder& Folder);
+	
+	// Unregister a folder from TEDS from a query callback
+	UE_EXPERIMENTAL(5.8, "The TEDS+Folders integration is experimental and the API subject to change.")
+	UE_API void UnregisterFolderFromTeds(IQueryContext& Context, ICoreProvider* Storage, const FFolder& Folder);
+
+	// Look up the row handle for a folder using the direct API
+	UE_EXPERIMENTAL(5.8, "The TEDS+Folders integration is experimental and the API subject to change.")
+	UE_API RowHandle LookupMappedRow(const ICoreProvider* Storage, const FFolder& Folder);
+
+	// Look up the row handle for a folder from inside a query callback
+	UE_EXPERIMENTAL(5.8, "The TEDS+Folders integration is experimental and the API subject to change.")
+	UE_API RowHandle LookupMappedRow(const IQueryContext& Context, const FFolder& Folder);
+
+	// Remap a folder if the index has changed due to data changes
+	UE_EXPERIMENTAL(5.8, "The TEDS+Folders integration is experimental and the API subject to change.")
+	UE_API RowHandle RemapFolder(ICoreProvider* Storage, const FFolder& InOldFolder, const FFolder& InNewFolder, UWorld* World);
+
+	// For a given ChildRow, add a parent folder that has not been registered in TEDS yet.
+	UE_EXPERIMENTAL(5.8, "The TEDS+Folders integration is experimental and the API subject to change.")
+	UE_API void SetUnresolvedFolderParent(IQueryContext& Context, RowHandle ChildRow, const FFolder& ParentFolder);
+	
+	// Get the mapping domain for a given folder
+	UE_EXPERIMENTAL(5.8, "The TEDS+Folders integration is experimental and the API subject to change.")
+	UE_API FName GetFolderMappingDomain(const FFolder& Folder);
+	
+	// Get the mapping key for a given folder
+	UE_EXPERIMENTAL(5.8, "The TEDS+Folders integration is experimental and the API subject to change.")
+	UE_API FMapKey GetFolderIndex(const FFolder& Folder);
+}
+
+#undef UE_API
\ No newline at end of file
diff --git a/Engine/Source/Editor/UnrealEd/Public/EditorActorFolders.h b/Engine/Source/Editor/UnrealEd/Public/EditorActorFolders.h
--- a/Engine/Source/Editor/UnrealEd/Public/EditorActorFolders.h
+++ b/Engine/Source/Editor/UnrealEd/Public/EditorActorFolders.h
@@ -30,7 +30,6 @@ DECLARE_MULTICAST_DELEGATE_TwoParams(FOnActorFolderCreated, UWorld&, const FFold
 DECLARE_MULTICAST_DELEGATE_TwoParams(FOnActorFolderDeleted, UWorld&, const FFolder&);
 DECLARE_MULTICAST_DELEGATE_ThreeParams(FOnActorFolderMoved, UWorld&, const FFolder& /* src */, const FFolder& /* dst */);
 
-
 UCLASS()
 class UActorEditorContextActorFolderState : public UActorEditorContextClientState
 {
@@ -83,6 +82,9 @@ struct FActorFolders : public FGCObject, public IActorEditorContextClient
 
 	/** Tests whether a folder container exists for the specified world */
 	UNREALED_API bool IsInitializedForWorld(UWorld& InWorld) const;
+	
+	/** Initialize a folder container for the specified world */
+	UNREALED_API void InitializeForWorld(UWorld& InWorld);
 
 	/** Get a default folder name under the specified parent path */
 	UNREALED_API FFolder GetDefaultFolderName(UWorld& InWorld, const FFolder& InParentFolder);
diff --git a/Engine/Source/Editor/UnrealEd/Public/EditorFolderUtils.h b/Engine/Source/Editor/UnrealEd/Public/EditorFolderUtils.h
--- a/Engine/Source/Editor/UnrealEd/Public/EditorFolderUtils.h
+++ b/Engine/Source/Editor/UnrealEd/Public/EditorFolderUtils.h
@@ -4,6 +4,7 @@
 
 #include "Containers/UnrealString.h"
 #include "CoreMinimal.h"
+#include "Folder.h"
 #include "HAL/Platform.h"
 #include "Misc/Paths.h"
 #include "UObject/NameTypes.h"
@@ -23,4 +24,5 @@ public:
 	/** Check if the specified path is a child of the specified parent */
 	static UNREALED_API bool PathIsChildOf(const FString& PotentialChild, const FString& Parent);
 	static UNREALED_API bool PathIsChildOf(const FName& PotentialChild, const FName& Parent);
+	static UNREALED_API void RenameFolder(const FFolder& InFolder, const FText& NewFolderName, UWorld* World);
 };
diff --git a/Engine/Source/Editor/UnrealEd/Public/WorldFolders.h b/Engine/Source/Editor/UnrealEd/Public/WorldFolders.h
--- a/Engine/Source/Editor/UnrealEd/Public/WorldFolders.h
+++ b/Engine/Source/Editor/UnrealEd/Public/WorldFolders.h
@@ -12,6 +12,11 @@
 #include "WorldTransientFolders.h"
 #include "WorldFolders.generated.h"
 
+namespace UE::Editor::DataStorage
+{
+	class ICoreProvider;
+}
+
 DECLARE_LOG_CATEGORY_EXTERN(LogWorldFolders, Log, All)
 
 USTRUCT()
@@ -80,9 +85,14 @@ public:
 	UE_INTERNAL UNREALED_API void ForEachFolderWithRootObject(const FFolder::FRootObject& InFolderRootObject, TFunctionRef<bool(const FFolder&)> Operation);
 	UE_INTERNAL UNREALED_API void SaveState();
 	UE_INTERNAL UNREALED_API UWorld* GetWorld() const;
+	
+	// Internal function called when this object is removed from the FActorFolders bookkeeping (when the owning world is destroyed)
+	UE_INTERNAL UNREALED_API void OnRemoved();
 
 	//~ Begin UObject
 	UE_INTERNAL UNREALED_API virtual void Serialize(FArchive& Ar) override;
+	UE_INTERNAL UNREALED_API virtual void PreEditUndo() override;
+	UE_INTERNAL UNREALED_API virtual void PostEditUndo() override;
 	//~ End UObject
 
 	//~ Begin Deprecated
@@ -102,6 +112,13 @@ private:
 	UNREALED_API void LoadState();
 
 	void RebuildListInternal();
+	
+	// Internal function for adding a folder to ensure all code paths go through a common place
+	void AddFolder_Internal(const FFolder& InFolder, const FActorFolderProps& FolderProperties = FActorFolderProps());
+	// Internal function for removing a folder to ensure all code paths go through a common place
+	void RemoveFolder_Internal(const FFolder& InFolder);
+	// Internal function for re-adding a folder (remove and add in one shot)  to ensure all code paths go through a common place
+	void ReAddFolder_Internal(const FFolder& InOldFolder, const FFolder& InNewFolder, const FActorFolderProps& InNewFolderProperties = FActorFolderProps());
 
 	TUniquePtr<FWorldPersistentFolders> PersistentFolders;
 	TUniquePtr<FWorldTransientFolders> TransientFolders;
@@ -109,6 +126,9 @@ private:
 	TWeakObjectPtr<UWorld> World;
 	TMap<FFolder, FActorFolderProps> FoldersProperties;
 	TMap<FFolder, FActorFolderProps> LoadedStateFoldersProperties;
+	
+	// Cached list of folders before an undo/redo operation
+	TSet<FFolder> PreUndoFolders;
 
 	UPROPERTY()
 	FActorPlacementFolder CurrentFolder;
@@ -116,6 +136,9 @@ private:
 	TArray<FActorPlacementFolder> CurrentFolderStack;
 
 	bool bQueuedRebuild;
+	
+	// Cache the editor data storage ptr to avoid having to grab it for every operation
+	UE::Editor::DataStorage::ICoreProvider* DataStorage = nullptr;
 
 	friend class FWorldFoldersImplementation;
 	friend class FWorldPersistentFolders;
